Removes "Reload Minimap" button and minimap generation from convert commandlet.
Removes "MinimapSize" from properties as only tile resolution is used to generate the texture.
MiniMap builder should be only used as an independant commandlet.

#changelist validated

diff --git a/Engine/Source/Editor/UnrealEd/Private/Commandlets/WorldPartitionConvertCommandlet.cpp b/Engine/Source/Editor/UnrealEd/Private/Commandlets/WorldPartitionConvertCommandlet.cpp
--- a/Engine/Source/Editor/UnrealEd/Private/Commandlets/WorldPartitionConvertCommandlet.cpp
+++ b/Engine/Source/Editor/UnrealEd/Private/Commandlets/WorldPartitionConvertCommandlet.cpp
@@ -664,22 +664,6 @@ void UWorldPartitionConvertCommandlet::OnWorldLoaded(UWorld* World)
 	}
 }
 
-void UWorldPartitionConvertCommandlet::CreateWorldMiniMapTexture(UWorld* World)
-{
-	TRACE_CPUPROFILER_EVENT_SCOPE(UWorldPartitionConvertCommandlet::CreateWorldMiniMapTexture);
-
-	UE_SCOPED_TIMER(TEXT("CreateWorldMiniMapTexture"), LogWorldPartitionConvertCommandlet, Display);
-
-	AWorldPartitionMiniMap* WorldMiniMap = FWorldPartitionMiniMapHelper::GetWorldPartitionMiniMap(World, /*bCreateNewMiniMap*/true);
-	if (!WorldMiniMap)
-	{
-		UE_LOG(LogWorldPartitionConvertCommandlet, Error, TEXT("Failed to create Minimap. WorldPartitionMiniMap actor not found in the persistent level."));
-		return;
-	}
-
-	FWorldPartitionMiniMapHelper::CaptureWorldMiniMapToTexture(World, WorldMiniMap, WorldMiniMap->MiniMapSize, static_cast<UTexture2D*&>(WorldMiniMap->MiniMapTexture), TEXT("MinimapTexture"), WorldMiniMap->MiniMapWorldBounds);
-}
-
 int32 UWorldPartitionConvertCommandlet::Main(const FString& Params)
 {
 	TRACE_CPUPROFILER_EVENT_SCOPE(UWorldPartitionConvertCommandlet::Main);
@@ -1356,12 +1340,6 @@ int32 UWorldPartitionConvertCommandlet::Main(const FString& Params)
 	bool bForceInitializeWorld = false;
 	bool bInitializedPhysicsSceneForSave = GEditor->InitializePhysicsSceneForSaveIfNecessary(MainWorld, bForceInitializeWorld);
 
-	// Create the world's minimap, do not create the minimap if only merging
-	if (!bOnlyMergeSubLevels && !Switches.Contains(TEXT("SkipMiniMapGeneration")))
-	{
-		CreateWorldMiniMapTexture(MainWorld);
-	}
-
 	// After conversion, convert actors to external actors
 	UPackage* LevelPackage = MainLevel->GetPackage();
 
diff --git a/Engine/Source/Editor/UnrealEd/Public/WorldPartition/WorldPartitionMiniMapBuilder.h b/Engine/Source/Editor/UnrealEd/Public/WorldPartition/WorldPartitionMiniMapBuilder.h
--- a/Engine/Source/Editor/UnrealEd/Public/WorldPartition/WorldPartitionMiniMapBuilder.h
+++ b/Engine/Source/Editor/UnrealEd/Public/WorldPartition/WorldPartitionMiniMapBuilder.h
@@ -32,6 +32,5 @@ private:
 	UPROPERTY(Transient, VisibleAnywhere, Category = WorldPartitionMiniMap)
 	TObjectPtr<AWorldPartitionMiniMap> WorldMiniMap = nullptr;
 	FBox EditorBounds;
-	int32 IterativeCellSize;
 	bool bAutoSubmit;
 };
\ No newline at end of file
diff --git a/Engine/Source/Editor/WorldPartitionEditor/Private/WorldPartition/SWorldPartitionEditorGrid2D.cpp b/Engine/Source/Editor/WorldPartitionEditor/Private/WorldPartition/SWorldPartitionEditorGrid2D.cpp
--- a/Engine/Source/Editor/WorldPartitionEditor/Private/WorldPartition/SWorldPartitionEditorGrid2D.cpp
+++ b/Engine/Source/Editor/WorldPartitionEditor/Private/WorldPartition/SWorldPartitionEditorGrid2D.cpp
@@ -39,8 +39,6 @@ public:
 	}
 };
 
-static TAutoConsoleVariable<int32> CVarShowReloadMiniMapButton(TEXT("wp.MiniMap.ShowReloadButton"), 0, TEXT("Show reload MiniMap button."));
-
 SWorldPartitionEditorGrid2D::FEditorCommands::FEditorCommands()
 	: TCommands<FEditorCommands>
 (
@@ -158,14 +156,6 @@ void SWorldPartitionEditorGrid2D::Construct(const FArguments& InArgs)
 						.IsEnabled(true)
 						.Text(LOCTEXT("ShowCellCoords", "Show Cell Coords"))
 					]
-					+SHorizontalBox::Slot()
-					.AutoWidth()
-					[
-						SNew(SButton)
-						.Text(LOCTEXT("ReloadMiniMap", "Reload MiniMap"))
-						.Visibility_Lambda([]() {return CVarShowReloadMiniMapButton.GetValueOnAnyThread() != 0 ? EVisibility::Visible : EVisibility::Hidden;})
-						.OnClicked(this, &SWorldPartitionEditorGrid2D::ReloadMiniMap)
-					]
 					+ SHorizontalBox::Slot()
 					.AutoWidth()
 					[
diff --git a/Engine/Source/Editor/WorldPartitionEditor/Private/WorldPartition/SWorldPartitionEditorGridSpatialHash.cpp b/Engine/Source/Editor/WorldPartitionEditor/Private/WorldPartition/SWorldPartitionEditorGridSpatialHash.cpp
--- a/Engine/Source/Editor/WorldPartitionEditor/Private/WorldPartition/SWorldPartitionEditorGridSpatialHash.cpp
+++ b/Engine/Source/Editor/WorldPartitionEditor/Private/WorldPartition/SWorldPartitionEditorGridSpatialHash.cpp
@@ -72,36 +72,6 @@ void SWorldPartitionEditorGridSpatialHash::Tick(const FGeometry& AllottedGeometr
 	}
 }
 
-FReply SWorldPartitionEditorGridSpatialHash::ReloadMiniMap()
-{
-	UE_LOG(LogTemp, Log, TEXT("Reload MiniMap has been clicked"));
-
-	//Create a new MiniMap if there isn't one.
-	AWorldPartitionMiniMap* WorldMiniMap = FWorldPartitionMiniMapHelper::GetWorldPartitionMiniMap(World, true);
-	if (!WorldMiniMap)
-	{
-		UE_LOG(LogTemp, Error, TEXT("Failed to create Minimap. WorldPartitionMiniMap actor not found in the persistent level."));
-		return FReply::Handled();
-	}
-
-	WorldMiniMap->Modify();
-
-
-	// Updating VT is not supported for now
-	if (WorldMiniMap->MiniMapTexture != nullptr)
-	{
-		WorldMiniMap->MiniMapTexture->VirtualTextureStreaming = false;
-	}
-
-	WorldMiniMap->UVOffset.bIsValid = false;
-
-	FWorldPartitionMiniMapHelper::CaptureWorldMiniMapToTexture(World, WorldMiniMap, WorldMiniMap->MiniMapSize, static_cast<UTexture2D*&>(WorldMiniMap->MiniMapTexture), TEXT("MinimapTexture"), WorldMiniMap->MiniMapWorldBounds);
-
-	UpdateWorldMiniMapDetails();
-
-	return FReply::Handled();
-}
-
 void SWorldPartitionEditorGridSpatialHash::UpdateWorldMiniMapDetails()
 {
 	auto WorldMiniMap = FWorldPartitionMiniMapHelper::GetWorldPartitionMiniMap(World);
diff --git a/Engine/Source/Editor/WorldPartitionEditor/Private/WorldPartitionEditorModule.cpp b/Engine/Source/Editor/WorldPartitionEditor/Private/WorldPartitionEditorModule.cpp
--- a/Engine/Source/Editor/WorldPartitionEditor/Private/WorldPartitionEditorModule.cpp
+++ b/Engine/Source/Editor/WorldPartitionEditor/Private/WorldPartitionEditorModule.cpp
@@ -396,11 +396,6 @@ FString UWorldPartitionConvertOptions::ToCommandletArgs() const
 		CommandletArgsBuilder.Append(TEXT(" -SkipStableGUIDValidation"));
 	}
 
-	if (bSkipMiniMapGeneration)
-	{
-		CommandletArgsBuilder.Append(TEXT(" -SkipMiniMapGeneration"));
-	}
-
 	if (bDeleteSourceLevels)
 	{
 		CommandletArgsBuilder.Append(TEXT(" -DeleteSourceLevels"));
diff --git a/Engine/Source/Editor/WorldPartitionEditor/Public/WorldPartition/SWorldPartitionEditorGrid2D.h b/Engine/Source/Editor/WorldPartitionEditor/Public/WorldPartition/SWorldPartitionEditorGrid2D.h
--- a/Engine/Source/Editor/WorldPartitionEditor/Public/WorldPartition/SWorldPartitionEditorGrid2D.h
+++ b/Engine/Source/Editor/WorldPartitionEditor/Public/WorldPartition/SWorldPartitionEditorGrid2D.h
@@ -55,7 +55,6 @@ public:
 	virtual uint32 PaintSelection(const FGeometry& AllottedGeometry, const FSlateRect& MyCullingRect, FSlateWindowElementList& OutDrawElements, uint32 LayerId) const;
 	virtual int32 PaintSoftwareCursor(const FGeometry& AllottedGeometry, const FSlateRect& MyCullingRect, FSlateWindowElementList& OutDrawElements, int32 LayerId) const;
 
-	virtual FReply ReloadMiniMap() { return FReply::Handled(); }
 	virtual FReply FocusSelection();
 
 protected:
diff --git a/Engine/Source/Editor/WorldPartitionEditor/Public/WorldPartition/SWorldPartitionEditorGridSpatialHash.h b/Engine/Source/Editor/WorldPartitionEditor/Public/WorldPartition/SWorldPartitionEditorGridSpatialHash.h
--- a/Engine/Source/Editor/WorldPartitionEditor/Public/WorldPartition/SWorldPartitionEditorGridSpatialHash.h
+++ b/Engine/Source/Editor/WorldPartitionEditor/Public/WorldPartition/SWorldPartitionEditorGridSpatialHash.h
@@ -20,8 +20,6 @@ public:
 
 	int32 PaintGrid(const FGeometry& AllottedGeometry, const FSlateRect& MyCullingRect, FSlateWindowElementList& OutDrawElements, int32 LayerId) const override;
 
-	virtual FReply ReloadMiniMap() override;
-
 private:
 	void UpdateWorldMiniMapDetails();
 	FBox2D          WorldMiniMapBounds;
diff --git a/Engine/Source/Editor/WorldPartitionEditor/Public/WorldPartition/WorldPartitionConvertOptions.h b/Engine/Source/Editor/WorldPartitionEditor/Public/WorldPartition/WorldPartitionConvertOptions.h
--- a/Engine/Source/Editor/WorldPartitionEditor/Public/WorldPartition/WorldPartitionConvertOptions.h
+++ b/Engine/Source/Editor/WorldPartitionEditor/Public/WorldPartition/WorldPartitionConvertOptions.h
@@ -35,9 +35,6 @@ public:
 	UPROPERTY(EditAnywhere, AdvancedDisplay, Category = Convert)
 	bool bSkipStableGUIDValidation;
 
-	UPROPERTY(EditAnywhere, AdvancedDisplay, Category = Convert)
-	bool bSkipMiniMapGeneration;
-
 	UPROPERTY(EditAnywhere, AdvancedDisplay, Category = Convert)
 	bool bOnlyMergeSubLevels;
 
diff --git a/Engine/Source/Runtime/Engine/Private/WorldPartition/WorldPartitionMiniMap.cpp b/Engine/Source/Runtime/Engine/Private/WorldPartition/WorldPartitionMiniMap.cpp
--- a/Engine/Source/Runtime/Engine/Private/WorldPartition/WorldPartitionMiniMap.cpp
+++ b/Engine/Source/Runtime/Engine/Private/WorldPartition/WorldPartitionMiniMap.cpp
@@ -14,7 +14,6 @@ AWorldPartitionMiniMap::AWorldPartitionMiniMap(const FObjectInitializer& ObjectI
 #if WITH_EDITOR
 void AWorldPartitionMiniMap::PostEditChangeProperty(FPropertyChangedEvent& PropertyChangedEvent)
 {
-	MiniMapSize = FMath::Clamp<uint32>(FMath::RoundUpToPowerOfTwo(MiniMapSize), 256, 8192);
 	MiniMapTileSize = FMath::Clamp<uint32>(FMath::RoundUpToPowerOfTwo(MiniMapTileSize), 256, 8192);
 }
 #endif
diff --git a/Engine/Source/Runtime/Engine/Public/WorldPartition/WorldPartitionMiniMap.h b/Engine/Source/Runtime/Engine/Public/WorldPartition/WorldPartitionMiniMap.h
--- a/Engine/Source/Runtime/Engine/Public/WorldPartition/WorldPartitionMiniMap.h
+++ b/Engine/Source/Runtime/Engine/Public/WorldPartition/WorldPartitionMiniMap.h
@@ -64,10 +64,6 @@ public:
 	UPROPERTY(VisibleAnywhere, Category = WorldPartitionMiniMap)
 	TObjectPtr<UTexture2D> MiniMapTexture;
 
-	/*MiniMap Size*/
-	UPROPERTY(EditAnywhere, Category = WorldPartitionMiniMap, meta=(UIMin = "256", UIMax = "8192"))
-	int32 MiniMapSize = 1024;
-
 	/*MiniMap Tile Size*/
 	UPROPERTY(EditAnywhere, Category = WorldPartitionMiniMap, meta = (UIMin = "256", UIMax = "8192"))
 	int32 MiniMapTileSize = 1024;
