World Partition: fixed WorldPartitionResaveActorsBuilder that would clear actor references (wrong redirector outering) by using newly introduced GenerateObjectsClusters to generate and load clustered actors to preserve references.

diff --git a/Engine/Source/Editor/UnrealEd/Private/WorldPartition/WorldPartitionResaveActorsBuilder.cpp b/Engine/Source/Editor/UnrealEd/Private/WorldPartition/WorldPartitionResaveActorsBuilder.cpp
--- a/Engine/Source/Editor/UnrealEd/Private/WorldPartition/WorldPartitionResaveActorsBuilder.cpp
+++ b/Engine/Source/Editor/UnrealEd/Private/WorldPartition/WorldPartitionResaveActorsBuilder.cpp
@@ -1,7 +1,6 @@
 // Copyright Epic Games, Inc. All Rights Reserved.
 
 #include "WorldPartition/WorldPartitionResaveActorsBuilder.h"
-
 #include "WorldPartition/WorldPartition.h"
 #include "WorldPartition/WorldPartitionHelpers.h"
 #include "PackageSourceControlHelper.h"
@@ -12,9 +11,9 @@
 #include "UObject/MetaData.h"
 #include "UObject/SavePackage.h"
 #include "Kismet2/BlueprintEditorUtils.h"
+#include "ReferenceCluster.h"
 #include "ActorFolder.h"
 
-
 DEFINE_LOG_CATEGORY_STATIC(LogWorldPartitionResaveActorsBuilder, All, All);
 
 UWorldPartitionResaveActorsBuilder::UWorldPartitionResaveActorsBuilder(const FObjectInitializer& ObjectInitializer)
@@ -210,155 +209,161 @@ bool UWorldPartitionResaveActorsBuilder::RunInternal(UWorld* World, const FCellI
 		}
 	};
 
-	// Rename pass
 	if (bSwitchActorPackagingSchemeToReduced)
 	{
 		World->PersistentLevel->ActorPackagingScheme = EActorPackagingScheme::Reduced;
 
-		FWorldPartitionHelpers::ForEachActorWithLoading(WorldPartition, ActorClass, [this, &LoadCount, &SaveCount, &FailCount, &SCCHelper, &PackagesToDelete, WorldPackage, &AddAndSavePackage, &CheckOutAndSavePackage](const FWorldPartitionActorDesc* ActorDesc)
+		// Build clusters
+		TArray<TPair<FGuid, TArray<FGuid>>> ActorsWithRefs;
+		for (FActorDescList::TIterator<> ActorDescIterator(WorldPartition); ActorDescIterator; ++ActorDescIterator)
 		{
-			ON_SCOPE_EXIT
-			{
-				UE_LOG(LogWorldPartitionResaveActorsBuilder, Display, TEXT("Processed %d packages (%d Saved / %d Failed)"), LoadCount, SaveCount, FailCount);
-			};
+			ActorsWithRefs.Emplace(ActorDescIterator->GetGuid(), ActorDescIterator->GetReferences());
+		}
 
-			AActor* Actor = ActorDesc->GetActor();
+		TArray<TArray<FGuid>> Clusters = GenerateObjectsClusters(ActorsWithRefs);
 
-			if (!Actor)
+		for (const TArray<FGuid>& Cluster : Clusters)
+		{
+			TArray<UPackage*> PackagesToSave;
+
+			// Load actor clusters
+			TArray<FWorldPartitionReference> ActorReferences;
+			for (const FGuid& ActorGuid : Cluster)
 			{
-				PackagesToDelete.Add(ActorDesc->GetActorPackage().ToString());
-				++FailCount;
-				return true;
+				ActorReferences.Emplace(WorldPartition, ActorGuid);
 			}
-			else
-			{
-				++LoadCount;
 
-				UPackage* Package = Actor->GetExternalPackage();
-				check(Package);
+			// Change packaging of all actors in the current cluster
+			for (FWorldPartitionReference& ActorReference : ActorReferences)
+			{
+				const FWorldPartitionActorDesc* ActorDesc = ActorReference.Get();
+				AActor* Actor = ActorDesc->GetActor();
 
-				if (!bReportOnly)
+				if (!Actor)
 				{
-					if (SCCHelper.Checkout(Package))
-					{
-						// Always mark this package for deletion, as it will contain a temporary redirector to fixup references
-						PackagesToDelete.Add(ActorDesc->GetActorPackage().ToString());
+					PackagesToDelete.Add(ActorDesc->GetActorPackage().ToString());
+					++FailCount;
+					continue;
+				}
+				else
+				{
+					++LoadCount;
 
-						// Move actor back into world's package
-						Actor->SetPackageExternal(false);
+					UPackage* Package = Actor->GetExternalPackage();
+					check(Package);
 
-						// Gather dependant objects that also needs to be moved along with the actor
-						TArray<UObject*> DependantObjects;
-						ForEachObjectWithPackage(Package, [&DependantObjects](UObject* Object)
+					if (!bReportOnly)
+					{
+						if (SCCHelper.Checkout(Package))
 						{
-							if (!Cast<UMetaData>(Object))
+							// Always mark this package for deletion, as it will contain a temporary redirector to fixup references
+							PackagesToDelete.Add(ActorDesc->GetActorPackage().ToString());
+
+							// Move actor back into world's package
+							Actor->SetPackageExternal(false);
+
+							// Gather dependant objects that also needs to be moved along with the actor
+							TArray<UObject*> DependantObjects;
+							ForEachObjectWithPackage(Package, [&DependantObjects](UObject* Object)
+							{
+								if (!Cast<UMetaData>(Object))
+								{
+									DependantObjects.Add(Object);
+								}
+								return true;
+							}, false);
+
+							// Move dependant objects into the new world package temporarily
+							for (UObject* DependantObject : DependantObjects)
 							{
-								DependantObjects.Add(Object);
+								DependantObject->Rename(nullptr, WorldPackage, REN_NonTransactional | REN_DontCreateRedirectors | REN_ForceNoResetLoaders);
 							}
-							return true;
-						}, false);
 
-						// Move dependant objects into the new world package temporarily, so we can save the original actor package with its redirector
-						for (UObject* DependantObject : DependantObjects)
-						{
-							DependantObject->Rename(nullptr, WorldPackage, REN_NonTransactional | REN_DontCreateRedirectors | REN_ForceNoResetLoaders);
-						}
+							// Move actor in its new package
+							Actor->SetPackageExternal(true);
 
-						// Replace old actor package with a redirector to the new actor
-						UObjectRedirector* Redirector = NewObject<UObjectRedirector>(Actor->GetWorld(), Actor->GetFName(), RF_Standalone | RF_Public);
-						Redirector->DestinationObject = Actor;
-						Redirector->SetExternalPackage(Package);
+							// Also move dependant objects into the new package
+							UPackage* NewActorPackage = Actor->GetExternalPackage();
+							for (UObject* DependantObject : DependantObjects)
+							{
+								DependantObject->Rename(nullptr, NewActorPackage, REN_NonTransactional | REN_DontCreateRedirectors | REN_ForceNoResetLoaders);
+							}
 
-						// Save original actor package containing a single redirector to the actor
-						FString PackageFileName = SourceControlHelpers::PackageFilename(Package);
-						FSavePackageArgs SaveArgs;
-						SaveArgs.TopLevelFlags = RF_Standalone;
-						if (!UPackage::SavePackage(Package, nullptr, *PackageFileName, SaveArgs))
-						{
-							UE_LOG(LogWorldPartitionResaveActorsBuilder, Error, TEXT("Error saving package redirector %s."), *Package->GetName());
-							++FailCount;
-							return true;
+							PackagesToSave.Add(NewActorPackage);
 						}
-
-						// Move actor in its new package
-						Actor->SetPackageExternal(true);
-
-						// Also move dependant objects into the new package
-						UPackage* NewActorPackage = Actor->GetExternalPackage();
-						for (UObject* DependantObject : DependantObjects)
+						else
 						{
-							DependantObject->Rename(nullptr, NewActorPackage, REN_NonTransactional | REN_DontCreateRedirectors | REN_ForceNoResetLoaders);
+							// It is possible the resave can't checkout everything. Continue processing.
+							UE_LOG(LogWorldPartitionResaveActorsBuilder, Warning, TEXT("Error checking out package %s."), *Package->GetName());
+							++FailCount;
+							continue;
 						}
-
-						return AddAndSavePackage(NewActorPackage, TEXT("Actor"));
-					}
-					else
-					{
-						// It is possible the resave can't checkout everything. Continue processing.
-						UE_LOG(LogWorldPartitionResaveActorsBuilder, Warning, TEXT("Error checking out package %s."), *Package->GetName());
-						++FailCount;
-						return true;
 					}
 				}
 			}
 
-			return true;
-		});
+			// Save modified actors
+			for (UPackage* PackageToSave : PackagesToSave)
+			{
+				AddAndSavePackage(PackageToSave, TEXT("Actor"));
+			}
+		}
 	}
-
-	// Normal save pass
-	FWorldPartitionHelpers::ForEachActorWithLoading(WorldPartition, ActorClass, [this, &LoadCount, &SaveCount, &FailCount, &SCCHelper, &PackagesToDelete, &CheckOutAndSavePackage](const FWorldPartitionActorDesc* ActorDesc)
+	else
 	{
-		ON_SCOPE_EXIT
+		FWorldPartitionHelpers::ForEachActorWithLoading(WorldPartition, ActorClass, [this, &LoadCount, &SaveCount, &FailCount, &SCCHelper, &PackagesToDelete, &CheckOutAndSavePackage](const FWorldPartitionActorDesc* ActorDesc)
 		{
-			UE_LOG(LogWorldPartitionResaveActorsBuilder, Display, TEXT("Processed %d packages (%d Saved / %d Failed)"), LoadCount, SaveCount, FailCount);
-		};
-
-		AActor* Actor = ActorDesc->GetActor();
+			ON_SCOPE_EXIT
+			{
+				UE_LOG(LogWorldPartitionResaveActorsBuilder, Display, TEXT("Processed %d packages (%d Saved / %d Failed)"), LoadCount, SaveCount, FailCount);
+			};
 
-		if (!Actor)
-		{
-			PackagesToDelete.Add(ActorDesc->GetActorPackage().ToString());
-			++FailCount;
-			return true;
-		}
-		else
-		{
-			++LoadCount;
+			AActor* Actor = ActorDesc->GetActor();
 
-			if (bEnableActorFolders)
+			if (!Actor)
 			{
-				if (!Actor->CreateOrUpdateActorFolder())
-				{
-					UE_LOG(LogWorldPartitionResaveActorsBuilder, Error, TEXT("Failed to create actor folder for actor %s."), *Actor->GetName());
-					++FailCount;
-					return true;
-				}
+				PackagesToDelete.Add(ActorDesc->GetActorPackage().ToString());
+				++FailCount;
+				return true;
 			}
-
-			UPackage* Package = Actor->GetExternalPackage();
-			check(Package);
-
-			if (bResaveDirtyActorDescsOnly)
+			else
 			{
-				TUniquePtr<FWorldPartitionActorDesc> NewActorDesc = Actor->CreateActorDesc();
+				++LoadCount;
 
-				if (ActorDesc->Equals(NewActorDesc.Get()))
+				if (bEnableActorFolders)
 				{
-					return true;
+					if (!Actor->CreateOrUpdateActorFolder())
+					{
+						UE_LOG(LogWorldPartitionResaveActorsBuilder, Error, TEXT("Failed to create actor folder for actor %s."), *Actor->GetName());
+						++FailCount;
+						return true;
+					}
 				}
 
-				UE_LOG(LogWorldPartitionResaveActorsBuilder, Log, TEXT("Package %s needs to be resaved."), *Package->GetName());
-			}
+				UPackage* Package = Actor->GetExternalPackage();
+				check(Package);
 
-			if (!bReportOnly)
-			{
-				return CheckOutAndSavePackage(Package, TEXT("Actor"), /*bErrorAtFailCheckout*/ false);
+				if (bResaveDirtyActorDescsOnly)
+				{
+					TUniquePtr<FWorldPartitionActorDesc> NewActorDesc = Actor->CreateActorDesc();
+
+					if (ActorDesc->Equals(NewActorDesc.Get()))
+					{
+						return true;
+					}
+
+					UE_LOG(LogWorldPartitionResaveActorsBuilder, Log, TEXT("Package %s needs to be resaved."), *Package->GetName());
+				}
+
+				if (!bReportOnly)
+				{
+					return CheckOutAndSavePackage(Package, TEXT("Actor"), /*bErrorAtFailCheckout*/ false);
+				}
 			}
-		}
 
-		return true;
-	});
+			return true;
+		});
+	}
 
 	if (!bReportOnly)
 	{
