FScopedEditorWorld - Check if package is null to prevent an assert in UWorld::FindWorldInPackage()

diff --git a/Engine/Source/Editor/UnrealEd/Private/EditorWorldUtils.cpp b/Engine/Source/Editor/UnrealEd/Private/EditorWorldUtils.cpp
--- a/Engine/Source/Editor/UnrealEd/Private/EditorWorldUtils.cpp
+++ b/Engine/Source/Editor/UnrealEd/Private/EditorWorldUtils.cpp
@@ -18,12 +18,12 @@ FScopedEditorWorld::FScopedEditorWorld(UWorld* InWorld, const UWorld::Initializa
 FScopedEditorWorld::FScopedEditorWorld(const FStringView InLongPackageName, const UWorld::InitializationValues& InInitializationValues)
 	: FScopedEditorWorld()
 {
-	UPackage* WorldPackage = LoadWorldPackageForEditor(InLongPackageName);
-	
-	UWorld* RuntimeWorld = UWorld::FindWorldInPackage(WorldPackage);
-	if (RuntimeWorld)
+	if (UPackage* WorldPackage = LoadWorldPackageForEditor(InLongPackageName))
 	{
-		Init(RuntimeWorld, InInitializationValues);
+		if (UWorld* RuntimeWorld = UWorld::FindWorldInPackage(WorldPackage))
+		{
+			Init(RuntimeWorld, InInitializationValues);
+		}
 	}
 }
 
