Fix no-unity builds

diff --git a/Engine/Source/Editor/UnrealEd/Classes/Commandlets/WorldPartitionDataLayerToAssetCommandLet.h b/Engine/Source/Editor/UnrealEd/Classes/Commandlets/WorldPartitionDataLayerToAssetCommandLet.h
--- a/Engine/Source/Editor/UnrealEd/Classes/Commandlets/WorldPartitionDataLayerToAssetCommandLet.h
+++ b/Engine/Source/Editor/UnrealEd/Classes/Commandlets/WorldPartitionDataLayerToAssetCommandLet.h
@@ -20,6 +20,7 @@ class UDataLayer;
 class UDataLayerAsset;
 class UDataLayerInstance;
 class UDataLayerFactory;
+class UWorld;
 
 UCLASS(Transient)
 class UDataLayerConversionInfo : public UObject
diff --git a/Engine/Source/Editor/UnrealEd/Private/Commandlets/WorldPartitionCommandletHelpers.cpp b/Engine/Source/Editor/UnrealEd/Private/Commandlets/WorldPartitionCommandletHelpers.cpp
--- a/Engine/Source/Editor/UnrealEd/Private/Commandlets/WorldPartitionCommandletHelpers.cpp
+++ b/Engine/Source/Editor/UnrealEd/Private/Commandlets/WorldPartitionCommandletHelpers.cpp
@@ -2,6 +2,7 @@
 
 #include "Commandlets/WorldPartitionCommandletHelpers.h"
 
+#include "Engine/World.h"
 #include "PackageSourceControlHelper.h"
 #include "SourceControlHelpers.h"
 #include "UObject/SavePackage.h"
diff --git a/Engine/Source/Editor/UnrealEd/Private/Commandlets/WorldPartitionDataLayerToAssetCommandlet.cpp b/Engine/Source/Editor/UnrealEd/Private/Commandlets/WorldPartitionDataLayerToAssetCommandlet.cpp
--- a/Engine/Source/Editor/UnrealEd/Private/Commandlets/WorldPartitionDataLayerToAssetCommandlet.cpp
+++ b/Engine/Source/Editor/UnrealEd/Private/Commandlets/WorldPartitionDataLayerToAssetCommandlet.cpp
@@ -8,8 +8,11 @@
 
 #include "Algo/Copy.h"
 #include "Algo/Accumulate.h"
+#include "AssetRegistry/AssetRegistryModule.h"
+#include "AssetRegistry/IAssetRegistry.h"
 #include "AssetToolsModule.h"
 #include "DataLayer/DataLayerFactory.h"
+#include "Engine/World.h"
 #include "Logging/LogVerbosity.h"
 #include "ProfilingDebugging/ScopedTimers.h"
 #include "WorldPartition/DataLayer/WorldDataLayers.h"
diff --git a/Engine/Source/Editor/UnrealEd/Public/Commandlets/WorldPartitionCommandletHelpers.h b/Engine/Source/Editor/UnrealEd/Public/Commandlets/WorldPartitionCommandletHelpers.h
--- a/Engine/Source/Editor/UnrealEd/Public/Commandlets/WorldPartitionCommandletHelpers.h
+++ b/Engine/Source/Editor/UnrealEd/Public/Commandlets/WorldPartitionCommandletHelpers.h
@@ -7,6 +7,7 @@
 #include "UObject/UObjectAnnotation.h"
 
 class FPackageSourceControlHelper;
+class UWorld;
 
 namespace WorldPartitionCommandletHelpers
 {
