Add a new switch '-OnlyVirtualized' to the ResavePackages commandlet. When used, only packages containing virtualized bulkdata payloads will be resaved.

#rnx

- This can be used to more easily rehydrate a project. With the current setup, saving a package that is virtualized will cause it to be rehydrated. So running with this new switch will cause only the virtualized packages to save and become hydrated.
-- When we change the default save behaviour to not rehydrate packages we will need to add another switch to this commandlet, or add an easy way to switch back to this behaviour.

diff --git a/Engine/Source/Editor/UnrealEd/Classes/Commandlets/ResavePackagesCommandlet.h b/Engine/Source/Editor/UnrealEd/Classes/Commandlets/ResavePackagesCommandlet.h
--- a/Engine/Source/Editor/UnrealEd/Classes/Commandlets/ResavePackagesCommandlet.h
+++ b/Engine/Source/Editor/UnrealEd/Classes/Commandlets/ResavePackagesCommandlet.h
@@ -102,6 +102,9 @@ protected:
 	/** Only save packages that been saved by a licensee **/
 	bool bOnlyLicenseed;
 
+	/** Only save packages containing bulkdata payloads that have been virtualized **/
+	bool bOnlyVirtualized;
+
 	/** Should we build navigation data for the packages we are saving? **/
 	bool bShouldBuildNavigationData;
 
diff --git a/Engine/Source/Editor/UnrealEd/Private/Commandlets/ContentCommandlets.cpp b/Engine/Source/Editor/UnrealEd/Private/Commandlets/ContentCommandlets.cpp
--- a/Engine/Source/Editor/UnrealEd/Private/Commandlets/ContentCommandlets.cpp
+++ b/Engine/Source/Editor/UnrealEd/Private/Commandlets/ContentCommandlets.cpp
@@ -1,5 +1,4 @@
 // Copyright Epic Games, Inc. All Rights Reserved.
-// Copyright Epic Games, Inc. All Rights Reserved.
 
 /*=============================================================================
 	UCContentCommandlets.cpp: Various commmandlets.
@@ -55,6 +54,7 @@
 
 #include "PackageHelperFunctions.h"
 #include "PackageTools.h"
+#include "UObject/PackageTrailer.h"
 
 DEFINE_LOG_CATEGORY(LogContentCommandlet);
 
@@ -1041,6 +1041,8 @@ int32 UResavePackagesCommandlet::Main( const FString& Params )
 	bOnlyUnversioned = Switches.Contains(TEXT("OnlyUnversioned"));
 	/** whether we should only save packages saved by licenseed */
 	bOnlyLicenseed = Switches.Contains(TEXT("OnlyLicenseed"));
+	/** whether we should only save packages contianing virtualized bulkdata payloads */
+	bOnlyVirtualized = Switches.Contains(TEXT("OnlyVirtualized"));
 	/** only process packages containing materials */
 	bOnlyMaterials = Switches.Contains(TEXT("onlymaterials"));
 	/** determine if we are building navigation data for the map packages on the pass. **/
@@ -1376,7 +1378,7 @@ void UResavePackagesCommandlet::PerformPreloadOperations( FLinkerLoad* PackageLi
 	const int32 LicenseeUEPackageVersion = PackageLinker->Summary.GetFileVersionLicenseeUE();
 
 	// validate that this package meets the minimum requirement
-	if ( MinResaveUEVersion != IGNORE_PACKAGE_VERSION && UEPackageVersion < MinResaveUEVersion )
+	if (MinResaveUEVersion != IGNORE_PACKAGE_VERSION && UEPackageVersion < MinResaveUEVersion)
 	{
 		bSavePackage = false;
 		return;
@@ -1388,7 +1390,7 @@ void UResavePackagesCommandlet::PerformPreloadOperations( FLinkerLoad* PackageLi
 						 (MaxResaveUEVersion != IGNORE_PACKAGE_VERSION && UEPackageVersion <= MaxResaveUEVersion) ||
 						 (MaxResaveLicenseeUEVersion != IGNORE_PACKAGE_VERSION && LicenseeUEPackageVersion <= MaxResaveLicenseeUEVersion);
 	// If not, don't resave it.
-	if ( !bAllowResave )
+	if (!bAllowResave)
 	{
 		bSavePackage = false;
 		return;
@@ -1414,7 +1416,7 @@ void UResavePackagesCommandlet::PerformPreloadOperations( FLinkerLoad* PackageLi
 	}
 
 	// Check if the package was saved by licensees
-	if ( bOnlyLicenseed)
+	if (bOnlyLicenseed)
 	{
 		if (!PackageLinker->Summary.SavedByEngineVersion.IsLicenseeVersion())
 		{
@@ -1429,8 +1431,19 @@ void UResavePackagesCommandlet::PerformPreloadOperations( FLinkerLoad* PackageLi
 		}
 	}
 
+	// Check if the package contains virtualized bulkdata payloads
+	if (bOnlyVirtualized)
+	{
+		const UE::FPackageTrailer* Trailer = PackageLinker->GetPackageTrailer();
+		if (Trailer == nullptr || Trailer->GetNumPayloads(UE::EPayloadFilter::Virtualized) == 0)
+		{
+			bSavePackage = false;
+			return;
+		}
+	}
+
 	// Check if the package contains any instances of the class that needs to be resaved.
-	if ( ResaveClasses.Num() > 0 )
+	if (ResaveClasses.Num() > 0)
 	{
 		bSavePackage = false;
 		for (int32 ExportIndex = 0; !bSavePackage && ExportIndex < PackageLinker->ExportMap.Num(); ExportIndex++)
