Added a new analyzer to the sumarize trace commandlet to summarize LoadModule_* CPU scope, hierarchically.

- The new analyzer compute the time to load each module, excluding the time required to load dependent modules, hense the need to do hierarchical analysis.
- The new analyzer publishes the total number of module scope encountered (which is roughly equivalent to number of module loaded if we ignore the few circular dependencies), the total time measured to load all modules, the average load time of a module, the longuest load time and a list of the 10 longuest module load time.

diff --git a/Engine/Source/Editor/UnrealEd/Private/Commandlets/SummarizeTraceCommandlet.cpp b/Engine/Source/Editor/UnrealEd/Private/Commandlets/SummarizeTraceCommandlet.cpp
--- a/Engine/Source/Editor/UnrealEd/Private/Commandlets/SummarizeTraceCommandlet.cpp
+++ b/Engine/Source/Editor/UnrealEd/Private/Commandlets/SummarizeTraceCommandlet.cpp
@@ -42,38 +42,136 @@ static uint64 Decode7bit(const uint8*& Cursor)
 	return Value;
 }
 
-class FCpuAnalyzer
-	: public UE::Trace::IAnalyzer
+/**
+ * Base class to extend for CPU scope analysis. Derived class instances are meant to be registered as
+ * delegates to the FCpuScopeStreamProcessor to handle scope events and perform meaningful analysis.
+ */
+class FCpuScopeAnalyzer
 {
 public:
-	struct FScopeName
+	enum class EScopeEventType : uint32
 	{
-		const TCHAR*	Name;
-		uint32			Id;
+		Enter,
+		Exit
 	};
 
+	struct FScopeEvent
+	{
+		EScopeEventType ScopeEventType;
+		uint32 ScopeId;
+		uint32 ThreadId;
+		double Timestamp; // As Seconds
+	};
+
+	struct FScope
+	{
+		uint32 ScopeId;
+		uint32 ThreadId;
+		double EnterTimestamp; // As Seconds
+		double ExitTimestamp;  // As Seconds
+	};
+
+public:
+	virtual ~FCpuScopeAnalyzer() = default;
+
+	/** Invoked when a CPU scope is discovered. This function is always invoked first when a CPU scope is encountered for the first time.*/
+	virtual void OnCpuScopeDiscovered(uint32 ScopeId) {}
+
+	/** Invoked when CPU scope specification is encountered in the trace stream. */
+	virtual void OnCpuScopeName(uint32 ScopeId, const FString& ScopeName) {};
+
+	/** Invoked when a scope is entered. The scope name might not be known yet. */
+	virtual void OnCpuScopeEnter(const FScopeEvent& ScopeEnter, const FString* ScopeName) {};
+
+	/** Invoked when a scope is exited. The scope name might not be known yet. */
+	virtual void OnCpuScopeExit(const FScope& Scope, const FString* ScopeName) {};
+
+	/** Invoked when a root event on the specified thread along with all child events down to the leaves are know. */
+	virtual void OnCpuScopeTree(uint32 ThreadId, const TArray<FCpuScopeAnalyzer::FScopeEvent>& ScopeEvents, const TFunction<const FString*(uint32)>& ScopeLookupNameFn) {};
+
+	/** Invoked when the trace stream has been fully consumed/processed. */
+	virtual void OnCpuScopeAnalysisEnd() {};
+};
+
+/**
+ * Decodes, format and routes CPU scope events embedded into a trace stream to specialized CPU scope analyzers. This processor
+ * decodes the low level events, keeps a small state and publishes higher level events to registered analyzers. The purpose
+ * is to decode the stream only once and let several registered analyzers reuse the small state built up by this processor. This
+ * matters when processing very large traces with possibly billions of scope events.
+ */
+class FCpuScopeStreamProcessor
+	: public UE::Trace::IAnalyzer
+{
+public:
+	FCpuScopeStreamProcessor();
+
+	/** Register an analyzer with this processor. The processor decodes the trace stream and invokes the registered analyzers when a CPU scope event occurs.*/
+	void AddCpuScopeAnalyzer(TSharedPtr<FCpuScopeAnalyzer> Analyzer);
+
+private:
 	struct FScopeEnter
 	{
-		double			TimeStamp;
-		uint32			ScopeId;
-		uint32			ThreadId;
+		uint32 ScopeId;
+		double Timestamp; // As Seconds
+	};
+
+	// Contains scope events for a root scope and its children along with extra info to analyze that tree at once.
+	struct FScopeTreeInfo
+	{
+		// Records the current root scope and its children to run analysis that needs to know the parent/child relationship.
+		TArray<FCpuScopeAnalyzer::FScopeEvent> ScopeEvents;
+
+		// Indicates if one of the scope in the current hierarchy is nameless. (Its names specs hasn't been received yet).
+		bool bHasNamelessScopes = false;
+
+		void Reset()
+		{
+			ScopeEvents.Reset();
+			bHasNamelessScopes = false;
+		}
 	};
 
-	struct FScopeExit
+	// For each thread we track what the stack of scopes are, for matching end-to-start
+	struct FThread
 	{
-		double			TimeStamp;
-		uint32			ThreadId;
+		TArray<FScopeEnter> ScopeStack;
+
+		// The events recorded for the current root scope and its children to run analysis that needs to know the parent/child relationship, for example to compute time including childreen and time
+		// excluding childreen.
+		FScopeTreeInfo ScopeTreeInfo;
+
+		// Scope trees for which as least one scope name was unknown. Some analysis need scope names, but scope names/event can be emitted out of order by the engine depending on thread scheduling.
+		// Some scope tree cannot be analyzed right away and need to be delayed until all scope names are discovered.
+		TArray<FScopeTreeInfo> DelayedScopeTreeInfo;
 	};
 
-	virtual void		OnCpuScopeName(const FScopeName& ScopeName) = 0;
-	virtual void		OnCpuScopeEnter(const FScopeEnter& ScopeEnter) = 0;
-	virtual void		OnCpuScopeExit(const FScopeExit& ScopeExit) = 0;
+private:
+	// UE::Trace::IAnalyzer interface.
+	virtual void OnAnalysisBegin(const FOnAnalysisContext& Context) override;
+	virtual void OnAnalysisEnd() override;
+	virtual bool OnEvent(uint16 RouteId, EStyle Style, const FOnEventContext& Context) override;
+
+	// Internal implementation.
+	void OnEventSpec(const FOnEventContext& Context);
+	void OnBatch(const FOnEventContext& Context);
+	void OnCpuScopeSpec(uint32 ScopeId, const FString* ScopeName);
+	void OnCpuScopeEnter(uint32 ScopeId, uint32 ThreadId, double Timestamp);
+	void OnCpuScopeExit(uint32 ThreadId, double Timestamp);
+	void OnCpuScopeTree(uint32 ThreadId, const FScopeTreeInfo& ScopeTreeInfo);
+	const FString* LookupScopeName(uint32 ScopeId) { return ScopeId < static_cast<uint32>(ScopeNames.Num()) && ScopeNames[ScopeId] ? &ScopeNames[ScopeId].GetValue() : nullptr; }
 
 private:
-	virtual void		OnAnalysisBegin(const FOnAnalysisContext& Context) override;
-	virtual bool		OnEvent(uint16 RouteId, EStyle Style, const FOnEventContext& Context) override;
-	void				OnEventSpec(const FOnEventContext& Context);
-	void				OnBatch(const FOnEventContext& Context);
+	// The state at any moment of the threads, indexes are doled out on the process-side
+	TArray<FThread> Threads;
+
+	// The scope names, the array index correspond to the scope Id. If the optional is not set, the scope hasn't been encountered yet.
+	TArray<TOptional<FString>> ScopeNames;
+
+	// List of analyzers to invoke when a scope event is decoded.
+	TArray<TSharedPtr<FCpuScopeAnalyzer>> ScopeAnalyzers;
+
+	// Scope name lookup function, cached for efficiency.
+	TFunction<const FString*(uint32 ScopeId)> LookupScopeNameFn;
 };
 
 enum
@@ -84,14 +182,46 @@ enum
 	RouteId_CpuProfiler_EndCapture,
 };
 
-void FCpuAnalyzer::OnAnalysisBegin(const FOnAnalysisContext& Context)
+FCpuScopeStreamProcessor::FCpuScopeStreamProcessor()
+ : LookupScopeNameFn([this](uint32 ScopeId) { return LookupScopeName(ScopeId); })
+{
+}
+
+void FCpuScopeStreamProcessor::AddCpuScopeAnalyzer(TSharedPtr<FCpuScopeAnalyzer> Analyzer)
+{
+	ScopeAnalyzers.Add(MoveTemp(Analyzer));
+}
+
+void FCpuScopeStreamProcessor::OnAnalysisBegin(const FOnAnalysisContext& Context)
 {
 	Context.InterfaceBuilder.RouteEvent(RouteId_CpuProfiler_EventSpec, "CpuProfiler", "EventSpec");
 	Context.InterfaceBuilder.RouteEvent(RouteId_CpuProfiler_EventBatch, "CpuProfiler", "EventBatch");
 	Context.InterfaceBuilder.RouteEvent(RouteId_CpuProfiler_EndCapture, "CpuProfiler", "EndCapture");
 }
 
-bool FCpuAnalyzer::OnEvent(uint16 RouteId, EStyle Style, const FOnEventContext& Context)
+void FCpuScopeStreamProcessor::OnAnalysisEnd()
+{
+	// Analyze scope trees that contained 'nameless' context when they were captured. Unless the trace was truncated,
+	// all scope names should be known now.
+	uint32 ThreadId = 0;
+	for (FThread& Thread : Threads)
+	{
+		for (FScopeTreeInfo& DelayedScopeTree : Threads[ThreadId].DelayedScopeTreeInfo)
+		{
+			// Run summary analysis for this delayed hierarchy.
+			OnCpuScopeTree(ThreadId, DelayedScopeTree);
+		}
+
+		++ThreadId;
+	}
+
+	for (TSharedPtr<FCpuScopeAnalyzer>& Analyzer : ScopeAnalyzers)
+	{
+		Analyzer->OnCpuScopeAnalysisEnd();
+	}
+}
+
+bool FCpuScopeStreamProcessor::OnEvent(uint16 RouteId, EStyle Style, const FOnEventContext& Context)
 {
 	switch (RouteId)
 	{
@@ -108,16 +238,16 @@ bool FCpuAnalyzer::OnEvent(uint16 RouteId, EStyle Style, const FOnEventContext&
 	return true;
 }
 
-void FCpuAnalyzer::OnEventSpec(const FOnEventContext& Context)
+void FCpuScopeStreamProcessor::OnEventSpec(const FOnEventContext& Context)
 {
 	const FEventData& EventData = Context.EventData;
 	FString Name;
 	uint32 Id = EventData.GetValue<uint32>("Id");
 	EventData.GetString("Name", Name);
-	OnCpuScopeName({ *Name, Id - 1 });
+	OnCpuScopeSpec(Id - 1, &Name);
 }
 
-void FCpuAnalyzer::OnBatch(const FOnEventContext& Context)
+void FCpuScopeStreamProcessor::OnBatch(const FOnEventContext& Context)
 {
 	const FEventData& EventData = Context.EventData;
 	const FEventTime& EventTime = Context.EventTime;
@@ -138,15 +268,114 @@ void FCpuAnalyzer::OnBatch(const FOnEventContext& Context)
 		if (Value & 1)
 		{
 			uint64 ScopeId = Decode7bit(Cursor);
-			OnCpuScopeEnter({ TimeStamp, uint32(ScopeId - 1), ThreadId });
+			OnCpuScopeEnter(static_cast<uint32>(ScopeId - 1), ThreadId, TimeStamp);
 		}
 		else
 		{
-			OnCpuScopeExit({ TimeStamp, ThreadId });
+			OnCpuScopeExit(ThreadId, TimeStamp);
+		}
+	}
+}
+
+void FCpuScopeStreamProcessor::OnCpuScopeSpec(uint32 ScopeId, const FString* ScopeName)
+{
+	if (ScopeId >= uint32(ScopeNames.Num()))
+	{
+		uint32 Num = (ScopeId + 128) & ~127;
+		ScopeNames.SetNum(Num);
+	}
+
+	// If the optional is not set, the scope hasn't been encountered yet (Set to a blank names means encountered, but names hasn't be encountered yet)
+	bool bDiscovered = !ScopeNames[ScopeId].IsSet();
+
+	// Store the discovery state (setting the optional) and the actual name if known.
+	ScopeNames[ScopeId].Emplace(ScopeName ? *ScopeName : FString());
+
+	// Notify the analyzers.
+	for (TSharedPtr<FCpuScopeAnalyzer>& Analyzer : ScopeAnalyzers)
+	{
+		if (bDiscovered)
+		{
+			Analyzer->OnCpuScopeDiscovered(ScopeId);
+		}
+		if (ScopeName)
+		{
+			Analyzer->OnCpuScopeName(ScopeId, *ScopeName);
 		}
 	}
 }
 
+void FCpuScopeStreamProcessor::OnCpuScopeEnter(uint32 ScopeId, uint32 ThreadId, double Timestamp)
+{
+	if (ThreadId >= uint32(Threads.Num()))
+	{
+		Threads.SetNum(ThreadId + 1);
+	}
+
+	FCpuScopeAnalyzer::FScopeEvent ScopeEvent{FCpuScopeAnalyzer::EScopeEventType::Enter, ScopeId, ThreadId, Timestamp};
+	Threads[ThreadId].ScopeStack.Add(FScopeEnter{ScopeId, Timestamp});
+	Threads[ThreadId].ScopeTreeInfo.ScopeEvents.Add(ScopeEvent);
+
+	// Scope specs/events can be received out of order depending on engine thread scheduling.
+	if (ScopeId >= static_cast<uint32>(ScopeNames.Num()) || !ScopeNames[ScopeId].IsSet())
+	{
+		// This is a newly discovered scope. Create an entry for it and notify the discovery.
+		OnCpuScopeSpec(ScopeId, nullptr);
+	}
+
+	// Notify the registered scope analyzers.
+	for (TSharedPtr<FCpuScopeAnalyzer>& Analyzer : ScopeAnalyzers)
+	{
+		Analyzer->OnCpuScopeEnter(ScopeEvent, LookupScopeName(ScopeId));
+	}
+}
+
+void FCpuScopeStreamProcessor::OnCpuScopeExit(uint32 ThreadId, double Timestamp)
+{
+	if (ThreadId >= uint32(Threads.Num()) || Threads[ThreadId].ScopeStack.Num() <= 0)
+	{
+		return;
+	}
+
+	FScopeEnter ScopeEnter = Threads[ThreadId].ScopeStack.Pop();
+	
+	Threads[ThreadId].ScopeTreeInfo.ScopeEvents.Add(FCpuScopeAnalyzer::FScopeEvent{FCpuScopeAnalyzer::EScopeEventType::Exit, ScopeEnter.ScopeId, ThreadId, Timestamp});
+	Threads[ThreadId].ScopeTreeInfo.bHasNamelessScopes |= ScopeNames[ScopeEnter.ScopeId].GetValue().IsEmpty();
+
+	// Notify the registered scope analyzers.
+	FCpuScopeAnalyzer::FScope Scope{ScopeEnter.ScopeId, ThreadId, ScopeEnter.Timestamp, Timestamp};
+	for (TSharedPtr<FCpuScopeAnalyzer>& Analyzer : ScopeAnalyzers)
+	{
+		Analyzer->OnCpuScopeExit(Scope, LookupScopeName(ScopeEnter.ScopeId));
+	}
+
+	// The root scope on this thread just poped out.
+	if (Threads[ThreadId].ScopeStack.IsEmpty())
+	{
+		if (Threads[ThreadId].ScopeTreeInfo.bHasNamelessScopes)
+		{
+			// Delay the analysis until all the scope names are known.
+			Threads[ThreadId].DelayedScopeTreeInfo.Add(MoveTemp(Threads[ThreadId].ScopeTreeInfo));
+		}
+		else
+		{
+			// Run analysis for this scope tree.
+			OnCpuScopeTree(ThreadId, Threads[ThreadId].ScopeTreeInfo);
+		}
+
+		Threads[ThreadId].ScopeTreeInfo.Reset();
+	}
+}
+
+void FCpuScopeStreamProcessor::OnCpuScopeTree(uint32 ThreadId, const FScopeTreeInfo& ScopeTreeInfo)
+{
+	for (TSharedPtr<FCpuScopeAnalyzer>& Analyzer : ScopeAnalyzers)
+	{
+		Analyzer->OnCpuScopeTree(ThreadId, ScopeTreeInfo.ScopeEvents, LookupScopeNameFn);
+	}
+}
+
+
 class FCountersAnalyzer
 	: public UE::Trace::IAnalyzer
 {
@@ -582,73 +811,218 @@ static uint32 GetTypeHash(const FSummarizeBookmark& Bookmark)
 	return FCrc::StrCrc32(*Bookmark.Name);
 }
 
-//
-// FSummarizeCpuAnalyzer - Generate Scopes from cpu channel scope enter/exit events
-//
-
+/**
+ * This analyzer aggregates CPU scopes having the same name together, counting the number of occurrences, total duration,
+ * average occurrence duration, for each scope name.
+ */
 class FSummarizeCpuAnalyzer
-	: public FCpuAnalyzer
+	: public FCpuScopeAnalyzer
 {
 public:
-	virtual void OnCpuScopeName(const FScopeName& ScopeName) override;
-	virtual void OnCpuScopeEnter(const FScopeEnter& ScopeEnter) override;
-	virtual void OnCpuScopeExit(const FScopeExit& ScopeExit) override;
+	FSummarizeCpuAnalyzer(TFunction<void(const TArray<FSummarizeScope>&)> InPublishFn);
+
+protected:
+	virtual void OnCpuScopeDiscovered(uint32 ScopeId) override;
+	virtual void OnCpuScopeName(uint32 ScopeId, const FString& ScopeName) override;
+	virtual void OnCpuScopeExit(const FScope& Scope, const FString* ScopeName) override;
+	virtual void OnCpuScopeAnalysisEnd() override;
 
+private:
 	// Scopes is an array only because indexes are doled out on the process-side
-	//  As such, there could be different scopes with the same name
-	//  We merge these together later to ensure name is a pkey in the csv
+	// As such, there could be different scopes with the same name
+	// We merge these together later to ensure name is a pkey in the csv
 	TArray<FSummarizeScope> Scopes;
 
-	// For each thread we track what the stack of scopes are, for matching end-to-start
-	struct FThread
-	{
-		TArray<FScopeEnter> ScopeStack;
-	};
-
-	// The state at any moment of the threads, indexes are doled out on the process-side
-	TArray<FThread> Threads;
+	// Function invoked to publish the list of scopes.
+	TFunction<void(const TArray<FSummarizeScope>&)> PublishFn;
 };
 
-void FSummarizeCpuAnalyzer::OnCpuScopeName(const FScopeName& ScopeName)
+FSummarizeCpuAnalyzer::FSummarizeCpuAnalyzer(TFunction<void(const TArray<FSummarizeScope>&)> InPublishFn)
+	: PublishFn(MoveTemp(InPublishFn))
 {
-	if (ScopeName.Id >= uint32(Scopes.Num()))
+}
+
+void FSummarizeCpuAnalyzer::OnCpuScopeDiscovered(uint32 ScopeId)
+{
+	if (ScopeId >= uint32(Scopes.Num()))
 	{
-		uint32 Num = (ScopeName.Id + 128) & ~127;
+		uint32 Num = (ScopeId + 128) & ~127;
 		Scopes.SetNum(Num);
 	}
-	Scopes[ScopeName.Id].Name = ScopeName.Name;
 }
 
-void FSummarizeCpuAnalyzer::OnCpuScopeEnter(const FScopeEnter& ScopeEnter)
+void FSummarizeCpuAnalyzer::OnCpuScopeName(uint32 ScopeId, const FString& ScopeName)
 {
-	uint32 ThreadId = ScopeEnter.ThreadId;
-	if (ThreadId >= uint32(Threads.Num()))
-	{
-		Threads.SetNum(ThreadId + 1);
-	}
-	Threads[ThreadId].ScopeStack.Add(ScopeEnter);
+	Scopes[ScopeId].Name = ScopeName;
 }
 
-void FSummarizeCpuAnalyzer::OnCpuScopeExit(const FScopeExit& ScopeExit)
+void FSummarizeCpuAnalyzer::OnCpuScopeExit(const FScope& Scope, const FString* ScopeName)
 {
-	uint32 ThreadId = ScopeExit.ThreadId;
-	if (ThreadId >= uint32(Threads.Num()) || Threads[ThreadId].ScopeStack.Num() <= 0)
+	Scopes[Scope.ScopeId].AddDuration(Scope.EnterTimestamp, Scope.ExitTimestamp);
+}
+
+void FSummarizeCpuAnalyzer::OnCpuScopeAnalysisEnd()
+{
+	// Eliminates scopes that don't have a name. (On scope discovery, the array is expended to creates blank scopes that may never be filled).
+	Scopes.RemoveAll([](const FSummarizeScope& Scope) { return Scope.Name.IsEmpty(); });
+
+	// Publish the scopes.
+	PublishFn(Scopes);
+}
+
+
+/**
+ * Summarizes matched CPU scopes, excluding time consumed by immediate children if any. The analyzer uses pattern matching to
+ * selects the scopes of interest and detects parent/child relationship. Once a parent/child relationship is established in a
+ * scope tree, the analyzer can substract the time consumed by its immediate children if any.
+ *
+ * Such analysis is often meaningful for reentrant/recursive scope. For example, the UE modules can be loaded recursively and
+ * it is useful to know how much time a module used to load itself vs how much time it use to recursively load its dependent
+ * modules. In that example, we need to know which scope timers are actually the recursion vs the other intermediate scopes.
+ * So, pattern-matching scope names is used to deduce a relationship between scopes in a tree of scopes.
+ *
+ * For the LoadModule example described above, if the analyzer gets this scope tree as input:
+ *
+ * |-LoadModule_Module1----------------------------------------------------------|
+ *    |- StartupModule -------------------------------------------------------|
+ *      |-LoadModule_Module1Dep1------------|  |-LoadModule_Module1Dep2------|
+ *        |-StartupModule-----------------|      |-StartupModule----------|
+ *
+ * It would turn it into the one below if the REGEX to match was "LoadModule_.*"
+ * 
+ * |-LoadModule_Module1----------------------------------------------------------|
+ *      |-LoadModule_Module1Dep1------------|  |-LoadModule_Module1Dep2------|
+ *
+ * And it would compute the exclusive time required to load Module1 by substracting the time consumed to
+ * load Module1Dep1 and Module1Dep2.
+ *
+ * @note If the matching expression was to match all scopes, the analyser would summarize all scopes,
+ *       accounting for their exclusive time.
+ */
+class FCpuScopeHierarchyAnalyzer : public FCpuScopeAnalyzer
+{
+public:
+	/**
+	 * Constructs the analyzer
+	 * @param InAnalyzerName    The name of this analyzer. Some output statistics will also be derived from this name.
+	 * @param InMatchFn         Invoked by the analyzer to determine if a scope should be accounted by the analyzer. If it returns true, the scope is kept, otherwise, it is ignored.
+	 * @param InPublishFn       Invoked at the end of the analysis to post process the scopes summaries procuded by this analysis, possibly eliminating or renaming them then to publish them.
+	 *                          The first parameter is the summary of all matched scopes together while the array contains summary of scopes that matched the expression by grouped by exact name match.
+	 * 
+	 * @note This analyzer publishes summarized scope names with a suffix ".excl" as it computes exclusive time to  prevent name collisions with other analyzers. The summary of all scopes
+	 *       suffixed with ".excl.all" and gets its base name from the analyzer name. The scopes in the array gets their names from the scope name themselves, but suffixed with .excl.
+	 */
+	FCpuScopeHierarchyAnalyzer(const FString& InAnalyzerName, TFunction<bool(const FString&)> InMatchFn, TFunction<void(const FSummarizeScope&, TArray<FSummarizeScope>&&)> InPublishFn);
+
+	/**
+	 * Runs analysis on a collection of scopes events. The scopes are expected to be from the same thread and form a 'tree', meaning
+	 * it has the root scope events on that thread as well as all children below the root down to the leaves.
+	 * @param ThreadId The thread on which the scopes were recorded.
+	 * @param ScopeEvents The scopes events containing one root event along with its hierarchy.
+	 * @param InScopeNameLookup Callback function to lookup scope names from scope ID.
+	 */
+	virtual void OnCpuScopeTree(uint32 ThreadId, const TArray<FCpuScopeAnalyzer::FScopeEvent>& ScopeEvents, const TFunction<const FString*(uint32 /*ScopeId*/)>& InScopeNameLookup) override;
+
+	/**
+	 * Invoked to notify that the trace session ended and that the analyzer can publish the statistics gathered.
+	 * The analyzer calls the publishing function passed at construction time to publish the analysis results.
+	 */
+	virtual void OnCpuScopeAnalysisEnd() override;
+
+private:
+	// The function invoked to filter (by name) the scopes of interest. A scopes is kept if this function returns true.
+	TFunction<bool(const FString&)> MatchesFn;
+
+	// Aggregate all scopes matching the filter function, accounting for the parent/child relationship, so the duration stats will be from
+	// the 'exclusive' time (itself - duration of immediate children) of the matched scope.
+	FSummarizeScope MatchedScopesSummary;
+
+	// Among the scopes matching the filter function, grouped by scope name summaries, accounting for the parent/child relationship. So the duration stats will be
+	// from the 'exclusive' time (itself - duration of immediate children).
+	TMap<FString, FSummarizeScope> ExactNameMatchScopesSummaries;
+
+	// Invoked at the end of the analysis to publich scope summaries.
+	TFunction<void(const FSummarizeScope&, TArray<FSummarizeScope>&&)> PublishFn;
+};
+
+FCpuScopeHierarchyAnalyzer::FCpuScopeHierarchyAnalyzer(const FString& InAnalyzerName, TFunction<bool(const FString&)> InMatchFn, TFunction<void(const FSummarizeScope&, TArray<FSummarizeScope>&&)> InPublishFn)
+ : MatchesFn(MoveTemp(InMatchFn))
+ , PublishFn(MoveTemp(InPublishFn))
+{
+	MatchedScopesSummary.Name = InAnalyzerName;
+}
+
+void FCpuScopeHierarchyAnalyzer::OnCpuScopeTree(uint32 ThreadId, const TArray<FCpuScopeAnalyzer::FScopeEvent>& ScopeEvents, const TFunction<const FString*(uint32 /*ScopeId*/)>& InScopeNameLookup)
+{
+	// Scope matching the pattern.
+	struct FMatchScopeEnter
 	{
-		return;
+		FMatchScopeEnter(uint32 InScopeId, double InTimestamp) : ScopeId(InScopeId), Timestamp(InTimestamp) {}
+		uint32 ScopeId;
+		double Timestamp;
+		FTimespan ChildrenDuration;
+	};
+
+	TArray<FMatchScopeEnter> ScopeStack;
+
+	// Replay and filter this scope hierarchy to only keep the ones matching the condition/regex. (See class documentation for a visual example)
+	for (const FCpuScopeAnalyzer::FScopeEvent& ScopeEvent : ScopeEvents)
+	{
+		const FString* ScopeName = InScopeNameLookup(ScopeEvent.ScopeId);
+		if (!ScopeName || !MatchesFn(*ScopeName))
+		{
+			continue;
+		}
+
+		if (ScopeEvent.ScopeEventType == FCpuScopeAnalyzer::EScopeEventType::Enter)
+		{
+			ScopeStack.Emplace(ScopeEvent.ScopeId, ScopeEvent.Timestamp);
+		}
+		else // Scope Exit
+		{
+			FMatchScopeEnter EnterScope = ScopeStack.Pop();
+			double EnterTimestampSecs = EnterScope.Timestamp;
+			uint32 ScopeId = EnterScope.ScopeId;
+
+			// Total time consumed by this scope.
+			FTimespan InclusiveDuration = FTimespan::FromSeconds(ScopeEvent.Timestamp - EnterTimestampSecs);
+
+			// Total time consumed by this scope, excluding the time consumed by matched 'children' scopes.
+			FTimespan ExclusiveDuration = InclusiveDuration - EnterScope.ChildrenDuration;
+
+			if (ScopeStack.Num() > 0)
+			{
+				// Track how much time this 'child' consumed inside its parent.
+				ScopeStack.Last().ChildrenDuration += InclusiveDuration;
+			}
+
+			// Aggregate this scope with all other scopes of the exact same name, excluding children duration, so that we have the 'self only' starts.
+			FSummarizeScope& ExactNameScopeSummary = ExactNameMatchScopesSummaries.FindOrAdd(*ScopeName);
+			ExactNameScopeSummary.Name = *ScopeName;
+			ExactNameScopeSummary.AddDuration(EnterTimestampSecs, EnterTimestampSecs + ExclusiveDuration.GetTotalSeconds());
+
+			// Aggregate this scope with all other scopes matching the pattern, but excluding the children time, so that we have the stats of 'self only'.
+			MatchedScopesSummary.AddDuration(EnterTimestampSecs, EnterTimestampSecs + ExclusiveDuration.GetTotalSeconds());
+		}
 	}
+}
 
-	FScopeEnter ScopeEnter = Threads[ThreadId].ScopeStack.Pop();
-	double ScopeDuration = ScopeExit.TimeStamp - ScopeEnter.TimeStamp;
+void FCpuScopeHierarchyAnalyzer::OnCpuScopeAnalysisEnd()
+{
+	MatchedScopesSummary.Name += TEXT(".excl.all");
 
-	// unclear why we are getting ids that are out-of-bounds, fewer specs than scopes shouldn't be possible
-	//  maybe we are losing spec data?
-	//  maybe scope to spec id enc/dec has edge cases?
-	if (ScopeEnter.ScopeId < uint32(Scopes.Num()))
+	TArray<FSummarizeScope> ScopeSummaries;
+	for (TPair<FString, FSummarizeScope>& Pair : ExactNameMatchScopesSummaries)
 	{
-		Scopes[ScopeEnter.ScopeId].AddDuration(ScopeEnter.TimeStamp, ScopeExit.TimeStamp);
+		Pair.Value.Name += TEXT(".excl");
+		ScopeSummaries.Add(Pair.Value);
 	}
+
+	// Publish the scopes.
+	PublishFn(MatchedScopesSummary, MoveTemp(ScopeSummaries));
 }
 
+
 //
 // FSummarizeCountersAnalyzer - Tally Counters from counter set/increment events
 //
@@ -848,6 +1222,7 @@ FSummarizeBookmark* FSummarizeBookmarksAnalyzer::FindStartBookmarkForEndBookmark
 * Begin SummarizeTrace commandlet implementation
 */
 
+// Defined here to prevent adding logs in the code above, which will likely be moved elsewhere.
 DEFINE_LOG_CATEGORY_STATIC(LogSummarizeTrace, Log, All);
 
 /*
@@ -933,7 +1308,7 @@ struct StatisticDefinition
 			&& BaselineErrorThreshold == InStatistic.BaselineErrorThreshold;
 	}
 
-	static bool LoadFromCSV(const FString& FilePath, TMultiMap<FString, StatisticDefinition>& NameToDefinitionMap);
+	static bool LoadFromCSV(const FString& FilePath, TMultiMap<FString, StatisticDefinition>& NameToDefinitionMap, TSet<FString>& OutWildcardNames);
 
 	FString Name;
 	FString Statistic;
@@ -944,7 +1319,7 @@ struct StatisticDefinition
 	FString BaselineErrorThreshold;
 };
 
-bool StatisticDefinition::LoadFromCSV(const FString& FilePath, TMultiMap<FString, StatisticDefinition>& NameToDefinitionMap)
+bool StatisticDefinition::LoadFromCSV(const FString& FilePath, TMultiMap<FString, StatisticDefinition>& NameToDefinitionMap, TSet<FString>& OutWildcardNames)
 {
 	TArray<FString> ParsedCSVFile;
 	FFileHelper::LoadFileToStringArray(ParsedCSVFile, *FilePath);
@@ -1011,6 +1386,11 @@ bool StatisticDefinition::LoadFromCSV(const FString& FilePath, TMultiMap<FString
 			const FString& TelemetryUnit(Fields[TelemetryUnitColumn]);
 			const FString& BaselineWarningThreshold(Fields[BaselineWarningThresholdColumn]);
 			const FString& BaselineErrorThreshold(Fields[BaselineErrorThresholdColumn]);
+
+			if (Name.Contains("*") || Name.Contains("?")) // Wildcards.
+			{
+				OutWildcardNames.Add(Name);
+			}
 			NameToDefinitionMap.AddUnique(Name, StatisticDefinition(Name, Statistic, TelemetryContext, TelemetryDataPoint, TelemetryUnit, BaselineWarningThreshold, BaselineErrorThreshold));
 		}
 	}
@@ -1274,19 +1654,19 @@ int32 USummarizeTraceCommandlet::Main(const FString& CmdLineParams)
 	// load the stats file to know which event name and statistic name to generate in the telemetry csv
 	// the telemetry csv is ingested completely, so this just highlights specific data elements we want to track
 	TMultiMap<FString, StatisticDefinition> NameToDefinitionMap;
+	TSet<FString> CpuScopeNamesWithWildcards;
 	FString GlobalStatisticsFileName = FPaths::RootDir() / TEXT("Engine") / TEXT("Build") / TEXT("EditorPerfStats.csv");
 	if (FPaths::FileExists(GlobalStatisticsFileName))
 	{
 		UE_LOG(LogSummarizeTrace, Display, TEXT("Loading global statistics from %s"), *GlobalStatisticsFileName);
-		bool bCSVOk = StatisticDefinition::LoadFromCSV(GlobalStatisticsFileName, NameToDefinitionMap);
+		bool bCSVOk = StatisticDefinition::LoadFromCSV(GlobalStatisticsFileName, NameToDefinitionMap, CpuScopeNamesWithWildcards);
 		check(bCSVOk);
 	}
 	FString ProjectStatisticsFileName = FPaths::ProjectDir() / TEXT("Build") / TEXT("EditorPerfStats.csv");
 	if (FPaths::FileExists(ProjectStatisticsFileName))
 	{
 		UE_LOG(LogSummarizeTrace, Display, TEXT("Loading project statistics from %s"), *ProjectStatisticsFileName);
-		bool bCSVOk = StatisticDefinition::LoadFromCSV(ProjectStatisticsFileName, NameToDefinitionMap);
-		check(bCSVOk);
+		bool bCSVOk = StatisticDefinition::LoadFromCSV(ProjectStatisticsFileName, NameToDefinitionMap, CpuScopeNamesWithWildcards);
 	}
 
 	bool bFound;
@@ -1328,8 +1708,47 @@ int32 USummarizeTraceCommandlet::Main(const FString& CmdLineParams)
 
 	// setup analysis context with analyzers
 	UE::Trace::FAnalysisContext AnalysisContext;
-	FSummarizeCpuAnalyzer CpuAnalyzer;
-	AnalysisContext.AddAnalyzer(CpuAnalyzer);
+
+	// List of summarized scopes.
+	TArray<FSummarizeScope> CollectedScopeSummaries;
+
+	// Analyze CPU scope timer individually.
+	TSharedPtr<FSummarizeCpuAnalyzer> IndividualScopeAnalyzer = MakeShared<FSummarizeCpuAnalyzer>(
+		[&CollectedScopeSummaries](const TArray<FSummarizeScope>& ScopeSummaries)
+		{
+			// Collect all individual scopes summary from this analyzer.
+			CollectedScopeSummaries.Append(ScopeSummaries);
+		});
+
+	// Analyze 'LoadModule*' scope timer hierarchically to account individual load time only (substracting time consumed to load dependent module(s)).
+	TSharedPtr<FCpuScopeHierarchyAnalyzer> HierarchicalScopeAnalyzer = MakeShared<FCpuScopeHierarchyAnalyzer>(
+		TEXT("LoadModule"), // Analyzer Name.
+		[](const FString& ScopeName)
+		{
+			return ScopeName.StartsWith("LoadModule_"); // When analyzing a tree of scopes, only keeps scope with name starting with 'LoadModule'.
+		},
+		[&CollectedScopeSummaries](const FSummarizeScope& AllModulesStats, TArray<FSummarizeScope>&& ModuleStats)
+		{
+			// Module should be loaded only once and the check below should be true but the load function can start loading X, process some dependencies which could
+			// trigger loading X again within the first scope. The engine code gracefully handle this case and don't load twice, but we end up with two 'load x' scope.
+			// Both scope times be added together, providing the correct sum for that module though.
+			//check(AllModulesStats.Count == ModuleStats.Num())
+
+			// Publish the total nb. of module loaded, total time to the modules, avg time per module, etc (all module load times exclude the time to load sub-modules)
+			CollectedScopeSummaries.Add(AllModulesStats);
+
+			// Sort the summaries descending.
+			ModuleStats.Sort([](const FSummarizeScope& Lhs, const FSummarizeScope& Rhs) { return Lhs.TotalDurationSeconds >= Rhs.TotalDurationSeconds; });
+
+			// Publish top N longuest load module. The ModuleStats are pre-sorted from the longest to the shorted timer.
+			CollectedScopeSummaries.Append(ModuleStats.GetData(), 10);
+		});
+
+	FCpuScopeStreamProcessor CpuScopeStreamProcessor;
+	CpuScopeStreamProcessor.AddCpuScopeAnalyzer(IndividualScopeAnalyzer);
+	CpuScopeStreamProcessor.AddCpuScopeAnalyzer(HierarchicalScopeAnalyzer);
+	AnalysisContext.AddAnalyzer(CpuScopeStreamProcessor);
+
 	FSummarizeCountersAnalyzer CountersAnalyzer;
 	AnalysisContext.AddAnalyzer(CountersAnalyzer);
 	FSummarizeBookmarksAnalyzer BookmarksAnalyzer;
@@ -1364,7 +1783,7 @@ int32 USummarizeTraceCommandlet::Main(const FString& CmdLineParams)
 			DeduplicatedScopes.Add(Scope);
 		}
 	};
-	for (const FSummarizeScope& Scope : CpuAnalyzer.Scopes)
+	for (const FSummarizeScope& Scope : CollectedScopeSummaries)
 	{
 		IngestScope(DeduplicatedScopes, Scope);
 	}
@@ -1499,12 +1918,32 @@ int32 USummarizeTraceCommandlet::Main(const FString& CmdLineParams)
 					continue;
 				}
 
+				// Is that scope summary desired in the output, using an exact name match?
 				NameToDefinitionMap.MultiFind(Scope.Name, Statistics, true);
 				for (const StatisticDefinition& Statistic : Statistics)
 				{
 					TelemetryData.Add(TelemetryDefinition(TestName, Statistic.TelemetryContext, Statistic.TelemetryDataPoint, Statistic.TelemetryUnit, Scope.GetValue(Statistic.Statistic)));
 				}
 				Statistics.Reset();
+
+				// If the configuration file contains scope names with wildcard characters * and ?
+				for (const FString& Pattern : CpuScopeNamesWithWildcards)
+				{
+					// Check if the current scope names matches the pattern.
+					if (Scope.Name.MatchesWildcard(Pattern))
+					{
+						// Find the statistic definition for this pattern.
+						NameToDefinitionMap.MultiFind(Pattern, Statistics, true);
+						for (const StatisticDefinition& Statistic : Statistics)
+						{
+							// Use the scope name as data point. Normally, the data point is configured in the .csv as a 1:1 match (1 scopeName=> 1 DataPoint) but in this
+							// case, it is a one to many relationship (1 pattern => * matches).
+							const FString& DataPoint = Scope.Name;
+							TelemetryData.Add(TelemetryDefinition(TestName, Statistic.TelemetryContext, DataPoint, Statistic.TelemetryUnit, Scope.GetValue(Statistic.Statistic)));
+						}
+						Statistics.Reset();
+					}
+				}
 			}
 
 			// resolve counters to telemetry
diff --git a/Engine/Source/Runtime/Core/Private/Modules/ModuleManager.cpp b/Engine/Source/Runtime/Core/Private/Modules/ModuleManager.cpp
--- a/Engine/Source/Runtime/Core/Private/Modules/ModuleManager.cpp
+++ b/Engine/Source/Runtime/Core/Private/Modules/ModuleManager.cpp
@@ -417,8 +417,7 @@ IModuleInterface* FModuleManager::LoadModuleWithFailureReason(const FName InModu
 	}
 
 	SCOPED_BOOT_TIMING("LoadModule");
-//  TODO NDarnell we should be using trace metadata to label the module
-//	TRACE_CPUPROFILER_EVENT_SCOPE_TEXT(*FString::Printf(TEXT("LoadModule %s"), *InModuleName.ToString()));
+	TRACE_CPUPROFILER_EVENT_SCOPE_TEXT(*FString::Printf(TEXT("LoadModule_%s"), *InModuleName.ToString()));
 	DECLARE_SCOPE_CYCLE_COUNTER(TEXT("Module Load"), STAT_ModuleLoad, STATGROUP_LoadTime);
 
 #if	STATS
