Migrate editor default maps to a new user settings that allows proper overrides per-project

diff --git a/Engine/Source/Editor/NewLevelDialog/Private/NewLevelDialogModule.cpp b/Engine/Source/Editor/NewLevelDialog/Private/NewLevelDialogModule.cpp
--- a/Engine/Source/Editor/NewLevelDialog/Private/NewLevelDialogModule.cpp
+++ b/Engine/Source/Editor/NewLevelDialog/Private/NewLevelDialogModule.cpp
@@ -344,7 +344,7 @@ private:
 		// Build a list of items - one for each template
 		for (const FTemplateMapInfo& TemplateMapInfo : TemplateMapInfos)
 		{
-			if (!bShowPartitionedTemplates && ULevel::GetIsLevelPartitionedFromPackage(FName(*TemplateMapInfo.Map)))
+			if (!bShowPartitionedTemplates && ULevel::GetIsLevelPartitionedFromPackage(FName(*TemplateMapInfo.Map.ToString())))
 			{
 				continue;
 			}
@@ -355,7 +355,7 @@ private:
 			Item->Name = TemplateMapInfo.DisplayName;
 
 			UTexture2D* ThumbnailTexture = nullptr;
-			if (!TemplateMapInfo.Thumbnail.IsEmpty())
+			if (TemplateMapInfo.Thumbnail.IsValid())
 			{
 				FSoftObjectPath SoftObjectPath(TemplateMapInfo.Thumbnail);
 				ThumbnailTexture = Cast<UTexture2D>(SoftObjectPath.ResolveObject());
@@ -407,7 +407,7 @@ private:
 
 				if (Item->Name.IsEmpty())
 				{
-					Item->Name = FText::FromString(FPaths::GetBaseFilename(TemplateMapInfo.Map));
+					Item->Name = FText::FromString(FPaths::GetBaseFilename(TemplateMapInfo.Map.ToString()));
 				}
 			}
 
@@ -471,7 +471,7 @@ private:
 		const TSharedPtr<FNewLevelTemplateItem> Template = Items[0];
 		if (Template->Type == FNewLevelTemplateItem::NewLevelType::Template)
 		{
-			OutTemplateMapPackageName = Template->TemplateMapInfo.Map;
+			OutTemplateMapPackageName = Template->TemplateMapInfo.Map.ToString();
 		}
 		else if (Template->Type == FNewLevelTemplateItem::NewLevelType::EmptyWorldPartition)
 		{
diff --git a/Engine/Source/Editor/UnrealEd/Classes/Editor/UnrealEdEngine.h b/Engine/Source/Editor/UnrealEd/Classes/Editor/UnrealEdEngine.h
--- a/Engine/Source/Editor/UnrealEd/Classes/Editor/UnrealEdEngine.h
+++ b/Engine/Source/Editor/UnrealEd/Classes/Editor/UnrealEdEngine.h
@@ -90,12 +90,12 @@ struct FTemplateMapInfo
 	TSoftObjectPtr<UTexture2D> ThumbnailTexture;
 
 	/** The Texture associated with this map template */
-	UPROPERTY(config)
-	FString Thumbnail;
+	UPROPERTY(config, meta = (AllowedClasses = "Texture2D"))
+	FSoftObjectPath Thumbnail;
 
 	/** The object path to the template map */
-	UPROPERTY(config)
-	FString Map;
+	UPROPERTY(config, meta = (AllowedClasses = "World"))
+	FSoftObjectPath Map;
 
 	/** Optional display name override for this map template  */
 	UPROPERTY()
@@ -157,6 +157,7 @@ public:
 	UPROPERTY()
 	TArray<FString> SortedSpriteCategories_DEPRECATED;
 
+	UE_DEPRECATED(5.0, "This variable may no longer contain the correct template maps for the project and this property will no longer be publically accessible in the future. Use GetDefaultTemplateMapInfos to get the default list of template maps for a project")
 	/** List of info for all known template maps */
 	UPROPERTY(config)
 	TArray<FTemplateMapInfo> TemplateMapInfos;
@@ -786,6 +787,8 @@ public:
 	 */
 	bool IsTemplateMap( const FString& MapName ) const;
 
+	void RebuildTemplateMapData();
+
 	/**
 	 * Returns true if the user is currently interacting with a viewport.
 	 */
@@ -920,12 +923,11 @@ public:
 	DECLARE_DELEGATE_RetVal(const TArray<FTemplateMapInfo>&, FGetTemplateMapInfos);
 	FGetTemplateMapInfos& OnGetTemplateMapInfos() { return GetTemplateMapInfosDelegate; }
 
-	/** List template maps */
-	const TArray<FTemplateMapInfo>& GetTemplateMapInfos() const
-	{
-		return GetTemplateMapInfosDelegate.IsBound() ? GetTemplateMapInfosDelegate.Execute() : TemplateMapInfos;
-	}
+	/** Gets the canonical list of map templates that should be visible in new level picker. This function calls OnGetTemplateMapInfos to allow runtime override of the default maps */
+	const TArray<FTemplateMapInfo>& GetTemplateMapInfos() const;
 
+	/** Gets the project default map templates without any runtime overrides */
+	const TArray<FTemplateMapInfo>& GetProjectDefaultMapTemplates() const;
 protected:
 
 	/** Called when global editor selection changes */
@@ -981,6 +983,9 @@ private:
 	/* Delegate to override TemplateMapInfos */
 	FGetTemplateMapInfos GetTemplateMapInfosDelegate;
 
+	/** Transient unsaved version of template map infos used by the editor. */
+	TArray<FTemplateMapInfo> TemplateMapInfoCache;
+
 	/**
 	* Internal helper function to count how many dirty packages require checkout.
 	*
diff --git a/Engine/Source/Editor/UnrealEd/Private/UnrealEdEngine.cpp b/Engine/Source/Editor/UnrealEd/Private/UnrealEdEngine.cpp
--- a/Engine/Source/Editor/UnrealEd/Private/UnrealEdEngine.cpp
+++ b/Engine/Source/Editor/UnrealEd/Private/UnrealEdEngine.cpp
@@ -69,6 +69,7 @@
 #include "Logging/MessageLog.h"
 #include "Subsystems/EditorActorSubsystem.h"
 #include "ProfilingDebugging/StallDetector.h"
+#include "GameMapsSettings.h"
 
 DEFINE_LOG_CATEGORY_STATIC(LogUnrealEdEngine, Log, All);
 
@@ -79,6 +80,8 @@ void UUnrealEdEngine::Init(IEngineLoop* InEngineLoop)
 
 	RegisterEditorElements();
 
+	RebuildTemplateMapData();
+
 	// Display warnings to the user about disk space issues
 	ValidateFreeDiskSpace();
 
@@ -407,7 +410,8 @@ void UUnrealEdEngine::PreExit()
 	Super::PreExit();
 }
 
-
+// Disable deprecated property warnings which are accessed when destructed here
+PRAGMA_DISABLE_DEPRECATION_WARNINGS
 UUnrealEdEngine::~UUnrealEdEngine()
 {
 	if (this == GUnrealEd)
@@ -415,7 +419,7 @@ UUnrealEdEngine::~UUnrealEdEngine()
 		GUnrealEd = NULL; 
 	}
 }
-
+PRAGMA_ENABLE_DEPRECATION_WARNINGS
 
 void UUnrealEdEngine::FinishDestroy()
 {
@@ -891,6 +895,16 @@ bool UUnrealEdEngine::PreferToStreamLevelsInPIE() const
 	return GetDefault<ULevelEditorPlaySettings>()->bPreferToStreamLevelsInPIE;
 }
 
+const TArray<FTemplateMapInfo>& UUnrealEdEngine::GetTemplateMapInfos() const
+{
+	return GetTemplateMapInfosDelegate.IsBound() ? GetTemplateMapInfosDelegate.Execute() : GetProjectDefaultMapTemplates();
+}
+
+const TArray<FTemplateMapInfo>& UUnrealEdEngine::GetProjectDefaultMapTemplates() const
+{
+	return TemplateMapInfoCache;
+}
+
 void UUnrealEdEngine::OnHISMTreeBuilt(UHierarchicalInstancedStaticMeshComponent* Component, bool bWasAsyncBuild)
 {
 	if (bWasAsyncBuild)
@@ -930,6 +944,31 @@ void UUnrealEdEngine::TakeHighResScreenShots()
 }
 
 
+void UUnrealEdEngine::RebuildTemplateMapData()
+{
+	const TArray<FTemplateMapInfoOverride> TemplateMapOverrides = GetDefault<UGameMapsSettings>()->EditorTemplateMapOverrides;
+
+	if (TemplateMapOverrides.Num() > 0)
+	{
+		TemplateMapInfoCache.Reset(TemplateMapOverrides.Num());
+
+		for (const FTemplateMapInfoOverride& Override : TemplateMapOverrides)
+		{
+			FTemplateMapInfo NewInfo;
+			NewInfo.Thumbnail = Override.Thumbnail;
+			NewInfo.Map = Override.Map;
+			NewInfo.DisplayName = Override.DisplayName;
+			TemplateMapInfoCache.Add(MoveTemp(NewInfo));
+		}
+	}
+	else
+	{
+		PRAGMA_DISABLE_DEPRECATION_WARNINGS
+		TemplateMapInfoCache = TemplateMapInfos;
+		PRAGMA_ENABLE_DEPRECATION_WARNINGS
+	}
+}
+
 void UUnrealEdEngine::SetCurrentClass( UClass* InClass )
 {
 	USelection* SelectionSet = GetSelectedObjects();
diff --git a/Engine/Source/Runtime/EngineSettings/Classes/GameMapsSettings.h b/Engine/Source/Runtime/EngineSettings/Classes/GameMapsSettings.h
--- a/Engine/Source/Runtime/EngineSettings/Classes/GameMapsSettings.h
+++ b/Engine/Source/Runtime/EngineSettings/Classes/GameMapsSettings.h
@@ -66,6 +66,31 @@ enum class ESubLevelStripMode : uint8
 	IsChildOf
 };
 
+
+/** Used by new level dialog. */
+USTRUCT()
+struct FTemplateMapInfoOverride
+{
+	GENERATED_USTRUCT_BODY()
+
+	/** The thumbnail to display in the new level dialog */
+	UPROPERTY(config, EditAnywhere, Category = DefaultMaps, meta = (AllowedClasses = "Texture2D"))
+	FSoftObjectPath Thumbnail;
+
+	/** The path to the template map */
+	UPROPERTY(config, EditAnywhere, Category = DefaultMaps, meta = (AllowedClasses = "World"))
+	FSoftObjectPath Map;
+
+	/** The display name of the map template in the level dialog. If this is empty the thumbnail name will be used */
+	UPROPERTY(config, EditAnywhere, Category = DefaultMaps)
+	FText DisplayName;
+
+	FTemplateMapInfoOverride()
+	{
+	}
+};
+
+
 UCLASS(config=Engine, defaultconfig)
 class ENGINESETTINGS_API UGameMapsSettings
 	: public UObject
@@ -125,6 +150,10 @@ public:
 	/** If set, this map will be loaded when the Editor starts up. */
 	UPROPERTY(config, EditAnywhere, Category=DefaultMaps, meta=(AllowedClasses="World"))
 	FSoftObjectPath EditorStartupMap;
+
+	/** Map templates that should show up in the new level dialog. This will completely override the default maps chosen by the default editor */
+	UPROPERTY(config, EditAnywhere, Category=DefaultMaps, meta=(ConfigRestartRequired=true))
+	TArray<FTemplateMapInfoOverride> EditorTemplateMapOverrides;
 #endif
 
 	/** The default options that will be appended to a map being loaded. */
