ExternalActorsCommandlet
- Added support for repairing duplicated actor guids.
- Added suport for deleting non-actor files.
- Tweaked error logs.

diff --git a/Engine/Source/Editor/UnrealEd/Private/Commandlets/ExternalActorsCommandlet.cpp b/Engine/Source/Editor/UnrealEd/Private/Commandlets/ExternalActorsCommandlet.cpp
--- a/Engine/Source/Editor/UnrealEd/Private/Commandlets/ExternalActorsCommandlet.cpp
+++ b/Engine/Source/Editor/UnrealEd/Private/Commandlets/ExternalActorsCommandlet.cpp
@@ -20,8 +20,6 @@ UExternalActorsCommandlet::UExternalActorsCommandlet(const FObjectInitializer& O
 
 UWorld* UExternalActorsCommandlet::LoadWorld(const FString& LevelToLoad)
 {
-	TRACE_CPUPROFILER_EVENT_SCOPE(UWorldPartitionConvertCommandlet::LoadWorld);
-
 	SET_WARN_COLOR(COLOR_WHITE);
 	UE_LOG(LogExternalActorsCommandlet, Log, TEXT("Loading level %s."), *LevelToLoad);
 	CLEAR_WARN_COLOR();
@@ -68,10 +66,14 @@ int32 UExternalActorsCommandlet::Main(const FString& Params)
 	FString ExternalActorsPath = ULevel::GetExternalActorsPath(Tokens[0]);
 	FString ExternalActorsFilePath = FPackageName::LongPackageNameToFilename(ExternalActorsPath);
 
-	TArray<FString> PackagesToDelete;
+	// Look for duplicated actor GUIDs
+	TMap<FGuid, AActor*> ActorGuids;
+
+	TSet<UPackage*> PackagesToSave;
+	TArray<FString> PackagesToDelete;	
 	if (IFileManager::Get().DirectoryExists(*ExternalActorsFilePath))
 	{
-		bool bResult = IFileManager::Get().IterateDirectoryRecursively(*ExternalActorsFilePath, [this, bRepair, &PackagesToDelete](const TCHAR* FilenameOrDirectory, bool bIsDirectory)
+		bool bResult = IFileManager::Get().IterateDirectoryRecursively(*ExternalActorsFilePath, [this, bRepair, &PackagesToSave, &PackagesToDelete, &ActorGuids](const TCHAR* FilenameOrDirectory, bool bIsDirectory)
 		{
 			if (!bIsDirectory)
 			{
@@ -107,7 +109,7 @@ int32 UExternalActorsCommandlet::Main(const FString& Params)
 
 					if (!MainPackageActor)
 					{
-						UE_LOG(LogExternalActorsCommandlet, Error, TEXT("Invalid actor file '%s'"), *Filename);
+						UE_LOG(LogExternalActorsCommandlet, Error, TEXT("Missing main actor for file '%s'"), *Filename);
 
 						if (bRepair)
 						{
@@ -117,25 +119,64 @@ int32 UExternalActorsCommandlet::Main(const FString& Params)
 								PotentialMainPackageActor->SetPackageExternal(true);
 								
 								UPackage* PackageToSave = PotentialMainPackageActor->GetPackage();
+								PackagesToSave.Add(PackageToSave);
 
-								FString PackageFileName = SourceControlHelpers::PackageFilename(PackageToSave);
-								FSavePackageArgs SaveArgs;
-								SaveArgs.TopLevelFlags = RF_Standalone;
-								if (UPackage::SavePackage(PackageToSave, nullptr, *PackageFileName, SaveArgs))
-								{
-									PackageHelper.AddToSourceControl(PackageToSave);
-								}
+								MainPackageActor = PotentialMainPackageActor;
 							}
 
 							PackagesToDelete.Add(Filename);
 						}
 					}
+
+					if (MainPackageActor)
+					{
+						if (ActorGuids.Contains(MainPackageActor->GetActorGuid()))
+						{
+							UE_LOG(LogExternalActorsCommandlet, Error, TEXT("Duplicated actor guid for file '%s'"), *Filename);
+
+							if (bRepair)
+							{
+								FSetActorGuid SetActorGuid(MainPackageActor, FGuid::NewGuid());
+
+								UPackage* PackageToSave = MainPackageActor->GetPackage();
+								PackagesToSave.Add(PackageToSave);
+							}
+						}
+						else
+						{
+							ActorGuids.Add(MainPackageActor->GetActorGuid(), MainPackageActor);
+						}
+					}
+				}
+				else
+				{
+					UE_LOG(LogExternalActorsCommandlet, Error, TEXT("Invalid actor file '%s'"), *Filename);
+
+					if (bRepair)
+					{
+						PackagesToDelete.Add(Filename);
+					}
 				}
 			}
 			return true;
 		});
 	}
 
+	for (UPackage* PackageToSave: PackagesToSave)
+	{
+		FSavePackageArgs SaveArgs;
+		SaveArgs.TopLevelFlags = RF_Standalone;
+		FString PackageFileName = SourceControlHelpers::PackageFilename(PackageToSave);
+
+		if (PackageHelper.Checkout(PackageToSave))
+		{
+			if (UPackage::SavePackage(PackageToSave, nullptr, *PackageFileName, SaveArgs))
+			{
+				PackageHelper.AddToSourceControl(PackageToSave);
+			}
+		}
+	}
+
 	CollectGarbage(RF_NoFlags);
 
 	for (const FString& PackageToDelete : PackagesToDelete)
diff --git a/Engine/Source/Runtime/Engine/Classes/GameFramework/Actor.h b/Engine/Source/Runtime/Engine/Classes/GameFramework/Actor.h
--- a/Engine/Source/Runtime/Engine/Classes/GameFramework/Actor.h
+++ b/Engine/Source/Runtime/Engine/Classes/GameFramework/Actor.h
@@ -3910,6 +3910,7 @@ private:
 	}
 	friend class UWorld;
 	friend class UEngine;
+	friend class UExternalActorsCommandlet;
 	friend class UWorldPartitionConvertCommandlet;
 };
 #endif
