Add an optoin to open all errored blueprints when PIE'ing. Close the modal window upon opening a BP asset from the error window.

#rnx

diff --git a/Engine/Source/Editor/UnrealEd/Private/PlayLevel.cpp b/Engine/Source/Editor/UnrealEd/Private/PlayLevel.cpp
--- a/Engine/Source/Editor/UnrealEd/Private/PlayLevel.cpp
+++ b/Engine/Source/Editor/UnrealEd/Private/PlayLevel.cpp
@@ -1442,19 +1442,33 @@ static bool ShowBlueprintErrorDialog( TArray<UBlueprint*> ErroredBlueprints )
 
 			if (InDialog.IsValid())
 			{
-				// Opening the blueprint editor above may end up creating an invisible new window on top of the dialog, 
-				// thus making it not interactable, so we have to force the dialog back to the front
-				InDialog->BringToFront(true);
+				InDialog->RequestDestroyWindow();
+			}
+		}
+
+		static void OnOpenAllLinkClicked(const TArray<UBlueprint*>& BlueprintsToOpen, TSharedPtr<SCustomDialog> InDialog)
+		{
+			for(UBlueprint* BP : BlueprintsToOpen)
+			{
+				if (BP)
+				{
+					GEditor->EditObject(BP);
+				}	
+			}
+
+			if (InDialog.IsValid())
+			{
+				InDialog->RequestDestroyWindow();
 			}
 		}
 	};
 
 	TSharedRef<SVerticalBox> DialogContents = SNew(SVerticalBox)
 		+ SVerticalBox::Slot()
-		.Padding(0, 0, 0, 16)
+		.Padding(0.f, 0.f, 0.f, 16.f)
 		[
 			SNew(STextBlock)
-			.Text(NSLOCTEXT("PlayInEditor", "PrePIE_BlueprintErrors", "One or more blueprints has an unresolved compiler error, are you sure you want to Play in Editor?"))
+			.Text(NSLOCTEXT("PlayInEditor", "PrePIE_BlueprintErrors", "Are you sure you want to Play in Editor? The following blueprints have unresolved compiler errors."))
 		];
 
 	TSharedPtr<SCustomDialog> CustomDialog;
@@ -1475,17 +1489,25 @@ static bool ShowBlueprintErrorDialog( TArray<UBlueprint*> ErroredBlueprints )
 			];
 	}
 
-	DialogContents->AddSlot()
-		.Padding(0, 16, 0, 0)
-		[
-			SNew(STextBlock)
-			.Text(NSLOCTEXT("PlayInEditor", "PrePIE_BlueprintErrorsDelayedOpen", "Clicked blueprints will open once this dialog is closed."))
-		];
+	// Add an option to open all errored blueprints
+	if(ErroredBlueprints.Num() > 1)
+	{
+		DialogContents->AddSlot()
+			.Padding(0.f, 16.f, 0.f, 0.f)
+			.HAlign(HAlign_Left)
+			[
+				SNew(SHyperlink)
+				.Style(FEditorStyle::Get(), "Common.GotoBlueprintHyperlink")
+				.OnNavigate(FSimpleDelegate::CreateLambda([&ErroredBlueprints, &CustomDialog]() { Local::OnOpenAllLinkClicked(ErroredBlueprints, CustomDialog); }))
+				.Text(NSLOCTEXT("SourceHyperlink", "EditAllErroredBlueprints", "Open all errored blueprints"))
+				.ToolTipText(NSLOCTEXT("SourceHyperlink", "EditAllErroredBlueprints_ToolTip", "Opens all the errored blueprint in the editor"))
+			];
+	}
 
-	FText DialogTitle = NSLOCTEXT("PlayInEditor", "PrePIE_BlueprintErrorsTitle", "Blueprint Compilation Errors");
+	static const FText DialogTitle = NSLOCTEXT("PlayInEditor", "PrePIE_BlueprintErrorsTitle", "Blueprint Compilation Errors");
 
-	FText OKText = NSLOCTEXT("PlayInEditor", "PrePIE_OkText", "Play in Editor");
-	FText CancelText = NSLOCTEXT("Dialogs", "EAppReturnTypeCancel", "Cancel");
+	static const FText OKText = NSLOCTEXT("PlayInEditor", "PrePIE_OkText", "Play in Editor");
+	static const FText CancelText = NSLOCTEXT("Dialogs", "EAppReturnTypeCancel", "Cancel");
 
 	CustomDialog = SNew(SCustomDialog)
 		.Title(DialogTitle)
@@ -1493,7 +1515,7 @@ static bool ShowBlueprintErrorDialog( TArray<UBlueprint*> ErroredBlueprints )
 		.DialogContent(DialogContents)
 		.Buttons( { SCustomDialog::FButton(OKText), SCustomDialog::FButton(CancelText) } );
 
-	int32 ButtonPressed = CustomDialog->ShowModal();
+	const int32 ButtonPressed = CustomDialog->ShowModal();
 	return ButtonPressed == 0;
 }
 
