Fix crash on FBX static mesh reimport related to spline mesh component collision computation.

diff --git a/Engine/Source/Editor/UnrealEd/Private/Fbx/FbxStaticMeshImport.cpp b/Engine/Source/Editor/UnrealEd/Private/Fbx/FbxStaticMeshImport.cpp
--- a/Engine/Source/Editor/UnrealEd/Private/Fbx/FbxStaticMeshImport.cpp
+++ b/Engine/Source/Editor/UnrealEd/Private/Fbx/FbxStaticMeshImport.cpp
@@ -2370,6 +2370,11 @@ void UnFbx::FFbxImporter::PostImportStaticMesh(UStaticMesh* StaticMesh, TArray<F
 		GenerateKDopAsSimpleCollision(StaticMesh, Dirs);
 	}
 
+	// Refresh collision change back to staticmesh components, this must be done after the build.
+	// Some derived components rely on data computed in the build for collision computation. For example,
+	// USplineMeshComponent uses the static mesh's extended bounds in its convex collision computation.
+	RefreshCollisionChangeComponentsOnly(*StaticMesh);
+
 	//If there is less the 2 materials in the fbx file there is no need to reorder them
 	//If we have import a LOD other then the base, the material array cannot be sorted, because only the base LOD reorder the material array
 	if (LODIndex == 0 && StaticMesh->GetStaticMaterials().Num() > 1)
@@ -2688,6 +2693,8 @@ bool UnFbx::FFbxImporter::FillCollisionModelList(FbxNode* Node)
 
 bool UnFbx::FFbxImporter::ImportCollisionModels(UStaticMesh* StaticMesh, const FbxString& InNodeName)
 {
+	TRACE_CPUPROFILER_EVENT_SCOPE(FFbxImporter::ImportCollisionModels);
+
 	// find collision models
 	bool bRemoveEmptyKey = false;
 	FbxString EmptyKey;
@@ -2936,9 +2943,13 @@ bool UnFbx::FFbxImporter::ImportCollisionModels(UStaticMesh* StaticMesh, const F
 	// Create new GUID
 	StaticMesh->GetBodySetup()->InvalidatePhysicsData();
 
-	// refresh collision change back to staticmesh components
-	RefreshCollisionChange(*StaticMesh);
-		
+	// Update collision on the static mesh
+	// Note: do not refresh collision on components here, this must occur after the static mesh build.
+	{
+		TRACE_CPUPROFILER_EVENT_SCOPE(CreateStaticMeshNavCollision);
+		StaticMesh->CreateNavCollision(true);
+	}
+
 	return bAtLeastOneCollisionMeshImported;
 }
 #undef LOCTEXT_NAMESPACE
diff --git a/Engine/Source/Editor/UnrealEd/Private/GeomFitUtils.cpp b/Engine/Source/Editor/UnrealEd/Private/GeomFitUtils.cpp
--- a/Engine/Source/Editor/UnrealEd/Private/GeomFitUtils.cpp
+++ b/Engine/Source/Editor/UnrealEd/Private/GeomFitUtils.cpp
@@ -623,6 +623,26 @@ void RefreshCollisionChanges(const TArray<UStaticMesh*>& StaticMeshes)
 	FEditorSupportDelegates::RedrawAllViewports.Broadcast();
 }
 
+void RefreshCollisionChangeComponentsOnly(UStaticMesh& StaticMesh)
+{
+	TRACE_CPUPROFILER_EVENT_SCOPE(RefreshCollisionChangeComponentsOnly)
+
+	for (FThreadSafeObjectIterator Iter(UStaticMeshComponent::StaticClass()); Iter; ++Iter)
+	{
+		UStaticMeshComponent* StaticMeshComponent = Cast<UStaticMeshComponent>(*Iter);
+		if (StaticMeshComponent->GetStaticMesh() == &StaticMesh)
+		{
+			// it needs to recreate IF it already has been created
+			if (StaticMeshComponent->IsPhysicsStateCreated())
+			{
+				StaticMeshComponent->RecreatePhysicsState();
+			}
+		}
+	}
+
+	FEditorSupportDelegates::RedrawAllViewports.Broadcast();
+}
+
 /* *************************** DEPRECATED ******************************** */
 void RefreshCollisionChange(const UStaticMesh* StaticMesh)
 {
diff --git a/Engine/Source/Editor/UnrealEd/Private/GeomFitUtils.h b/Engine/Source/Editor/UnrealEd/Private/GeomFitUtils.h
--- a/Engine/Source/Editor/UnrealEd/Private/GeomFitUtils.h
+++ b/Engine/Source/Editor/UnrealEd/Private/GeomFitUtils.h
@@ -134,6 +134,17 @@ UNREALED_API void RefreshCollisionChange(UStaticMesh& StaticMesh);
  */
 UNREALED_API void RefreshCollisionChanges(const TArray<UStaticMesh*>& StaticMeshes);
 
+/**
+ * Refresh Collision Change Components Only
+ *
+ * Collision has been changed, so it will need to recreate physics state to reflect it
+ * Utilities functions to propagate BodySetup change for StaticMesh to its component
+ * Use this method if static mesh nav collision has already been updated elsewhere.
+ *
+ * @param	StaticMesh	StaticMesh that collision has been changed for
+ */
+UNREALED_API void RefreshCollisionChangeComponentsOnly(UStaticMesh& StaticMesh);
+
 UE_DEPRECATED(4.15, "This version of RefreshCollisionChange is deprecated. Please use RefreshCollisionChange(UStaticMesh&) instead.")
 UNREALED_API void RefreshCollisionChange(const UStaticMesh* StaticMesh);
 	
