Fix crash with -diffonly: during GetPreviousCookedBytes Record.Packages has not yet been populated by WritePackageData, and we need to use the function argument instead.
#rnx

diff --git a/Engine/Source/Editor/UnrealEd/Private/Cooker/LooseCookedPackageWriter.cpp b/Engine/Source/Editor/UnrealEd/Private/Cooker/LooseCookedPackageWriter.cpp
--- a/Engine/Source/Editor/UnrealEd/Private/Cooker/LooseCookedPackageWriter.cpp
+++ b/Engine/Source/Editor/UnrealEd/Private/Cooker/LooseCookedPackageWriter.cpp
@@ -355,10 +355,8 @@ void FLooseCookedPackageWriter::UpdateManifest(FRecord& Record)
 
 bool FLooseCookedPackageWriter::GetPreviousCookedBytes(const FPackageInfo& Info, FPreviousCookedBytesData& OutData)
 {
-	FPackageWriterRecords::FPackage& BaseRecord = Records.FindRecordChecked(Info.InputPackageName);
-	FRecord& Record = static_cast<FRecord&>(BaseRecord);
 	FArchiveStackTrace::FPackageData ExistingPackageData;
-	FArchiveStackTrace::LoadPackageIntoMemory(*Record.Packages[Info.MultiOutputIndex].Info.LooseFilePath, ExistingPackageData, OutData.Data);
+	FArchiveStackTrace::LoadPackageIntoMemory(*Info.LooseFilePath, ExistingPackageData, OutData.Data);
 	OutData.Size = ExistingPackageData.Size;
 	OutData.HeaderSize = ExistingPackageData.HeaderSize;
 	OutData.StartOffset = ExistingPackageData.StartOffset;
@@ -707,4 +705,4 @@ EPackageExtension FLooseCookedPackageWriter::BulkDataTypeToExtension(FBulkDataIn
 		checkNoEntry();
 		return EPackageExtension::Unspecified;
 	}
-}
\ No newline at end of file
+}
