DiffPackageWriter: Fix missing reset of bHasStartedSecondSave during each package; this was blocking the logging pass from running on every package after the first cooked package.
#rnx

diff --git a/Engine/Source/Editor/UnrealEd/Private/Cooker/DiffPackageWriter.cpp b/Engine/Source/Editor/UnrealEd/Private/Cooker/DiffPackageWriter.cpp
--- a/Engine/Source/Editor/UnrealEd/Private/Cooker/DiffPackageWriter.cpp
+++ b/Engine/Source/Editor/UnrealEd/Private/Cooker/DiffPackageWriter.cpp
@@ -115,6 +115,7 @@ void FDiffPackageWriter::BeginPackage(const FBeginPackageInfo& Info)
 {
 	bIsDifferent = false;
 	bDiffCallstack = false;
+	bHasStartedSecondSave = false;
 	DiffMap.Reset();
 
 	BeginInfo = Info;
@@ -310,6 +311,7 @@ FLinkerDiffPackageWriter::FLinkerDiffPackageWriter(TUniquePtr<ICookedPackageWrit
 
 void FLinkerDiffPackageWriter::BeginPackage(const FBeginPackageInfo& Info)
 {
+	bHasStartedSecondSave = false;
 	BeginInfo = Info;
 	Inner->BeginPackage(Info);
 	SetupOtherAlgorithm();
