Allow consolidating different subclasses of UPhysicalMaterial together

diff --git a/Engine/Source/Editor/UnrealEd/Private/ConsolidateWindow.cpp b/Engine/Source/Editor/UnrealEd/Private/ConsolidateWindow.cpp
--- a/Engine/Source/Editor/UnrealEd/Private/ConsolidateWindow.cpp
+++ b/Engine/Source/Editor/UnrealEd/Private/ConsolidateWindow.cpp
@@ -21,6 +21,7 @@
 #include "Widgets/Input/SCheckBox.h"
 #include "EditorStyleSet.h"
 #include "Materials/MaterialInterface.h"
+#include "PhysicalMaterials/PhysicalMaterial.h"
 #include "ISourceControlModule.h"
 #include "Engine/Texture.h"
 #include "AssetData.h"
@@ -614,6 +615,7 @@ bool SConsolidateToolWidget::DetermineAssetCompatibility( const TArray<UObject*>
 				// If the proposed object doesn't share a common class or a common base that is allowed as an exception, it is not a compatible object
 				if ( !( NearestCommonBase->IsChildOf( UTexture::StaticClass() ) )  &&
 					 !( NearestCommonBase->IsChildOf( UMaterialInterface::StaticClass() ) ) &&
+					 !( NearestCommonBase->IsChildOf( UPhysicalMaterial::StaticClass() ) ) &&
 					 !( NearestCommonBase->IsChildOf( UDataTable::StaticClass() ) ) )
 				{
 					bAllAssetsValid = false;
