UnrealEd: Fix error when converting brush actors to static mesh actors. Partially addresses jira UE-137902

#rnx

diff --git a/Engine/Source/Editor/UnrealEd/Private/EditorEngine.cpp b/Engine/Source/Editor/UnrealEd/Private/EditorEngine.cpp
--- a/Engine/Source/Editor/UnrealEd/Private/EditorEngine.cpp
+++ b/Engine/Source/Editor/UnrealEd/Private/EditorEngine.cpp
@@ -5670,7 +5670,7 @@ namespace ConvertHelpers
 		for( int32 ActorIdx = 0; ActorIdx < InActorsToConvert.Num(); ++ActorIdx )
 		{
 			AActor* ActorToConvert = InActorsToConvert[ActorIdx];
-			if (ActorToConvert->GetClass()->IsChildOf(ABrush::StaticClass()) && InConvertToClass == AStaticMeshActor::StaticClass())
+			if (IsValidChecked(ActorToConvert) && ActorToConvert->GetClass()->IsChildOf(ABrush::StaticClass()) && InConvertToClass == AStaticMeshActor::StaticClass())
 			{
 				GEditor->SelectActor(ActorToConvert, true, true);
 				OutBrushList.Add(Cast<ABrush>(ActorToConvert));
@@ -5768,20 +5768,8 @@ void UEditorEngine::DoConvertActors( const TArray<AActor*>& ActorsToConvert, UCl
 		int32 BrushIndexForReattachment = 0;
 
 		FVector CachePivotLocation = GetPivotLocation();
-		for( int32 ActorIdx = 0; ActorIdx < ActorsToConvert.Num(); ++ActorIdx )
-		{
-			AActor* ActorToConvert = ActorsToConvert[ActorIdx];
-			if (IsValidChecked(ActorToConvert) && ActorToConvert->GetClass()->IsChildOf(ABrush::StaticClass()) && ConvertToClass == AStaticMeshActor::StaticClass())
-			{
-				SelectActor(ActorToConvert, true, true);
-				BrushList.Add(Cast<ABrush>(ActorToConvert));
-
-				// If this is a single brush conversion then this index will be used for re-attachment.
-				BrushIndexForReattachment = ActorIdx;
-			}
-		}
+		ConvertHelpers::GetBrushList(ActorsToConvert, ConvertToClass, BrushList, BrushIndexForReattachment);
 
-		// If no package name is supplied, ask the user
 		if( BrushList.Num() )
 		{
 			AActor* ConvertedBrushActor = ConvertBrushesToStaticMesh(InStaticMeshPackageName, BrushList, CachePivotLocation);
@@ -5793,7 +5781,6 @@ void UEditorEngine::DoConvertActors( const TArray<AActor*>& ActorsToConvert, UCl
 			{
 				ReattachActorsHelper::CacheActorConvert(BrushList[0], ConvertedBrushActor, ConvertedMap, AttachmentInfo[BrushIndexForReattachment]);
 			}
-
 		}
 
 		ULayersSubsystem* LayersSubsystem = GetEditorSubsystem<ULayersSubsystem>();
@@ -5801,6 +5788,14 @@ void UEditorEngine::DoConvertActors( const TArray<AActor*>& ActorsToConvert, UCl
 		{
 			AActor* ActorToConvert = ActorsToConvert[ ActorIdx ];
 
+			if (ActorToConvert->GetClass()->IsChildOf(ABrush::StaticClass()) && ConvertToClass == AStaticMeshActor::StaticClass())
+			{
+				// We already converted this actor in ConvertBrushesToStaticMesh above, and it has been marked as pending
+				// kill (and hence is invalid) TODO: It would be good to refactor this function so there is a single place
+				// where conversion happens
+				ensure(!IsValid(ActorToConvert));
+				continue;
+			}
 
 			if (!IsValidChecked(ActorToConvert))
 			{
@@ -5810,13 +5805,10 @@ void UEditorEngine::DoConvertActors( const TArray<AActor*>& ActorsToConvert, UCl
 
 			// Source actor display label
 			FString ActorLabel = ActorToConvert->GetActorLabel();
-			// Low level source actor object name
-			FName ActorObjectName = ActorToConvert->GetFName();
 	
 			// The class of the actor we are about to replace
 			UClass* ClassToReplace = ActorToConvert->GetClass();
 
-			// Spawn the new actor
 			AActor* NewActor = NULL;
 
 			ABrush* Brush = Cast< ABrush >( ActorToConvert );
