Added a log line to get more info of a crash during a force delete of assets.

diff --git a/Engine/Source/Editor/UnrealEd/Private/ObjectTools.cpp b/Engine/Source/Editor/UnrealEd/Private/ObjectTools.cpp
--- a/Engine/Source/Editor/UnrealEd/Private/ObjectTools.cpp
+++ b/Engine/Source/Editor/UnrealEd/Private/ObjectTools.cpp
@@ -2886,7 +2886,8 @@ namespace ObjectTools
 			for(int32 I = 0; I < ShownObjectsToDelete.Num() && I < MAX_PACKAGES_TO_LOG; ++I)
 			{
 				Msg.Append(TEXT("\n"));
-				Msg.Append(FString::Printf(TEXT("\tAsset Name: %s"), *GetPathNameSafe(ShownObjectsToDelete[I])));
+				Msg.Append(FString::Printf(TEXT("\tAsset Name: %s\n"), *GetPathNameSafe(ShownObjectsToDelete[I])));
+				Msg.Append(FString::Printf(TEXT("\tAsset Type: %s"), *(ShownObjectsToDelete[I]->GetClass()->GetName())));
 			}
 			UE_LOG(LogUObjectGlobals, Log, TEXT("%s"), *Msg);
 		}
