Adds DataLayer exclusion to MiniMap generation.
Adds default MiniMap when creating WorldPartition.
Prevents user from removing MiniMap or assigning it to a DataLayer

#changelist validated

diff --git a/Engine/Source/Editor/UnrealEd/Private/WorldPartition/WorldPartitionBuilder.cpp b/Engine/Source/Editor/UnrealEd/Private/WorldPartition/WorldPartitionBuilder.cpp
--- a/Engine/Source/Editor/UnrealEd/Private/WorldPartition/WorldPartitionBuilder.cpp
+++ b/Engine/Source/Editor/UnrealEd/Private/WorldPartition/WorldPartitionBuilder.cpp
@@ -138,10 +138,12 @@ bool UWorldPartitionBuilder::Run(UWorld* World, FPackageSourceControlHelper& Pac
 		bool bUpdateEditorCells = false;
 		WorldDataLayers->ForEachDataLayer([&bUpdateEditorCells, this](UDataLayer* DataLayer)
 		{
-			// Load all Non DynamicallyLoaded Data Layers + Initially Active Data Layers + Data Layers provided by builder
-			const bool bLoadedInEditor = (bLoadNonDynamicDataLayers && !DataLayer->IsRuntime()) ||
+			const FName DataLayerLabel = DataLayer->GetDataLayerLabel();
+
+			// Load all Non Excluded Data Layers + Non DynamicallyLoaded Data Layers + Initially Active Data Layers + Data Layers provided by builder
+			const bool bLoadedInEditor = !ExcludedDataLayerLabels.Contains(DataLayerLabel) && ((bLoadNonDynamicDataLayers && !DataLayer->IsRuntime()) ||
 										 (bLoadInitiallyActiveDataLayers && DataLayer->GetInitialRuntimeState() == EDataLayerRuntimeState::Activated) ||
-										 DataLayerLabels.Contains(DataLayer->GetDataLayerLabel());
+										 DataLayerLabels.Contains(DataLayerLabel));
 			if (DataLayer->IsLoadedInEditor() != bLoadedInEditor)
 			{
 				bUpdateEditorCells = true;
diff --git a/Engine/Source/Editor/UnrealEd/Private/WorldPartition/WorldPartitionMiniMapBuilder.cpp b/Engine/Source/Editor/UnrealEd/Private/WorldPartition/WorldPartitionMiniMapBuilder.cpp
--- a/Engine/Source/Editor/UnrealEd/Private/WorldPartition/WorldPartitionMiniMapBuilder.cpp
+++ b/Engine/Source/Editor/UnrealEd/Private/WorldPartition/WorldPartitionMiniMapBuilder.cpp
@@ -19,6 +19,9 @@
 #include "SourceControlOperations.h"
 
 #include "WorldPartition/ActorDescContainer.h"
+#include "WorldPartition/DataLayer/ActorDataLayer.h"
+#include "WorldPartition/DataLayer/DataLayer.h"
+#include "WorldPartition/DataLayer/WorldDataLayers.h"
 #include "WorldPartition/HLOD/HLODActor.h"
 #include "WorldPartition/HLOD/HLODActorDesc.h"
 #include "WorldPartition/WorldPartition.h"
@@ -52,6 +55,24 @@ bool UWorldPartitionMiniMapBuilder::PreRun(UWorld* World, FPackageSourceControlH
 	WorldMiniMap->MiniMapTexture = nullptr;
 	MiniMapTiles.Empty();
 
+	AWorldDataLayers* WorldDataLayers = World->GetWorldDataLayers();
+
+	if (WorldDataLayers == nullptr)
+	{
+		UE_LOG(LogWorldPartitionMiniMapBuilder, Error, TEXT("Failed to retrieve WorldDataLayers."));
+		return false;
+	}
+
+	for (const FActorDataLayer& ActorDataLayer : WorldMiniMap->ExcludedDataLayers)
+	{
+		const UDataLayer* DataLayer = WorldDataLayers->GetDataLayerFromName(ActorDataLayer.Name);
+
+		if (DataLayer != nullptr)
+		{
+			ExcludedDataLayerLabels.Add(DataLayer->GetDataLayerLabel());
+		}
+	}
+
 	return true;
 }
 
diff --git a/Engine/Source/Editor/UnrealEd/Public/WorldPartition/WorldPartitionBuilder.h b/Engine/Source/Editor/UnrealEd/Public/WorldPartition/WorldPartitionBuilder.h
--- a/Engine/Source/Editor/UnrealEd/Public/WorldPartition/WorldPartitionBuilder.h
+++ b/Engine/Source/Editor/UnrealEd/Public/WorldPartition/WorldPartitionBuilder.h
@@ -81,6 +81,7 @@ protected:
 	int32 IterativeCellSize = 102400;
 	int32 IterativeCellOverlapSize = 0;
 	TSet<FName> DataLayerLabels;
+	TSet<FName> ExcludedDataLayerLabels;
 	bool bLoadNonDynamicDataLayers = true;
 	bool bLoadInitiallyActiveDataLayers = true;
 
diff --git a/Engine/Source/Runtime/Engine/Private/WorldPartition/WorldPartition.cpp b/Engine/Source/Runtime/Engine/Private/WorldPartition/WorldPartition.cpp
--- a/Engine/Source/Runtime/Engine/Private/WorldPartition/WorldPartition.cpp
+++ b/Engine/Source/Runtime/Engine/Private/WorldPartition/WorldPartition.cpp
@@ -38,6 +38,8 @@
 #include "WorldPartition/WorldPartitionEditorHash.h"
 #include "WorldPartition/WorldPartitionRuntimeHash.h"
 #include "WorldPartition/WorldPartitionEditorPerProjectUserSettings.h"
+#include "WorldPartition/WorldPartitionMiniMap.h"
+#include "WorldPartition/WorldPartitionMiniMapHelper.h"
 #include "WorldPartition/DataLayer/WorldDataLayers.h"
 #include "WorldPartition/WorldPartitionActorDescViewProxy.h"
 #include "Modules/ModuleManager.h"
@@ -699,6 +701,8 @@ UWorldPartition* UWorldPartition::CreateOrRepairWorldPartition(AWorldSettings* W
 		{
 			WorldDataLayers = AWorldDataLayers::Create(World);
 		}
+
+		FWorldPartitionMiniMapHelper::GetWorldPartitionMiniMap(World, true);
 	}
 
 	if (!WorldPartition->EditorHash)
diff --git a/Engine/Source/Runtime/Engine/Public/WorldPartition/WorldPartitionMiniMap.h b/Engine/Source/Runtime/Engine/Public/WorldPartition/WorldPartitionMiniMap.h
--- a/Engine/Source/Runtime/Engine/Public/WorldPartition/WorldPartitionMiniMap.h
+++ b/Engine/Source/Runtime/Engine/Public/WorldPartition/WorldPartitionMiniMap.h
@@ -3,6 +3,7 @@
 
 #include "CoreMinimal.h"
 #include "GameFramework/Info.h"
+#include "WorldPartition/DataLayer/ActorDataLayer.h"
 #include "WorldPartitionMiniMap.generated.h"
 
 /** Coordinates of the tile in Minimap texture */
@@ -50,6 +51,8 @@ public:
 
 #if WITH_EDITOR
 	virtual void PostEditChangeProperty(FPropertyChangedEvent& PropertyChangedEvent) override;
+	virtual bool IsUserManaged() const final { return false; }
+	virtual bool SupportsDataLayer() const final { return false; }
 #endif
 
 	/*WorldBounds for MinMapTexture*/
@@ -64,6 +67,10 @@ public:
 	UPROPERTY(VisibleAnywhere, Category = WorldPartitionMiniMap)
 	TObjectPtr<UTexture2D> MiniMapTexture;
 
+	/*Datalayers excluded from MiniMap rendering*/
+	UPROPERTY(EditAnywhere, Category = WorldPartitionMiniMap)
+	TSet<FActorDataLayer> ExcludedDataLayers;
+
 	/*MiniMap Tile Size*/
 	UPROPERTY(EditAnywhere, Category = WorldPartitionMiniMap, meta = (UIMin = "256", UIMax = "8192"))
 	int32 MiniMapTileSize = 1024;
