FExtender::Combine() now returns a TSharedRef rather than a TSharedPtr that is never null.

diff --git a/Engine/Source/Editor/UnrealEd/Private/Toolkits/SStandaloneAssetEditorToolkitHost.cpp b/Engine/Source/Editor/UnrealEd/Private/Toolkits/SStandaloneAssetEditorToolkitHost.cpp
--- a/Engine/Source/Editor/UnrealEd/Private/Toolkits/SStandaloneAssetEditorToolkitHost.cpp
+++ b/Engine/Source/Editor/UnrealEd/Private/Toolkits/SStandaloneAssetEditorToolkitHost.cpp
@@ -212,7 +212,7 @@ void SStandaloneAssetEditorToolkitHost::GenerateMenus(bool bForceCreateMenu)
 
 		UAssetEditorToolkitMenuContext* ContextObject = NewObject<UAssetEditorToolkitMenuContext>();
 		ContextObject->Toolkit = HostedAssetEditorToolkit;
-		FToolMenuContext ToolMenuContext(HostedAssetEditorToolkit->GetToolkitCommands(), FExtender::Combine(MenuExtenders).ToSharedRef(), ContextObject);
+		FToolMenuContext ToolMenuContext(HostedAssetEditorToolkit->GetToolkitCommands(), FExtender::Combine(MenuExtenders), ContextObject);
 		HostedAssetEditorToolkit->InitToolMenuContext(ToolMenuContext);
 		IMainFrameModule& MainFrameModule = FModuleManager::LoadModuleChecked<IMainFrameModule>( "MainFrame" );
 		DefaultMenuWidget = MainFrameModule.MakeMainMenu( MyTabManager, AssetEditorMenuName, ToolMenuContext );
diff --git a/Engine/Source/Runtime/Slate/Private/Framework/MultiBox/MultiBoxExtender.cpp b/Engine/Source/Runtime/Slate/Private/Framework/MultiBox/MultiBoxExtender.cpp
--- a/Engine/Source/Runtime/Slate/Private/Framework/MultiBox/MultiBoxExtender.cpp
+++ b/Engine/Source/Runtime/Slate/Private/Framework/MultiBox/MultiBoxExtender.cpp
@@ -226,9 +226,9 @@ void FExtender::Apply( FName ExtensionHook, EExtensionHook::Position HookPositio
 }
 
 
-TSharedPtr<FExtender> FExtender::Combine(const TArray< TSharedPtr<FExtender> >& Extenders)
+TSharedRef<FExtender> FExtender::Combine(const TArray< TSharedPtr<FExtender> >& Extenders)
 {
-	TSharedPtr<FExtender> OutExtender = MakeShareable(new FExtender);
+	TSharedRef<FExtender> OutExtender = MakeShared<FExtender>();
 	for (int32 i = 0; i < Extenders.Num(); ++i)
 	{
 		OutExtender->Extensions.Append(Extenders[i]->Extensions);
diff --git a/Engine/Source/Runtime/Slate/Public/Framework/MultiBox/MultiBoxExtender.h b/Engine/Source/Runtime/Slate/Public/Framework/MultiBox/MultiBoxExtender.h
--- a/Engine/Source/Runtime/Slate/Public/Framework/MultiBox/MultiBoxExtender.h
+++ b/Engine/Source/Runtime/Slate/Public/Framework/MultiBox/MultiBoxExtender.h
@@ -124,7 +124,7 @@ public:
 	 * @param Extenders The array of FExtenders to consolidate
 	 * @return A new FExtender pointer that has all of the extensions of the input
 	 */
-	static SLATE_API TSharedPtr<FExtender> Combine(const TArray< TSharedPtr<FExtender> >& Extenders);
+	static SLATE_API TSharedRef<FExtender> Combine(const TArray< TSharedPtr<FExtender> >& Extenders);
 
 	/** Returns the number of extensions in the extender */
 	SLATE_API int32 NumExtensions() const { return Extensions.Num(); }
