Fixed CPU Lightmass fails to launch. Adding a dependency on the Zen module to lightmass was also introducing a dependency on the sockets library that was not resolved properly

diff --git a/Engine/Source/Editor/UnrealEd/Private/Lightmass/Lightmass.cpp b/Engine/Source/Editor/UnrealEd/Private/Lightmass/Lightmass.cpp
--- a/Engine/Source/Editor/UnrealEd/Private/Lightmass/Lightmass.cpp
+++ b/Engine/Source/Editor/UnrealEd/Private/Lightmass/Lightmass.cpp
@@ -2931,6 +2931,7 @@ bool FLightmassProcessor::BeginRun()
 		TEXT("Binaries/Win32/UnrealLightmass-Core.dll"),
 		TEXT("Binaries/Win32/UnrealLightmass-CoreUObject.dll"),
 		TEXT("Binaries/Win32/UnrealLightmass-DerivedDataCache.dll"),
+		TEXT("Binaries/Win32/UnrealLightmass-Sockets.dll"),
 		TEXT("Binaries/Win32/UnrealLightmass-Zen.dll"),
 		TEXT("Binaries/Win32/UnrealLightmass-Projects.dll"),
 		TEXT("Binaries/Win32/UnrealLightmass-Json.dll"),
@@ -2950,6 +2951,7 @@ bool FLightmassProcessor::BeginRun()
 		TEXT("Binaries/Win64/UnrealLightmass-Core.dll"),
 		TEXT("Binaries/Win64/UnrealLightmass-CoreUObject.dll"),
 		TEXT("Binaries/Win64/UnrealLightmass-DerivedDataCache.dll"),
+		TEXT("Binaries/Win64/UnrealLightmass-Sockets.dll"),
 		TEXT("Binaries/Win64/UnrealLightmass-Zen.dll"),
 		TEXT("Binaries/Win64/UnrealLightmass-Projects.dll"),
 		TEXT("Binaries/Win64/UnrealLightmass-Json.dll"),
@@ -2968,6 +2970,7 @@ bool FLightmassProcessor::BeginRun()
 		TEXT("Binaries/Mac/UnrealLightmass-Core.dylib"),
 		TEXT("Binaries/Mac/UnrealLightmass-CoreUObject.dylib"),
 		TEXT("Binaries/Mac/UnrealLightmass-DerivedDataCache.dylib"),
+		TEXT("Binaries/Mac/UnrealLightmass-Sockets.dylib"),
 		TEXT("Binaries/Mac/UnrealLightmass-Zen.dylib"),
 		TEXT("Binaries/Mac/UnrealLightmass-Json.dylib"),
 		TEXT("Binaries/Mac/UnrealLightmass-Projects.dylib"),
@@ -2987,6 +2990,7 @@ bool FLightmassProcessor::BeginRun()
 		TEXT("Binaries/Linux/libUnrealLightmass-Core.so"),
 		TEXT("Binaries/Linux/libUnrealLightmass-CoreUObject.so"),
 		TEXT("Binaries/Linux/libUnrealLightmass-DerivedDataCache.so"),
+		TEXT("Binaries/Linux/libUnrealLightmass-Sockets.so"),
 		TEXT("Binaries/Linux/libUnrealLightmass-Zen.so"),
 		TEXT("Binaries/Linux/libUnrealLightmass-Json.so"),
 		TEXT("Binaries/Linux/libUnrealLightmass-Projects.so"),
