PR #7772: DataTable Factory Class-Picker Support (Contributed by [REDACTED_PERSON])

diff --git a/Engine/Source/Editor/UnrealEd/Classes/Factories/DataTableFactory.h b/Engine/Source/Editor/UnrealEd/Classes/Factories/DataTableFactory.h
--- a/Engine/Source/Editor/UnrealEd/Classes/Factories/DataTableFactory.h
+++ b/Engine/Source/Editor/UnrealEd/Classes/Factories/DataTableFactory.h
@@ -16,6 +16,8 @@ class UNREALED_API UDataTableFactory : public UFactory
 
 	UPROPERTY(BlueprintReadWrite, Category = "Data Table Factory")
 	TObjectPtr<const class UScriptStruct> Struct;
+	UPROPERTY(BlueprintReadWrite, Category = "Data Table Factory")
+	TObjectPtr<const UClass> TableClass;
 
 	//~ Begin UFactory Interface
 	virtual bool ConfigureProperties() override;
diff --git a/Engine/Source/Editor/UnrealEd/Private/Factories/DataTableFactory.cpp b/Engine/Source/Editor/UnrealEd/Private/Factories/DataTableFactory.cpp
--- a/Engine/Source/Editor/UnrealEd/Private/Factories/DataTableFactory.cpp
+++ b/Engine/Source/Editor/UnrealEd/Private/Factories/DataTableFactory.cpp
@@ -8,6 +8,8 @@
 #include "Modules/ModuleManager.h"
 #include "StructViewerModule.h"
 #include "StructViewerFilter.h"
+#include "ClassViewerModule.h"
+#include "ClassViewerFilter.h"
 
 #include "Widgets/SWindow.h"
 #include "Widgets/Input/SButton.h"
@@ -45,12 +47,29 @@ bool UDataTableFactory::ConfigureProperties()
 		}
 	};
 
+	class FDataTableClassFilter : public IClassViewerFilter
+	{
+	public:
+		virtual bool IsClassAllowed(const FClassViewerInitializationOptions& InInitOptions, const UClass* InClass, TSharedRef< class FClassViewerFilterFuncs > InFilterFuncs ) override
+		{
+			return InClass && !InClass->HasAnyClassFlags(EClassFlags::CLASS_Abstract) && InClass->IsChildOf(UDataTable::StaticClass());
+		}
+
+		virtual bool IsUnloadedClassAllowed(const FClassViewerInitializationOptions& InInitOptions, const TSharedRef< const class IUnloadedBlueprintData > InUnloadedClassData, TSharedRef< class FClassViewerFilterFuncs > InFilterFuncs) override
+		{
+			// DataTable cannot have Blueprint sub-classes.
+			return false;
+		}
+	};
+
 	class FDataTableFactoryUI : public TSharedFromThis<FDataTableFactoryUI>
 	{
 	public:
 		FReply OnCreate()
 		{
 			check(ResultStruct);
+			check(ResultClass);
+
 			if (PickerWindow.IsValid())
 			{
 				PickerWindow->RequestDestroyWindow();
@@ -61,6 +80,8 @@ bool UDataTableFactory::ConfigureProperties()
 		FReply OnCancel()
 		{
 			ResultStruct = nullptr;
+			ResultClass = nullptr;
+
 			if (PickerWindow.IsValid())
 			{
 				PickerWindow->RequestDestroyWindow();
@@ -68,9 +89,9 @@ bool UDataTableFactory::ConfigureProperties()
 			return FReply::Handled();
 		}
 
-		bool IsStructSelected() const
+		bool IsValidSelection() const
 		{
-			return ResultStruct != nullptr;
+			return ResultStruct && ResultClass;
 		}
 
 		void OnPickedStruct(const UScriptStruct* ChosenStruct)
@@ -79,13 +100,26 @@ bool UDataTableFactory::ConfigureProperties()
 			StructPickerAnchor->SetIsOpen(false);
 		}
 
-		FText OnGetComboTextValue() const
+		void OnPickedClass(UClass* ChosenClass)
+		{
+			ResultClass = ChosenClass;
+			ClassPickerAnchor->SetIsOpen(false);
+		}
+
+		FText OnGetStructComboTextValue() const
 		{
 			return ResultStruct
 				? FText::AsCultureInvariant(ResultStruct->GetName())
 				: LOCTEXT("None", "None");
 		}
 
+		FText OnGetClassComboTextValue() const
+		{
+			return ResultClass
+				? FText::AsCultureInvariant(ResultClass->GetName())
+				: LOCTEXT("None", "None");
+		}
+
 		TSharedRef<SWidget> GenerateStructPicker()
 		{
 			FStructViewerModule& StructViewerModule = FModuleManager::LoadModuleChecked<FStructViewerModule>("StructViewer");
@@ -115,18 +149,38 @@ bool UDataTableFactory::ConfigureProperties()
 				];
 		}
 
-		const UScriptStruct* OpenStructSelector()
+		TSharedRef<SWidget> GenerateClassPicker()
 		{
-			FStructViewerModule& StructViewerModule = FModuleManager::LoadModuleChecked<FStructViewerModule>("StructViewer");
-			ResultStruct = nullptr;
+			FClassViewerModule& ClassViewerModule = FModuleManager::LoadModuleChecked<FClassViewerModule>("ClassViewer");
 
-			// Fill in options
-			FStructViewerInitializationOptions Options;
-			Options.Mode = EStructViewerMode::StructPicker;
-			Options.StructFilter = MakeShared<FDataTableStructFilter>();
+			FClassViewerInitializationOptions Options;
+			Options.Mode = EClassViewerMode::ClassPicker;
+			Options.ClassFilters.Add(MakeShared<FDataTableClassFilter>());
+			Options.bShowNoneOption = false;
 
+			return SNew(SBox)
+				.WidthOverride(330)
+				[
+					SNew(SVerticalBox)
+
+					+SVerticalBox::Slot()
+					.FillHeight(1.0f)
+					.MaxHeight(500)
+					[
+						SNew(SBorder)
+						.Padding(4)
+						.BorderImage(FEditorStyle::GetBrush("ToolPanel.GroupBorder"))
+						[
+							ClassViewerModule.CreateClassViewer(Options, FOnClassPicked::CreateSP(this, &FDataTableFactoryUI::OnPickedClass))
+						]
+					]
+				];
+		}
+
+		bool OpenSelectorDialog(TObjectPtr<const UClass> ClassResult, TObjectPtr<const UScriptStruct> StructResult)
+		{
 			PickerWindow = SNew(SWindow)
-				.Title(LOCTEXT("DataTableFactoryOptions", "Pick Row Structure"))
+				.Title(LOCTEXT("DataTableFactoryOptions", "Pick Class & Row Structure"))
 				.ClientSize(FVector2D(350, 100))
 				.SupportsMinimize(false)
 				.SupportsMaximize(false)
@@ -136,6 +190,21 @@ bool UDataTableFactory::ConfigureProperties()
 					.Padding(10)
 					[
 						SNew(SVerticalBox)
+						// Class Picker
+						+SVerticalBox::Slot()
+						.AutoHeight()
+						[
+							SAssignNew(ClassPickerAnchor, SComboButton)
+							.ContentPadding(FMargin(2,2,2,1))
+							.MenuPlacement(MenuPlacement_BelowAnchor)
+							.ButtonContent()
+							[
+								SNew(STextBlock)
+								.Text(this, &FDataTableFactoryUI::OnGetClassComboTextValue)
+							]
+							.OnGetMenuContent(this, &FDataTableFactoryUI::GenerateClassPicker)
+						]
+						// Struct Picker
 						+SVerticalBox::Slot()
 						.AutoHeight()
 						[
@@ -145,7 +214,7 @@ bool UDataTableFactory::ConfigureProperties()
 							.ButtonContent()
 							[
 								SNew(STextBlock)
-								.Text(this, &FDataTableFactoryUI::OnGetComboTextValue)
+								.Text(this, &FDataTableFactoryUI::OnGetStructComboTextValue)
 							]
 							.OnGetMenuContent(this, &FDataTableFactoryUI::GenerateStructPicker)
 						]
@@ -159,7 +228,7 @@ bool UDataTableFactory::ConfigureProperties()
 							[
 								SNew(SButton)
 								.Text(LOCTEXT("OK", "OK"))
-								.IsEnabled(this, &FDataTableFactoryUI::IsStructSelected)
+								.IsEnabled(this, &FDataTableFactoryUI::IsValidSelection)
 								.OnClicked(this, &FDataTableFactoryUI::OnCreate)
 							]
 							+SHorizontalBox::Slot()
@@ -176,25 +245,31 @@ bool UDataTableFactory::ConfigureProperties()
 			GEditor->EditorAddModalWindow(PickerWindow.ToSharedRef());
 			PickerWindow.Reset();
 
-			return ResultStruct;
+			ClassResult = ResultClass;
+			StructResult = ResultStruct;
+
+			return IsValidSelection();
 		}
 
 	private:
 		TSharedPtr<SWindow> PickerWindow;
 		TSharedPtr<SComboButton> StructPickerAnchor;
+		TSharedPtr<SComboButton> ClassPickerAnchor;
 		const UScriptStruct* ResultStruct = nullptr;
+		const UClass* ResultClass = UDataTable::StaticClass();
 	};
 
-	TSharedRef<FDataTableFactoryUI> StructSelector = MakeShareable(new FDataTableFactoryUI());
-	Struct = StructSelector->OpenStructSelector();
+	FDataTableFactoryUI ConfigSelector = FDataTableFactoryUI();
 
-	return Struct != nullptr;
+	TableClass = nullptr;
+	Struct = nullptr;
+	return ConfigSelector.OpenSelectorDialog(TableClass, Struct);
 }
 
 UObject* UDataTableFactory::FactoryCreateNew(UClass* Class, UObject* InParent, FName Name, EObjectFlags Flags, UObject* Context, FFeedbackContext* Warn)
 {
 	UDataTable* DataTable = nullptr;
-	if (Struct && ensure(SupportedClass == Class))
+	if (Struct && TableClass && (SupportedClass == Class))
 	{
 		ensure(0 != (RF_Public & Flags));
 		DataTable = MakeNewDataTable(InParent, Name, Flags);
@@ -208,7 +283,7 @@ UObject* UDataTableFactory::FactoryCreateNew(UClass* Class, UObject* InParent, F
 
 UDataTable* UDataTableFactory::MakeNewDataTable(UObject* InParent, FName Name, EObjectFlags Flags)
 {
-	return NewObject<UDataTable>(InParent, Name, Flags);
+	return NewObject<UDataTable>(InParent, TableClass, Name, Flags);
 }
 
 #undef LOCTEXT_NAMESPACE // "DataTableFactory"
