Ensure that all compilation has been submitted when querying attempting to render a thumbnail for that material.

- This is important when running with "compile what you use" ODSC in the editor.

diff --git a/Engine/Source/Editor/UnrealEd/Private/ObjectTools.cpp b/Engine/Source/Editor/UnrealEd/Private/ObjectTools.cpp
--- a/Engine/Source/Editor/UnrealEd/Private/ObjectTools.cpp
+++ b/Engine/Source/Editor/UnrealEd/Private/ObjectTools.cpp
@@ -4668,6 +4668,10 @@ namespace ThumbnailTools
 				FMaterialResource* CurrentResource = InMaterial->GetMaterialResource(GMaxRHIFeatureLevel);
 				if (CurrentResource)
 				{
+					if (!CurrentResource->IsGameThreadShaderMapComplete())
+					{
+						CurrentResource->SubmitCompileJobs(EShaderCompileJobPriority::High);
+					}
 					CurrentResource->FinishCompilation();
 				}
 			}
