Anim node Function/Blueprint Function in-use deletion warnings and expanded Find References implementation

diff --git a/Engine/Plugins/Developer/PropertyAccessNode/Source/PropertyAccessNode/Private/K2Node_PropertyAccess.cpp b/Engine/Plugins/Developer/PropertyAccessNode/Source/PropertyAccessNode/Private/K2Node_PropertyAccess.cpp
--- a/Engine/Plugins/Developer/PropertyAccessNode/Source/PropertyAccessNode/Private/K2Node_PropertyAccess.cpp
+++ b/Engine/Plugins/Developer/PropertyAccessNode/Source/PropertyAccessNode/Private/K2Node_PropertyAccess.cpp
@@ -363,7 +363,7 @@ bool UK2Node_PropertyAccess::ReferencesVariable(const FName& InVarName, const US
 	IPropertyAccessEditor& PropertyAccessEditor = IModularFeatures::Get().GetModularFeature<IPropertyAccessEditor>("PropertyAccessEditor");
 
 	const UClass* SkeletonVariableClass = FBlueprintEditorUtils::GetSkeletonClass(Cast<UClass>(InScope));
-	
+
 	// See if the path references the variable
 	bool bReferencesVariable = false;
 	
@@ -372,8 +372,9 @@ bool UK2Node_PropertyAccess::ReferencesVariable(const FName& InVarName, const US
 	{
 		if(SkeletonVariableClass)
 		{
-			UStruct* OwnerStruct = InProperty->GetOwnerStruct();
-			if(OwnerStruct && InProperty->GetFName() == InVarName && OwnerStruct->IsChildOf(SkeletonVariableClass))
+			const UClass* OwnerSkeletonVariableClass = FBlueprintEditorUtils::GetSkeletonClass(Cast<UClass>(InProperty->GetOwnerStruct()));
+
+			if(OwnerSkeletonVariableClass && InProperty->GetFName() == InVarName && OwnerSkeletonVariableClass->IsChildOf(SkeletonVariableClass))
 			{
 				bReferencesVariable = true;
 			}
@@ -389,6 +390,38 @@ bool UK2Node_PropertyAccess::ReferencesVariable(const FName& InVarName, const US
 	return bReferencesVariable;
 }
 
+bool UK2Node_PropertyAccess::ReferencesFunction(const FName& InFunctionName, const UStruct* InScope) const
+{
+	IPropertyAccessEditor& PropertyAccessEditor = IModularFeatures::Get().GetModularFeature<IPropertyAccessEditor>("PropertyAccessEditor");
+
+	const UClass* SkeletonFunctionClass = FBlueprintEditorUtils::GetSkeletonClass(Cast<UClass>(InScope));
+
+	// See if the path references the function
+	bool bReferencesFunction = false;
+
+	IPropertyAccessEditor::FResolvePropertyAccessArgs ResolveArgs;
+	ResolveArgs.FunctionFunction = [InFunctionName, SkeletonFunctionClass, &bReferencesFunction](int32 InSegmentIndex, UFunction* InFunction, FProperty* InProperty)
+	{
+		if (SkeletonFunctionClass)
+		{
+			const UClass* OwnerSkeletonFunctionClass = FBlueprintEditorUtils::GetSkeletonClass(InFunction->GetOuterUClass());
+
+			if (OwnerSkeletonFunctionClass && InFunction->GetFName() == InFunctionName && OwnerSkeletonFunctionClass->IsChildOf(SkeletonFunctionClass))
+			{
+				bReferencesFunction = true;
+			}
+		}
+		else if (InFunction->GetFName() == InFunctionName)
+		{
+			bReferencesFunction = true;
+		}
+	};
+
+	PropertyAccessEditor.ResolvePropertyAccess(GetBlueprint()->SkeletonGeneratedClass, Path, ResolveArgs);
+
+	return bReferencesFunction;
+}
+
 const FProperty* UK2Node_PropertyAccess::GetResolvedProperty() const
 {
 	if(const FProperty* Property = ResolvedProperty.Get())
@@ -401,4 +434,4 @@ const FProperty* UK2Node_PropertyAccess::GetResolvedProperty() const
 	return ResolvedProperty.Get();
 }
 
-#undef LOCTEXT_NAMESPACE
\ No newline at end of file
+#undef LOCTEXT_NAMESPACE
diff --git a/Engine/Plugins/Developer/PropertyAccessNode/Source/PropertyAccessNode/Private/K2Node_PropertyAccess.h b/Engine/Plugins/Developer/PropertyAccessNode/Source/PropertyAccessNode/Private/K2Node_PropertyAccess.h
--- a/Engine/Plugins/Developer/PropertyAccessNode/Source/PropertyAccessNode/Private/K2Node_PropertyAccess.h
+++ b/Engine/Plugins/Developer/PropertyAccessNode/Source/PropertyAccessNode/Private/K2Node_PropertyAccess.h
@@ -118,6 +118,7 @@ private:
 	virtual void HandleVariableRenamed(UBlueprint* InBlueprint, UClass* InVariableClass, UEdGraph* InGraph, const FName& InOldVarName, const FName& InNewVarName) override;
 	virtual void ReplaceReferences(UBlueprint* InBlueprint, UBlueprint* InReplacementBlueprint, const FMemberReference& InSource, const FMemberReference& InReplacement) override;
 	virtual bool ReferencesVariable(const FName& InVarName, const UStruct* InScope) const override;
+	virtual bool ReferencesFunction(const FName& InFunctionName,const UStruct* InScope) const override;
 	
 	/** Helper function for pin allocation */
 	void AllocatePins(UEdGraphPin* InOldOutputPin = nullptr);
diff --git a/Engine/Source/Editor/AnimGraph/Private/AnimBlueprintExtension.cpp b/Engine/Source/Editor/AnimGraph/Private/AnimBlueprintExtension.cpp
--- a/Engine/Source/Editor/AnimGraph/Private/AnimBlueprintExtension.cpp
+++ b/Engine/Source/Editor/AnimGraph/Private/AnimBlueprintExtension.cpp
@@ -83,13 +83,8 @@ void UAnimBlueprintExtension::RequestExtensionsForNode(UAnimGraphNode_Base* InAn
 		// and request extensions if it *looks* like we have a valid reference, rather than a verified concrete ref.
 		// Later validation code will do the actual validation for us anyways and fail compilation if the reference is
 		// actually invalid.
-		auto MemberReferencePotentiallyValid = [](const FMemberReference& InMemberReference)
-		{
-			return InMemberReference.GetMemberGuid().IsValid() || InMemberReference.GetMemberName() != NAME_None;
-		};
-
-		if( MemberReferencePotentiallyValid(InAnimGraphNode->InitialUpdateFunction) ||
-			MemberReferencePotentiallyValid(InAnimGraphNode->BecomeRelevantFunction))
+		if(UAnimGraphNode_Base::IsPotentiallyBoundFunction(InAnimGraphNode->InitialUpdateFunction) ||
+			UAnimGraphNode_Base::IsPotentiallyBoundFunction(InAnimGraphNode->BecomeRelevantFunction))
 		{
 			ExtensionClasses.Add(UAnimBlueprintExtension_NodeRelevancy::StaticClass());
 		}
diff --git a/Engine/Source/Editor/AnimGraph/Private/AnimGraphNode_Base.cpp b/Engine/Source/Editor/AnimGraph/Private/AnimGraphNode_Base.cpp
--- a/Engine/Source/Editor/AnimGraph/Private/AnimGraphNode_Base.cpp
+++ b/Engine/Source/Editor/AnimGraph/Private/AnimGraphNode_Base.cpp
@@ -746,7 +746,21 @@ FString UAnimGraphNode_Base::GetPinMetaData(FName InPinName, FName InKey)
 void UAnimGraphNode_Base::AddSearchMetaDataInfo(TArray<struct FSearchTagDataPair>& OutTaggedMetaData) const
 {
 	Super::AddSearchMetaDataInfo(OutTaggedMetaData);
-	
+
+	const auto ConditionallyTagNodeFuncRef = [&OutTaggedMetaData](const FMemberReference& FuncMember, const FText& LocText)
+	{
+		if (IsPotentiallyBoundFunction(FuncMember))
+		{
+			const FText FunctionName = FText::FromName(FuncMember.GetMemberName());
+			OutTaggedMetaData.Add(FSearchTagDataPair(LocText, FunctionName));
+		}
+	};
+
+	// Conditionally include anim node function references as part of the node's search metadata
+	ConditionallyTagNodeFuncRef(InitialUpdateFunction, LOCTEXT("InitialUpdateFunctionName", "Initial Update"));
+	ConditionallyTagNodeFuncRef(BecomeRelevantFunction, LOCTEXT("BecomeRelevantFunctionName", "Become Relevant"));
+	ConditionallyTagNodeFuncRef(UpdateFunction, LOCTEXT("UpdateFunctionName", "Update"));
+
 	if(Tag != NAME_None)
 	{
 		OutTaggedMetaData.Add(FSearchTagDataPair(LOCTEXT("Tag", "Tag"), FText::FromName(Tag)));
@@ -1712,8 +1726,9 @@ bool UAnimGraphNode_Base::ReferencesVariable(const FName& InVarName, const UStru
 		{
 			if(SkeletonVariableClass)
 			{
-				UStruct* OwnerStruct = InProperty->GetOwnerStruct();
-				if(OwnerStruct && InProperty->GetFName() == InVarName && OwnerStruct->IsChildOf(SkeletonVariableClass))
+				const UClass* OwnerSkeletonVariableClass = FBlueprintEditorUtils::GetSkeletonClass(Cast<UClass>(InProperty->GetOwnerStruct()));
+
+				if(OwnerSkeletonVariableClass && InProperty->GetFName() == InVarName && OwnerSkeletonVariableClass->IsChildOf(SkeletonVariableClass))
 				{
 					bReferencesVariable = true;
 				}
@@ -1735,6 +1750,47 @@ bool UAnimGraphNode_Base::ReferencesVariable(const FName& InVarName, const UStru
 	return false;
 }
 
+bool UAnimGraphNode_Base::ReferencesFunction(const FName& InFunctionName,const UStruct* InScope) const
+{
+	IPropertyAccessEditor& PropertyAccessEditor = IModularFeatures::Get().GetModularFeature<IPropertyAccessEditor>("PropertyAccessEditor");
+
+	const UClass* SkeletonFunctionClass = FBlueprintEditorUtils::GetSkeletonClass(Cast<UClass>(InScope));
+
+	// See if any of bindings reference the function
+	for (const auto& BindingPair : PropertyBindings)
+	{
+		bool bReferencesFunction = false;
+
+		IPropertyAccessEditor::FResolvePropertyAccessArgs ResolveArgs;
+		ResolveArgs.FunctionFunction = [InFunctionName, SkeletonFunctionClass, &bReferencesFunction](int32 InSegmentIndex, UFunction* InFunction, FProperty* InProperty)
+		{
+			if (SkeletonFunctionClass)
+			{
+				const UClass* OwnerSkeletonFunctionClass = FBlueprintEditorUtils::GetSkeletonClass(InFunction->GetOuterUClass());
+
+				if (OwnerSkeletonFunctionClass && InFunction->GetFName() == InFunctionName && OwnerSkeletonFunctionClass->IsChildOf(SkeletonFunctionClass))
+				{
+					bReferencesFunction = true;
+				}
+			}
+			else if (InFunction->GetFName() == InFunctionName)
+			{
+				bReferencesFunction = true;
+			}
+		};
+
+		PropertyAccessEditor.ResolvePropertyAccess(GetAnimBlueprint()->SkeletonGeneratedClass, BindingPair.Value.PropertyPath, ResolveArgs);
+
+		if (bReferencesFunction)
+		{
+			return true;
+		}
+	}
+
+	// Check private member anim node function binding 
+	return InitialUpdateFunction.GetMemberName() == InFunctionName || BecomeRelevantFunction.GetMemberName() == InFunctionName || UpdateFunction.GetMemberName() == InFunctionName;
+}
+
 FAnimNode_Base* UAnimGraphNode_Base::GetDebuggedAnimNode() const
 {
 	if (UBlueprint* Blueprint = FBlueprintEditorUtils::FindBlueprintForNode(this))
@@ -1774,4 +1830,9 @@ void UAnimGraphNode_Base::PostEditRefreshDebuggedComponent()
 	}
 }
 
+bool UAnimGraphNode_Base::IsPotentiallyBoundFunction(const FMemberReference& FunctionReference)
+{
+	return FunctionReference.GetMemberGuid().IsValid() || FunctionReference.GetMemberName() != NAME_None;
+}
+
 #undef LOCTEXT_NAMESPACE
diff --git a/Engine/Source/Editor/AnimGraph/Public/AnimGraphNode_Base.h b/Engine/Source/Editor/AnimGraph/Public/AnimGraphNode_Base.h
--- a/Engine/Source/Editor/AnimGraph/Public/AnimGraphNode_Base.h
+++ b/Engine/Source/Editor/AnimGraph/Public/AnimGraphNode_Base.h
@@ -270,6 +270,7 @@ public:
 	virtual void HandleVariableRenamed(UBlueprint* InBlueprint, UClass* InVariableClass, UEdGraph* InGraph, const FName& InOldVarName, const FName& InNewVarName) override;
 	virtual void ReplaceReferences(UBlueprint* InBlueprint, UBlueprint* InReplacementBlueprint, const FMemberReference& InSource, const FMemberReference& InReplacement) override;
 	virtual bool ReferencesVariable(const FName& InVarName, const UStruct* InScope) const override;
+	virtual bool ReferencesFunction(const FName& InFunctionName, const UStruct* InScope) const override;
 	// End of UK2Node interface
 
 	// UAnimGraphNode_Base interface
@@ -505,7 +506,10 @@ public:
 	// This is required to see changes as the component may be either an editor-only component that is not ticking,
 	// or in a paused PIE world
 	void PostEditRefreshDebuggedComponent();
-	
+
+	// Check if a specified function reference appears to be valid by inspecting only the validity of the name and guid
+	static bool IsPotentiallyBoundFunction(const FMemberReference& FunctionReference);
+
 protected:
 	friend class FAnimBlueprintCompilerContext;
 	friend class FAnimGraphNodeDetails;
diff --git a/Engine/Source/Editor/BlueprintGraph/Classes/K2Node.h b/Engine/Source/Editor/BlueprintGraph/Classes/K2Node.h
--- a/Engine/Source/Editor/BlueprintGraph/Classes/K2Node.h
+++ b/Engine/Source/Editor/BlueprintGraph/Classes/K2Node.h
@@ -409,6 +409,9 @@ public:
 
 	/** Return whether this node references the specified variable, give the supplied scope. Used when variable types are changed. */
 	virtual bool ReferencesVariable(const FName& InVarName, const UStruct* InScope) const { return false; }
+	
+	/** Return whether this node references the specified function, identified by a name and guid pair **/
+	virtual bool ReferencesFunction(const FName& InFunctionName, const UStruct* InScope) const { return false; };
 
 	/** 
 	 * Replace any member references of source with replacement 
diff --git a/Engine/Source/Editor/Kismet/Private/SMyBlueprint.cpp b/Engine/Source/Editor/Kismet/Private/SMyBlueprint.cpp
--- a/Engine/Source/Editor/Kismet/Private/SMyBlueprint.cpp
+++ b/Engine/Source/Editor/Kismet/Private/SMyBlueprint.cpp
@@ -2901,30 +2901,57 @@ void SMyBlueprint::OnDeleteDelegate(FEdGraphSchemaAction_K2Delegate* InDelegateA
 	}
 }
 
+namespace UE::Blueprint::Private
+{
+	// Given a type and value name, display a deletion confirmation warning.
+	// Returns true if the user 'cancelled' the action, interpreted as an early exit prior to deletion.
+	static bool DisplayInUseWarningAndEarlyExit(const FName& DisplayTypeName, const FName& DisplayValueName)
+	{
+		const FText DeleteConfirmationPrompt = FText::Format(LOCTEXT("DeleteConfirmationPrompt", "{0} {1} is in use! Do you really want to delete it?")
+			, { FText::FromName(DisplayTypeName), FText::FromName(DisplayValueName) }
+		);
+		const FText DeleteConfirmationTitle = FText::Format(LOCTEXT("DeleteConfirmationTitle", "Delete {0}")
+			, { FText::FromName(DisplayTypeName) }
+		);
+		const FString DeleteConfirmationIniSetting = FString::Format(TEXT("DeleteConfirmation{0}_Warning"), { DisplayTypeName.ToString() });
+
+		// Warn the user that this may result in data loss
+		FSuppressableWarningDialog::FSetupInfo Info(DeleteConfirmationPrompt, DeleteConfirmationTitle, DeleteConfirmationIniSetting);
+		Info.ConfirmText = LOCTEXT("DeleteConfirmation_Yes", "Yes");
+		Info.CancelText = LOCTEXT("DeleteConfirmation_No", "No");
+
+		FSuppressableWarningDialog DeleteFunctionInUse(Info);
+		return DeleteFunctionInUse.ShowModal() == FSuppressableWarningDialog::Cancel;
+	}
+}
+
 void SMyBlueprint::OnDeleteEntry()
 {
 	if (FEdGraphSchemaAction_K2Graph* GraphAction = SelectionAsGraph())
 	{
+		// Currently only function graphs are supported for in-use detection and deletion warnings
+		if (GraphAction->GraphType == EEdGraphSchemaAction_K2Graph::Function)
+		{
+			if (FBlueprintEditorUtils::IsFunctionUsed(GetBlueprintObj(), GraphAction->FuncName))
+			{
+				if (UE::Blueprint::Private::DisplayInUseWarningAndEarlyExit("Function", GraphAction->FuncName))
+				{
+					return;
+				}
+			}
+		}
+
 		OnDeleteGraph(GraphAction->EdGraph, GraphAction->GraphType);
 	}
 	else if (FEdGraphSchemaAction_K2Delegate* DelegateAction = SelectionAsDelegate())
 	{
 		OnDeleteDelegate(DelegateAction);
 	}
-	else if ( FEdGraphSchemaAction_K2Var* VarAction = SelectionAsVar() )
+	else if (FEdGraphSchemaAction_K2Var* VarAction = SelectionAsVar())
 	{
 		if(FBlueprintEditorUtils::IsVariableUsed(GetBlueprintObj(), VarAction->GetVariableName()))
 		{
-			FText ConfirmDelete = FText::Format(LOCTEXT( "ConfirmDeleteVariableInUse", "Variable {0} is in use! Do you really want to delete it?"),
-				FText::FromName( VarAction->GetVariableName() ) );
-
-			// Warn the user that this may result in data loss
-			FSuppressableWarningDialog::FSetupInfo Info( ConfirmDelete, LOCTEXT("DeleteVar", "Delete Variable"), "DeleteVariableInUse_Warning" );
-			Info.ConfirmText = LOCTEXT( "DeleteVariable_Yes", "Yes");
-			Info.CancelText = LOCTEXT( "DeleteVariable_No", "No");	
-
-			FSuppressableWarningDialog DeleteVariableInUse( Info );
-			if ( DeleteVariableInUse.ShowModal() == FSuppressableWarningDialog::Cancel )
+			if (UE::Blueprint::Private::DisplayInUseWarningAndEarlyExit("Variable", VarAction->GetVariableName()))
 			{
 				return;
 			}
@@ -2935,20 +2962,11 @@ void SMyBlueprint::OnDeleteEntry()
 		GetBlueprintObj()->Modify();
 		FBlueprintEditorUtils::RemoveMemberVariable(GetBlueprintObj(), VarAction->GetVariableName());
 	}
-	else if ( FEdGraphSchemaAction_K2LocalVar* LocalVarAction = SelectionAsLocalVar() )
+	else if (FEdGraphSchemaAction_K2LocalVar* LocalVarAction = SelectionAsLocalVar())
 	{
-		if(FBlueprintEditorUtils::IsVariableUsed(GetBlueprintObj(), LocalVarAction->GetVariableName(), FBlueprintEditorUtils::FindScopeGraph(GetBlueprintObj(), CastChecked<UStruct>(LocalVarAction->GetVariableScope()))))
+		if (FBlueprintEditorUtils::IsVariableUsed(GetBlueprintObj(), LocalVarAction->GetVariableName(), FBlueprintEditorUtils::FindScopeGraph(GetBlueprintObj(), CastChecked<UStruct>(LocalVarAction->GetVariableScope()))))
 		{
-			FText ConfirmDelete = FText::Format(LOCTEXT( "ConfirmDeleteLocalVariableInUse", "Local Variable {0} is in use! Do you really want to delete it?"),
-				FText::FromName( LocalVarAction->GetVariableName() ) );
-
-			// Warn the user that this may result in data loss
-			FSuppressableWarningDialog::FSetupInfo Info( ConfirmDelete, LOCTEXT("DeleteVar", "Delete Variable"), "DeleteVariableInUse_Warning" );
-			Info.ConfirmText = LOCTEXT( "DeleteVariable_Yes", "Yes");
-			Info.CancelText = LOCTEXT( "DeleteVariable_No", "No");	
-
-			FSuppressableWarningDialog DeleteVariableInUse( Info );
-			if ( DeleteVariableInUse.ShowModal() == FSuppressableWarningDialog::Cancel )
+			if (UE::Blueprint::Private::DisplayInUseWarningAndEarlyExit("Local Variable", LocalVarAction->GetVariableName()))
 			{
 				return;
 			}
@@ -2966,20 +2984,11 @@ void SMyBlueprint::OnDeleteEntry()
 
 		FBlueprintEditorUtils::RemoveLocalVariable(GetBlueprintObj(), CastChecked<UStruct>(LocalVarAction->GetVariableScope()), LocalVarAction->GetVariableName());
 	}
-	else if ( FEdGraphSchemaAction_BlueprintVariableBase* BPVarAction = SelectionAsBlueprintVariable() )
+	else if (FEdGraphSchemaAction_BlueprintVariableBase* BPVarAction = SelectionAsBlueprintVariable())
 	{
-		if(BPVarAction->IsVariableUsed())
+		if (BPVarAction->IsVariableUsed())
 		{
-			FText ConfirmDelete = FText::Format(LOCTEXT( "ConfirmDeleteLocalVariableInUse", "Variable {0} is in use! Do you really want to delete it?"),
-				FText::FromName( BPVarAction->GetVariableName() ) );
-
-			// Warn the user that this may result in data loss
-			FSuppressableWarningDialog::FSetupInfo Info( ConfirmDelete, LOCTEXT("DeleteVar", "Delete Variable"), "DeleteVariableInUse_Warning" );
-			Info.ConfirmText = LOCTEXT( "DeleteVariable_Yes", "Yes");
-			Info.CancelText = LOCTEXT( "DeleteVariable_No", "No");	
-
-			FSuppressableWarningDialog DeleteVariableInUse( Info );
-			if ( DeleteVariableInUse.ShowModal() == FSuppressableWarningDialog::Cancel )
+			if (UE::Blueprint::Private::DisplayInUseWarningAndEarlyExit("Variable", BPVarAction->GetVariableName()))
 			{
 				return;
 			}
@@ -2997,7 +3006,7 @@ void SMyBlueprint::OnDeleteEntry()
 		GetBlueprintObj()->Modify();
 		FBlueprintEditorUtils::RemoveNode(GetBlueprintObj(), EventAction->NodeTemplate);
 	}
-	else if ( SelectionIsCategory() )
+	else if (SelectionIsCategory())
 	{
 		TArray<TSharedPtr<FEdGraphSchemaAction>> Actions;
 		GraphActionMenu->GetSelectedCategorySubActions(Actions);
diff --git a/Engine/Source/Editor/UnrealEd/Private/Kismet2/BlueprintEditorUtils.cpp b/Engine/Source/Editor/UnrealEd/Private/Kismet2/BlueprintEditorUtils.cpp
--- a/Engine/Source/Editor/UnrealEd/Private/Kismet2/BlueprintEditorUtils.cpp
+++ b/Engine/Source/Editor/UnrealEd/Private/Kismet2/BlueprintEditorUtils.cpp
@@ -1,6 +1,7 @@
 // Copyright Epic Games, Inc. All Rights Reserved.
 
 #include "Kismet2/BlueprintEditorUtils.h"
+#include "Algo/AnyOf.h"
 #include "Algo/Copy.h"
 #include "Algo/Transform.h"
 #include "BlueprintCompilationManager.h"
@@ -5843,99 +5844,215 @@ bool FBlueprintEditorUtils::IsVariableComponent(const FBPVariableDescription& Va
 	return false;
 }
 
-bool FBlueprintEditorUtils::IsVariableUsed(const UBlueprint* VariableBlueprint, const FName& VariableName, UEdGraph* LocalGraphScope/* = nullptr*/)
+namespace UE::Blueprint::Private
 {
+	// Given a specified search criteria algorithm, walk the current blueprint with a specified scope.
+	// When no explicit scope is provided, the asset registry will also be walked with the algorithm on additional blueprints.
+	template<typename SearchFunc> static bool SearchBlueprintWithFunc(const SearchFunc& Func, const UBlueprint* Blueprint, const UEdGraph* LocalGraphScope)
+	{
+		// Search the initial blueprint
+		if (Func(Blueprint))
+		{
+			return true;
+		}
+
+		// Optionally walk the asset registry for other blueprints
+		if (!LocalGraphScope)
+		{
+			const FAssetRegistryModule& AssetRegistryModule = FModuleManager::LoadModuleChecked<FAssetRegistryModule>(TEXT("AssetRegistry"));
+		
+			// Discover additional packages which reference the initial blueprint package name
+			FARFilter Filter;
+			AssetRegistryModule.Get().GetReferencers(Blueprint->GetPackage()->GetFName(), Filter.PackageNames, UE::AssetRegistry::EDependencyCategory::Package, UE::AssetRegistry::EDependencyQuery::Hard);
+
+			if (Filter.PackageNames.Num() > 0)
+			{
+				GWarn->BeginSlowTask(LOCTEXT("LoadingReferencerAssets", "Loading Referencing Assets ..."), true);
+				{
+					Filter.TagsAndValues.Add(FBlueprintTags::IsDataOnly, TOptional<FString>(TEXT("false")));
+					TArray<FAssetData> ReferencerAssetData;
+					AssetRegistryModule.Get().GetAssets(Filter, ReferencerAssetData);
+
+					// For each referencing asset
+					for (const FAssetData& ReferencerData : ReferencerAssetData)
+					{
+						const UObject* AssetReferencer = ReferencerData.GetAsset();
+
+						// Conditionally search the asset if it is a blueprint
+						if (const UBlueprint* BlueprintReferencer = Cast<const UBlueprint>(AssetReferencer))
+						{
+							if (BlueprintReferencer && Func(BlueprintReferencer))
+							{
+								GWarn->EndSlowTask();
+								return true;
+							}
+						}
+						// Otherwise check to see if a corresponding world level blueprint is in scope to search
+						else if (const UWorld* WorldReferencer = Cast<const UWorld>(AssetReferencer))
+						{
+							const auto& PersistentLevel = WorldReferencer->PersistentLevel;
+							if (!PersistentLevel.IsNull() && PersistentLevel->OwningWorld && Func(PersistentLevel->GetLevelScriptBlueprint()))
+							{
+								GWarn->EndSlowTask();
+								return true;
+							}
+						}
+					}
+				}
+				GWarn->EndSlowTask();
+			}
+		}
+
+		return false;
+	}
+}
+
+bool FBlueprintEditorUtils::IsVariableUsed(const UBlueprint* Blueprint, const FName& VariableName, const UEdGraph* LocalGraphScope /* = nullptr */)
+{
+	if (VariableName.IsNone())
+	{
+		return false;
+	}
+
+	// Retrieve the corresponding variable guid from the blueprint
 	FGuid VariableGuid;
-	UBlueprint::GetGuidFromClassByFieldName<FProperty>(VariableBlueprint->GeneratedClass, VariableName, VariableGuid);
+	UBlueprint::GetGuidFromClassByFieldName<FProperty>(Blueprint->SkeletonGeneratedClass, VariableName, VariableGuid);
 	
-	auto CheckSingleBlueprint = [VariableGuid, VariableName, VariableBlueprint, LocalGraphScope](const UBlueprint* BlueprintToSearch) -> bool
+	if (!VariableGuid.IsValid())
+	{
+		return false;
+	}
+
+	// Blueprint variable search algorithm
+	const auto SearchBlueprint = [VariableGuid, VariableName, Blueprint, LocalGraphScope](const UBlueprint* CurrentBlueprint) -> bool
 	{
 		TArray<UEdGraph*> AllGraphs;
-		BlueprintToSearch->GetAllGraphs(AllGraphs);
+		CurrentBlueprint->GetAllGraphs(AllGraphs);
+
+		// For each blueprint subgraph
 		for (TArray<UEdGraph*>::TConstIterator it(AllGraphs); it; ++it)
 		{
 			const UEdGraph* CurrentGraph = *it;
-			check(CurrentGraph);
-			if (CurrentGraph == LocalGraphScope || LocalGraphScope == nullptr)
+
+			// If the current graph is the specified scope or unbounded
+			if (CurrentGraph && (CurrentGraph == LocalGraphScope || LocalGraphScope == nullptr))
 			{
-				TArray<UK2Node_Variable*> GraphNodes;
-				CurrentGraph->GetNodesOfClass(GraphNodes);
+				// Check all variable nodes, ignoring connectivity
+				TArray<UK2Node_Variable*> VariableNodes;
+				CurrentGraph->GetNodesOfClass(VariableNodes);
 
-				for (const UK2Node_Variable* CurrentNode : GraphNodes)
+				if (Algo::AnyOf(VariableNodes, [&VariableGuid, &VariableName](const UK2Node_Variable* VariableNode)
 				{
-					if (VariableGuid == CurrentNode->VariableReference.GetMemberGuid() &&
-						VariableName == CurrentNode->GetVarName())
-					{
-						return true;
-					}
+					return VariableGuid == VariableNode->VariableReference.GetMemberGuid() && VariableName == VariableNode->GetVarName();
+				}))
+				{
+					return true;
 				}
 
-				// Also consider "used" if there's a GetClassDefaults node that exposes the variable as an output pin that's connected to something.
+				// Check all GetClassDefaults nodes that exposes the variable as an output pin connected to something
 				TArray<UK2Node_GetClassDefaults*> ClassDefaultsNodes;
 				CurrentGraph->GetNodesOfClass(ClassDefaultsNodes);
-				for (const UK2Node_GetClassDefaults* ClassDefaultsNode : ClassDefaultsNodes)
+
+				if (Algo::AnyOf(ClassDefaultsNodes, [&VariableName, &Blueprint](const UK2Node_GetClassDefaults* GraphNode)
 				{
-					if (ClassDefaultsNode->GetInputClass() == VariableBlueprint->SkeletonGeneratedClass)
+					if (GraphNode->GetInputClass() == Blueprint->SkeletonGeneratedClass)
 					{
-						const UEdGraphPin* VarPin = ClassDefaultsNode->FindPin(VariableName);
+						const UEdGraphPin* VarPin = GraphNode->FindPin(VariableName);
 						if (VarPin && VarPin->Direction == EGPD_Output && VarPin->LinkedTo.Num() > 0)
 						{
 							return true;
 						}
 					}
+
+					return false;
+				}))
+				{
+					return true;
+				}
+
+				// Check all K2Node's which specify private/internal function referencing behavior
+				TArray<const UK2Node*> GraphNodes;
+				CurrentGraph->GetNodesOfClass(GraphNodes);
+
+				if (Algo::AnyOf(GraphNodes, [&VariableName, &VariableGuid, &Blueprint](const UK2Node* GraphNode)
+				{
+					return GraphNode->ReferencesVariable(VariableName, Blueprint->SkeletonGeneratedClass);
+				}))
+				{
+					return true;
 				}
+
 			}
 		}
 
 		return false;
 	};
 
-	if (CheckSingleBlueprint(VariableBlueprint))
+	// Given the specified variable search algorithm, walk the blueprint asset
+	return UE::Blueprint::Private::SearchBlueprintWithFunc(SearchBlueprint, Blueprint, LocalGraphScope);
+}
+
+bool FBlueprintEditorUtils::IsFunctionUsed(const UBlueprint* Blueprint, const FName& FunctionName, const UEdGraph* LocalGraphScope /* = nullptr */)
+{
+	if (FunctionName.IsNone())
 	{
-		return true;
+		return false;
 	}
 
-	if (!LocalGraphScope)
+	// Retrieve the corresponding function guid from the blueprint
+	FGuid FunctionGuid;
+	UBlueprint::GetFunctionGuidFromClassByFieldName(Blueprint->SkeletonGeneratedClass, FunctionName, FunctionGuid);
+
+	if (!FunctionGuid.IsValid())
 	{
-		FAssetRegistryModule& AssetRegistryModule = FModuleManager::LoadModuleChecked<FAssetRegistryModule>(TEXT("AssetRegistry"));
-		FARFilter Filter;
-		AssetRegistryModule.Get().GetReferencers(VariableBlueprint->GetPackage()->GetFName(), Filter.PackageNames, UE::AssetRegistry::EDependencyCategory::Package, UE::AssetRegistry::EDependencyQuery::Hard);
-		if (Filter.PackageNames.Num() > 0)
+		return false;
+	}
+
+	// Blueprint function search algorithm
+	const auto SearchBlueprint = [&LocalGraphScope, &FunctionGuid, &FunctionName, &Blueprint](const UBlueprint* CurrentBlueprint) -> bool
+	{
+		TArray<UEdGraph*> BlueprintGraphs;
+		CurrentBlueprint->GetAllGraphs(BlueprintGraphs);
+
+		// For each blueprint subgraph
+		for (TArray<UEdGraph*>::TConstIterator it(BlueprintGraphs); it; ++it)
 		{
-			GWarn->BeginSlowTask(LOCTEXT("LoadingReferencerAssets", "Loading Referencers..."), true);
+			const UEdGraph* CurrentGraph = *it;
 
-			Filter.TagsAndValues.Add(TEXT("IsDataOnly"), TOptional<FString>(TEXT("false")));
-			TArray<FAssetData> ReferencersAssetData;
-			AssetRegistryModule.Get().GetAssets(Filter, ReferencersAssetData);
-			for (const FAssetData& ReferencerData : ReferencersAssetData)
+			// If the current graph is the specified scope or unbounded
+			if (CurrentGraph && (CurrentGraph == LocalGraphScope || !LocalGraphScope))
 			{
-				UObject* ReferencerAsset = ReferencerData.GetAsset();
-				if (UBlueprint* BlueprintReferencer = Cast<UBlueprint>(ReferencerAsset))
+				// Check all function graph nodes, ignoring connectivity
+				TArray<UK2Node_CallFunction*> CallFunctionNodes;
+				CurrentGraph->GetNodesOfClass(CallFunctionNodes);
+
+				if (Algo::AnyOf(CallFunctionNodes, [&FunctionGuid, &FunctionName](const UK2Node_CallFunction* GraphNode)
 				{
-					if (CheckSingleBlueprint(BlueprintReferencer))
-					{
-						GWarn->EndSlowTask();
-						return true;
-					}
+					return FunctionGuid == GraphNode->FunctionReference.GetMemberGuid() && FunctionName == GraphNode->GetFunctionName();
+				}))
+				{
+					return true;
 				}
-				else if (UWorld* WorldReferencer = Cast<UWorld>(ReferencerAsset))
+
+				// Check all K2Nodes which specify internal function referencing behavior
+				TArray<const UK2Node*> GraphNodes;
+				CurrentGraph->GetNodesOfClass(GraphNodes);
+
+				if (Algo::AnyOf(GraphNodes, [&FunctionName, &Blueprint](const UK2Node* GraphNode)
 				{
-					if (WorldReferencer->PersistentLevel && WorldReferencer->PersistentLevel->OwningWorld)
-					{
-						if (CheckSingleBlueprint(WorldReferencer->PersistentLevel->GetLevelScriptBlueprint()))
-						{
-							GWarn->EndSlowTask();
-							return true;
-						}
-					}
+					return GraphNode->ReferencesFunction(FunctionName, Blueprint->SkeletonGeneratedClass);
+				}))
+				{
+					return true;
 				}
 			}
-
-			GWarn->EndSlowTask();
 		}
-	}
 
-	return false;
+		return false;
+	};
+
+	// Given the specified function search algorithm, walk the blueprint asset
+	return UE::Blueprint::Private::SearchBlueprintWithFunc(SearchBlueprint, Blueprint, LocalGraphScope);
 }
 
 bool FBlueprintEditorUtils::ValidateAllMemberVariables(UBlueprint* InBlueprint, UBlueprint* InParentBlueprint, const FName InVariableName)
diff --git a/Engine/Source/Editor/UnrealEd/Public/Kismet2/BlueprintEditorUtils.h b/Engine/Source/Editor/UnrealEd/Public/Kismet2/BlueprintEditorUtils.h
--- a/Engine/Source/Editor/UnrealEd/Public/Kismet2/BlueprintEditorUtils.h
+++ b/Engine/Source/Editor/UnrealEd/Public/Kismet2/BlueprintEditorUtils.h
@@ -1318,8 +1318,11 @@ public:
 	/** Returns whether or not the specified member var is a component */
 	static bool IsVariableComponent(const FBPVariableDescription& Variable);
 
-	/** Indicates if the variable is used on any graphs in this Blueprint*/
-	static bool IsVariableUsed(const UBlueprint* VariableBlueprint, const FName& VariableName, UEdGraph* LocalGraphScope = nullptr);
+	/** Indicates if the variable is used on any graphs in this Blueprint */
+	static bool IsVariableUsed(const UBlueprint* VariableBlueprint, const FName& VariableName, const UEdGraph* LocalGraphScope = nullptr);
+
+	/** Indicates if the function is used on any graphs in this Blueprint */
+	static bool IsFunctionUsed(const UBlueprint* FunctionBlueprint, const FName& FunctionName, const UEdGraph* LocalGraphScope = nullptr);
 
 	/** 
 	 * Copies the value from the passed in string into a property. ContainerMem points to the Struct or Class containing Property 
diff --git a/Engine/Source/Runtime/Engine/Private/MemberReference.cpp b/Engine/Source/Runtime/Engine/Private/MemberReference.cpp
--- a/Engine/Source/Runtime/Engine/Private/MemberReference.cpp
+++ b/Engine/Source/Runtime/Engine/Private/MemberReference.cpp
@@ -139,7 +139,7 @@ FString FMemberReference::GetReferenceSearchString(UClass* InFieldOwner) const
 		{
 			if (MemberGuid.IsValid())
 			{
-				return FString::Printf(TEXT("Nodes(VariableReference(MemberName=+\"%s\" && MemberGuid(A=%i && B=%i && C=%i && D=%i)) || Name=\"(%s)\") || Pins(Binding=\"%s\")"), *MemberName.ToString(), MemberGuid.A, MemberGuid.B, MemberGuid.C, MemberGuid.D, *MemberName.ToString(), *MemberName.ToString());
+				return FString::Printf(TEXT("Nodes(VariableReference(MemberName=+\"%s\" && MemberGuid(A=%i && B=%i && C=%i && D=%i)) || Name=\"(%s)\") || Pins(Binding=\"%s\") || Binding=\"%s\""), *MemberName.ToString(), MemberGuid.A, MemberGuid.B, MemberGuid.C, MemberGuid.D, *MemberName.ToString(), *MemberName.ToString(), *MemberName.ToString());
 			}
 			else
 			{
@@ -149,21 +149,21 @@ FString FMemberReference::GetReferenceSearchString(UClass* InFieldOwner) const
 				ExportMemberParentName += InFieldOwner->GetAuthoritativeClass()->GetPathName();
 				ExportMemberParentName.AppendChar('\'');
 
-				return FString::Printf(TEXT("Nodes(VariableReference(MemberName=+\"%s\" && (MemberParent=\"%s\" || bSelfContext=true) ) || Name=\"(%s)\") || Pins(Binding=\"%s\")"), *MemberName.ToString(), *ExportMemberParentName, *MemberName.ToString(), *MemberName.ToString());
+				return FString::Printf(TEXT("Nodes(VariableReference(MemberName=+\"%s\" && (MemberParent=\"%s\" || bSelfContext=true) ) || Name=\"(%s)\") || Pins(Binding=\"%s\") || Binding=\"%s\""), *MemberName.ToString(), *ExportMemberParentName, *MemberName.ToString(), *MemberName.ToString(), *MemberName.ToString());
 			}
 		}
 		else if (MemberGuid.IsValid())
 		{
-			return FString::Printf(TEXT("Nodes(VariableReference(MemberName=+\"%s\" && MemberGuid(A=%i && B=%i && C=%i && D=%i)) || Name=\"(%s)\") || Pins(Binding=\"%s\")"), *MemberName.ToString(), MemberGuid.A, MemberGuid.B, MemberGuid.C, MemberGuid.D, *MemberName.ToString(), *MemberName.ToString());
+			return FString::Printf(TEXT("Nodes(VariableReference(MemberName=+\"%s\" && MemberGuid(A=%i && B=%i && C=%i && D=%i)) || Name=\"(%s)\") || Pins(Binding=\"%s\") || Binding=\"%s\""), *MemberName.ToString(), MemberGuid.A, MemberGuid.B, MemberGuid.C, MemberGuid.D, *MemberName.ToString(), *MemberName.ToString(), *MemberName.ToString());
 		}
 		else
 		{
-			return FString::Printf(TEXT("Nodes(VariableReference(MemberName=+\"%s\") || Name=\"(%s)\") || Pins(Binding=\"%s\")"), *MemberName.ToString(), *MemberName.ToString(), *MemberName.ToString());
+			return FString::Printf(TEXT("Nodes(VariableReference(MemberName=+\"%s\") || Name=\"(%s)\") || Pins(Binding=\"%s\") || Binding=\"%s\""), *MemberName.ToString(), *MemberName.ToString(), *MemberName.ToString(), *MemberName.ToString());
 		}
 	}
 	else
 	{
-		return FString::Printf(TEXT("Nodes(VariableReference((MemberName=+\"%s\" && MemberScope=+\"%s\")))"), *MemberName.ToString(), *GetMemberScopeName());
+		return FString::Printf(TEXT("Nodes(VariableReference((MemberName=+\"%s\" && MemberScope=+\"%s\"))) || Binding=\"%s\""), *MemberName.ToString(), *GetMemberScopeName(), *MemberName.ToString());
 	}
 }
 
