Remove unused function CompileShadersComplexityViewMode and associated classes.

diff --git a/Engine/Source/Editor/UnrealEd/Private/EditorBuildUtils.cpp b/Engine/Source/Editor/UnrealEd/Private/EditorBuildUtils.cpp
--- a/Engine/Source/Editor/UnrealEd/Private/EditorBuildUtils.cpp
+++ b/Engine/Source/Editor/UnrealEd/Private/EditorBuildUtils.cpp
@@ -1577,111 +1577,4 @@ void FEditorBuildUtils::EditorBuildAllLandscape(UWorld* InWorld)
 	}
 }
 
-/** classed used to compile shaders for a specific (mobile) platform and copy the number of instruction to the editor-emulated (mobile) platform */
-class FMaterialOfflineCompilation : public FMaterialResource
-{
-public:
-	FMaterialOfflineCompilation() {}
-	virtual ~FMaterialOfflineCompilation() {}
-
-	/** this will pass paths to (eventual) offline shader compilers */
-	virtual void SetupExtaCompilationSettings(const EShaderPlatform Platform, FExtraShaderCompilerSettings& Settings) const override;
-
-	/** this function will copy the number of instruction in each of its shaders to editor's emulated shaders */
-	void CopyPlatformSpecificStats();
-};
-
-void FMaterialOfflineCompilation::SetupExtaCompilationSettings(const EShaderPlatform Platform, FExtraShaderCompilerSettings& Settings) const
-{
-	Settings.OfflineCompilerPath = FMaterialStatsUtils::GetPlatformOfflineCompilerPath(Platform);
-}
-
-bool FEditorBuildUtils::CompileShadersComplexityViewMode(EMaterialQualityLevel::Type QualityLevel, ERHIFeatureLevel::Type FeatureLevel, TSet<UMaterialInterface*>& Materials, FSlowTask& ProgressTask)
-{
-	check(Materials.Num());
-
-	// Finish compiling pending shaders first.
-	if (!WaitForShaderCompilation(LOCTEXT("CompileShaders_Complexity_FinishPendingShadersCompilation", "Waiting For Pending Shaders Compilation"), &ProgressTask))
-	{
-		return false;
-	}
-
-	TArray<TSharedPtr<FMaterialOfflineCompilation>> OfflineShaderResources;
-
-	const double StartTime = FPlatformTime::Seconds();
-	const float OneOverNumMaterials = 1.f / (float)Materials.Num();
-
-	const auto SimulatedShaderPlatform = GetFeatureLevelShaderPlatform(FeatureLevel);
-	const auto ShaderPlatform = GetSimulatedPlatform(SimulatedShaderPlatform);
-
-	bool bResult = false;
-
-	// trigger shader compilation/loading for each of the passed materials
-	for (UMaterialInterface* MaterialInterface : Materials)
-	{
-		check(MaterialInterface);
-
-		TSharedPtr<FMaterialOfflineCompilation> SpecialResource = MakeShareable(new FMaterialOfflineCompilation());
-		SpecialResource->SetMaterial(MaterialInterface->GetMaterial(), Cast<UMaterialInstance>(MaterialInterface), FeatureLevel, QualityLevel);
-
-		SpecialResource->CacheShaders(ShaderPlatform);
-
-		OfflineShaderResources.Add(SpecialResource);
-	}
-
-	// wait for compilation to be done and copy the number of instruction from the compiled shaders to the emulated shader set
-	if (WaitForShaderCompilation(LOCTEXT("OfflineShaderCompilation", "Offline Shader Compilation"), &ProgressTask))
-	{
-		FSuspendRenderingThread SuspendObject(false);
-
-		for (int32 i = 0; i < OfflineShaderResources.Num(); ++i)
-		{
-			OfflineShaderResources[i]->CopyPlatformSpecificStats();
-		}
-
-		UE_LOG(LogShaders, Display, TEXT("Offline shader compilation took %.3f seconds."), FPlatformTime::Seconds() - StartTime);
-		bResult = true;
-	}
-
-	OfflineShaderResources.Reset();
-	return bResult;
-}
-
-void FMaterialOfflineCompilation::CopyPlatformSpecificStats()
-{
-	auto Quality = GetQualityLevel();
-	auto Feature = GetFeatureLevel();
-
-	FMaterialResource* Resource = GetMaterialInterface()->GetMaterialResource(Feature, Quality);
-
-	if (Resource == nullptr)
-	{
-		return;
-	}
-
-	const FMaterialShaderMap* DstShaderMap = Resource->GetGameThreadShaderMap();
-	const FMaterialShaderMap* SrcShaderMap = GetGameThreadShaderMap();
-
-	if (DstShaderMap == nullptr || SrcShaderMap == nullptr)
-	{
-		return;
-	}
-
-	TMap<FShaderId, TShaderRef<FShader>> SrcShaders;
-	SrcShaderMap->GetShaderList(SrcShaders);
-
-	TMap<FShaderId, TShaderRef<FShader>> DstShaders;
-	DstShaderMap->GetShaderList(DstShaders);
-
-	for (auto Pair : SrcShaders)
-	{
-		auto *DestinationShaderPtr = DstShaders.Find(Pair.Key);
-		if (DestinationShaderPtr != nullptr)
-		{
-			auto NumInstructions = Pair.Value->GetNumInstructions();
-			(*DestinationShaderPtr)->SetNumInstructions(NumInstructions);
-		}
-	}
-}
-
 #undef LOCTEXT_NAMESPACE
diff --git a/Engine/Source/Editor/UnrealEd/Public/EditorBuildUtils.h b/Engine/Source/Editor/UnrealEd/Public/EditorBuildUtils.h
--- a/Engine/Source/Editor/UnrealEd/Public/EditorBuildUtils.h
+++ b/Engine/Source/Editor/UnrealEd/Public/EditorBuildUtils.h
@@ -269,15 +269,6 @@ private:
 	 */
 	static void TriggerHierarchicalLODBuilder(UWorld* InWorld, FName Id);
 
-	/** this will attempt to trigger the associated real shader in respect to the current (mobile) emulated versions used by the editor and copy
-	 * obtained number of instructions in place of the emulated ones
-	 *
-	 * @param	QualityLevel	Material quality setting for which the shaders to be compiled
-	 * @param	FeatureLevel	RHI feature level
-	 * @param	Materials		Set of materials to compile shaders for
-	 */
-	static bool CompileShadersComplexityViewMode(EMaterialQualityLevel::Type QualityLevel, ERHIFeatureLevel::Type FeatureLevel, TSet<class UMaterialInterface*>& Materials, struct FSlowTask& ProgressTask);
-
 	/** Intentionally hide constructors, etc. to prevent instantiation */
 	FEditorBuildUtils();
 	~FEditorBuildUtils();
