Components filtering is applied even if there's no custom filter on SComponentClassCombpo in order to apply the global class filter (FClassViewerModule:GetGlobalClassViewerFilter)
Bonuses:
- Better handling of uncooked vs cooked BP assets and dealing with aberrations where the BPName_C naming convention is not followed
- Use class path as keys rather than class name to avoid clashes

diff --git a/Engine/Source/Editor/UnrealEd/Private/ComponentTypeRegistry.cpp b/Engine/Source/Editor/UnrealEd/Private/ComponentTypeRegistry.cpp
--- a/Engine/Source/Editor/UnrealEd/Private/ComponentTypeRegistry.cpp
+++ b/Engine/Source/Editor/UnrealEd/Private/ComponentTypeRegistry.cpp
@@ -418,15 +418,15 @@ FComponentTypeRegistryData::FComponentTypeRegistryData()
 		Parent->PendingAssetData.Push(Data);
 	};
 
-	FAssetRegistryModule& AssetRegistryModule = FModuleManager::LoadModuleChecked<FAssetRegistryModule>(TEXT("AssetRegistry"));
-	AssetRegistryModule.Get().OnAssetAdded().AddStatic(HandleAdded, this);
-	AssetRegistryModule.Get().OnAssetRemoved().AddStatic(HandleAdded, this);
+	IAssetRegistry& AssetRegistry = FModuleManager::LoadModuleChecked<FAssetRegistryModule>(AssetRegistryConstants::ModuleName).Get();
+	AssetRegistry.OnAssetAdded().AddStatic(HandleAdded, this);
+	AssetRegistry.OnAssetRemoved().AddStatic(HandleAdded, this);
 
 	const auto HandleRenamed = [](const FAssetData& Data, const FString&, FComponentTypeRegistryData* Parent)
 	{
 		Parent->PendingAssetData.Push(Data);
 	};
-	AssetRegistryModule.Get().OnAssetRenamed().AddStatic(HandleRenamed, this);
+	AssetRegistry.OnAssetRenamed().AddStatic(HandleRenamed, this);
 }
 
 void FComponentTypeRegistryData::ForceRefreshComponentList()
@@ -494,22 +494,20 @@ void FComponentTypeRegistryData::ForceRefreshComponentList()
 
 	AddBasicShapeComponents(SortedClassList);
 
-	TArray<FName> InMemoryClasses;
+	// Add loaded component classes
+	TSet<FName> InMemoryClassPaths;
 	for (TObjectIterator<UClass> It; It; ++It)
 	{
 		UClass* Class = *It;
 		// If this is a subclass of Actor Component, not abstract, and tagged as spawnable from Kismet
 		if (Class->IsChildOf(UActorComponent::StaticClass()))
 		{
-			InMemoryClasses.Push(Class->GetFName());
+			InMemoryClassPaths.Add(FName(Class->GetPathName()));
 
 			const bool bOutOfDateClass = Class->HasAnyClassFlags(CLASS_NewerVersionExists);
 			const bool bBlueprintSkeletonClass = FKismetEditorUtilities::IsClassABlueprintSkeleton(Class);
-			const bool bPassesAllowedClasses = GetDefault<UClassViewerSettings>()->AllowedClasses.Num() == 0 || GetDefault<UClassViewerSettings>()->AllowedClasses.Contains(Class->GetName());
 
-			if (!bOutOfDateClass &&
-				!bBlueprintSkeletonClass &&
-				bPassesAllowedClasses)
+			if (!bOutOfDateClass &&	!bBlueprintSkeletonClass)
 			{
 				if (FKismetEditorUtilities::IsClassABlueprintSpawnableComponent(Class))
 				{
@@ -542,90 +540,86 @@ void FComponentTypeRegistryData::ForceRefreshComponentList()
 				}
 
 				FComponentTypeEntry Entry = { Class->GetName(), FString(), Class };
-				ComponentTypeList.Add(Entry);
+				ComponentTypeList.Add(MoveTemp(Entry));
 			}
 		}
 	}
 
+	// Add unloaded component classes
 	{
-		// make sure that we add any user created classes immediately, generally this will not create anything (because assets have not been discovered yet), 
-		// but asset discovery should be allowed to take place at any time:
-		FAssetRegistryModule& AssetRegistryModule = FModuleManager::LoadModuleChecked<FAssetRegistryModule>(TEXT("AssetRegistry"));
-		TArray<FName> ClassNames;
-		ClassNames.Add(UActorComponent::StaticClass()->GetFName());
-		TSet<FName> DerivedClassNames;
-		AssetRegistryModule.Get().GetDerivedClassNames(ClassNames, TSet<FName>(), DerivedClassNames);
-
-		TSet<FName> InMemoryClassesSet = TSet<FName>(InMemoryClasses);
-		TSet<FName> OnDiskClasses = DerivedClassNames.Difference(InMemoryClassesSet);
-
-		if (OnDiskClasses.Num() > 0)
+		auto AddUnloadedComponentClass = [this, &SortedClassList](FName ClassPath, const FAssetData& AssetData)
 		{
-			TMap<FString, FAssetData> BlueprintNames;
-			{
-				// GetAssetsByClass call is a kludge to get the full asset paths for the blueprints we care about
-				// Bob T. thinks that the Asset Registry could just keep asset paths
-				TArray<FAssetData> Assets;
-				AssetRegistryModule.Get().GetAssetsByClass(UBlueprint::StaticClass()->GetFName(), Assets, true);
-				for (FAssetData& Asset : Assets)
-				{
-					BlueprintNames.Add(Asset.AssetName.ToString(), MoveTemp(Asset));
-				}
+			// The blueprint is unloaded, so we need to work out which icon to use for it using its asset data
+			const UClass* IconClass = FClassIconFinder::GetIconClassForAssetData(AssetData);
 
-				Assets.Reset();
-				AssetRegistryModule.Get().GetAssetsByClass(UBlueprintGeneratedClass::StaticClass()->GetFName(), Assets, true);
-				for (FAssetData& Asset : Assets)
-				{
-					FString BlueprintName = Asset.AssetName.ToString();
-					BlueprintName.RemoveFromEnd(TEXT("_C"));
-					BlueprintNames.Add(MoveTemp(BlueprintName), MoveTemp(Asset));
-				}
-			}
+			FString ClassPathStr = ClassPath.ToString();
+			FString ClassName = FPackageName::ObjectPathToObjectName(ClassPathStr);
+			ClassName.RemoveFromEnd(TEXT("_C"));
+
+			FComponentClassComboEntryPtr NewEntry(new FComponentClassComboEntry(BlueprintComponents, ClassName, ClassPath, IconClass, /*bIncludeInFilter=*/true));
+			SortedClassList.Add(NewEntry);
 
-			for (const FName& OnDiskClass : OnDiskClasses)
+			// Create an unloaded blueprint data object to assist with class filtering
 			{
-				FString FixedString = OnDiskClass.ToString();
-				FixedString.RemoveFromEnd(TEXT("_C"));
+				using namespace UE::Editor::ComponentTypeRegistry;
+				TSharedPtr<IUnloadedBlueprintData> UnloadedBlueprintData = MakeShared<Private::FUnloadedBlueprintData>(AssetData);
 
-				const bool bPassesAllowedClasses = GetDefault<UClassViewerSettings>()->AllowedClasses.Num() == 0 || GetDefault<UClassViewerSettings>()->AllowedClasses.Contains(FixedString);
-				if (bPassesAllowedClasses)
-				{
-					FAssetData AssetData;
-					if (const FAssetData* Value = BlueprintNames.Find(FixedString))
-					{
-						AssetData = *Value;
-					}
+				const uint32 ClassFlags = AssetData.GetTagValueRef<uint32>(FBlueprintTags::ClassFlags);
+				UnloadedBlueprintData->SetClassFlags(ClassFlags);
 
-					FComponentTypeEntry Entry = { FixedString, AssetData.ObjectPath.ToString(), nullptr };
-					ComponentTypeList.Add(Entry);
+				const FString BlueprintType = AssetData.GetTagValueRef<FString>(FBlueprintTags::BlueprintType);
+				UnloadedBlueprintData->SetNormalBlueprintType(BlueprintType == TEXT("BPType_Normal"));
 
-					// The blueprint is unloaded, so we need to work out which icon to use for it using its asset data
-					const UClass* BlueprintIconClass = FClassIconFinder::GetIconClassForAssetData(AssetData);
+				NewEntry->SetUnloadedBlueprintData(UnloadedBlueprintData);
+			}
 
-					const bool bIncludeInFilter = true;
-					FComponentClassComboEntryPtr NewEntry = MakeShared<FComponentClassComboEntry>(BlueprintComponents, FixedString, AssetData.ObjectPath, BlueprintIconClass, bIncludeInFilter);
-					
-					// Create an unloaded blueprint data object to assist with class filtering
-					TSharedPtr<IUnloadedBlueprintData> UnloadedBlueprintData;
-					{
-						using namespace UE::Editor::ComponentTypeRegistry;
-						UnloadedBlueprintData = MakeShared<Private::FUnloadedBlueprintData>(AssetData);
+			FComponentTypeEntry Entry = { MoveTemp(ClassName), MoveTemp(ClassPathStr), nullptr };
+			ComponentTypeList.Add(MoveTemp(Entry));
+		};
 
-						const uint32 ClassFlags = AssetData.GetTagValueRef<uint32>(FBlueprintTags::ClassFlags);
-						UnloadedBlueprintData->SetClassFlags(ClassFlags);
+		IAssetRegistry& AssetRegistry = FModuleManager::LoadModuleChecked<FAssetRegistryModule>(AssetRegistryConstants::ModuleName).Get();
 
-						const FString BlueprintType = AssetData.GetTagValueRef<FString>(FBlueprintTags::BlueprintType);
-						UnloadedBlueprintData->SetNormalBlueprintType(BlueprintType == TEXT("BPType_Normal"));
+		TSet<FName> ActorComponentDerivedClassNames;
+		{
+			TArray<FName> ActorComponentClassNames;
+			ActorComponentClassNames.Add(UActorComponent::StaticClass()->GetFName());
+			AssetRegistry.GetDerivedClassNames(ActorComponentClassNames, TSet<FName>(), ActorComponentDerivedClassNames);
+		}
 
-						NewEntry->SetUnloadedBlueprintData(UnloadedBlueprintData);
-					}
-					
-					SortedClassList.Add(NewEntry);
+		// GetAssetsByClass call is a kludge to get the full asset paths for the blueprints we care about
+		// Bob T. thinks that the Asset Registry could just keep asset paths
+		TArray<FAssetData> Assets;
+
+		AssetRegistry.GetAssetsByClass(UBlueprint::StaticClass()->GetFName(), Assets, true);
+		for (FAssetData& BPAsset : Assets)
+		{
+			const FName ClassPath = FEditorClassUtils::GetClassPathFromAssetTag(BPAsset);
+			if (!ClassPath.IsNone())
+			{
+				if (!InMemoryClassPaths.Contains(ClassPath) && 
+					ActorComponentDerivedClassNames.Contains(FName(FPackageName::ObjectPathToObjectName(FStringView(FNameBuilder(ClassPath))))))
+				{
+					AddUnloadedComponentClass(ClassPath, BPAsset);
 				}
 			}
+			else
+			{
+				UE_LOG(LogBlueprint, Warning, TEXT("Blueprint %s is missing %s asset tag"), *BPAsset.PackageName.ToString(), *FBlueprintTags::GeneratedClassPath.ToString());
+			}
+
 		}
-	}
 
+		Assets.Reset();
+		AssetRegistry.GetAssetsByClass(UBlueprintGeneratedClass::StaticClass()->GetFName(), Assets, true);
+		for (FAssetData& BPGCAsset : Assets)
+		{
+			if (ActorComponentDerivedClassNames.Contains(BPGCAsset.AssetName) && !InMemoryClassPaths.Contains(BPGCAsset.ObjectPath))
+			{
+				AddUnloadedComponentClass(BPGCAsset.ObjectPath, BPGCAsset);
+			}
+		}
+	}
+	
 	if (SortedClassList.Num() > 0)
 	{
 		Sort(SortedClassList.GetData(), SortedClassList.Num(), SortComboEntry());
@@ -770,6 +764,11 @@ void FComponentTypeRegistry::InvalidateClass(TSubclassOf<UActorComponent> /*Clas
 	Data->Invalidate();
 }
 
+void FComponentTypeRegistry::Invalidate()
+{
+	Data->Invalidate();
+}
+
 FComponentClassComboEntryPtr FComponentTypeRegistry::FindClassEntryForObjectPath(FName InObjectPath) const
 {
 	if (int32* ClassListIndexPtr = Data->ClassPathToClassListIndexMap.Find(InObjectPath))
diff --git a/Engine/Source/Editor/UnrealEd/Private/SComponentClassCombo.cpp b/Engine/Source/Editor/UnrealEd/Private/SComponentClassCombo.cpp
--- a/Engine/Source/Editor/UnrealEd/Private/SComponentClassCombo.cpp
+++ b/Engine/Source/Editor/UnrealEd/Private/SComponentClassCombo.cpp
@@ -54,17 +54,11 @@ void SComponentClassCombo::Construct(const FArguments& InArgs)
 	OnComponentClassSelected = InArgs._OnComponentClassSelected;
 	OnSubobjectClassSelected = InArgs._OnSubobjectClassSelected;
 	TextFilter = MakeShared<FTextFilterExpressionEvaluator>(ETextFilterExpressionEvaluatorMode::BasicString);
-	
-	if(InArgs._CustomClassFilters.Num() > 0)
-	{
-		ComponentClassFilterData = MakeUnique<FComponentClassFilterData>();
 
-		ComponentClassFilterData->InitOptions = MakeShared<FClassViewerInitializationOptions>();
-		ComponentClassFilterData->InitOptions->ClassFilters.Append(InArgs._CustomClassFilters);
-
-		ComponentClassFilterData->ClassFilter = FModuleManager::LoadModuleChecked<FClassViewerModule>("ClassViewer").CreateClassFilter(*ComponentClassFilterData->InitOptions);
-		ComponentClassFilterData->FilterFuncs = FModuleManager::LoadModuleChecked<FClassViewerModule>("ClassViewer").CreateFilterFuncs();
-	}
+	ComponentClassFilterData.InitOptions = MakeShared<FClassViewerInitializationOptions>();
+	ComponentClassFilterData.InitOptions->ClassFilters.Append(InArgs._CustomClassFilters);
+	ComponentClassFilterData.ClassFilter = FModuleManager::LoadModuleChecked<FClassViewerModule>("ClassViewer").CreateClassFilter(*ComponentClassFilterData.InitOptions);
+	ComponentClassFilterData.FilterFuncs = FModuleManager::LoadModuleChecked<FClassViewerModule>("ClassViewer").CreateFilterFuncs();
 
 	FComponentTypeRegistry::Get().SubscribeToComponentList(ComponentClassList).AddRaw(this, &SComponentClassCombo::UpdateComponentClassList);
 
@@ -168,87 +162,79 @@ void SComponentClassCombo::ClearSelection()
 
 void SComponentClassCombo::GenerateFilteredComponentList()
 {
-	const bool bNoClassFilter = !ComponentClassFilterData.IsValid();
-	const bool bHasFilterText = !TextFilter->GetFilterText().IsEmpty();
+	FilteredComponentClassList.Reset();
 
-	if (bNoClassFilter && !bHasFilterText)
-	{
-		FilteredComponentClassList = *ComponentClassList;
-	}
-	else
-	{
-		FilteredComponentClassList.Empty();
+	int32 LastHeadingIndex = INDEX_NONE;
+	FComponentClassComboEntryPtr* LastHeadingPtr = nullptr;
 
-		int32 LastHeadingIndex = INDEX_NONE;
-		FComponentClassComboEntryPtr* LastHeadingPtr = nullptr;
+	int32 LastSeparatorIndex = INDEX_NONE;
+	FComponentClassComboEntryPtr* LastSeparatorPtr = nullptr;
 
-		int32 LastSeparatorIndex = INDEX_NONE;
-		FComponentClassComboEntryPtr* LastSeparatorPtr = nullptr;
+	const bool bHasFilterText = !TextFilter->GetFilterText().IsEmpty();
 
-		for (int32 ComponentIndex = 0; ComponentIndex < ComponentClassList->Num(); ComponentIndex++)
-		{
-			FComponentClassComboEntryPtr& CurrentEntry = (*ComponentClassList)[ComponentIndex];
+	for (int32 ComponentIndex = 0; ComponentIndex < ComponentClassList->Num(); ComponentIndex++)
+	{
+		FComponentClassComboEntryPtr& CurrentEntry = (*ComponentClassList)[ComponentIndex];
 
-			if (CurrentEntry->IsHeading())
-			{
-				LastHeadingIndex = FilteredComponentClassList.Num();
-				LastHeadingPtr = &CurrentEntry;
-			}
-			else if (CurrentEntry->IsSeparator())
+		if (CurrentEntry->IsHeading())
+		{
+			LastHeadingIndex = FilteredComponentClassList.Num();
+			LastHeadingPtr = &CurrentEntry;
+		}
+		else if (CurrentEntry->IsSeparator())
+		{
+			LastSeparatorIndex = FilteredComponentClassList.Num();
+			LastSeparatorPtr = &CurrentEntry;
+		}
+		else if(CurrentEntry->IsClass())
+		{
+			// Disallow class entries that are not to be seen when searching via text.
+			bool bAllowEntry = !bHasFilterText || CurrentEntry->IsIncludedInFilter();
+			if(bAllowEntry)
 			{
-				LastSeparatorIndex = FilteredComponentClassList.Num();
-				LastSeparatorPtr = &CurrentEntry;
+				// Disallow class entries that don't match the custom class filter, if set.
+				bAllowEntry = IsComponentClassAllowed(CurrentEntry);
+				if (bAllowEntry && bHasFilterText)
+				{
+					// Finally, disallow class entries that don't match the search box text.
+					FString FriendlyComponentName = GetSanitizedComponentName(CurrentEntry);
+					bAllowEntry = TextFilter->TestTextFilter(FBasicStringFilterExpressionContext(FriendlyComponentName));
+				}
 			}
-			else if(CurrentEntry->IsClass())
+
+			if (bAllowEntry)
 			{
-				// Disallow class entries that are not to be seen when searching via text.
-				bool bAllowEntry = !bHasFilterText || CurrentEntry->IsIncludedInFilter();
-				if(bAllowEntry)
+				// Add the heading first if it hasn't already been added
+				if (LastHeadingPtr && LastHeadingIndex != INDEX_NONE)
 				{
-					// Disallow class entries that don't match the custom class filter, if set.
-					bAllowEntry = bNoClassFilter || IsComponentClassAllowed(CurrentEntry);
-					if (bAllowEntry && bHasFilterText)
-					{
-						// Finally, disallow class entries that don't match the search box text.
-						FString FriendlyComponentName = GetSanitizedComponentName(CurrentEntry);
-						bAllowEntry = TextFilter->TestTextFilter(FBasicStringFilterExpressionContext(FriendlyComponentName));
-					}
+					FilteredComponentClassList.Insert(*LastHeadingPtr, LastHeadingIndex);
+					LastHeadingIndex = INDEX_NONE;
+					LastHeadingPtr = nullptr;
 				}
 
-				if (bAllowEntry)
+				// Add the separator next so that it will precede the heading
+				if (LastSeparatorPtr && LastSeparatorIndex != INDEX_NONE)
 				{
-					// Add the heading first if it hasn't already been added
-					if (LastHeadingPtr && LastHeadingIndex != INDEX_NONE)
-					{
-						FilteredComponentClassList.Insert(*LastHeadingPtr, LastHeadingIndex);
-						LastHeadingIndex = INDEX_NONE;
-						LastHeadingPtr = nullptr;
-					}
-
-					// Add the separator next so that it will precede the heading
-					if (LastSeparatorPtr && LastSeparatorIndex != INDEX_NONE)
-					{
-						FilteredComponentClassList.Insert(*LastSeparatorPtr, LastSeparatorIndex);
-						LastSeparatorIndex = INDEX_NONE;
-						LastSeparatorPtr = nullptr;
-					}
-
-					// Add the class
-					FilteredComponentClassList.Add(CurrentEntry);
+					FilteredComponentClassList.Insert(*LastSeparatorPtr, LastSeparatorIndex);
+					LastSeparatorIndex = INDEX_NONE;
+					LastSeparatorPtr = nullptr;
 				}
+
+				// Add the class
+				FilteredComponentClassList.Add(CurrentEntry);
 			}
 		}
+	}
 
-		if (ComponentClassListView.IsValid())
+	if (ComponentClassListView.IsValid())
+	{
+		// Select the first non-category item that passed the filter
+		for (FComponentClassComboEntryPtr& TestEntry : FilteredComponentClassList)
 		{
-			// Select the first non-category item that passed the filter
-			for (FComponentClassComboEntryPtr& TestEntry : FilteredComponentClassList)
+			if (TestEntry->IsClass())
 			{
-				if (TestEntry->IsClass())
-				{
-					ComponentClassListView->SetSelection(TestEntry, ESelectInfo::OnNavigation);
-					break;
-				}
+				ComponentClassListView->SetSelection(TestEntry, ESelectInfo::OnNavigation);
+				break;
 			}
 		}
 	}
@@ -581,22 +567,22 @@ TSharedRef<SToolTip> SComponentClassCombo::GetComponentToolTip(FComponentClassCo
 
 bool SComponentClassCombo::IsComponentClassAllowed(FComponentClassComboEntryPtr Entry) const
 {
-	if (Entry.IsValid() && Entry->IsClass() && ComponentClassFilterData.IsValid())
+	if (Entry.IsValid() && Entry->IsClass())
 	{
-		check(ComponentClassFilterData->InitOptions.IsValid());
-		check(ComponentClassFilterData->ClassFilter.IsValid());
-		check(ComponentClassFilterData->FilterFuncs.IsValid());
+		check(ComponentClassFilterData.InitOptions.IsValid());
+		check(ComponentClassFilterData.ClassFilter.IsValid());
+		check(ComponentClassFilterData.FilterFuncs.IsValid());
 
 		if (const UClass* ComponentClass = Entry->GetComponentClass())
 		{
-			return ComponentClassFilterData->ClassFilter->IsClassAllowed(*ComponentClassFilterData->InitOptions, ComponentClass, ComponentClassFilterData->FilterFuncs.ToSharedRef());
+			return ComponentClassFilterData.ClassFilter->IsClassAllowed(*ComponentClassFilterData.InitOptions, ComponentClass, ComponentClassFilterData.FilterFuncs.ToSharedRef());
 		}
 		else
 		{
 			TSharedPtr<IUnloadedBlueprintData> UnloadedBlueprintData = Entry->GetUnloadedBlueprintData();
 			if (UnloadedBlueprintData.IsValid())
 			{
-				return ComponentClassFilterData->ClassFilter->IsUnloadedClassAllowed(*ComponentClassFilterData->InitOptions, UnloadedBlueprintData.ToSharedRef(), ComponentClassFilterData->FilterFuncs.ToSharedRef());
+				return ComponentClassFilterData.ClassFilter->IsUnloadedClassAllowed(*ComponentClassFilterData.InitOptions, UnloadedBlueprintData.ToSharedRef(), ComponentClassFilterData.FilterFuncs.ToSharedRef());
 			}
 		}
 	}
diff --git a/Engine/Source/Editor/UnrealEd/Public/ComponentTypeRegistry.h b/Engine/Source/Editor/UnrealEd/Public/ComponentTypeRegistry.h
--- a/Engine/Source/Editor/UnrealEd/Public/ComponentTypeRegistry.h
+++ b/Engine/Source/Editor/UnrealEd/Public/ComponentTypeRegistry.h
@@ -43,6 +43,9 @@ struct UNREALED_API FComponentTypeRegistry
 	 */
 	void InvalidateClass(TSubclassOf<UActorComponent> ClassToUpdate);
 
+	/** Schedules a full update of the component type registry on the next frame */
+	void Invalidate();
+
 	/**
 	 * Attempts to locate the class entry corresponding to the given object path.
 	 */
diff --git a/Engine/Source/Editor/UnrealEd/Public/SComponentClassCombo.h b/Engine/Source/Editor/UnrealEd/Public/SComponentClassCombo.h
--- a/Engine/Source/Editor/UnrealEd/Public/SComponentClassCombo.h
+++ b/Engine/Source/Editor/UnrealEd/Public/SComponentClassCombo.h
@@ -283,7 +283,7 @@ private:
 	TSharedPtr< SListView<FComponentClassComboEntryPtr> > ComponentClassListView;
 
 	/** Internal data that facilitates custom class filtering */
-	TUniquePtr<FComponentClassFilterData> ComponentClassFilterData;
+	FComponentClassFilterData ComponentClassFilterData;
 
 	/** Cached selection index used to skip over unselectable items */
 	int32 PrevSelectedIndex;
