Fix compile errors in SComponentClassCombo (because of merged changes that don't account for UE5 only code)

diff --git a/Engine/Source/Editor/UnrealEd/Private/SComponentClassCombo.cpp b/Engine/Source/Editor/UnrealEd/Private/SComponentClassCombo.cpp
--- a/Engine/Source/Editor/UnrealEd/Private/SComponentClassCombo.cpp
+++ b/Engine/Source/Editor/UnrealEd/Private/SComponentClassCombo.cpp
@@ -593,16 +593,16 @@ bool SComponentClassCombo::IsComponentClassAllowed(FComponentClassComboEntryPtr
 
 void SComponentClassCombo::GetComponentClassFilterOptions(TArray<TSharedRef<FClassViewerFilterOption>>& OutFilterOptions) const
 {
-	if (ComponentClassFilterData.IsValid() && ComponentClassFilterData->InitOptions.IsValid())
+	if (ComponentClassFilterData.InitOptions.IsValid())
 	{
 		TArray<TSharedRef<FClassViewerFilterOption>> FilterOptions;
-		if(ComponentClassFilterData->ClassFilter.IsValid())
+		if(ComponentClassFilterData.ClassFilter.IsValid())
 		{
-			ComponentClassFilterData->ClassFilter->GetFilterOptions(FilterOptions);
+			ComponentClassFilterData.ClassFilter->GetFilterOptions(FilterOptions);
 			OutFilterOptions.Append(FilterOptions);
 		}
 
-		for (const TSharedRef<IClassViewerFilter>& ClassFilter : ComponentClassFilterData->InitOptions->ClassFilters)
+		for (const TSharedRef<IClassViewerFilter>& ClassFilter : ComponentClassFilterData.InitOptions->ClassFilters)
 		{
 			FilterOptions.Reset();
 			ClassFilter->GetFilterOptions(FilterOptions);
