Modify DumpMaterialShaderTypes commandlet to add FMaterialShader types to the final shader type histogram

diff --git a/Engine/Source/Editor/UnrealEd/Private/Commandlets/DumpMaterialShaderTypes.cpp b/Engine/Source/Editor/UnrealEd/Private/Commandlets/DumpMaterialShaderTypes.cpp
--- a/Engine/Source/Editor/UnrealEd/Private/Commandlets/DumpMaterialShaderTypes.cpp
+++ b/Engine/Source/Editor/UnrealEd/Private/Commandlets/DumpMaterialShaderTypes.cpp
@@ -269,6 +269,8 @@ void PrintDebugShaderInfo(FShaderStatsGatheringContext& Output, const TArray<FDe
 				{
 					ShaderTypeMap.FindOrAdd(ShaderTypeName, 1);
 				}
+
+				Output.AddToHistogram(nullptr, nullptr, ShaderType->GetName());
 			}
 
 			if (ShaderTypeMap.Num() > 0)
