WorldPartition: Newly created maps now prompt for save if they were modified

diff --git a/Engine/Source/Editor/UnrealEd/Private/FileHelpers.cpp b/Engine/Source/Editor/UnrealEd/Private/FileHelpers.cpp
--- a/Engine/Source/Editor/UnrealEd/Private/FileHelpers.cpp
+++ b/Engine/Source/Editor/UnrealEd/Private/FileHelpers.cpp
@@ -841,7 +841,8 @@ static bool SaveWorld(UWorld* World,
 			// This will save external actors
 			if (bSuccess)
 			{
-				bSuccess = SaveExternalPackages(Package, /*bCheckDirty*/true);
+				// Do not check dirty flag if world is newly created
+				bSuccess = SaveExternalPackages(Package, /*bCheckDirty*/!bNewlyCreated);
 			}
 		}
 				
@@ -4486,6 +4487,7 @@ void FEditorFileUtils::GetDirtyWorldPackages(TArray<UPackage*>& OutDirtyPackages
 			&& (!bHasWritableFolderFilter || WritableFolderFilter->PassesStartsWithFilter(WorldPackage->GetName()))
 			)
 		{
+			bool bDirtyNewWorldPackage = false;
 			if (WorldPackage->IsDirty() && !ShouldIgnorePackageFunction(WorldPackage))
 			{
 				// IF the package is dirty and its not a pie package, add the world package to the list of packages to save
@@ -4511,21 +4513,7 @@ void FEditorFileUtils::GetDirtyWorldPackages(TArray<UPackage*>& OutDirtyPackages
 
 					if (BuiltDataPackage->IsDirty())
 					{
-						// If built data package does not have a name yet add the world package so a user is prompted to have a name chosen
-						if (!WorldPackage->IsDirty())
-						{
-							const FString WorldPackageName = WorldPackage->GetName();
-							const bool bIncludeReadOnlyRoots = false;
-							const bool bIsValidPath = FPackageName::IsValidLongPackageName(WorldPackageName, bIncludeReadOnlyRoots);
-							if (!bIsValidPath)
-							{
-								WorldPackage->MarkPackageDirty();
-								if (!ShouldIgnorePackageFunction(WorldPackage))
-								{
-									OutDirtyPackages.Add(WorldPackage);
-								}
-							}
-						}
+						bDirtyNewWorldPackage = true;
 
 						if (!ShouldIgnorePackageFunction(BuiltDataPackage))
 						{
@@ -4559,38 +4547,53 @@ void FEditorFileUtils::GetDirtyWorldPackages(TArray<UPackage*>& OutDirtyPackages
 			{
 				for (UPackage* ExternalPackage : WorldIt->PersistentLevel->GetLoadedExternalObjectPackages())
 				{
-					if (ExternalPackage->IsDirty() && !ShouldIgnorePackageFunction(ExternalPackage))
+					if (ExternalPackage->IsDirty())
 					{
-						bool bActorPackageNeedsToSave = true;
+						bDirtyNewWorldPackage = true;
 
-						// Skip unsaved packages containing only pending kill actors
-						if (ExternalPackage->HasAnyPackageFlags(PKG_NewlyCreated))
+						if (!ShouldIgnorePackageFunction(ExternalPackage))
 						{
-							bActorPackageNeedsToSave = false;
-							ForEachObjectWithPackage(ExternalPackage, [&bActorPackageNeedsToSave](UObject* Object)
+							bool bActorPackageNeedsToSave = true;
+
+							// Skip unsaved packages containing only pending kill actors
+							if (ExternalPackage->HasAnyPackageFlags(PKG_NewlyCreated))
 							{
-								// @todo_ow: Find better way
-								if (Object->IsA<AActor>() || Object->IsA<UActorFolder>())
-								{
-									if (IsValidChecked(Object))
+								bActorPackageNeedsToSave = false;
+								ForEachObjectWithPackage(ExternalPackage, [&bActorPackageNeedsToSave](UObject* Object)
 									{
-										bActorPackageNeedsToSave = true;
-										return false;
-									}
-								}
-								return true;
-							}, false);
-						}
+										// @todo_ow: Find better way
+										if (Object->IsA<AActor>() || Object->IsA<UActorFolder>())
+										{
+											if (IsValidChecked(Object))
+											{
+												bActorPackageNeedsToSave = true;
+												return false;
+											}
+										}
+										return true;
+									}, false);
+							}
 
-						// Filter out Actors that might be unsaved (/Temp folder)
-						bActorPackageNeedsToSave &= FPackageName::IsValidLongPackageName(ExternalPackage->GetName());
-						if (bActorPackageNeedsToSave)
-						{
-							OutDirtyPackages.Add(ExternalPackage);
+							// Filter out Actors that might be unsaved (/Temp folder)
+							bActorPackageNeedsToSave &= FPackageName::IsValidLongPackageName(ExternalPackage->GetName());
+							if (bActorPackageNeedsToSave)
+							{
+								OutDirtyPackages.Add(ExternalPackage);
+							}
 						}
 					}
 				}
 			}
+
+			if (bDirtyNewWorldPackage && !WorldPackage->IsDirty() && WorldPackage->HasAnyPackageFlags(PKG_NewlyCreated))
+			{
+				// If world package does not have a name yet add the world package so a user is prompted to have a name chosen
+				WorldPackage->MarkPackageDirty();
+				if (!ShouldIgnorePackageFunction(WorldPackage))
+				{
+					OutDirtyPackages.Add(WorldPackage);
+				}
+			}
 		}
 	}
 }
diff --git a/Engine/Source/Runtime/Engine/Private/World.cpp b/Engine/Source/Runtime/Engine/Private/World.cpp
--- a/Engine/Source/Runtime/Engine/Private/World.cpp
+++ b/Engine/Source/Runtime/Engine/Private/World.cpp
@@ -2085,8 +2085,12 @@ UWorld* UWorld::CreateWorld(const EWorldType::Type InWorldType, bool bInformEngi
 	NewWorld->FeatureLevel = InFeatureLevel;
 	NewWorld->InitializeNewWorld(InIVS ? *InIVS : UWorld::InitializationValues().CreatePhysicsScene(InWorldType != EWorldType::Inactive).ShouldSimulatePhysics(false).EnableTraceCollision(true).CreateNavigation(InWorldType == EWorldType::Editor).CreateAISystem(InWorldType == EWorldType::Editor), bInSkipInitWorld);
 
-	// Clear the dirty flag set during SpawnActor and UpdateLevelComponents
+	// Clear the dirty flags set during SpawnActor and UpdateLevelComponents
 	WorldPackage->SetDirtyFlag(false);
+	for (UPackage* ExternalPackage : WorldPackage->GetExternalPackages())
+	{
+		ExternalPackage->SetDirtyFlag(false);
+	}
 
 	if ( bAddToRoot )
 	{
