Display runtime name/version in VR PIE title bar

diff --git a/Engine/Plugins/Runtime/Oculus/OculusVR/Source/OculusHMD/Private/OculusHMD.cpp b/Engine/Plugins/Runtime/Oculus/OculusVR/Source/OculusHMD/Private/OculusHMD.cpp
--- a/Engine/Plugins/Runtime/Oculus/OculusVR/Source/OculusHMD/Private/OculusHMD.cpp
+++ b/Engine/Plugins/Runtime/Oculus/OculusVR/Source/OculusHMD/Private/OculusHMD.cpp
@@ -125,7 +125,7 @@ namespace OculusHMD
 			Version = "Unknown";
 		}
 
-		return FString::Printf(TEXT("%s, OVRPlugin: %s"), *FEngineVersion::Current().ToString(), UTF8_TO_TCHAR(Version));
+		return FString::Printf(TEXT("OVRPlugin: %s"), UTF8_TO_TCHAR(Version));
 	}
 
 
diff --git a/Engine/Plugins/Runtime/OpenXR/Source/OpenXRHMD/Private/OpenXRHMD.cpp b/Engine/Plugins/Runtime/OpenXR/Source/OpenXRHMD/Private/OpenXRHMD.cpp
--- a/Engine/Plugins/Runtime/OpenXR/Source/OpenXRHMD/Private/OpenXRHMD.cpp
+++ b/Engine/Plugins/Runtime/OpenXR/Source/OpenXRHMD/Private/OpenXRHMD.cpp
@@ -311,7 +311,7 @@ FName FOpenXRHMD::GetHMDName() const
 
 FString FOpenXRHMD::GetVersionString() const
 {
-	return FString::Printf(TEXT("%s: %d.%d.%d"),
+	return FString::Printf(TEXT("%s (%d.%d.%d)"),
 		UTF8_TO_TCHAR(InstanceProperties.runtimeName),
 		XR_VERSION_MAJOR(InstanceProperties.runtimeVersion),
 		XR_VERSION_MINOR(InstanceProperties.runtimeVersion),
diff --git a/Engine/Source/Editor/UnrealEd/Private/PlayLevel.cpp b/Engine/Source/Editor/UnrealEd/Private/PlayLevel.cpp
--- a/Engine/Source/Editor/UnrealEd/Private/PlayLevel.cpp
+++ b/Engine/Source/Editor/UnrealEd/Private/PlayLevel.cpp
@@ -3102,10 +3102,12 @@ FText GeneratePIEViewportWindowTitle(const EPlayNetMode InNetMode, const ERHIFea
 		GEngine->XRSystem && GEngine->XRSystem.IsValid())
 	{
 		Args.Add(TEXT("XRSystemName"), FText::FromName(GEngine->XRSystem->GetSystemName()));
+		Args.Add(TEXT("XRRuntimeVersion"), FText::FromString(GEngine->XRSystem->GetVersionString()));
 	}
 	else
 	{
 		Args.Add(TEXT("XRSystemName"), FText::GetEmpty());
+		Args.Add(TEXT("XRRuntimeVersion"), FText::GetEmpty());
 	}
 
 	if (FixedTick > 0.f)
@@ -3118,7 +3120,7 @@ FText GeneratePIEViewportWindowTitle(const EPlayNetMode InNetMode, const ERHIFea
 		Args.Add(TEXT("FixedFPS"), FText::GetEmpty());
 	}
 
-	return FText::Format(NSLOCTEXT("UnrealEd", "PlayInEditor_WindowTitleFormat", "{GameName} Preview [NetMode: {NetMode}] {FixedFPS} ({PlatformBits}-bit/{RHIName}) {XRSystemName}"), Args);
+	return FText::TrimTrailing(FText::Format(NSLOCTEXT("UnrealEd", "PlayInEditor_WindowTitleFormat", "{GameName} Preview [NetMode: {NetMode}] {FixedFPS} ({PlatformBits}-bit/{RHIName}) {XRSystemName} {XRRuntimeVersion}"), Args));
 }
 
 void UEditorEngine::TransferEditorSelectionToPlayInstances(const bool bInSelectInstances)
