Fix an invalid BP compiler warning for generated Blueprint assets that may not have been otherwise loaded.

diff --git a/Engine/Source/Editor/UnrealEd/Private/Kismet2/BlueprintEditorUtils.cpp b/Engine/Source/Editor/UnrealEd/Private/Kismet2/BlueprintEditorUtils.cpp
--- a/Engine/Source/Editor/UnrealEd/Private/Kismet2/BlueprintEditorUtils.cpp
+++ b/Engine/Source/Editor/UnrealEd/Private/Kismet2/BlueprintEditorUtils.cpp
@@ -513,10 +513,10 @@ FBlueprintEditorUtils::FOnRefreshAllNodes FBlueprintEditorUtils::OnRefreshAllNod
 
 void FBlueprintEditorUtils::RefreshAllNodes(UBlueprint* Blueprint)
 {
-	if (!Blueprint || !Blueprint->HasAllFlags(RF_LoadCompleted))
+	if (!Blueprint || (Blueprint->HasAllFlags(RF_WasLoaded) && !Blueprint->HasAllFlags(RF_LoadCompleted)))
 	{
 		UE_LOG(LogBlueprint, Warning, 
-			TEXT("RefreshAllNodes called on incompletly loaded blueprint '%s'"), 
+			TEXT("RefreshAllNodes was called on an invalid or incompletely loaded blueprint '%s'"), 
 			Blueprint ? *Blueprint->GetFullName() : TEXT("NULL"));
 		return;
 	}
@@ -569,10 +569,10 @@ FBlueprintEditorUtils::FOnReconstructAllNodes FBlueprintEditorUtils::OnReconstru
 
 void FBlueprintEditorUtils::ReconstructAllNodes(UBlueprint* Blueprint)
 {
-	if (!Blueprint || !Blueprint->HasAllFlags(RF_LoadCompleted))
+	if (!Blueprint || (Blueprint->HasAllFlags(RF_WasLoaded) && !Blueprint->HasAllFlags(RF_LoadCompleted)))
 	{
 		UE_LOG(LogBlueprint, Warning,
-			TEXT("ReconstructAllNodes called on incompletly loaded blueprint '%s'"),
+			TEXT("ReconstructAllNodes was called on an invalid or incompletely loaded blueprint '%s'"),
 			Blueprint ? *Blueprint->GetFullName() : TEXT("NULL"));
 		return;
 	}
@@ -606,10 +606,10 @@ void FBlueprintEditorUtils::RefreshExternalBlueprintDependencyNodes(UBlueprint*
 {
 	BP_SCOPED_COMPILER_EVENT_STAT(EKismetCompilerStats_RefreshExternalDependencyNodes);
 
-	if (!Blueprint || !Blueprint->HasAllFlags(RF_LoadCompleted))
+	if (!Blueprint || (Blueprint->HasAllFlags(RF_WasLoaded) && !Blueprint->HasAllFlags(RF_LoadCompleted)))
 	{
 		UE_LOG(LogBlueprint, Warning,
-			TEXT("RefreshAllNodes called on incompletly loaded blueprint '%s'"),
+			TEXT("RefreshExternalBlueprintDependencyNodes was called on an invalid or incompletely loaded blueprint '%s'"),
 			Blueprint ? *Blueprint->GetFullName() : TEXT("NULL"));
 		return;
 	}
