Remove old cooking config to cook slate assets. If still needed, you may add the setting in your project settings.

diff --git a/Engine/Config/BaseEditor.ini b/Engine/Config/BaseEditor.ini
--- a/Engine/Config/BaseEditor.ini
+++ b/Engine/Config/BaseEditor.ini
@@ -238,13 +238,6 @@ bExpandClassPickerClassList=False
 +NewAssetDefaultClasses=(ClassName="/Script/Engine.ActorComponent", AssetClass="/Script/Engine.Blueprint")
 +NewAssetDefaultClasses=(ClassName="/Script/Engine.SceneComponent", AssetClass="/Script/Engine.Blueprint")
 
-[UI]
-; Directories specifying assets needed by Slate UI, assets in these directories are always cooked even if not referenced
-+ContentDirectories=/Game/UI
-+ContentDirectories=/Game/Widget
-+ContentDirectories=/Game/Widgets
-+ContentDirectories=/Engine/MobileResources
-
 [Cooker.CleanStaticMeshMtrlSkip]
 Class=UnrealEd.ActorFactory
 
diff --git a/Engine/Source/Editor/UnrealEd/Classes/CookOnTheSide/CookOnTheFlyServer.h b/Engine/Source/Editor/UnrealEd/Classes/CookOnTheSide/CookOnTheFlyServer.h
--- a/Engine/Source/Editor/UnrealEd/Classes/CookOnTheSide/CookOnTheFlyServer.h
+++ b/Engine/Source/Editor/UnrealEd/Classes/CookOnTheSide/CookOnTheFlyServer.h
@@ -69,7 +69,6 @@ enum class ECookByTheBookOptions
 	NoGameAlwaysCookPackages =			0x00000080, // don't include the packages specified by the game in the cook (this cook will probably be missing content unless you know what you are doing)
 	NoAlwaysCookMaps =					0x00000100, // don't include always cook maps (this cook will probably be missing content unless you know what you are doing)
 	NoDefaultMaps =						0x00000200, // don't include default cook maps (this cook will probably be missing content unless you know what you are doing)
-	NoSlatePackages =					0x00000400, // don't include slate content (this cook will probably be missing content unless you know what you are doing)
 	NoInputPackages =					0x00000800, // don't include slate content (this cook will probably be missing content unless you know what you are doing)
 	SkipSoftReferences =				0x00001000, // Don't follow soft references when cooking. Usually not viable for a real cook and the results probably wont load properly, but can be useful for debugging
 	SkipHardReferences =				0x00002000, // Don't follow hard references when cooking. Not viable for a real cook, only useful for debugging
@@ -79,6 +78,7 @@ enum class ECookByTheBookOptions
 	ZenStore =							0x00040000, // Store cooked data in Zen Store
 
 	// Deprecated flags
+	NoSlatePackages UE_DEPRECATED(5.0, "The [UI]ContentDirectories is deprecated. You may use DirectoriesToAlwaysCook in your project settings instead.") = 0x00000400, // don't include slate content
 	DisableUnsolicitedPackages UE_DEPRECATED(4.26, "Use SkipSoftReferences and/or SkipHardReferences instead") = SkipSoftReferences | SkipHardReferences,
 };
 ENUM_CLASS_FLAGS(ECookByTheBookOptions);
diff --git a/Engine/Source/Editor/UnrealEd/Private/Commandlets/CookCommandlet.cpp b/Engine/Source/Editor/UnrealEd/Private/Commandlets/CookCommandlet.cpp
--- a/Engine/Source/Editor/UnrealEd/Private/Commandlets/CookCommandlet.cpp
+++ b/Engine/Source/Editor/UnrealEd/Private/Commandlets/CookCommandlet.cpp
@@ -917,7 +917,7 @@ bool UCookCommandlet::CookByTheBook( const TArray<ITargetPlatform*>& Platforms)
 	if (bCookSinglePackage)
 	{
 		const ECookByTheBookOptions SinglePackageFlags = ECookByTheBookOptions::NoAlwaysCookMaps | ECookByTheBookOptions::NoDefaultMaps | ECookByTheBookOptions::NoGameAlwaysCookPackages |
-			ECookByTheBookOptions::NoInputPackages | ECookByTheBookOptions::NoSlatePackages | ECookByTheBookOptions::SkipSoftReferences | ECookByTheBookOptions::ForceDisableSaveGlobalShaders;
+			ECookByTheBookOptions::NoInputPackages | ECookByTheBookOptions::SkipSoftReferences | ECookByTheBookOptions::ForceDisableSaveGlobalShaders;
 		CookOptions |= SinglePackageFlags;
 		CookOptions |= bKeepSinglePackageRefs ? ECookByTheBookOptions::None : ECookByTheBookOptions::SkipHardReferences;
 	}
diff --git a/Engine/Source/Editor/UnrealEd/Private/Commandlets/GenerateDistillFileSetsCommandlet.cpp b/Engine/Source/Editor/UnrealEd/Private/Commandlets/GenerateDistillFileSetsCommandlet.cpp
--- a/Engine/Source/Editor/UnrealEd/Private/Commandlets/GenerateDistillFileSetsCommandlet.cpp
+++ b/Engine/Source/Editor/UnrealEd/Private/Commandlets/GenerateDistillFileSetsCommandlet.cpp
@@ -221,30 +221,12 @@ int32 UGenerateDistillFileSetsCommandlet::Main( const FString& InParams )
 	// Form a full unique package list
 	TSet<FString> AllPackageNames;
 
-	//@todo SLATE: This is a hack to ensure all slate referenced assets get cooked.
-	// Slate needs to be refactored to properly identify required assets at cook time.
-	// Simply jamming everything in a given directory into the cook list is error-prone
-	// on many levels - assets not required getting cooked/shipped; assets not put under 
-	// the correct folder; etc.
+	// Slate
 	{
 		TArray<FString> UIContentPaths;
 		if (GConfig->GetArray(TEXT("UI"), TEXT("ContentDirectories"), UIContentPaths, GEditorIni) > 0)
 		{
-			for (int32 DirIdx = 0; DirIdx < UIContentPaths.Num(); DirIdx++)
-			{
-				FString ContentPath = FPackageName::LongPackageNameToFilename(UIContentPaths[DirIdx]);
-
-				TArray<FString> Files;
-				IFileManager::Get().FindFilesRecursive(Files, *ContentPath, *(FString(TEXT("*")) + FPackageName::GetAssetPackageExtension()), true, false);
-				for (int32 Index = 0; Index < Files.Num(); Index++)
-				{
-					FString StdFile = Files[Index];
-					FPaths::MakeStandardFilename(StdFile);
-					StdFile = FPackageName::FilenameToLongPackageName(StdFile);
-					AllPackageNames.Add(StdFile);
-				}
-			}
-
+			UE_LOG(LogGenerateDistillFileSetsCommandlet, Warning, TEXT("The [UI]ContentDirectories is deprecated. You may use DirectoriesToAlwaysCook in your project settings instead."));
 		}
 	}
 
diff --git a/Engine/Source/Editor/UnrealEd/Private/CookOnTheFlyServer.cpp b/Engine/Source/Editor/UnrealEd/Private/CookOnTheFlyServer.cpp
--- a/Engine/Source/Editor/UnrealEd/Private/CookOnTheFlyServer.cpp
+++ b/Engine/Source/Editor/UnrealEd/Private/CookOnTheFlyServer.cpp
@@ -4695,7 +4695,7 @@ bool UCookOnTheFlyServer::Exec(class UWorld* InWorld, const TCHAR* Cmd, FOutputD
 				StartupOptions.CookMaps.Add(StandardPackageName.ToString());
 			}
 		}
-		StartupOptions.CookOptions = ECookByTheBookOptions::NoAlwaysCookMaps | ECookByTheBookOptions::NoDefaultMaps | ECookByTheBookOptions::NoGameAlwaysCookPackages | ECookByTheBookOptions::NoInputPackages | ECookByTheBookOptions::NoSlatePackages | ECookByTheBookOptions::SkipSoftReferences | ECookByTheBookOptions::ForceDisableSaveGlobalShaders;
+		StartupOptions.CookOptions = ECookByTheBookOptions::NoAlwaysCookMaps | ECookByTheBookOptions::NoDefaultMaps | ECookByTheBookOptions::NoGameAlwaysCookPackages | ECookByTheBookOptions::SkipSoftReferences | ECookByTheBookOptions::ForceDisableSaveGlobalShaders;
 		
 		StartCookByTheBook(StartupOptions);
 	}
@@ -6300,41 +6300,12 @@ void UCookOnTheFlyServer::CollectFilesToCook(TArray<FName>& FilesInPath, TMap<FN
 			}
 		}
 	}
-	//@todo SLATE: This is a hack to ensure all slate referenced assets get cooked.
-	// Slate needs to be refactored to properly identify required assets at cook time.
-	// Simply jamming everything in a given directory into the cook list is error-prone
-	// on many levels - assets not required getting cooked/shipped; assets not put under 
-	// the correct folder; etc.
-	if ( !(FilesToCookFlags & ECookByTheBookOptions::NoSlatePackages))
+
 	{
 		TArray<FString> UIContentPaths;
-		TSet <FName> ContentDirectoryAssets; 
 		if (GConfig->GetArray(TEXT("UI"), TEXT("ContentDirectories"), UIContentPaths, GEditorIni) > 0)
 		{
-			for (int32 DirIdx = 0; DirIdx < UIContentPaths.Num(); DirIdx++)
-			{
-				FString ContentPath = FPackageName::LongPackageNameToFilename(UIContentPaths[DirIdx]);
-				FName ContentPathName(*ContentPath);
-
-				TArray<FString> Files;
-				IFileManager::Get().FindFilesRecursive(Files, *ContentPath, *(FString(TEXT("*")) + FPackageName::GetAssetPackageExtension()), true, false);
-				for (int32 Index = 0; Index < Files.Num(); Index++)
-				{
-					FString StdFile = Files[Index];
-					FName PackageName = FName(*FPackageName::FilenameToLongPackageName(StdFile));
-					ContentDirectoryAssets.Add(PackageName);
-					FPaths::MakeStandardFilename(StdFile);
-					AddFileToCook( FilesInPath, Instigators, StdFile, FInstigator(EInstigator::UIContentDirectory, ContentPathName));
-				}
-			}
-		}
-
-		if (CookByTheBookOptions && CookByTheBookOptions->bGenerateDependenciesForMaps) 
-		{
-			for (auto& MapDependencyGraph : CookByTheBookOptions->MapDependencyGraphs)
-			{
-				MapDependencyGraph.Value.Add(FName(TEXT("ContentDirectoryAssets")), ContentDirectoryAssets);
-			}
+			UE_LOG(LogCook, Warning, TEXT("The [UI]ContentDirectories is deprecated. You may use DirectoriesToAlwaysCook in your project settings instead."));
 		}
 	}
 }
diff --git a/Engine/Source/Runtime/Engine/Public/Engine/ICookInfo.h b/Engine/Source/Runtime/Engine/Public/Engine/ICookInfo.h
--- a/Engine/Source/Runtime/Engine/Public/Engine/ICookInfo.h
+++ b/Engine/Source/Runtime/Engine/Public/Engine/ICookInfo.h
@@ -33,7 +33,6 @@ namespace UE::Cook
 	CallbackMacro(FullDepotSearch, true) \
 	CallbackMacro(GameDefaultObject, false) \
 	CallbackMacro(InputSettingsIni, true) \
-	CallbackMacro(UIContentDirectory, false) \
 	CallbackMacro(StartupSoftObjectPath, true) \
 	CallbackMacro(PackagingSettingsMapToCook, true) \
 	CallbackMacro(CookModificationDelegate, true) \
