WorldPartition - HLOD Builders - Added some trace bookmarks

diff --git a/Engine/Source/Editor/UnrealEd/Private/WorldPartition/WorldPartitionHLODsBuilder.cpp b/Engine/Source/Editor/UnrealEd/Private/WorldPartition/WorldPartitionHLODsBuilder.cpp
--- a/Engine/Source/Editor/UnrealEd/Private/WorldPartition/WorldPartitionHLODsBuilder.cpp
+++ b/Engine/Source/Editor/UnrealEd/Private/WorldPartition/WorldPartitionHLODsBuilder.cpp
@@ -444,6 +444,8 @@ bool UWorldPartitionHLODsBuilder::BuildHLODActors()
 
 	for (int32 CurrentActor = ResumeBuildIndex; CurrentActor < HLODActorsToBuild.Num(); ++CurrentActor)
 	{
+		TRACE_BOOKMARK(TEXT("BuildHLOD Start - %d"), CurrentActor);
+
 		const FGuid& HLODActorGuid = HLODActorsToBuild[CurrentActor];
 
 		FWorldPartitionReference ActorRef(WorldPartition, HLODActorGuid);
@@ -473,6 +475,8 @@ bool UWorldPartitionHLODsBuilder::BuildHLODActors()
 		{
 			FWorldPartitionHelpers::DoCollectGarbage();
 		}
+
+		TRACE_BOOKMARK(TEXT("BuildHLOD End - %d"), CurrentActor);
 	}
 
 	UE_LOG(LogWorldPartitionHLODsBuilder, Display, TEXT("#### Built %d HLOD actors ####"), HLODActorsToBuild.Num());
