Fix a reference replacement oversight during object consolidation to a duplicated Blueprint asset that could lead to an editor crash in some situations.

Also pull changes from CL 15794345 (initial fix for MH-3535) into the UE5 engine stream.

diff --git a/Engine/Source/Editor/UnrealEd/Private/ObjectTools.cpp b/Engine/Source/Editor/UnrealEd/Private/ObjectTools.cpp
--- a/Engine/Source/Editor/UnrealEd/Private/ObjectTools.cpp
+++ b/Engine/Source/Editor/UnrealEd/Private/ObjectTools.cpp
@@ -48,6 +48,7 @@
 #include "GameFramework/Volume.h"
 #include "UObject/MetaData.h"
 #include "Serialization/ArchiveReplaceObjectRef.h"
+#include "Serialization/ArchiveReplaceObjectAndStructPropertyRef.h"
 #include "GameFramework/WorldSettings.h"
 #include "Engine/BlueprintGeneratedClass.h"
 #include "PhysicalMaterials/PhysicalMaterial.h"
@@ -992,7 +993,7 @@ namespace ObjectTools
 		}
 
 		// Objects already loaded and in memory have to have any of their references to the objects to replace swapped with a reference to
-		// the "object to replace with". FArchiveReplaceObjectRef can serve this purpose, but it expects a TMap of object to replace : object to replace with.
+		// the "object to replace with". FArchiveReplaceObjectAndStructPropertyRef can serve this purpose, but it expects a TMap of object to replace : object to replace with.
 		// Therefore, populate a map with all of the valid objects to replace as keys, with the object to replace with as the value for each one.
 		TMap<UObject*, UObject*> ReplacementMap;
 		for ( TArray<UObject*>::TConstIterator ReplaceItr( ObjectsToReplace ); ReplaceItr; ++ReplaceItr )
@@ -1120,9 +1121,10 @@ namespace ObjectTools
 					UBlueprint* BPObjectToUpdate = Cast<UBlueprint>(CurObject);
 					if (BPObjectToUpdate)
 					{
-						FArchiveReplaceObjectRef<UObject> ReplaceAr(BPObjectToUpdate->GeneratedClass->ClassDefaultObject, ReplacementMap, EArchiveReplaceObjectFlags::IncludeClassGeneratedByRef);
+						FArchiveReplaceObjectAndStructPropertyRef<UObject> ReplaceInBPClassObject_Ar(BPObjectToUpdate->GeneratedClass, ReplacementMap, EArchiveReplaceObjectFlags::IncludeClassGeneratedByRef);
+						FArchiveReplaceObjectAndStructPropertyRef<UObject> ReplaceInBPClassDefaultObject_Ar(BPObjectToUpdate->GeneratedClass->ClassDefaultObject, ReplacementMap, EArchiveReplaceObjectFlags::IncludeClassGeneratedByRef);
 					}
-					FArchiveReplaceObjectRef<UObject> ReplaceAr(CurObject, ReplacementMap, EArchiveReplaceObjectFlags::IncludeClassGeneratedByRef);
+					FArchiveReplaceObjectAndStructPropertyRef<UObject> ReplaceAr(CurObject, ReplacementMap, EArchiveReplaceObjectFlags::IncludeClassGeneratedByRef);
 				}
 			}
 		}
@@ -1136,7 +1138,7 @@ namespace ObjectTools
 				GWarn->StatusUpdate( NumObjsReplaced, ReferencingPropertiesMapKeys.Num(), NSLOCTEXT("UnrealEd", "ConsolidateAssetsUpdate_ReplacingReferences", "Replacing Asset References...") );
 
 				UObject* CurReplaceObj = ReferencingPropertiesMapKeys[Index];
-				FArchiveReplaceObjectRef<UObject> ReplaceAr( CurReplaceObj, ReplacementMap, EArchiveReplaceObjectFlags::IncludeClassGeneratedByRef);
+				FArchiveReplaceObjectAndStructPropertyRef<UObject> ReplaceAr(CurReplaceObj, ReplacementMap, EArchiveReplaceObjectFlags::IncludeClassGeneratedByRef);
 			}
 		}
 		// Now alter the referencing objects the change has completed via PostEditChange,
@@ -1414,6 +1416,14 @@ namespace ObjectTools
 					}
 
 					ReplaceInfo.AppendUnique(GeneratedClassReplaceInfo);
+
+					// Find and cache all Blueprints that have a new dependency on the consolidation target after reference replacement.
+					TArray<UBlueprint*> DependentBPs;
+					FBlueprintEditorUtils::FindDependentBlueprints(ObjectToConsolidateTo_BP, DependentBPs);
+					for (UBlueprint* DependentBP : DependentBPs)
+					{
+						ObjectToConsolidateTo_BP->CachedDependents.Add(DependentBP);
+					}
 				}
 				DirtiedPackages.Append( ReplaceInfo.DirtiedPackages );
 				UnconsolidatableObjects.Append( ReplaceInfo.UnreplaceableObjects );
diff --git a/Engine/Source/Runtime/Engine/Classes/Engine/Blueprint.h b/Engine/Source/Runtime/Engine/Classes/Engine/Blueprint.h
--- a/Engine/Source/Runtime/Engine/Classes/Engine/Blueprint.h
+++ b/Engine/Source/Runtime/Engine/Classes/Engine/Blueprint.h
@@ -724,21 +724,49 @@ public:
 	UPROPERTY(transient, duplicatetransient)
 	uint32 CrcLastCompiledSignature;
 
+	/**
+	 * Transient flag that indicates whether or not the internal dependency
+	 * cache needs to be updated. This is not carried forward by duplication so
+	 * that the post-duplicate compile path is forced to reinitialize the cache.
+	 */
+	UPROPERTY(transient, duplicatetransient)
 	bool bCachedDependenciesUpToDate;
+
 	/**
 	 * Set of blueprints that we reference - i.e. blueprints that we have
-	 * some kind of reference to (variable of that blueprints type or function 
+	 * some kind of reference to (variable of that blueprints type or function
 	 * call
+	 *
+	 * We need this to be serializable so that its references can be collected.
+	 *
+	 * This is intentionally marked 'duplicatetransient' so that it won't carry
+	 * over to a duplicated Blueprint object. The post-duplicate compile path
+	 * will instead regenerate this set to be relative to the duplicated object.
 	 */
+	UPROPERTY(transient, duplicatetransient)
 	TSet<TWeakObjectPtr<UBlueprint>> CachedDependencies;
 
-	/** 
+	/**
 	 * Transient cache of dependent blueprints - i.e. blueprints that call
-	 * functions declared in this blueprint. Used to speed up compilation checks 
+	 * functions declared in this blueprint. Used to speed up compilation checks
+	 *
+	 * This is intentionally marked 'duplicatetransient' so that it won't carry
+	 * over to a duplicated Blueprint object. However, the post-duplicate compile
+	 * path will not regenerate this set, since it is populated by each dependent
+	 * Blueprint's compile. In this case, a duplicated Blueprint equates to a
+	 * new Blueprint, so it won't initially have any dependents to include here.
 	 */
+	UPROPERTY(transient, duplicatetransient)
 	TSet<TWeakObjectPtr<UBlueprint>> CachedDependents;
 
-	// User Defined Structures, the blueprint depends on
+	/**
+	 * User Defined Structures the blueprint depends on
+	 *
+	 * This is intentionally marked 'duplicatetransient' so that it won't carry
+	 * over to a duplicated Blueprint object. The post-duplicate compile path
+	 * will instead regenerate this set to be relative to the duplicated object.
+	 */
+	UPROPERTY(transient, duplicatetransient)
 	TSet<TWeakObjectPtr<UStruct>> CachedUDSDependencies;
 
 	enum class EIsBPNonReducible : uint8
