Fix missing indirect dependencies from a macro graph expansion which could lead to a script VM crash during PIE.

diff --git a/Engine/Source/Editor/Kismet/Private/BlueprintCompilationManager.cpp b/Engine/Source/Editor/Kismet/Private/BlueprintCompilationManager.cpp
--- a/Engine/Source/Editor/Kismet/Private/BlueprintCompilationManager.cpp
+++ b/Engine/Source/Editor/Kismet/Private/BlueprintCompilationManager.cpp
@@ -541,6 +541,8 @@ void FBlueprintCompilationManagerImpl::FlushCompilationQueueImpl(bool bSuppressB
 				BP->bCachedDependenciesUpToDate = false;
 			}
 
+			const bool bWasDependencyCacheOutOfDate = !BP->bCachedDependenciesUpToDate;
+
 			FBlueprintEditorUtils::EnsureCachedDependenciesUpToDate(BP);
 
 			if ((CompileJob.UserData.CompileOptions & 
@@ -558,6 +560,9 @@ void FBlueprintCompilationManagerImpl::FlushCompilationQueueImpl(bool bSuppressB
 				{
 					if(!IsQueuedForCompilation(DependentBlueprint))
 					{
+						// The macro may have updated its dependency cache above; if so, we'll need to regenerate the dependent's set as well.
+						DependentBlueprint->bCachedDependenciesUpToDate &= !bWasDependencyCacheOutOfDate;
+
 						DependentBlueprint->bQueuedForCompilation = true;
 						CurrentlyCompilingBPs.Emplace(
 							FCompilerData(
diff --git a/Engine/Source/Editor/UnrealEd/Private/Kismet2/BlueprintEditorUtils.cpp b/Engine/Source/Editor/UnrealEd/Private/Kismet2/BlueprintEditorUtils.cpp
--- a/Engine/Source/Editor/UnrealEd/Private/Kismet2/BlueprintEditorUtils.cpp
+++ b/Engine/Source/Editor/UnrealEd/Private/Kismet2/BlueprintEditorUtils.cpp
@@ -3075,6 +3075,19 @@ void FBlueprintEditorUtils::EnsureCachedDependenciesUpToDate(UBlueprint* Bluepri
 	{
 		GatherDependencies(Blueprint, Blueprint->CachedDependencies, Blueprint->CachedUDSDependencies);
 		Blueprint->bCachedDependenciesUpToDate = true;
+
+		// A macro dependency will result in an expansion from an external graph rather than a local one, so we must also include its dependencies.
+		TSet<TWeakObjectPtr<UBlueprint>> LocalCopyOfCachedDependencies = Blueprint->CachedDependencies;
+		for (const TWeakObjectPtr<UBlueprint>& Dependency : LocalCopyOfCachedDependencies)
+		{
+			UBlueprint* ResolvedDependency = Dependency.Get();
+			if (ResolvedDependency && ResolvedDependency->BlueprintType == BPTYPE_MacroLibrary)
+			{
+				EnsureCachedDependenciesUpToDate(ResolvedDependency);
+				Blueprint->CachedDependencies.Append(ResolvedDependency->CachedDependencies);
+				Blueprint->CachedUDSDependencies.Append(ResolvedDependency->CachedUDSDependencies);
+			}
+		}
 	}
 }
 
