Fixed missing view extensions registration when creating view families

diff --git a/Engine/Source/Editor/UnrealEd/Classes/ThumbnailRendering/LevelThumbnailRenderer.h b/Engine/Source/Editor/UnrealEd/Classes/ThumbnailRendering/LevelThumbnailRenderer.h
--- a/Engine/Source/Editor/UnrealEd/Classes/ThumbnailRendering/LevelThumbnailRenderer.h
+++ b/Engine/Source/Editor/UnrealEd/Classes/ThumbnailRendering/LevelThumbnailRenderer.h
@@ -14,6 +14,7 @@
 
 class FCanvas;
 class FRenderTarget;
+class FSceneView;
 class FSceneViewFamily;
 class ULevel;
 
@@ -28,6 +29,6 @@ class ULevelThumbnailRenderer : public UDefaultSizedThumbnailRenderer
 	// End UThumbnailRenderer Object
 
 private:
-	void GetView(ULevel* Level, FSceneViewFamily* ViewFamily, int32 X, int32 Y, uint32 SizeX, uint32 SizeY) const; 
+	FSceneView* GetView(ULevel* Level, FSceneViewFamily* ViewFamily, int32 X, int32 Y, uint32 SizeX, uint32 SizeY) const;
 };
 
diff --git a/Engine/Source/Editor/UnrealEd/Private/ThumbnailRendering/LevelThumbnailRenderer.cpp b/Engine/Source/Editor/UnrealEd/Private/ThumbnailRendering/LevelThumbnailRenderer.cpp
--- a/Engine/Source/Editor/UnrealEd/Private/ThumbnailRendering/LevelThumbnailRenderer.cpp
+++ b/Engine/Source/Editor/UnrealEd/Private/ThumbnailRendering/LevelThumbnailRenderer.cpp
@@ -6,6 +6,7 @@
 #include "Engine/Level.h"
 #include "ShowFlags.h"
 #include "SceneView.h"
+#include "SceneViewExtension.h"
 #include "Engine/LevelBounds.h"
 
 ULevelThumbnailRenderer::ULevelThumbnailRenderer(const FObjectInitializer& ObjectInitializer)
@@ -27,16 +28,23 @@ void ULevelThumbnailRenderer::Draw(UObject* Object, int32 X, int32 Y, uint32 Wid
 		ViewFamily.EngineShowFlags.SetDistanceCulledPrimitives(true); // show distance culled objects
 		ViewFamily.EngineShowFlags.SetPostProcessing(false);
 
-		GetView(Level, &ViewFamily, X, Y, Width, Height);
+		FSceneView* NewView = GetView(Level, &ViewFamily, X, Y, Width, Height);
 
-		if (ViewFamily.Views.Num() > 0)
+		ViewFamily.ViewExtensions = GEngine->ViewExtensions->GatherActiveExtensions(FSceneViewExtensionContext(Level->OwningWorld->Scene));
+		for (const FSceneViewExtensionRef& Extension : ViewFamily.ViewExtensions)
+		{
+			Extension->SetupViewFamily(ViewFamily);
+			Extension->SetupView(ViewFamily, *NewView);
+		}
+
+		if (NewView)
 		{
 			RenderViewFamily(Canvas, &ViewFamily);
 		}
 	}
 }
 
-void ULevelThumbnailRenderer::GetView(ULevel* Level, FSceneViewFamily* ViewFamily, int32 X, int32 Y, uint32 SizeX, uint32 SizeY) const 
+FSceneView* ULevelThumbnailRenderer::GetView(ULevel* Level, FSceneViewFamily* ViewFamily, int32 X, int32 Y, uint32 SizeX, uint32 SizeY) const
 {
 	check(ViewFamily);
 
@@ -47,6 +55,7 @@ void ULevelThumbnailRenderer::GetView(ULevel* Level, FSceneViewFamily* ViewFamil
 		FMath::Max<int32>(Y+SizeY,0));
 
 	FBox LevelBox(ForceInit);
+	FSceneView* NewView = nullptr;
 
 	if (Level->LevelBoundsActor.IsValid())
 	{
@@ -80,8 +89,10 @@ void ULevelThumbnailRenderer::GetView(ULevel* Level, FSceneViewFamily* ViewFamil
 			ZOffset
 			);
 
-		FSceneView* NewView = new FSceneView(ViewInitOptions);
+		NewView = new FSceneView(ViewInitOptions);
 
 		ViewFamily->Views.Add(NewView);
 	}
+
+	return NewView;
 }
diff --git a/Engine/Source/Runtime/Renderer/Private/ReflectionEnvironmentCapture.cpp b/Engine/Source/Runtime/Renderer/Private/ReflectionEnvironmentCapture.cpp
--- a/Engine/Source/Runtime/Renderer/Private/ReflectionEnvironmentCapture.cpp
+++ b/Engine/Source/Runtime/Renderer/Private/ReflectionEnvironmentCapture.cpp
@@ -27,6 +27,7 @@
 #include "GlobalShader.h"
 #include "SceneRenderTargetParameters.h"
 #include "SceneRendering.h"
+#include "SceneViewExtension.h"
 #include "ScenePrivate.h"
 #include "PostProcess/SceneFilterRendering.h"
 #include "PostProcess/PostProcessing.h"
@@ -1436,6 +1437,13 @@ void CaptureSceneIntoScratchCubemap(
 
 		FSceneView* View = new FSceneView(ViewInitOptions);
 
+		ViewFamily.ViewExtensions = GEngine->ViewExtensions->GatherActiveExtensions(FSceneViewExtensionContext(Scene));
+		for (const FSceneViewExtensionRef& Extension : ViewFamily.ViewExtensions)
+		{
+			Extension->SetupViewFamily(ViewFamily);
+			Extension->SetupView(ViewFamily, *View);
+		}
+
 		// Force all surfaces diffuse
 		View->RoughnessOverrideParameter = FVector2D( 1.0f, 0.0f );
 
