Remove experimental save flags added to cookig by mistake

diff --git a/Engine/Source/Editor/UnrealEd/Private/CookOnTheFlyServer.cpp b/Engine/Source/Editor/UnrealEd/Private/CookOnTheFlyServer.cpp
--- a/Engine/Source/Editor/UnrealEd/Private/CookOnTheFlyServer.cpp
+++ b/Engine/Source/Editor/UnrealEd/Private/CookOnTheFlyServer.cpp
@@ -4110,7 +4110,7 @@ void UCookOnTheFlyServer::SaveCookedPackage(UE::Cook::FSaveCookedPackageContext&
 			SaveArgs.TopLevelFlags = Context.FlagsToCook;
 			SaveArgs.bForceByteSwapping = Context.bEndianSwap;
 			SaveArgs.bWarnOfLongFilename = false;
-			SaveArgs.SaveFlags = Context.SaveFlags | SAVE_Optional;
+			SaveArgs.SaveFlags = Context.SaveFlags;
 			SaveArgs.TargetPlatform = TargetPlatform;
 			SaveArgs.bSlowTask = false;
 			SaveArgs.SavePackageContext = Context.SavePackageContext;
@@ -8812,7 +8812,7 @@ uint32 UCookOnTheFlyServer::FullLoadAndSave(uint32& CookedPackageCount)
 						SaveArgs.TopLevelFlags = FlagsToCook;
 						SaveArgs.bForceByteSwapping = bSwap;
 						SaveArgs.bWarnOfLongFilename = false;
-						SaveArgs.SaveFlags = SaveFlags | SAVE_Optional;
+						SaveArgs.SaveFlags = SaveFlags;
 						SaveArgs.TargetPlatform = Target;
 						SaveArgs.bSlowTask = false;
 						SaveArgs.SavePackageContext = &SavePackageContext;
