Make sure RHI resources are flushed even when nothing is rendered to prevent leak in RHI (D3D12/Vulkan) when Realtime is turned off in Editor.

diff --git a/Engine/Source/Editor/UnrealEd/Private/EditorEngine.cpp b/Engine/Source/Editor/UnrealEd/Private/EditorEngine.cpp
--- a/Engine/Source/Editor/UnrealEd/Private/EditorEngine.cpp
+++ b/Engine/Source/Editor/UnrealEd/Private/EditorEngine.cpp
@@ -1921,8 +1921,9 @@ void UEditorEngine::Tick( float DeltaSeconds, bool bIdleMode )
 	}
 
 	// Do not redraw if the application is hidden
-	bool bAllWindowsHidden = !bHasFocus && AreAllWindowsHidden();
-	if( !bAllWindowsHidden || bRunDrawWithEditorHidden)
+	const bool bAllWindowsHidden = !bHasFocus && AreAllWindowsHidden();
+	bool bEditorFrameNonRealtimeViewportDrawn = false;
+	if (!bAllWindowsHidden || bRunDrawWithEditorHidden)
 	{
 		FPixelInspectorModule& PixelInspectorModule = FModuleManager::LoadModuleChecked<FPixelInspectorModule>(TEXT("PixelInspectorModule"));
 		if (!bAllWindowsHidden && PixelInspectorModule.IsPixelInspectorEnable())
@@ -1931,7 +1932,6 @@ void UEditorEngine::Tick( float DeltaSeconds, bool bIdleMode )
 		}
 
 		// Render view parents, then view children.
-		bool bEditorFrameNonRealtimeViewportDrawn = false;
 		if (GCurrentLevelEditingViewportClient && GCurrentLevelEditingViewportClient->IsVisible())
 		{
 			if (!bAllWindowsHidden || GCurrentLevelEditingViewportClient->WantsDrawWhenAppIsHidden())
@@ -1946,7 +1946,7 @@ void UEditorEngine::Tick( float DeltaSeconds, bool bIdleMode )
 		}
 		for (int32 bRenderingChildren = 0; bRenderingChildren < 2; bRenderingChildren++)
 		{
-			for(FEditorViewportClient* ViewportClient : AllViewportClients)
+			for (FEditorViewportClient* ViewportClient : AllViewportClients)
 			{
 				if (ViewportClient == GCurrentLevelEditingViewportClient)
 				{
@@ -1974,7 +1974,20 @@ void UEditorEngine::Tick( float DeltaSeconds, bool bIdleMode )
 				}
 			}
 		}
+	}
 
+	// If we're not Realtime and no NonRealtime viewports are drawn, or if everything is hidden, make sure RHI gets its flush to prevent memory from accumulating
+	if ((bAllWindowsHidden || !IsRealtime) && !bEditorFrameNonRealtimeViewportDrawn && IsRunningRHIInSeparateThread())
+	{
+		ENQUEUE_RENDER_COMMAND(FlushPendingDeleteRHIResources_NonRealtime)(
+			[](FRHICommandListImmediate& RHICmdList)
+			{
+				RHICmdList.ImmediateFlush(EImmediateFlushType::FlushRHIThreadFlushResources);
+			});
+	}
+
+	if (!bAllWindowsHidden || bRunDrawWithEditorHidden)
+	{
 		// Some tasks can only be done once we finish all scenes/viewports
 		GetRendererModule().PostRenderAllViewports();
 	}
