CookPackageData: prevent crash when LookupPackageNameOnDisk converts a filename to a packagename but then fails to map that packagename back to a filename.
#rnx

diff --git a/Engine/Source/Editor/UnrealEd/Private/Cooker/CookPackageData.cpp b/Engine/Source/Editor/UnrealEd/Private/Cooker/CookPackageData.cpp
--- a/Engine/Source/Editor/UnrealEd/Private/Cooker/CookPackageData.cpp
+++ b/Engine/Source/Editor/UnrealEd/Private/Cooker/CookPackageData.cpp
@@ -1741,6 +1741,12 @@ FPackageData* FPackageDatas::TryAddPackageDataByStandardFileName(const FName& Fi
 	{
 		return nullptr;
 	}
+	if (ExistingFileName.IsNone())
+	{
+		UE_LOG(LogCook, Warning, TEXT("Unexpected failure to cook filename '%s'. It is mapped to PackageName '%s', but that PackageName does not have a corresponding filename."),
+			*FileName.ToString(), *PackageName.ToString());
+		return nullptr;
+	}
 	FoundFileName = ExistingFileName;
 	return &CreatePackageData(PackageName, ExistingFileName);
 }
