Fix FVector2D variant casts
Submitted on behalf of fred.kimberley

diff --git a/Engine/Source/Developer/MaterialBaking/Private/MaterialRenderItem.cpp b/Engine/Source/Developer/MaterialBaking/Private/MaterialRenderItem.cpp
--- a/Engine/Source/Developer/MaterialBaking/Private/MaterialRenderItem.cpp
+++ b/Engine/Source/Developer/MaterialBaking/Private/MaterialRenderItem.cpp
@@ -238,7 +238,7 @@ void FMeshMaterialRenderItem::PopulateWithMeshData()
 				if (!bUseNewUVs)
 				{
 					// compute vertex position from original UV
-					const FVector2D& UV = VertexInstanceUVs.Get(SrcVertexInstanceID, MeshSettings->TextureCoordinateIndex);
+					const FVector2D& UV = FVector2D(VertexInstanceUVs.Get(SrcVertexInstanceID, MeshSettings->TextureCoordinateIndex));
 					Vert->Position.Set(UV.X * ScaleX, UV.Y * ScaleY, 0);
 				}
 				else
diff --git a/Engine/Source/Developer/MaterialUtilities/Private/MeshRendering.cpp b/Engine/Source/Developer/MaterialUtilities/Private/MeshRendering.cpp
--- a/Engine/Source/Developer/MaterialUtilities/Private/MeshRendering.cpp
+++ b/Engine/Source/Developer/MaterialUtilities/Private/MeshRendering.cpp
@@ -218,7 +218,7 @@ public:
 					if (!bUseNewUVs)
 					{
 						// compute vertex position from original UV
-						const FVector2D& UV = VertexInstanceUVs.Get(SrcVertexInstanceID, 0);
+						const FVector2f& UV = VertexInstanceUVs.Get(SrcVertexInstanceID, 0);
 						Vert->Position.Set((UV.X - U) * ScaleX, (UV.Y - V) * ScaleY, 0);
 					}
 					else
@@ -341,7 +341,7 @@ public:
 						FDynamicMeshVertex* DstVert = new(OutVerts)FDynamicMeshVertex();
 
 						// compute vertex position from original UV
-						const FVector2D UV = LODData.StaticVertexBuffers.StaticMeshVertexBuffer.GetVertexUV(SrcVertIndex, 0);
+						const FVector2f UV = LODData.StaticVertexBuffers.StaticMeshVertexBuffer.GetVertexUV(SrcVertIndex, 0);
 						DstVert->Position.Set((UV.X - U) * ScaleX, (UV.Y - V) * ScaleY, 0);
 
 						DstVert->TangentX = LODData.StaticVertexBuffers.StaticMeshVertexBuffer.VertexTangentX(SrcVertIndex);
@@ -403,7 +403,7 @@ public:
 							int32 SrcVertIndex = IndexData[CornerIndex];
 							FDynamicMeshVertex* DstVert = new(OutVerts)FDynamicMeshVertex();
 
-							const FVector2D UV = LODData.StaticVertexBuffers.StaticMeshVertexBuffer.GetVertexUV(SrcVertIndex, 0);
+							const FVector2f UV = LODData.StaticVertexBuffers.StaticMeshVertexBuffer.GetVertexUV(SrcVertIndex, 0);
 							DstVert->Position.Set(UV.X * ScaleX, UV.Y * ScaleY, 0);
 
 							DstVert->TangentX = LODData.StaticVertexBuffers.StaticMeshVertexBuffer.VertexTangentX(SrcVertIndex);
diff --git a/Engine/Source/Developer/MeshBuilderCommon/Private/ThirdPartyBuildOptimizationHelper.cpp b/Engine/Source/Developer/MeshBuilderCommon/Private/ThirdPartyBuildOptimizationHelper.cpp
--- a/Engine/Source/Developer/MeshBuilderCommon/Private/ThirdPartyBuildOptimizationHelper.cpp
+++ b/Engine/Source/Developer/MeshBuilderCommon/Private/ThirdPartyBuildOptimizationHelper.cpp
@@ -204,7 +204,7 @@ namespace BuildOptimizationThirdParty
 
 			if (VertexBuffer.GetNumTexCoords())
 			{
-				const FVector2D UV = VertexBuffer.GetVertexUV(Index, 0);
+				const FVector2f UV = VertexBuffer.GetVertexUV(Index, 0);
 				Vertex.uv.x = UV.X;
 				Vertex.uv.y = UV.Y;
 			}
diff --git a/Engine/Source/Developer/MeshMergeUtilities/Private/MeshMergeHelpers.cpp b/Engine/Source/Developer/MeshMergeUtilities/Private/MeshMergeHelpers.cpp
--- a/Engine/Source/Developer/MeshMergeUtilities/Private/MeshMergeHelpers.cpp
+++ b/Engine/Source/Developer/MeshMergeUtilities/Private/MeshMergeHelpers.cpp
@@ -582,7 +582,7 @@ bool FMeshMergeHelpers::CheckWrappingUVs(const FMeshDescription& MeshDescription
 
 	for (const FVertexInstanceID VertexInstanceID : MeshDescription.VertexInstances().GetElementIDs())
 	{
-		const FVector2D& Coordinate = VertexInstanceUVs.Get(VertexInstanceID, UVChannelIndex);
+		const FVector2D& Coordinate = FVector2D(VertexInstanceUVs.Get(VertexInstanceID, UVChannelIndex));
 		if ((FMath::IsNegativeFloat(Coordinate.X) || FMath::IsNegativeFloat(Coordinate.Y)) || (Coordinate.X > (1.0f + KINDA_SMALL_NUMBER) || Coordinate.Y > (1.0f + KINDA_SMALL_NUMBER)))
 		{
 			bResult = true;
@@ -1019,7 +1019,7 @@ void FMeshMergeHelpers::CalculateTextureCoordinateBoundsForMesh(const FMeshDescr
 			{
 				for (int32 UVIndex = 0; UVIndex < VertexInstanceUVs.GetNumChannels(); ++UVIndex)
 				{
-					OutBounds[MaterialIndex] += VertexInstanceUVs.Get(VertexInstanceID, UVIndex);
+					OutBounds[MaterialIndex] += FVector2D(VertexInstanceUVs.Get(VertexInstanceID, UVIndex));
 				}
 			}
 		}
@@ -1181,11 +1181,11 @@ void FMeshMergeHelpers::MergeImpostersToMesh(TArray<const UStaticMeshComponent*>
 
 				UVOne.X = ActorPosition.X;
 				UVOne.Y = ActorPosition.Y;
-				VertexInstanceUVs.Set(VertexInstanceID, UVOneIndex, UVOne);
+				VertexInstanceUVs.Set(VertexInstanceID, UVOneIndex, FVector2f(UVOne));	// LWC_TODO: Precision loss
 
 				UVTwo.X = ActorPosition.Z;
 				UVTwo.Y = FMath::Abs(ActorToWorld.GetScale3D().X);
-				VertexInstanceUVs.Set(VertexInstanceID, UVTwoIndex, UVTwo);
+				VertexInstanceUVs.Set(VertexInstanceID, UVTwoIndex, FVector2f(UVTwo));	// LWC_TODO: Precision loss
 			}
 		}
 		else if (!InPivot.IsZero())
diff --git a/Engine/Source/Developer/MeshMergeUtilities/Private/MeshMergeUtilities.cpp b/Engine/Source/Developer/MeshMergeUtilities/Private/MeshMergeUtilities.cpp
--- a/Engine/Source/Developer/MeshMergeUtilities/Private/MeshMergeUtilities.cpp
+++ b/Engine/Source/Developer/MeshMergeUtilities/Private/MeshMergeUtilities.cpp
@@ -1164,7 +1164,7 @@ public:
 			TVertexInstanceAttributesRef<FVector2f> VertexInstanceUVs = FStaticMeshAttributes(MeshDescription).GetVertexInstanceUVs();
 			for (int32 Idx = 0; Idx < TexCoords.Num(); Idx++)
 			{
-				VertexInstanceUVs.Set(Idx, TexCoords[Idx]);
+				VertexInstanceUVs.Set(Idx, FVector2f(TexCoords[Idx]));	// LWC_TODO: Precision loss
 			}
 		}
 		else
@@ -1408,7 +1408,7 @@ TArray<FMeshData> PrepareBakingMeshes(const struct FMeshProxySettings& InMeshPro
 			    MeshSettings.CustomTextureCoordinates.Reset(VertexInstanceUVs.GetNumElements());
 			    for (const FVertexInstanceID VertexInstanceID : MeshDescription.VertexInstances().GetElementIDs())
 			    {
-				    MeshSettings.CustomTextureCoordinates.Add(VertexInstanceUVs.Get(VertexInstanceID, LightMapCoordinateIndex));
+				    MeshSettings.CustomTextureCoordinates.Add(FVector2D(VertexInstanceUVs.Get(VertexInstanceID, LightMapCoordinateIndex)));
 			    }
 			    MeshSettings.TextureCoordinateBox = FBox2D(FVector2D::ZeroVector, FVector2D(1, 1));
 			    ScaleTextureCoordinatesToBox(MeshSettings.TextureCoordinateBox, MeshSettings.CustomTextureCoordinates);
@@ -1877,7 +1877,7 @@ void FMeshMergeUtilities::CreateProxyMesh(const TArray<UStaticMeshComponent*>& I
 					MergeData.NewUVs.Reset(MeshData->MeshDescription->VertexInstances().Num());
 					for (const FVertexInstanceID VertexInstanceID : MeshData->MeshDescription->VertexInstances().GetElementIDs())
 					{
-						MergeData.NewUVs.Add(VertexInstanceUVs.Get(VertexInstanceID, MeshData->TextureCoordinateIndex));
+						MergeData.NewUVs.Add(FVector2D(VertexInstanceUVs.Get(VertexInstanceID, MeshData->TextureCoordinateIndex)));
 					}
 				}
 				MergeData.TexCoordBounds[0] = FBox2D(FVector2D(0.0f, 0.0f), FVector2D(1.0f, 1.0f));
@@ -2668,7 +2668,7 @@ void FMeshMergeUtilities::MergeComponentsToStaticMesh(const TArray<UPrimitiveCom
 							int32 VertexIndex = 0;
 							for (FVertexInstanceID VertexInstanceID : RawMesh->VertexInstances().GetElementIDs())
 							{
-								FVector2D UV = VertexInstanceUVs.Get(VertexInstanceID, UVChannelIdx);
+								FVector2D UV = FVector2D(VertexInstanceUVs.Get(VertexInstanceID, UVChannelIdx));
 								if (UVChannelIdx == 0 && !InSettings.bCreateMergedMaterial)
 								{
 									if (MeshData.CustomTextureCoordinates.Num())
@@ -2678,7 +2678,7 @@ void FMeshMergeUtilities::MergeComponentsToStaticMesh(const TArray<UPrimitiveCom
 									else if (MeshData.TextureCoordinateIndex != 0)
 									{
 										check(MeshData.TextureCoordinateIndex < NumUVChannel);
-										UV = VertexInstanceUVs.Get(VertexInstanceID, MeshData.TextureCoordinateIndex);
+										UV = FVector2D(VertexInstanceUVs.Get(VertexInstanceID, MeshData.TextureCoordinateIndex));
 									}
 								}
 
@@ -2690,7 +2690,7 @@ void FMeshMergeUtilities::MergeComponentsToStaticMesh(const TArray<UPrimitiveCom
 									{
 										if (UVTransform.Value != FVector2D::ZeroVector)
 										{
-											VertexInstanceUVs.Set(VertexInstanceID, UVChannelIdx, UV * UVTransform.Value + UVTransform.Key);
+											VertexInstanceUVs.Set(VertexInstanceID, UVChannelIdx, FVector2f(UV * UVTransform.Value + UVTransform.Key));	// LWC_TODO: Precision loss
 											break;
 										}
 									}
@@ -2707,7 +2707,7 @@ void FMeshMergeUtilities::MergeComponentsToStaticMesh(const TArray<UPrimitiveCom
 							for(FVertexInstanceID VertexInstanceID : RawMesh->VertexInstances().GetElementIDs())
 							{
 								FVector2D UV = MeshData.CustomTextureCoordinates[VertexIndex];
-								VertexInstanceUVs.Set(VertexInstanceID, MergedMatUVChannel, UV * UVTransform.Value + UVTransform.Key);
+								VertexInstanceUVs.Set(VertexInstanceID, MergedMatUVChannel, FVector2f(UV * UVTransform.Value + UVTransform.Key));	// LWC_TODO: Precision loss
 								VertexIndex++;
 							}
 						}
diff --git a/Engine/Source/Developer/MeshUtilities/Private/MeshUtilities.cpp b/Engine/Source/Developer/MeshUtilities/Private/MeshUtilities.cpp
--- a/Engine/Source/Developer/MeshUtilities/Private/MeshUtilities.cpp
+++ b/Engine/Source/Developer/MeshUtilities/Private/MeshUtilities.cpp
@@ -1148,7 +1148,7 @@ static void ComputeTriangleTangents(
 	for (int32 TriangleIndex = 0; TriangleIndex < NumTriangles; TriangleIndex++)
 	{
 		FVector3f Positions[3];
-		FVector2D UVs[3];
+		FVector2f UVs[3];
 		for (int32 Index = 0; Index < 3; ++Index)
 		{
 			Positions[Index] = InVertices[InIndices[TriangleIndex * 3 + Index]];
@@ -1433,8 +1433,8 @@ static void ComputeTangents(
 												CommonVertices++;
 
 
-												const FVector2D& UVOne = InUVs[NextFace.FaceIndex * 3 + NextCornerIndex];
-												const FVector2D& UVTwo = InUVs[OtherFace.FaceIndex * 3 + OtherCornerIndex];
+												const FVector2f& UVOne = InUVs[NextFace.FaceIndex * 3 + NextCornerIndex];
+												const FVector2f& UVTwo = InUVs[OtherFace.FaceIndex * 3 + OtherCornerIndex];
 
 												if (UVsEqual(UVOne, UVTwo))
 												{
@@ -1644,7 +1644,7 @@ static void MikkSetTSpaceBasic(const SMikkTSpaceContext* Context, const float Ta
 static void MikkGetTexCoord(const SMikkTSpaceContext* Context, float UV[2], const int FaceIdx, const int VertIdx)
 {
 	MikkTSpace_Mesh *UserData = (MikkTSpace_Mesh*)(Context->m_pUserData);
-	const FVector2D &TexCoord = UserData->UVs[FaceIdx * 3 + VertIdx];
+	const FVector2f &TexCoord = UserData->UVs[FaceIdx * 3 + VertIdx];
 	UV[0] = TexCoord.X;
 	UV[1] = TexCoord.Y;
 }
@@ -1754,7 +1754,7 @@ static void MikkSetTSpaceBasic_Skeletal(const SMikkTSpaceContext* Context, const
 static void MikkGetTexCoord_Skeletal(const SMikkTSpaceContext* Context, float UV[2], const int FaceIdx, const int VertIdx)
 {
 	MikkTSpace_Skeletal_Mesh *UserData = (MikkTSpace_Skeletal_Mesh*)(Context->m_pUserData);
-	const FVector2D &TexCoord = UserData->wedges[UserData->faces[FaceIdx].iWedge[VertIdx]].UVs[0];
+	const FVector2f &TexCoord = UserData->wedges[UserData->faces[FaceIdx].iWedge[VertIdx]].UVs[0];
 	UV[0] = TexCoord.X;
 	UV[1] = TexCoord.Y;
 }
@@ -3199,7 +3199,7 @@ public:
 
 	virtual FVector2D GetVertexUV(uint32 FaceIndex, uint32 TriIndex, uint32 UVIndex) override
 	{
-		return Wedges[Faces[FaceIndex].iWedge[TriIndex]].UVs[UVIndex];
+		return FVector2D(Wedges[Faces[FaceIndex].iWedge[TriIndex]].UVs[UVIndex]);
 	}
 
 	virtual uint32 GetFaceSmoothingGroups(uint32 FaceIndex)
@@ -3343,11 +3343,11 @@ public:
 		{
 			const int32 UVIndex = 0;
 			FVector3f Positions[3];
-			FVector2D UVs[3];
+			FVector2f UVs[3];
 			for (int32 Index = 0; Index < 3; ++Index)
 			{
 				Positions[Index] = BuildData->GetVertexPosition(TriangleIndex, Index);
-				UVs[Index] = BuildData->GetVertexUV(TriangleIndex, Index, UVIndex);
+				UVs[Index] = FVector2f(BuildData->GetVertexUV(TriangleIndex, Index, UVIndex));	// LWC_TODO: Precision loss
 			}
 			
 			CalculateTriangleTangentInternal(Positions[0], UVs[0], Positions[1], UVs[1], Positions[2], UVs[2], TriangleTangents, RealComparisonThreshold);
@@ -5139,7 +5139,7 @@ void FMeshUtilities::CalculateTextureCoordinateBoundsForSkeletalMesh(const FSkel
 			uint32 VertexIndex = LODModel.IndexBuffer[Index];
 			FSoftSkinVertex& Vertex = Vertices[VertexIndex];
 
-			FVector2D TexCoord = Vertex.UVs[0];
+			FVector2D TexCoord = FVector2D(Vertex.UVs[0]);
 			OutBounds[MaterialIndex] += TexCoord;
 		}
 	}
diff --git a/Engine/Source/Developer/MeshUtilities/Private/SkeletalMeshTools.cpp b/Engine/Source/Developer/MeshUtilities/Private/SkeletalMeshTools.cpp
--- a/Engine/Source/Developer/MeshUtilities/Private/SkeletalMeshTools.cpp
+++ b/Engine/Source/Developer/MeshUtilities/Private/SkeletalMeshTools.cpp
@@ -221,7 +221,7 @@ namespace SkeletalMeshTools
 				const FVector& TangentXRef = SoftSkinVertRef.TangentX;
 				const FVector& TangentYRef = SoftSkinVertRef.TangentY;
 				const FVector& TangentZRef = SoftSkinVertRef.TangentZ;
-				const FVector2D& UVRef = SoftSkinVertRef.UVs[0];
+				const FVector2f& UVRef = SoftSkinVertRef.UVs[0];
 				const FColor& ColorRef = SoftSkinVertRef.Color;
 				const TArray<int32>& AdjacentFaces = VertexIndexToAdjacentFaces.FindChecked(VertexIndex);
 				for (int32 AdjacentFaceArrayIndex = 0; AdjacentFaceArrayIndex < AdjacentFaces.Num(); ++AdjacentFaceArrayIndex)
diff --git a/Engine/Source/Developer/MeshUtilities/Private/SkeletalMeshTools.h b/Engine/Source/Developer/MeshUtilities/Private/SkeletalMeshTools.h
--- a/Engine/Source/Developer/MeshUtilities/Private/SkeletalMeshTools.h
+++ b/Engine/Source/Developer/MeshUtilities/Private/SkeletalMeshTools.h
@@ -80,8 +80,8 @@ namespace SkeletalMeshTools
 {
 	inline bool SkeletalMesh_UVsEqual(const SkeletalMeshImportData::FMeshWedge& V1, const SkeletalMeshImportData::FMeshWedge& V2, const FOverlappingThresholds& OverlappingThresholds, const int32 UVIndex = 0)
 	{
-		const FVector2D& UV1 = V1.UVs[UVIndex];
-		const FVector2D& UV2 = V2.UVs[UVIndex];
+		const FVector2f& UV1 = V1.UVs[UVIndex];
+		const FVector2f& UV2 = V2.UVs[UVIndex];
 
 		if(FMath::Abs(UV1.X - UV2.X) > OverlappingThresholds.ThresholdUV)
 			return 0;
diff --git a/Engine/Source/Developer/SkeletalMeshUtilitiesCommon/Private/LODUtilities.cpp b/Engine/Source/Developer/SkeletalMeshUtilitiesCommon/Private/LODUtilities.cpp
--- a/Engine/Source/Developer/SkeletalMeshUtilitiesCommon/Private/LODUtilities.cpp
+++ b/Engine/Source/Developer/SkeletalMeshUtilitiesCommon/Private/LODUtilities.cpp
@@ -613,7 +613,7 @@ bool FindTriangleUVMatch(const FVector2D& TargetUV, const TArray<FTriangleElemen
 	for (uint32 TriangleIndex : QuadTreeTriangleResults)
 	{
 		const FTriangleElement& TriangleElement = Triangles[TriangleIndex];
-		if (PointInTriangle(TriangleElement.Vertices[0].UVs[0], TriangleElement.Vertices[1].UVs[0], TriangleElement.Vertices[2].UVs[0], TargetUV))
+		if (PointInTriangle(FVector2D(TriangleElement.Vertices[0].UVs[0]), FVector2D(TriangleElement.Vertices[1].UVs[0]), FVector2D(TriangleElement.Vertices[2].UVs[0]), TargetUV))
 		{
 			MatchTriangleIndexes.Add(TriangleIndex);
 		}
@@ -684,7 +684,7 @@ void ProjectTargetOnBase(const TArray<FSoftSkinVertex>& BaseVertices, const TArr
 				const FSoftSkinVertex& BaseVertex = BaseVertices[CornerIndice];
 				TriangleElement.Indexes.Add(CornerIndice);
 				TriangleElement.Vertices.Add(BaseVertex);
-				TriangleElement.UVsBound += BaseVertex.UVs[0];
+				TriangleElement.UVsBound += FVector2D(BaseVertex.UVs[0]);
 				BaseMeshPositionBound += BaseVertex.Position;
 			}
 			BaseMeshUVBound += TriangleElement.UVsBound;
@@ -714,7 +714,7 @@ void ProjectTargetOnBase(const TArray<FSoftSkinVertex>& BaseVertices, const TArr
 		QuadTreeTriangleResults.Reserve(Triangles.Num() / 10); //Reserve 10% to speed up the query
 		for (uint32 TargetVertexIndex = 0; TargetVertexIndex < (uint32)TargetVertices.Num(); ++TargetVertexIndex)
 		{
-			FVector2D TargetUV = TargetVertices[TargetVertexIndex].UVs[0];
+			FVector2D TargetUV = FVector2D(TargetVertices[TargetVertexIndex].UVs[0]);
 			//Reset the last data without flushing the memmery allocation
 			QuadTreeTriangleResults.Reset();
 			const uint32 FullTargetIndex = TargetSections[SectionIndex].BaseVertexIndex + TargetVertexIndex;
@@ -816,7 +816,7 @@ void ProjectTargetOnBase(const TArray<FSoftSkinVertex>& BaseVertices, const TArr
 				FTriangleElement& BestTriangle = Triangles[FoundIndexMatch];
 				//Found the surface area of the 3 barycentric triangles from the UVs
 				FVector3f BarycentricWeight;
-				BarycentricWeight = GetBaryCentric(FVector3f(TargetUV, 0.0f), FVector3f(BestTriangle.Vertices[0].UVs[0], 0.0f), FVector3f(BestTriangle.Vertices[1].UVs[0], 0.0f), FVector3f(BestTriangle.Vertices[2].UVs[0], 0.0f));
+				BarycentricWeight = GetBaryCentric(FVector3f(FVector2f(TargetUV), 0.0f), FVector3f(BestTriangle.Vertices[0].UVs[0], 0.0f), FVector3f(BestTriangle.Vertices[1].UVs[0], 0.0f), FVector3f(BestTriangle.Vertices[2].UVs[0], 0.0f));	// LWC_TODO: Precision loss
 				//Fill the target match
 				for (int32 Corner = 0; Corner < 3; ++Corner)
 				{
@@ -1578,13 +1578,12 @@ void MatchVertexIndexUsingPosition(
 		{
 			const uint32 WedgeIndexDest = Triangle.WedgeIndex[Corner];
 			const uint32 VertexIndexDest = ImportDataDest.Wedges[WedgeIndexDest].VertexIndex;
-			const FVector2D UVsDest = ImportDataDest.Wedges[WedgeIndexDest].UVs[0];
 			TriangleElement.Indexes.Add(WedgeIndexDest);
 			FSoftSkinVertex SoftSkinVertex;
 			SoftSkinVertex.Position = ImportDataDest.Points[VertexIndexDest];
 			SoftSkinVertex.UVs[0] = ImportDataDest.Wedges[WedgeIndexDest].UVs[0];
 			TriangleElement.Vertices.Add(SoftSkinVertex);
-			TriangleElement.UVsBound += SoftSkinVertex.UVs[0];
+			TriangleElement.UVsBound += FVector2D(SoftSkinVertex.UVs[0]);
 			TrianglePositionBound += SoftSkinVertex.Position;
 			BaseMeshPositionBound += SoftSkinVertex.Position;
 		}
@@ -3360,9 +3359,9 @@ bool FLODUtilities::StripLODGeometry(USkeletalMesh* SkeletalMesh, const int32 LO
 			LODModel.GetSectionFromVertexIndex(VertexIndexC, SectionIndex, SectionVertexIndexC);
 			FSkelMeshSection& Section = LODModel.Sections[SectionIndex];
 			//Get the UV triangle, add the small number that will act like threshold when converting the UV into pixel coordinate.
-			FVector2D UvA = Section.SoftVertices[SectionVertexIndexA].UVs[0] + KINDA_SMALL_NUMBER;
-			FVector2D UvB = Section.SoftVertices[SectionVertexIndexB].UVs[0] + KINDA_SMALL_NUMBER;
-			FVector2D UvC = Section.SoftVertices[SectionVertexIndexC].UVs[0] + KINDA_SMALL_NUMBER;
+			FVector2D UvA = FVector2D(Section.SoftVertices[SectionVertexIndexA].UVs[0]) + KINDA_SMALL_NUMBER;
+			FVector2D UvB = FVector2D(Section.SoftVertices[SectionVertexIndexB].UVs[0]) + KINDA_SMALL_NUMBER;
+			FVector2D UvC = FVector2D(Section.SoftVertices[SectionVertexIndexC].UVs[0]) + KINDA_SMALL_NUMBER;
 
 			if (ShouldStripTriangle(UvA, UvB, UvC))
 			{
diff --git a/Engine/Source/Developer/SlateReflector/Private/Widgets/NavigationSimulationOverlay.cpp b/Engine/Source/Developer/SlateReflector/Private/Widgets/NavigationSimulationOverlay.cpp
--- a/Engine/Source/Developer/SlateReflector/Private/Widgets/NavigationSimulationOverlay.cpp
+++ b/Engine/Source/Developer/SlateReflector/Private/Widgets/NavigationSimulationOverlay.cpp
@@ -150,7 +150,7 @@ namespace NavigationSimulationOverlay
 				FArrangedWidget ArrangedWidget = WidgetPath.FindArrangedWidget(Widget).Get(FArrangedWidget::GetNullWidget());
 				ArrangedWidget.Geometry.AppendTransform(FSlateLayoutTransform(Inverse(CurrentWindowRef->GetPositionInScreen())));
 
-				const FVector2D InflateAmount = FVector2D(1.0f, 1.0f) / ArrangedWidget.Geometry.GetAccumulatedRenderTransform().GetMatrix().GetScale().GetVector();
+				const FVector2D InflateAmount = FVector2D(1, 1) / FVector2D(ArrangedWidget.Geometry.GetAccumulatedRenderTransform().GetMatrix().GetScale().GetVector());
 				FPaintGeometry PaintGeometry = ArrangedWidget.Geometry.ToInflatedPaintGeometry(InflateAmount);
 
 				Map.Add(WidgetPtr, PaintGeometry);
@@ -169,7 +169,7 @@ namespace NavigationSimulationOverlay
 
 			if (WidgetInfo.bHasGeometry)
 			{
-				const FVector2D InflateAmount = FVector2D(1.0f, 1.0f) / WidgetInfo.WidgetGeometry.GetAccumulatedRenderTransform().GetMatrix().GetScale().GetVector();
+				const FVector2D InflateAmount = FVector2D(1, 1) / FVector2D(WidgetInfo.WidgetGeometry.GetAccumulatedRenderTransform().GetMatrix().GetScale().GetVector());
 				const FVector2D NewLocalOffset = RootDrawOffset + WidgetInfo.WidgetGeometry.GetAccumulatedLayoutTransform().GetTranslation();
 				const FVector2D NewLocalSize = TransformPoint(WidgetInfo.WidgetGeometry.GetAccumulatedLayoutTransform().GetScale(), WidgetInfo.WidgetGeometry.GetLocalSize());
 
diff --git a/Engine/Source/Developer/StandaloneRenderer/Private/OpenGL/SlateOpenGLRenderingPolicy.cpp b/Engine/Source/Developer/StandaloneRenderer/Private/OpenGL/SlateOpenGLRenderingPolicy.cpp
--- a/Engine/Source/Developer/StandaloneRenderer/Private/OpenGL/SlateOpenGLRenderingPolicy.cpp
+++ b/Engine/Source/Developer/StandaloneRenderer/Private/OpenGL/SlateOpenGLRenderingPolicy.cpp
@@ -300,8 +300,8 @@ void FSlateOpenGLRenderingPolicy::DrawElements( const FMatrix& ViewProjectionMat
 
 					glEnable(GL_SCISSOR_TEST);
 					
-					const float ScissorWidth = FVector2D::Distance(ScissorRect.TopLeft, ScissorRect.TopRight);
-					const float ScissorHeight = FVector2D::Distance(ScissorRect.TopLeft, ScissorRect.BottomLeft);
+					const float ScissorWidth = FVector2f::Distance(ScissorRect.TopLeft, ScissorRect.TopRight);
+					const float ScissorHeight = FVector2f::Distance(ScissorRect.TopLeft, ScissorRect.BottomLeft);
 					glScissor(ScissorRect.TopLeft.X, ViewportSize.Y - ScissorRect.BottomLeft.Y, ScissorWidth, ScissorHeight);
 				}
 				else
diff --git a/Engine/Source/Developer/StandaloneRenderer/Private/OpenGL/SlateOpenGLTextureManager.cpp b/Engine/Source/Developer/StandaloneRenderer/Private/OpenGL/SlateOpenGLTextureManager.cpp
--- a/Engine/Source/Developer/StandaloneRenderer/Private/OpenGL/SlateOpenGLTextureManager.cpp
+++ b/Engine/Source/Developer/StandaloneRenderer/Private/OpenGL/SlateOpenGLTextureManager.cpp
@@ -282,8 +282,8 @@ FSlateShaderResourceProxy* FSlateOpenGLTextureManager::GenerateTextureResource(c
 		NewProxy = new FSlateShaderResourceProxy;
 		NewProxy->Resource = Atlas->GetAtlasTexture();
 		// Compute the sub-uvs for the location of this texture in the atlas, accounting for padding
-		NewProxy->StartUV = FVector2D((float)(NewSlot->X + Padding) / Atlas->GetWidth(), (float)(NewSlot->Y + Padding) / Atlas->GetHeight());
-		NewProxy->SizeUV = FVector2D((float)(NewSlot->Width - Padding * 2) / Atlas->GetWidth(), (float)(NewSlot->Height - Padding * 2) / Atlas->GetHeight());
+		NewProxy->StartUV = FVector2f((float)(NewSlot->X + Padding) / Atlas->GetWidth(), (float)(NewSlot->Y + Padding) / Atlas->GetHeight());
+		NewProxy->SizeUV = FVector2f((float)(NewSlot->Width - Padding * 2) / Atlas->GetWidth(), (float)(NewSlot->Height - Padding * 2) / Atlas->GetHeight());
 		NewProxy->ActualSize = FIntPoint(Width, Height);
 	}
 	else
@@ -299,8 +299,8 @@ FSlateShaderResourceProxy* FSlateOpenGLTextureManager::GenerateTextureResource(c
 		NewProxy = new FSlateShaderResourceProxy;
 
 		NewProxy->Resource = NewTexture.Get();
-		NewProxy->StartUV = FVector2D(0.0f, 0.0f);
-		NewProxy->SizeUV = FVector2D(1.0f, 1.0f);
+		NewProxy->StartUV = FVector2f(0.0f, 0.0f);
+		NewProxy->SizeUV = FVector2f(1.0f, 1.0f);
 		NewProxy->ActualSize = FIntPoint(Width, Height);
 
 		NonAtlasedTextures.Add(MoveTemp(NewTexture));
diff --git a/Engine/Source/Developer/StandaloneRenderer/Private/Windows/D3D/SlateD3DTextureManager.cpp b/Engine/Source/Developer/StandaloneRenderer/Private/Windows/D3D/SlateD3DTextureManager.cpp
--- a/Engine/Source/Developer/StandaloneRenderer/Private/Windows/D3D/SlateD3DTextureManager.cpp
+++ b/Engine/Source/Developer/StandaloneRenderer/Private/Windows/D3D/SlateD3DTextureManager.cpp
@@ -455,8 +455,8 @@ FSlateShaderResourceProxy* FSlateD3DTextureManager::GenerateTextureResource( con
 		NewProxy = new FSlateShaderResourceProxy;
 		NewProxy->Resource = Atlas->GetAtlasTexture();
 		// Compute the sub-uvs for the location of this texture in the atlas, accounting for padding
-		NewProxy->StartUV = FVector2D( (float)(NewSlot->X+Padding) / Atlas->GetWidth(), (float)(NewSlot->Y+Padding) / Atlas->GetHeight() );
-		NewProxy->SizeUV = FVector2D( (float)(NewSlot->Width-Padding*2) / Atlas->GetWidth(), (float)(NewSlot->Height-Padding*2) / Atlas->GetHeight() );
+		NewProxy->StartUV = FVector2f( (float)(NewSlot->X+Padding) / Atlas->GetWidth(), (float)(NewSlot->Y+Padding) / Atlas->GetHeight() );
+		NewProxy->SizeUV = FVector2f( (float)(NewSlot->Width-Padding*2) / Atlas->GetWidth(), (float)(NewSlot->Height-Padding*2) / Atlas->GetHeight() );
 		NewProxy->ActualSize = FIntPoint( Width, Height );
 	}
 	else
@@ -468,8 +468,8 @@ FSlateShaderResourceProxy* FSlateD3DTextureManager::GenerateTextureResource( con
 		TUniquePtr<FSlateD3DTexture> Texture = MakeUnique<FSlateD3DTexture>( Width, Height );
 	
 		NewProxy->Resource = Texture.Get();
-		NewProxy->StartUV = FVector2D(0.0f,0.0f);
-		NewProxy->SizeUV = FVector2D(1.0f, 1.0f);
+		NewProxy->StartUV = FVector2f(0.0f,0.0f);
+		NewProxy->SizeUV = FVector2f(1.0f, 1.0f);
 		NewProxy->ActualSize = FIntPoint( Width, Height );
 
 		D3D11_SUBRESOURCE_DATA InitData;
diff --git a/Engine/Source/Developer/TextureCompressor/Private/TextureCompressorModule.cpp b/Engine/Source/Developer/TextureCompressor/Private/TextureCompressorModule.cpp
--- a/Engine/Source/Developer/TextureCompressor/Private/TextureCompressorModule.cpp
+++ b/Engine/Source/Developer/TextureCompressor/Private/TextureCompressorModule.cpp
@@ -2584,10 +2584,10 @@ private:
 			if (BuildSettings.bHasColorSpaceDefinition)
 			{
 				Mip->TransformToWorkingColorSpace(
-					BuildSettings.RedChromaticityCoordinate,
-					BuildSettings.GreenChromaticityCoordinate,
-					BuildSettings.BlueChromaticityCoordinate,
-					BuildSettings.WhiteChromaticityCoordinate,
+					FVector2D(BuildSettings.RedChromaticityCoordinate),
+					FVector2D(BuildSettings.GreenChromaticityCoordinate),
+					FVector2D(BuildSettings.BlueChromaticityCoordinate),
+					FVector2D(BuildSettings.WhiteChromaticityCoordinate),
 					static_cast<UE::Color::EChromaticAdaptationMethod>(BuildSettings.ChromaticAdaptationMethod));
 			}
 
diff --git a/Engine/Source/Developer/TraceInsights/Private/Insights/ViewModels/GraphTrack.cpp b/Engine/Source/Developer/TraceInsights/Private/Insights/ViewModels/GraphTrack.cpp
--- a/Engine/Source/Developer/TraceInsights/Private/Insights/ViewModels/GraphTrack.cpp
+++ b/Engine/Source/Developer/TraceInsights/Private/Insights/ViewModels/GraphTrack.cpp
@@ -288,17 +288,17 @@ void FGraphTrack::DrawSeries(const FGraphSeries& Series, FDrawContext& DrawConte
 					const float X = PrevLinePoint.X + (LinePoint.X - PrevLinePoint.X) / ((BaselineY - LinePoint.Y) / (PrevLinePoint.Y - BaselineY) + 1.0f);
 
 					// Add an intersection point vertex.
-					FVector2D UV(X / Size.X, BaselineV);
-					Verts.Add(FSlateVertex::Make<ESlateVertexRounding::Disabled>(RenderTransform, UV * Size, AtlasOffset + UV * AtlasUVSize, FillColor));
+					FVector2f UV(X / Size.X, BaselineV);	// LWC_TODO: Precision loss
+					Verts.Add(FSlateVertex::Make<ESlateVertexRounding::Disabled>(RenderTransform, UV * FVector2f(Size), AtlasOffset + UV * AtlasUVSize, FillColor));	// LWC_TODO: Precision loss
 
 					// Add a value point vertex.
 					UV.X = LinePoint.X / Size.X;
 					UV.Y = TopV + LinePoint.Y / Size.Y;
-					Verts.Add(FSlateVertex::Make<ESlateVertexRounding::Disabled>(RenderTransform, UV * Size, AtlasOffset + UV * AtlasUVSize, FillColor));
+					Verts.Add(FSlateVertex::Make<ESlateVertexRounding::Disabled>(RenderTransform, UV * FVector2f(Size), AtlasOffset + UV * AtlasUVSize, FillColor));	// LWC_TODO: Precision loss
 
 					// Add a baseline vertex.
 					UV.Y = BaselineV;
-					Verts.Add(FSlateVertex::Make<ESlateVertexRounding::Disabled>(RenderTransform, UV * Size, AtlasOffset + FVector2D(UV.X, 0.5f) * AtlasUVSize, FillColor));
+					Verts.Add(FSlateVertex::Make<ESlateVertexRounding::Disabled>(RenderTransform, UV * FVector2f(Size), AtlasOffset + FVector2f(UV.X, 0.5f) * AtlasUVSize, FillColor));	// LWC_TODO: Precision loss
 
 					int32 Index0 = Verts.Num() - 5;
 					int32 Index1 = Verts.Num() - 4;
@@ -317,12 +317,12 @@ void FGraphTrack::DrawSeries(const FGraphSeries& Series, FDrawContext& DrawConte
 				else
 				{
 					// Add a value point vertex.
-					FVector2D UV(LinePoint.X / Size.X, TopV + LinePoint.Y / Size.Y);
-					Verts.Add(FSlateVertex::Make<ESlateVertexRounding::Disabled>(RenderTransform, UV * Size, AtlasOffset + UV * AtlasUVSize, FillColor));
+					FVector2f UV(LinePoint.X / Size.X, TopV + LinePoint.Y / Size.Y);	// LWC_TODO: Precision loss
+					Verts.Add(FSlateVertex::Make<ESlateVertexRounding::Disabled>(RenderTransform, UV * FVector2f(Size), AtlasOffset + UV * AtlasUVSize, FillColor));	// LWC_TODO: Precision loss
 
 					// Add a baseline vertex.
 					UV.Y = BaselineV;
-					Verts.Add(FSlateVertex::Make<ESlateVertexRounding::Disabled>(RenderTransform, UV * Size, AtlasOffset + FVector2D(UV.X, 0.5f) * AtlasUVSize, FillColor));
+					Verts.Add(FSlateVertex::Make<ESlateVertexRounding::Disabled>(RenderTransform, UV * FVector2f(Size), AtlasOffset + FVector2f(UV.X, 0.5f) * AtlasUVSize, FillColor));	// LWC_TODO: Precision loss
 
 					if (Verts.Num() >= 4)
 					{
diff --git a/Engine/Source/Editor/BehaviorTreeEditor/Private/SGraphNode_BehaviorTree.cpp b/Engine/Source/Editor/BehaviorTreeEditor/Private/SGraphNode_BehaviorTree.cpp
--- a/Engine/Source/Editor/BehaviorTreeEditor/Private/SGraphNode_BehaviorTree.cpp
+++ b/Engine/Source/Editor/BehaviorTreeEditor/Private/SGraphNode_BehaviorTree.cpp
@@ -610,7 +610,7 @@ EVisibility SGraphNode_BehaviorTree::GetDebuggerSearchFailedMarkerVisibility() c
 void SGraphNode_BehaviorTree::Tick( const FGeometry& AllottedGeometry, const double InCurrentTime, const float InDeltaTime )
 {
 	SGraphNode::Tick( AllottedGeometry, InCurrentTime, InDeltaTime );
-	CachedPosition = AllottedGeometry.AbsolutePosition / AllottedGeometry.Scale;
+	CachedPosition = FVector2D(AllottedGeometry.AbsolutePosition) / AllottedGeometry.Scale;
 
 	UBehaviorTreeGraphNode* MyNode = Cast<UBehaviorTreeGraphNode>(GraphNode);
 	if (MyNode && MyNode->DebuggerUpdateCounter != DebuggerStateCounter)
diff --git a/Engine/Source/Editor/BlueprintGraph/Private/EdGraphSchema_K2_Actions.cpp b/Engine/Source/Editor/BlueprintGraph/Private/EdGraphSchema_K2_Actions.cpp
--- a/Engine/Source/Editor/BlueprintGraph/Private/EdGraphSchema_K2_Actions.cpp
+++ b/Engine/Source/Editor/BlueprintGraph/Private/EdGraphSchema_K2_Actions.cpp
@@ -615,7 +615,7 @@ UEdGraphNode* FEdGraphSchemaAction_K2AddCallOnActor::PerformAction(class UEdGrap
 	LocalLocation.Y = FMath::GridSnap( (float)Location.Y, GridSnapSize);
 
 	// First use the base functionality to spawn the 'call function' node
-	UEdGraphNode* CallNode = FEdGraphSchemaAction_K2NewNode::PerformAction(ParentGraph, FromPin, LocalLocation);
+	UEdGraphNode* CallNode = FEdGraphSchemaAction_K2NewNode::PerformAction(ParentGraph, FromPin, FVector2D(LocalLocation));
 	const float FunctionNodeHeightUnsnapped = UEdGraphSchema_K2::EstimateNodeHeight( CallNode );
 
 	// this is the guesstimate of the function node's height, snapped to grid units
diff --git a/Engine/Source/Editor/CurveEditor/Private/CurveEditorHelpers.cpp b/Engine/Source/Editor/CurveEditor/Private/CurveEditorHelpers.cpp
--- a/Engine/Source/Editor/CurveEditor/Private/CurveEditorHelpers.cpp
+++ b/Engine/Source/Editor/CurveEditor/Private/CurveEditorHelpers.cpp
@@ -11,7 +11,7 @@ namespace CurveEditor
 FVector2D ComputeScreenSpaceTangentOffset(const FCurveEditorScreenSpace& CurveSpace, float Tangent, float Weight)
 {
 	const float Angle = FMath::Atan(-Tangent);
-	FVector2f Offset;
+	FVector2D Offset;
 	FMath::SinCos(&Offset.Y, &Offset.X, Angle);
 	Offset *= Weight;
 
diff --git a/Engine/Source/Editor/GraphEditor/Private/ConnectionDrawingPolicy.cpp b/Engine/Source/Editor/GraphEditor/Private/ConnectionDrawingPolicy.cpp
--- a/Engine/Source/Editor/GraphEditor/Private/ConnectionDrawingPolicy.cpp
+++ b/Engine/Source/Editor/GraphEditor/Private/ConnectionDrawingPolicy.cpp
@@ -28,13 +28,13 @@ FVector2D FGeometryHelper::VerticalMiddleRightOf(const FGeometry& SomeGeometry)
 FVector2D FGeometryHelper::CenterOf(const FGeometry& SomeGeometry)
 {
 	const FVector2D GeometryDrawSize = SomeGeometry.GetDrawSize();
-	return SomeGeometry.AbsolutePosition + (GeometryDrawSize * 0.5f);
+	return FVector2D(SomeGeometry.AbsolutePosition) + (GeometryDrawSize * 0.5f);
 }
 
 void FGeometryHelper::ConvertToPoints(const FGeometry& Geom, TArray<FVector2D>& Points)
 {
 	const FVector2D Size = Geom.GetDrawSize();
-	const FVector2D Location = Geom.AbsolutePosition;
+	const FVector2D Location = FVector2D(Geom.AbsolutePosition);
 
 	int32 Index = Points.AddUninitialized(4);
 	Points[Index++] = Location;
diff --git a/Engine/Source/Editor/GraphEditor/Private/SCommentBubble.cpp b/Engine/Source/Editor/GraphEditor/Private/SCommentBubble.cpp
--- a/Engine/Source/Editor/GraphEditor/Private/SCommentBubble.cpp
+++ b/Engine/Source/Editor/GraphEditor/Private/SCommentBubble.cpp
@@ -73,7 +73,7 @@ void SCommentBubble::Construct( const FArguments& InArgs )
 FCursorReply SCommentBubble::OnCursorQuery( const FGeometry& MyGeometry, const FPointerEvent& CursorEvent ) const
 {
 	const FVector2D Size( GetDesiredSize().X, GetDesiredSize().Y - SCommentBubbleDefs::BubbleArrowHeight );
-	const FSlateRect TestRect( MyGeometry.AbsolutePosition, MyGeometry.AbsolutePosition + Size );
+	const FSlateRect TestRect( FVector2D(MyGeometry.AbsolutePosition), FVector2D(MyGeometry.AbsolutePosition) + Size );
 
 	if( TestRect.ContainsPoint( CursorEvent.GetScreenSpacePosition() ))
 	{
diff --git a/Engine/Source/Editor/GraphEditor/Private/SGraphNode.cpp b/Engine/Source/Editor/GraphEditor/Private/SGraphNode.cpp
--- a/Engine/Source/Editor/GraphEditor/Private/SGraphNode.cpp
+++ b/Engine/Source/Editor/GraphEditor/Private/SGraphNode.cpp
@@ -449,7 +449,7 @@ void SGraphNode::OnToolTipClosing()
 
 void SGraphNode::Tick( const FGeometry& AllottedGeometry, const double InCurrentTime, const float InDeltaTime )
 {
-	CachedUnscaledPosition = AllottedGeometry.AbsolutePosition/AllottedGeometry.Scale;
+	CachedUnscaledPosition = FVector2D(AllottedGeometry.AbsolutePosition/AllottedGeometry.Scale);
 
 	SNodePanel::SNode::Tick(AllottedGeometry, InCurrentTime, InDeltaTime);
 
diff --git a/Engine/Source/Editor/GraphEditor/Private/SGraphPanel.cpp b/Engine/Source/Editor/GraphEditor/Private/SGraphPanel.cpp
--- a/Engine/Source/Editor/GraphEditor/Private/SGraphPanel.cpp
+++ b/Engine/Source/Editor/GraphEditor/Private/SGraphPanel.cpp
@@ -269,7 +269,7 @@ int32 SGraphPanel::OnPaint( const FPaintArgs& Args, const FGeometry& AllottedGeo
 
 								// Handle bouncing during PIE
 								const float BounceValue = FMath::Sin(2.0f * PI * BounceCurve.GetLerp());
-								BouncedGeometry.DrawPosition += (OverlayInfo.AnimationEnvelope * BounceValue * ZoomFactor);
+								BouncedGeometry.DrawPosition += FVector2f(OverlayInfo.AnimationEnvelope * BounceValue * ZoomFactor);
 
 								FLinearColor FinalColorAndOpacity(InWidgetStyle.GetColorAndOpacityTint()* OverlayBrush->GetTint(InWidgetStyle));
 								//FinalColorAndOpacity.A = Alpha;
@@ -359,7 +359,7 @@ int32 SGraphPanel::OnPaint( const FPaintArgs& Args, const FGeometry& AllottedGeo
 				ChildNode->GetPins(NodePins);
 
 				const FVector2D NodeLoc = ChildNode->GetPosition();
-				const FGeometry SynthesizedNodeGeometry(GraphCoordToPanelCoord(NodeLoc) * AllottedGeometry.Scale, AllottedGeometry.AbsolutePosition, FVector2D::ZeroVector, 1.f);
+				const FGeometry SynthesizedNodeGeometry(GraphCoordToPanelCoord(NodeLoc) * AllottedGeometry.Scale, FVector2D(AllottedGeometry.AbsolutePosition), FVector2D::ZeroVector, 1.f);
 
 				for (TArray< TSharedRef<SWidget> >::TConstIterator NodePinIterator(NodePins); NodePinIterator; ++NodePinIterator)
 				{
@@ -368,7 +368,7 @@ int32 SGraphPanel::OnPaint( const FPaintArgs& Args, const FGeometry& AllottedGeo
 					{
 						FVector2D PinLoc = NodeLoc + PinWidget.GetNodeOffset();
 
-						const FGeometry SynthesizedPinGeometry(GraphCoordToPanelCoord(PinLoc) * AllottedGeometry.Scale, AllottedGeometry.AbsolutePosition, FVector2D::ZeroVector, 1.f);
+						const FGeometry SynthesizedPinGeometry(GraphCoordToPanelCoord(PinLoc) * AllottedGeometry.Scale, FVector2D(AllottedGeometry.AbsolutePosition), FVector2D::ZeroVector, 1.f);
 						PinGeometries.Add(*NodePinIterator, FArrangedWidget(*NodePinIterator, SynthesizedPinGeometry));
 					}
 				}
@@ -551,7 +551,7 @@ void SGraphPanel::OnArrangeChildren( const FGeometry& AllottedGeometry, FArrange
 		{
 			FOverlayWidgetInfo& OverlayInfo = OverlayWidgets[WidgetIndex];
 
-			MyArrangedChildren.AddWidget(AllottedGeometry.MakeChild( OverlayInfo.Widget.ToSharedRef(), CurWidget.Geometry.Position + OverlayInfo.OverlayOffset, OverlayInfo.Widget->GetDesiredSize(), GetZoomAmount() ));
+			MyArrangedChildren.AddWidget(AllottedGeometry.MakeChild( OverlayInfo.Widget.ToSharedRef(), FVector2D(CurWidget.Geometry.Position) + OverlayInfo.OverlayOffset, OverlayInfo.Widget->GetDesiredSize(), GetZoomAmount() ));
 		}
 	}
 
diff --git a/Engine/Source/Editor/GraphEditor/Private/SGraphPin.cpp b/Engine/Source/Editor/GraphEditor/Private/SGraphPin.cpp
--- a/Engine/Source/Editor/GraphEditor/Private/SGraphPin.cpp
+++ b/Engine/Source/Editor/GraphEditor/Private/SGraphPin.cpp
@@ -771,18 +771,18 @@ FReply SGraphPin::OnDrop( const FGeometry& MyGeometry, const FDragDropEvent& Dra
 		TSharedPtr<SGraphNode> OwnerNode = OwnerNodePtr.Pin();
 		if (OwnerNode.IsValid())
 		{
-			NodeAddPosition	= OwnerNode->GetPosition() + MyGeometry.Position;
+			NodeAddPosition	= OwnerNode->GetPosition() + FVector2D(MyGeometry.Position);
 
-			//Don't have access to bounding information for node, using fixed offet that should work for most cases.
+			//Don't have access to bounding information for node, using fixed offset that should work for most cases.
 			const float FixedOffset = 200.0f;
 
 			//Line it up vertically with pin
 			NodeAddPosition.Y += MyGeometry.Size.Y;
 
 			// if the pin widget is nested into another compound
-			if (MyGeometry.Position == FVector2D::ZeroVector)
+			if (MyGeometry.Position == FVector2f::ZeroVector)
 			{
-				FVector2D PinOffsetPosition = MyGeometry.AbsolutePosition - NodeWidget->GetTickSpaceGeometry().AbsolutePosition;
+				FVector2D PinOffsetPosition = FVector2D(MyGeometry.AbsolutePosition) - FVector2D(NodeWidget->GetTickSpaceGeometry().AbsolutePosition);
 				NodeAddPosition = OwnerNode->GetPosition() + PinOffsetPosition;
 			}
 
@@ -812,7 +812,7 @@ FReply SGraphPin::OnDrop( const FGeometry& MyGeometry, const FDragDropEvent& Dra
 			TSharedPtr<FAssetDragDropOp> AssetOp = StaticCastSharedPtr<FAssetDragDropOp>(Operation);
 			if (AssetOp->HasAssets())
 			{
-				Node->GetSchema()->DroppedAssetsOnPin(AssetOp->GetAssets(), NodeWidget->GetPosition() + MyGeometry.Position, GraphPinObj);
+				Node->GetSchema()->DroppedAssetsOnPin(AssetOp->GetAssets(), NodeWidget->GetPosition() + FVector2D(MyGeometry.Position), GraphPinObj);
 			}
 		}
 		return FReply::Handled();
@@ -823,7 +823,7 @@ FReply SGraphPin::OnDrop( const FGeometry& MyGeometry, const FDragDropEvent& Dra
 
 void SGraphPin::Tick( const FGeometry& AllottedGeometry, const double InCurrentTime, const float InDeltaTime )
 {
-	CachedNodeOffset = AllottedGeometry.AbsolutePosition/AllottedGeometry.Scale - OwnerNodePtr.Pin()->GetUnscaledPosition();
+	CachedNodeOffset = FVector2D(AllottedGeometry.AbsolutePosition)/AllottedGeometry.Scale - OwnerNodePtr.Pin()->GetUnscaledPosition();
 	CachedNodeOffset.Y += AllottedGeometry.Size.Y * 0.5f;
 
 	if (!ValueInspectorTooltip.IsValid() && IsHovered() && FKismetDebugUtilities::CanInspectPinValue(GetPinObj()))
diff --git a/Engine/Source/Editor/GraphEditor/Private/SNodePanel.cpp b/Engine/Source/Editor/GraphEditor/Private/SNodePanel.cpp
--- a/Engine/Source/Editor/GraphEditor/Private/SNodePanel.cpp
+++ b/Engine/Source/Editor/GraphEditor/Private/SNodePanel.cpp
@@ -395,7 +395,7 @@ FVector2D SNodePanel::ComputeEdgePanAmount(const FGeometry& MyGeometry, const FV
 		EdgePanThisTick.Y = FMath::Min( MaxPanSpeed, EdgePanSpeedCoefficient * FMath::Pow(LocalCursorPos.Y - float(MyGeometry.GetLocalSize().Y) + EdgePanForgivenessZone, EdgePanSpeedPower) );
 	}
 
-	return EdgePanThisTick;
+	return FVector2D(EdgePanThisTick);
 }
 
 void SNodePanel::UpdateViewOffset (const FGeometry& MyGeometry, const FVector2D& TargetPosition)
diff --git a/Engine/Source/Editor/IntroTutorials/Private/STutorialButton.cpp b/Engine/Source/Editor/IntroTutorials/Private/STutorialButton.cpp
--- a/Engine/Source/Editor/IntroTutorials/Private/STutorialButton.cpp
+++ b/Engine/Source/Editor/IntroTutorials/Private/STutorialButton.cpp
@@ -128,7 +128,7 @@ int32 STutorialButton::OnPaint(const FPaintArgs& Args, const FGeometry& Allotted
 		{
 			FVector2D PulseOffset = FVector2D(PulseFactor0 * TutorialButtonConstants::MaxPulseOffset, PulseFactor0 * TutorialButtonConstants::MaxPulseOffset);
 
-			FVector2D BorderPosition = (AllottedGeometry.AbsolutePosition - ((FVector2D(PulseBrush->Margin.Left, PulseBrush->Margin.Top) * PulseBrush->ImageSize * AllottedGeometry.Scale) + PulseOffset));
+			FVector2D BorderPosition = (FVector2D(AllottedGeometry.AbsolutePosition) - ((FVector2D(PulseBrush->Margin.Left, PulseBrush->Margin.Top) * PulseBrush->ImageSize * AllottedGeometry.Scale) + PulseOffset));
 			FVector2D BorderSize = ((AllottedGeometry.GetLocalSize() * AllottedGeometry.Scale) + (PulseOffset * 2.0f) + (FVector2D(PulseBrush->Margin.Right * 2.0f, PulseBrush->Margin.Bottom * 2.0f) * PulseBrush->ImageSize * AllottedGeometry.Scale));
 
 			FPaintGeometry BorderGeometry(BorderPosition, BorderSize, AllottedGeometry.Scale);
@@ -140,7 +140,7 @@ int32 STutorialButton::OnPaint(const FPaintArgs& Args, const FGeometry& Allotted
 		{
 			FVector2D PulseOffset = FVector2D(PulseFactor1 * TutorialButtonConstants::MaxPulseOffset, PulseFactor1 * TutorialButtonConstants::MaxPulseOffset);
 
-			FVector2D BorderPosition = (AllottedGeometry.AbsolutePosition - ((FVector2D(PulseBrush->Margin.Left, PulseBrush->Margin.Top) * PulseBrush->ImageSize * AllottedGeometry.Scale) + PulseOffset));
+			FVector2D BorderPosition = (FVector2D(AllottedGeometry.AbsolutePosition) - ((FVector2D(PulseBrush->Margin.Left, PulseBrush->Margin.Top) * PulseBrush->ImageSize * AllottedGeometry.Scale) + PulseOffset));
 			FVector2D BorderSize = ((AllottedGeometry.Size * AllottedGeometry.Scale) + (PulseOffset * 2.0f) + (FVector2D(PulseBrush->Margin.Right * 2.0f, PulseBrush->Margin.Bottom * 2.0f) * PulseBrush->ImageSize * AllottedGeometry.Scale));
 
 			FPaintGeometry BorderGeometry(BorderPosition, BorderSize, AllottedGeometry.Scale);
diff --git a/Engine/Source/Editor/IntroTutorials/Private/STutorialContent.cpp b/Engine/Source/Editor/IntroTutorials/Private/STutorialContent.cpp
--- a/Engine/Source/Editor/IntroTutorials/Private/STutorialContent.cpp
+++ b/Engine/Source/Editor/IntroTutorials/Private/STutorialContent.cpp
@@ -301,7 +301,7 @@ int32 STutorialContent::OnPaint(const FPaintArgs& Args, const FGeometry& Allotte
 		// We should be clipped by the window size, not our containing widget, as we want to draw outside the widget
 		FSlateRect WindowClippingRect(0.0f, 0.0f, WindowSize.X, WindowSize.Y);
 
-		FPaintGeometry ShadowGeometry((WidgetGeometry.AbsolutePosition - FVector2D(ShadowBrush->Margin.Left, ShadowBrush->Margin.Top) * ShadowBrush->ImageSize * WidgetGeometry.Scale * TutorialConstants::ShadowScale),
+		FPaintGeometry ShadowGeometry((FVector2D(WidgetGeometry.AbsolutePosition) - FVector2D(ShadowBrush->Margin.Left, ShadowBrush->Margin.Top) * ShadowBrush->ImageSize * WidgetGeometry.Scale * TutorialConstants::ShadowScale),
 										((WidgetGeometry.GetLocalSize() * WidgetGeometry.Scale) + (FVector2D(ShadowBrush->Margin.Right * 2.0f, ShadowBrush->Margin.Bottom * 2.0f) * ShadowBrush->ImageSize * WidgetGeometry.Scale * TutorialConstants::ShadowScale)),
 										WidgetGeometry.Scale * TutorialConstants::ShadowScale);
 		// draw highlight shadow
@@ -309,7 +309,7 @@ int32 STutorialContent::OnPaint(const FPaintArgs& Args, const FGeometry& Allotte
 
 		FVector2D PulseOffset = FVector2D(PulseFactor * TutorialConstants::MaxBorderOffset, PulseFactor * TutorialConstants::MaxBorderOffset);
 
-		FVector2D BorderPosition = (WidgetGeometry.AbsolutePosition - ((FVector2D(BorderBrush->Margin.Left, BorderBrush->Margin.Top) * BorderBrush->ImageSize * WidgetGeometry.Scale) + PulseOffset));
+		FVector2D BorderPosition = (FVector2D(WidgetGeometry.AbsolutePosition) - ((FVector2D(BorderBrush->Margin.Left, BorderBrush->Margin.Top) * BorderBrush->ImageSize * WidgetGeometry.Scale) + PulseOffset));
 		FVector2D BorderSize = ((WidgetGeometry.Size * WidgetGeometry.Scale) + (PulseOffset * 2.0f) + (FVector2D(BorderBrush->Margin.Right * 2.0f, BorderBrush->Margin.Bottom * 2.0f) * BorderBrush->ImageSize * WidgetGeometry.Scale));
 
 		FPaintGeometry BorderGeometry(BorderPosition, BorderSize, WidgetGeometry.Scale);
diff --git a/Engine/Source/Editor/LandscapeEditor/Private/LandscapeEdModeBrushes.cpp b/Engine/Source/Editor/LandscapeEditor/Private/LandscapeEdModeBrushes.cpp
--- a/Engine/Source/Editor/LandscapeEditor/Private/LandscapeEdModeBrushes.cpp
+++ b/Engine/Source/Editor/LandscapeEditor/Private/LandscapeEdModeBrushes.cpp
@@ -120,7 +120,7 @@ public:
 	virtual void BeginStroke(float LandscapeX, float LandscapeY, FLandscapeTool* CurrentTool) override
 	{
 		FLandscapeBrush::BeginStroke(LandscapeX, LandscapeY, CurrentTool);
-		LastMousePosition = FVector2D(LandscapeX, LandscapeY);
+		LastMousePosition = FVector2f(LandscapeX, LandscapeY);
 	}
 
 	virtual void Tick(FEditorViewportClient* ViewportClient, float DeltaTime) override
@@ -248,7 +248,7 @@ public:
 
 	virtual void MouseMove(float LandscapeX, float LandscapeY) override
 	{
-		LastMousePosition = FVector2D(LandscapeX, LandscapeY);
+		LastMousePosition = FVector2f(LandscapeX, LandscapeY);
 	}
 
 	virtual FLandscapeBrushData ApplyBrush(const TArray<FLandscapeToolInteractorPosition>& InInteractorPositions) override
diff --git a/Engine/Source/Editor/LevelEditor/Private/LevelViewportLayout.cpp b/Engine/Source/Editor/LevelEditor/Private/LevelViewportLayout.cpp
--- a/Engine/Source/Editor/LevelEditor/Private/LevelViewportLayout.cpp
+++ b/Engine/Source/Editor/LevelEditor/Private/LevelViewportLayout.cpp
@@ -435,20 +435,20 @@ void FLevelViewportLayout::MaximizeViewport( FName ViewportToMaximize, const boo
 				{
 					FWidgetPath ViewportsOverlayWidgetPath = ViewportWidgetPath.GetPathDownTo( ViewportsOverlayPtr.Pin().ToSharedRef() );
 					const FArrangedWidget& ViewportsOverlayGeometry = ViewportsOverlayWidgetPath.Widgets.Last();
-					MaximizedViewportStartPosition = ViewportsOverlayGeometry.Geometry.AbsolutePosition - WindowScreenPos;
+					MaximizedViewportStartPosition = FVector2D(ViewportsOverlayGeometry.Geometry.AbsolutePosition) - WindowScreenPos;
 					MaximizedViewportStartSize = ViewportsOverlayPtr.Pin()->GetCachedSize();
 				}
 				else
 				{
 					const FArrangedWidget& ViewportGeometry = ViewportWidgetPath.Widgets.Last();	
-					MaximizedViewportStartPosition = ViewportGeometry.Geometry.AbsolutePosition - WindowScreenPos;
+					MaximizedViewportStartPosition = FVector2D(ViewportGeometry.Geometry.AbsolutePosition) - WindowScreenPos;
 					MaximizedViewportStartSize = ViewportGeometry.Geometry.Size;
 				}
 			}
 			else
 			{
 				const FArrangedWidget& ViewportGeometry = ViewportWidgetPath.Widgets.Last();
-				MaximizedViewportStartPosition = ViewportGeometry.Geometry.Position;
+				MaximizedViewportStartPosition = FVector2D(ViewportGeometry.Geometry.Position);
 				MaximizedViewportStartSize = ViewportGeometry.Geometry.Size;
 			}
 		}
diff --git a/Engine/Source/Editor/LevelEditor/Private/SCaptureRegionWidget.cpp b/Engine/Source/Editor/LevelEditor/Private/SCaptureRegionWidget.cpp
--- a/Engine/Source/Editor/LevelEditor/Private/SCaptureRegionWidget.cpp
+++ b/Engine/Source/Editor/LevelEditor/Private/SCaptureRegionWidget.cpp
@@ -62,14 +62,14 @@ FReply SCaptureRegionWidget::OnMouseButtonDown( const FGeometry& MyGeometry, con
 {
 	if (IsEnabled() && MouseEvent.GetEffectingButton() == EKeys::LeftMouseButton)
 	{
-		FVector2D ViewportPosition = MouseEvent.GetScreenSpacePosition() - MyGeometry.AbsolutePosition;
+		FVector2D ViewportPosition = MouseEvent.GetScreenSpacePosition() - FVector2D(MyGeometry.AbsolutePosition);
 		FIntRect& CurrentCaptureRegion = GetHighResScreenshotConfig().UnscaledCaptureRegion;
 
 		switch (PotentialInteraction)
 		{
 		case PI_DrawNewCaptureRegion:
 			{
-				DragStartPosition = MouseEvent.GetScreenSpacePosition() - MyGeometry.AbsolutePosition;
+				DragStartPosition = MouseEvent.GetScreenSpacePosition() - FVector2D(MyGeometry.AbsolutePosition);
 				BuildNewCaptureRegion(DragStartPosition, DragStartPosition);
 				CurrentState = State_Dragging;
 				break;
@@ -153,7 +153,7 @@ FReply SCaptureRegionWidget::OnMouseButtonUp( const FGeometry& MyGeometry, const
 {
 	if (IsEnabled() && MouseEvent.GetEffectingButton() == EKeys::LeftMouseButton)
 	{
-		FVector2D ViewportPosition = MouseEvent.GetScreenSpacePosition() - MyGeometry.AbsolutePosition;
+		FVector2D ViewportPosition = MouseEvent.GetScreenSpacePosition() - FVector2D(MyGeometry.AbsolutePosition);
 		FIntRect& CurrentCaptureRegion = GetHighResScreenshotConfig().UnscaledCaptureRegion;
 
 		switch (CurrentState)
@@ -214,7 +214,7 @@ FReply SCaptureRegionWidget::OnMouseMove( const FGeometry& MyGeometry, const FPo
 {
 	if (IsEnabled())
 	{
-		FVector2D ViewportPosition = MouseEvent.GetScreenSpacePosition() - MyGeometry.AbsolutePosition;
+		FVector2D ViewportPosition = MouseEvent.GetScreenSpacePosition() - FVector2D(MyGeometry.AbsolutePosition);
 		FIntRect& CurrentCaptureRegion = GetHighResScreenshotConfig().UnscaledCaptureRegion;
 
 		switch (CurrentState)
diff --git a/Engine/Source/Editor/MeshPaint/Private/MeshPaintHelpers.cpp b/Engine/Source/Editor/MeshPaint/Private/MeshPaintHelpers.cpp
--- a/Engine/Source/Editor/MeshPaint/Private/MeshPaintHelpers.cpp
+++ b/Engine/Source/Editor/MeshPaint/Private/MeshPaintHelpers.cpp
@@ -1398,7 +1398,7 @@ void MeshPaintHelpers::ImportVertexColorsToStaticMesh(UStaticMesh* StaticMesh, c
 	const FColor ColorMask = Options->CreateColorMask();
 	for (uint32 VertexIndex = 0; VertexIndex < LODModel.VertexBuffers.StaticMeshVertexBuffer.GetNumVertices(); ++VertexIndex)
 	{
-		const FVector2D UV = LODModel.VertexBuffers.StaticMeshVertexBuffer.GetVertexUV(VertexIndex, UVIndex);
+		const FVector2D UV = FVector2D(LODModel.VertexBuffers.StaticMeshVertexBuffer.GetVertexUV(VertexIndex, UVIndex));
 		LODModel.VertexBuffers.ColorVertexBuffer.VertexColor(VertexIndex) = PickVertexColorFromTextureData(MipData, UV, Texture, ColorMask);
 	}
 
@@ -1461,7 +1461,7 @@ void MeshPaintHelpers::ImportVertexColorsToStaticMeshComponent(UStaticMeshCompon
 		const FColor ColorMask = Options->CreateColorMask();
 		for (uint32 VertexIndex = 0; VertexIndex < LODModel.VertexBuffers.StaticMeshVertexBuffer.GetNumVertices(); ++VertexIndex)
 		{
-			const FVector2D UV = LODModel.VertexBuffers.StaticMeshVertexBuffer.GetVertexUV(VertexIndex, UVIndex);
+			const FVector2D UV = FVector2D(LODModel.VertexBuffers.StaticMeshVertexBuffer.GetVertexUV(VertexIndex, UVIndex));
 			InstanceMeshLODInfo.OverrideVertexColors->VertexColor(VertexIndex) = PickVertexColorFromTextureData(MipData, UV, Texture, ColorMask);
 		}
 
@@ -1508,7 +1508,7 @@ void MeshPaintHelpers::ImportVertexColorsToSkeletalMesh(USkeletalMesh* SkeletalM
 
 		for (uint32 VertexIndex = 0; VertexIndex < LODData.GetNumVertices(); ++VertexIndex)
 		{
-			const FVector2D UV = LODData.StaticVertexBuffers.StaticMeshVertexBuffer.GetVertexUV(VertexIndex, UVIndex);
+			const FVector2D UV = FVector2D(LODData.StaticVertexBuffers.StaticMeshVertexBuffer.GetVertexUV(VertexIndex, UVIndex));
 			LODData.StaticVertexBuffers.ColorVertexBuffer.VertexColor(VertexIndex) = PickVertexColorFromTextureData(MipData, UV, Texture, ColorMask);
 		}
 		
@@ -1525,7 +1525,7 @@ void MeshPaintHelpers::ImportVertexColorsToSkeletalMesh(USkeletalMesh* SkeletalM
 		int32 SectionVertexIndex = INDEX_NONE;
 		LODModel.GetSectionFromVertexIndex(VertexIndex, SectionIndex, SectionVertexIndex);
 
-		const FVector2D UV = LODModel.Sections[SectionIndex].SoftVertices[SectionVertexIndex].UVs[UVIndex];
+		const FVector2D UV = FVector2D(LODModel.Sections[SectionIndex].SoftVertices[SectionVertexIndex].UVs[UVIndex]);
 		LODModel.Sections[SectionIndex].SoftVertices[SectionVertexIndex].Color = PickVertexColorFromTextureData(MipData, UV, Texture, ColorMask);
 	}
 
diff --git a/Engine/Source/Editor/MeshPaint/Private/MeshPaintSkeletalMeshAdapter.cpp b/Engine/Source/Editor/MeshPaint/Private/MeshPaintSkeletalMeshAdapter.cpp
--- a/Engine/Source/Editor/MeshPaint/Private/MeshPaintSkeletalMeshAdapter.cpp
+++ b/Engine/Source/Editor/MeshPaint/Private/MeshPaintSkeletalMeshAdapter.cpp
@@ -426,7 +426,7 @@ void FMeshPaintGeometryAdapterForSkeletalMeshes::AddReferencedObjects(FReference
 
 void FMeshPaintGeometryAdapterForSkeletalMeshes::GetTextureCoordinate(int32 VertexIndex, int32 ChannelIndex, FVector2D& OutTextureCoordinate) const
 {
-	OutTextureCoordinate = LODData->StaticVertexBuffers.StaticMeshVertexBuffer.GetVertexUV(VertexIndex, ChannelIndex);
+	OutTextureCoordinate = FVector2D(LODData->StaticVertexBuffers.StaticMeshVertexBuffer.GetVertexUV(VertexIndex, ChannelIndex));
 }
 
 void FMeshPaintGeometryAdapterForSkeletalMeshes::PreEdit()
diff --git a/Engine/Source/Editor/MeshPaint/Private/MeshPaintStaticMeshAdapter.cpp b/Engine/Source/Editor/MeshPaint/Private/MeshPaintStaticMeshAdapter.cpp
--- a/Engine/Source/Editor/MeshPaint/Private/MeshPaintStaticMeshAdapter.cpp
+++ b/Engine/Source/Editor/MeshPaint/Private/MeshPaintStaticMeshAdapter.cpp
@@ -375,7 +375,7 @@ FMatrix FMeshPaintGeometryAdapterForStaticMeshes::GetComponentToWorldMatrix() co
 
 void FMeshPaintGeometryAdapterForStaticMeshes::GetTextureCoordinate(int32 VertexIndex, int32 ChannelIndex, FVector2D& OutTextureCoordinate) const
 {
-	OutTextureCoordinate = LODModel->VertexBuffers.StaticMeshVertexBuffer.GetVertexUV(VertexIndex, ChannelIndex);
+	OutTextureCoordinate = FVector2D(LODModel->VertexBuffers.StaticMeshVertexBuffer.GetVertexUV(VertexIndex, ChannelIndex));
 }
 
 void FMeshPaintGeometryAdapterForStaticMeshes::PreEdit()
diff --git a/Engine/Source/Editor/MeshPaint/Private/TexturePaintHelpers.cpp b/Engine/Source/Editor/MeshPaint/Private/TexturePaintHelpers.cpp
--- a/Engine/Source/Editor/MeshPaint/Private/TexturePaintHelpers.cpp
+++ b/Engine/Source/Editor/MeshPaint/Private/TexturePaintHelpers.cpp
@@ -224,7 +224,7 @@ bool TexturePaintHelpers::GenerateSeamMask(UMeshComponent* MeshComponent, int32
 			for (int32 TriVertexNum = 0; TriVertexNum < 3; ++TriVertexNum)
 			{
 				const int32 VertexIndex = Indices[TriIndex * 3 + TriVertexNum];
-				TriUVs[TriVertexNum] = LODModel.VertexBuffers.StaticMeshVertexBuffer.GetVertexUV(VertexIndex, UVSet);
+				TriUVs[TriVertexNum] = FVector2D(LODModel.VertexBuffers.StaticMeshVertexBuffer.GetVertexUV(VertexIndex, UVSet));
 
 				// Update bounds
 				float U = TriUVs[TriVertexNum].X;
diff --git a/Engine/Source/Editor/Persona/Private/SAnimNotifyPanel.cpp b/Engine/Source/Editor/Persona/Private/SAnimNotifyPanel.cpp
--- a/Engine/Source/Editor/Persona/Private/SAnimNotifyPanel.cpp
+++ b/Engine/Source/Editor/Persona/Private/SAnimNotifyPanel.cpp
@@ -1552,7 +1552,7 @@ void SAnimNotifyNode::Construct(const FArguments& InArgs)
 
 FReply SAnimNotifyNode::OnDragDetected( const FGeometry& MyGeometry, const FPointerEvent& MouseEvent )
 {
-	FVector2D ScreenNodePosition = MyGeometry.AbsolutePosition;
+	FVector2D ScreenNodePosition = FVector2D(MyGeometry.AbsolutePosition);
 	
 	// Whether the drag has hit a duration marker
 	bool bDragOnMarker = false;
@@ -1891,7 +1891,7 @@ FReply SAnimNotifyNode::OnMouseMove( const FGeometry& MyGeometry, const FPointer
 
 		if(MouseEvent.GetScreenSpacePosition().X >= TrackScreenSpaceXPosition && MouseEvent.GetScreenSpacePosition().X <= TrackScreenSpaceXPosition + CachedAllotedGeometrySize.X)
 		{
-			float NewDisplayTime = ScaleInfo.LocalXToInput((MouseEvent.GetScreenSpacePosition() - MyGeometry.AbsolutePosition + XPositionInTrack).X);
+			float NewDisplayTime = ScaleInfo.LocalXToInput((FVector2f(MouseEvent.GetScreenSpacePosition()) - MyGeometry.AbsolutePosition + XPositionInTrack).X);	// LWC_TODO: Precision loss
 			float NewDuration = NodeObjectInterface->GetDuration() + OldDisplayTime - NewDisplayTime;
 
 			// Check to make sure the duration is not less than the minimum allowed
@@ -1911,7 +1911,7 @@ FReply SAnimNotifyNode::OnMouseMove( const FGeometry& MyGeometry, const FPointer
 			ScaleInfo.ViewMinInput = ViewInputMin.Get();
 			ScaleInfo.ViewMaxInput = ViewInputMax.Get();
 
-			float NewDisplayTime = ScaleInfo.LocalXToInput((MouseEvent.GetScreenSpacePosition() - MyGeometry.AbsolutePosition + XPositionInTrack).X);
+			float NewDisplayTime = ScaleInfo.LocalXToInput((FVector2f(MouseEvent.GetScreenSpacePosition()) - MyGeometry.AbsolutePosition + XPositionInTrack).X);	// LWC_TODO: Precision loss
 			NodeObjectInterface->SetTime(FMath::Max(0.0f, NewDisplayTime));
 			NodeObjectInterface->SetDuration(NodeObjectInterface->GetDuration() + OldDisplayTime - NodeObjectInterface->GetTime());
 
@@ -1953,7 +1953,7 @@ FReply SAnimNotifyNode::OnMouseMove( const FGeometry& MyGeometry, const FPointer
 	{
 		if(MouseEvent.GetScreenSpacePosition().X >= TrackScreenSpaceXPosition && MouseEvent.GetScreenSpacePosition().X <= TrackScreenSpaceXPosition + CachedAllotedGeometrySize.X)
 		{
-			float NewDuration = ScaleInfo.LocalXToInput((MouseEvent.GetScreenSpacePosition() - MyGeometry.AbsolutePosition + XPositionInTrack).X) - NodeObjectInterface->GetTime();
+			float NewDuration = ScaleInfo.LocalXToInput((FVector2f(MouseEvent.GetScreenSpacePosition()) - MyGeometry.AbsolutePosition + XPositionInTrack).X) - NodeObjectInterface->GetTime();	// LWC_TODO: Precision loss
 
 			NodeObjectInterface->SetDuration(FMath::Max(NewDuration, MinimumStateDuration));
 		}
@@ -1965,7 +1965,7 @@ FReply SAnimNotifyNode::OnMouseMove( const FGeometry& MyGeometry, const FPointer
 			ScaleInfo.ViewMinInput = ViewInputMin.Get();
 			ScaleInfo.ViewMaxInput = ViewInputMax.Get();
 
-			float NewDuration = ScaleInfo.LocalXToInput((MouseEvent.GetScreenSpacePosition() - MyGeometry.AbsolutePosition + XPositionInTrack).X) - NodeObjectInterface->GetTime();
+			float NewDuration = ScaleInfo.LocalXToInput((FVector2f(MouseEvent.GetScreenSpacePosition()) - MyGeometry.AbsolutePosition + XPositionInTrack).X) - NodeObjectInterface->GetTime();	// LWC_TODO: Precision loss
 			NodeObjectInterface->SetDuration(FMath::Max(NewDuration, MinimumStateDuration));
 		}
 
@@ -2119,7 +2119,7 @@ void SAnimNotifyNode::DrawHandleOffset( const float& Offset, const float& Handle
 
 void SAnimNotifyNode::Tick( const FGeometry& AllottedGeometry, const double InCurrentTime, const float InDeltaTime )
 {
-	ScreenPosition = AllottedGeometry.AbsolutePosition;
+	ScreenPosition = FVector2D(AllottedGeometry.AbsolutePosition);
 }
 
 void SAnimNotifyNode::OnFocusLost(const FFocusEvent& InFocusEvent)
diff --git a/Engine/Source/Editor/Persona/Private/SAnimationBlendSpaceGridWidget.cpp b/Engine/Source/Editor/Persona/Private/SAnimationBlendSpaceGridWidget.cpp
--- a/Engine/Source/Editor/Persona/Private/SAnimationBlendSpaceGridWidget.cpp
+++ b/Engine/Source/Editor/Persona/Private/SAnimationBlendSpaceGridWidget.cpp
@@ -87,7 +87,7 @@ static void PaintTriangle(
 	{
 		Vertices.AddZeroed();
 		FSlateVertex& NewVert = Vertices.Last();
-		NewVert.Position = AllottedGeometry.LocalToAbsolute(*Points[PointIndex]);
+		NewVert.Position = FVector2f(AllottedGeometry.LocalToAbsolute(*Points[PointIndex]));	// LWC_TODO: Precision loss
 		NewVert.Color = Color.ToFColor(false);
 	}
 
@@ -120,7 +120,7 @@ static void PaintPolygon(
 	{
 		Vertices.AddZeroed();
 		FSlateVertex& NewVert = Vertices.Last();
-		NewVert.Position = AllottedGeometry.LocalToAbsolute(Points[PointIndex]);
+		NewVert.Position = FVector2f(AllottedGeometry.LocalToAbsolute(Points[PointIndex]));	// LWC_TODO: Precision loss
 		NewVert.Color = FillColor.ToFColor(false);
 		MidVertex.Position += NewVert.Position;
 	}
@@ -133,8 +133,8 @@ static void PaintPolygon(
 		FComparePoints(const FSlateVertex& Mid) : MidPoint(Mid) {}
 		bool operator()(const FSlateVertex& A, const FSlateVertex& B) const
 		{
-			FVector2D DeltaA = A.Position - MidPoint.Position;
-			FVector2D DeltaB = B.Position - MidPoint.Position;
+			FVector2D DeltaA = FVector2D(A.Position) - FVector2D(MidPoint.Position);
+			FVector2D DeltaB = FVector2D(B.Position) - FVector2D(MidPoint.Position);
 			float AngleA = FMath::Atan2(DeltaA.Y, DeltaA.X);
 			float AngleB = FMath::Atan2(DeltaB.Y, DeltaB.X);
 			return AngleA < AngleB;
@@ -164,7 +164,7 @@ static void PaintPolygon(
 		LinePoints.Reserve(Points.Num() + 1);
 		for (int VertexIndex = 1; VertexIndex <= Vertices.Num(); ++VertexIndex)
 		{
-			LinePoints.Add(VertexIndex < Vertices.Num() ? Vertices[VertexIndex].Position : Vertices[1].Position);
+			LinePoints.Add(FVector2D(VertexIndex < Vertices.Num() ? Vertices[VertexIndex].Position : Vertices[1].Position));
 		}
 		FSlateDrawElement::MakeLines(
 			OutDrawElements, DrawLayerId + 1, FPaintGeometry(), LinePoints, ESlateDrawEffect::None, OutlineColor, true, 1.0f);
diff --git a/Engine/Source/Editor/Persona/Private/STrack.cpp b/Engine/Source/Editor/Persona/Private/STrack.cpp
--- a/Engine/Source/Editor/Persona/Private/STrack.cpp
+++ b/Engine/Source/Editor/Persona/Private/STrack.cpp
@@ -241,7 +241,7 @@ FReply STrackNode::BeginDrag( const FGeometry& MyGeometry, const FPointerEvent&
 
 	FVector2D ScreenCursorPos = MouseEvent.GetScreenSpacePosition();
 	FVector2D CursorPos = MyGeometry.AbsoluteToLocal(ScreenCursorPos);
-	FVector2D ScreenNodePosition = MyGeometry.AbsolutePosition;// + GetOffsetRelativeToParent(MyGeometry);
+	FVector2D ScreenNodePosition = FVector2D(MyGeometry.AbsolutePosition);// + GetOffsetRelativeToParent(MyGeometry);
 	
 	bBeingDragged = true;
 	LastSize = MyGeometry.GetLocalSize();
diff --git a/Engine/Source/Editor/PixelInspector/Private/PixelInspector.cpp b/Engine/Source/Editor/PixelInspector/Private/PixelInspector.cpp
--- a/Engine/Source/Editor/PixelInspector/Private/PixelInspector.cpp
+++ b/Engine/Source/Editor/PixelInspector/Private/PixelInspector.cpp
@@ -515,7 +515,7 @@ namespace PixelInspector
 		if (BufferIndex == -1)
 			return;
 		
-		Requests[BufferIndex].SetRequestData(InspectViewportUV, BufferIndex, viewportUniqueId, GBufferFormat, AllowStaticLighting, InPreExposure);
+		Requests[BufferIndex].SetRequestData(FVector2f(InspectViewportUV), BufferIndex, viewportUniqueId, GBufferFormat, AllowStaticLighting, InPreExposure);	// LWC_TODO: Precision loss
 		SceneInterface->AddPixelInspectorRequest(&(Requests[BufferIndex]));
 	}
 
@@ -706,14 +706,14 @@ namespace PixelInspector
 			{
 				if (Request.FrameCountAfterRenderingCommandSend >= WAIT_FRAMENUMBER_BEFOREREADING)
 				{
-					if (Request.SourceViewportUV == FVector2D(-1, -1))
+					if (Request.SourceViewportUV == FVector2f(-1, -1))
 					{
 						continue;
 					}
 
 
 					PixelInspectorResult PixelResult;
-					PixelResult.ViewportUV = Request.SourceViewportUV;
+					PixelResult.ViewportUV = FVector2D(Request.SourceViewportUV);
 					PixelResult.ViewUniqueId = Request.ViewId;
 					PixelResult.PreExposure = Request.PreExposure;
 					PixelResult.OneOverPreExposure = Request.PreExposure > 0.f ? (1.f / Request.PreExposure) : 1.f;;
diff --git a/Engine/Source/Editor/SequenceRecorder/Private/SequenceRecorder.cpp b/Engine/Source/Editor/SequenceRecorder/Private/SequenceRecorder.cpp
--- a/Engine/Source/Editor/SequenceRecorder/Private/SequenceRecorder.cpp
+++ b/Engine/Source/Editor/SequenceRecorder/Private/SequenceRecorder.cpp
@@ -487,12 +487,12 @@ void FSequenceRecorder::DrawDebug(UCanvas* InCanvas, APlayerController* InPlayer
 	if(bCountingDown)
 	{
 		const FVector2D IconSize(128.0f, 128.0f);
-		const FVector2D HalfIconSize(64.0f, 64.0f);
+		const FVector2f HalfIconSize(64.0f, 64.0f);
 		const float LineThickness = 2.0f;
 
 		FVector2f Center;
 		InCanvas->GetCenter(Center.X, Center.Y);
-		FVector2D IconPosition = Center - HalfIconSize;
+		FVector2f IconPosition = Center - HalfIconSize;
 
 		InCanvas->SetDrawColor(FColor::White);
 
@@ -503,8 +503,8 @@ void FSequenceRecorder::DrawDebug(UCanvas* InCanvas, APlayerController* InPlayer
 		const float Angle = 2.0f * PI * FMath::Fmod(CurrentDelay, 1.0f);
 		const FVector2D AxisX(0.f, -1.f);
 		const FVector2D AxisY(-1.f, 0.f);
-		const FVector2D EndPos = Center + (AxisX * FMath::Cos(Angle) + AxisY * FMath::Sin(Angle)) * (InCanvas->SizeX + InCanvas->SizeY);
-		FCanvasLineItem LineItem(Center, EndPos);
+		const FVector2D EndPos = FVector2D(Center) + (AxisX * FMath::Cos(Angle) + AxisY * FMath::Sin(Angle)) * (InCanvas->SizeX + InCanvas->SizeY);
+		FCanvasLineItem LineItem(FVector2D(Center), EndPos);
 		LineItem.LineThickness = LineThickness;
 		LineItem.SetColor(FLinearColor::Black);
 		InCanvas->DrawItem(LineItem);
diff --git a/Engine/Source/Editor/Sequencer/Private/SSequencerDebugVisualizer.cpp b/Engine/Source/Editor/Sequencer/Private/SSequencerDebugVisualizer.cpp
--- a/Engine/Source/Editor/Sequencer/Private/SSequencerDebugVisualizer.cpp
+++ b/Engine/Source/Editor/Sequencer/Private/SSequencerDebugVisualizer.cpp
@@ -264,7 +264,7 @@ void SSequencerEvaluationTemplateDebugVisualizer::OnArrangeChildren( const FGeom
  			{
  				ArrangedChildren.AddWidget( 
  					WidgetVisibility, 
- 					AllottedGeometry.MakeChild(Child, SegmentGeometry.Position, SegmentGeometry.GetLocalSize())
+ 					AllottedGeometry.MakeChild(Child, FVector2D(SegmentGeometry.Position), SegmentGeometry.GetLocalSize())
  					);
  			}
  		}
@@ -627,7 +627,7 @@ void SSequencerEntityComponentSystemDebugVisualizer::OnArrangeChildren( const FG
  			{
  				ArrangedChildren.AddWidget( 
  					WidgetVisibility, 
- 					AllottedGeometry.MakeChild(Child, SegmentGeometry.Position, SegmentGeometry.GetLocalSize())
+ 					AllottedGeometry.MakeChild(Child, FVector2D(SegmentGeometry.Position), SegmentGeometry.GetLocalSize())
  					);
  			}
  		}
diff --git a/Engine/Source/Editor/Sequencer/Private/SSequencerSection.cpp b/Engine/Source/Editor/Sequencer/Private/SSequencerSection.cpp
--- a/Engine/Source/Editor/Sequencer/Private/SSequencerSection.cpp
+++ b/Engine/Source/Editor/Sequencer/Private/SSequencerSection.cpp
@@ -440,7 +440,7 @@ struct FSequencerSectionPainterImpl : FSequencerSectionPainter
 
 		FSlateRenderTransform RenderTransform;
 
-		const FVector2D Pos = RangeGeometry.GetAbsolutePosition();
+		const FVector2f Pos = FVector2f(RangeGeometry.GetAbsolutePosition());	// LWC_TODO: Precision loss
 		const FVector2D Size = RangeGeometry.GetLocalSize();
 
 		FLinearColor EaseSelectionColor = FEditorStyle::GetSlateColor(SequencerSectionConstants::SelectionColorName).GetColor(FWidgetStyle());
@@ -477,6 +477,8 @@ struct FSequencerSectionPainterImpl : FSequencerSectionPainter
 			Indices.Reserve(CurvePoints.Num()*6);
 			Verts.Reserve(CurvePoints.Num()*2);
 
+			const FVector2f SizeAsFloatVec = FVector2f(Size);	// LWC_TODO: Precision loss
+
 			for (const FEasingCurvePoint& Point : CurvePoints)
 			{
 				float SegmentStartTime = UE::MovieScene::DiscreteInclusiveLower(Segment.Range) / TimeToPixelConverter.GetTickResolution();
@@ -484,12 +486,12 @@ struct FSequencerSectionPainterImpl : FSequencerSectionPainter
 
 				// Add verts top->bottom
 				FVector2f UV(U, 0.f);
-				Verts.Add(FSlateVertex::Make<ESlateVertexRounding::Disabled>(RenderTransform, (Pos + UV*Size*RangeGeometry.Scale), AtlasOffset + UV*AtlasUVSize, FillColor));
+				Verts.Add(FSlateVertex::Make<ESlateVertexRounding::Disabled>(RenderTransform, (Pos + UV*SizeAsFloatVec*RangeGeometry.Scale), AtlasOffset + UV*AtlasUVSize, FillColor));	// LWC_TODO: Precision loss
 
 				UV.Y = 1.f - Point.Location.Y;
-				BorderPoints.Add(UV*Size);
+				BorderPoints.Add(FVector2D(UV)*Size);
 				BorderPointColors.Add(Point.Color);
-				Verts.Add(FSlateVertex::Make<ESlateVertexRounding::Disabled>(RenderTransform, (Pos + UV*Size*RangeGeometry.Scale), AtlasOffset + FVector2f(UV.X, 0.5f)*AtlasUVSize, FillColor));
+				Verts.Add(FSlateVertex::Make<ESlateVertexRounding::Disabled>(RenderTransform, (Pos + UV*SizeAsFloatVec*RangeGeometry.Scale), AtlasOffset + FVector2f(UV.X, 0.5f)*AtlasUVSize, FillColor));	// LWC_TODO: Precision loss
 
 				if (Verts.Num() >= 4)
 				{
diff --git a/Engine/Source/Editor/Sequencer/Private/SSequencerSectionAreaView.cpp b/Engine/Source/Editor/Sequencer/Private/SSequencerSectionAreaView.cpp
--- a/Engine/Source/Editor/Sequencer/Private/SSequencerSectionAreaView.cpp
+++ b/Engine/Source/Editor/Sequencer/Private/SSequencerSectionAreaView.cpp
@@ -262,7 +262,7 @@ void SSequencerSectionAreaView::OnArrangeChildren( const FGeometry& AllottedGeom
 			
 			ArrangedChildren.AddWidget( 
 				WidgetVisibility, 
-				AllottedGeometry.MakeChild( Widget, SectionGeometry.Position, SectionGeometry.GetLocalSize() )
+				AllottedGeometry.MakeChild( Widget, FVector2D(SectionGeometry.Position), SectionGeometry.GetLocalSize() )
 				);
 		}
 	}
diff --git a/Engine/Source/Editor/StaticMeshEditor/Private/StaticMeshEditorViewportClient.cpp b/Engine/Source/Editor/StaticMeshEditor/Private/StaticMeshEditorViewportClient.cpp
--- a/Engine/Source/Editor/StaticMeshEditor/Private/StaticMeshEditorViewportClient.cpp
+++ b/Engine/Source/Editor/StaticMeshEditor/Private/StaticMeshEditorViewportClient.cpp
@@ -1269,8 +1269,8 @@ void FStaticMeshEditorViewportClient::ProcessClick(class FSceneView& InView, cla
 								if (RenderData.VertexBuffers.StaticMeshVertexBuffer.GetNumTexCoords() > TexCoordIndex)
 								{
 									FVector2D UVIndex1, UVIndex2;
-									UVIndex1 = RenderData.VertexBuffers.StaticMeshVertexBuffer.GetVertexUV(VertexIndex, TexCoordIndex);
-									UVIndex2 = RenderData.VertexBuffers.StaticMeshVertexBuffer.GetVertexUV(VertexIndex2, TexCoordIndex);
+									UVIndex1 = FVector2D(RenderData.VertexBuffers.StaticMeshVertexBuffer.GetVertexUV(VertexIndex, TexCoordIndex));
+									UVIndex2 = FVector2D(RenderData.VertexBuffers.StaticMeshVertexBuffer.GetVertexUV(VertexIndex2, TexCoordIndex));
 									SelectedEdgeTexCoords[TexCoordIndex].Add(UVIndex1);
 									SelectedEdgeTexCoords[TexCoordIndex].Add(UVIndex2);
 								}
diff --git a/Engine/Source/Editor/ToolMenusEditor/Private/SEditToolMenuDialog.cpp b/Engine/Source/Editor/ToolMenusEditor/Private/SEditToolMenuDialog.cpp
--- a/Engine/Source/Editor/ToolMenusEditor/Private/SEditToolMenuDialog.cpp
+++ b/Engine/Source/Editor/ToolMenusEditor/Private/SEditToolMenuDialog.cpp
@@ -334,7 +334,7 @@ FReply SMultiBlockDragHandle::OnDragDetected( const FGeometry& MyGeometry, const
 			Block->IsPartOfHeading(),
 			NAME_None, 
 			CustomDecorator,
-			MyGeometry.AbsolutePosition-MouseEvent.GetScreenSpacePosition()
+		FVector2D(MyGeometry.AbsolutePosition)-MouseEvent.GetScreenSpacePosition()
 		);
 
 	NewOp->SetOnDropNotification( FSimpleDelegate::CreateSP( BaseWidget.Pin().ToSharedRef(), &SMultiBoxWidget::OnDropExternal ) );
diff --git a/Engine/Source/Editor/UMGEditor/Private/Designer/SDesignerView.cpp b/Engine/Source/Editor/UMGEditor/Private/Designer/SDesignerView.cpp
--- a/Engine/Source/Editor/UMGEditor/Private/Designer/SDesignerView.cpp
+++ b/Engine/Source/Editor/UMGEditor/Private/Designer/SDesignerView.cpp
@@ -108,7 +108,7 @@ public:
 		if (MouseEvent.GetEffectingButton() == EKeys::LeftMouseButton)
 		{
 			bResizing = true;
-			AbsoluteOffset = MouseEvent.GetScreenSpacePosition() - MyGeometry.AbsolutePosition;
+			AbsoluteOffset = MouseEvent.GetScreenSpacePosition() - FVector2D(MyGeometry.AbsolutePosition);
 			return FReply::Handled().CaptureMouse(SharedThis(this));
 		}
 
@@ -1324,11 +1324,11 @@ FVector2D SDesignerView::GetWidgetOriginAbsolute() const
 		FGeometry Geometry;
 		if (GetWidgetGeometry(PreviewWidget, Geometry))
 		{
-			return Geometry.AbsolutePosition;
+			return FVector2D(Geometry.AbsolutePosition);
 		}
 	}
 
-	return FVector2D(0, 0);
+	return FVector2D::ZeroVector;
 }
 
 void SDesignerView::MarkDesignModifed(bool bRequiresRecompile)
@@ -1462,7 +1462,7 @@ FVector2D SDesignerView::GetExtensionPosition(TSharedRef<FDesignerSurfaceElement
 
 		if ( GetWidgetGeometry(SelectedWidget, SelectedWidgetGeometry) && GetWidgetParentGeometry(SelectedWidget, SelectedWidgetParentGeometry) )
 		{
-			const FVector2D ParentPostion_DesignerSpace = (SelectedWidgetParentGeometry.AbsolutePosition - GetDesignerGeometry().AbsolutePosition) / GetDesignerGeometry().Scale;
+			const FVector2D ParentPostion_DesignerSpace = FVector2D(SelectedWidgetParentGeometry.AbsolutePosition - GetDesignerGeometry().AbsolutePosition) / GetDesignerGeometry().Scale;
 			const FVector2D ParentSize = SelectedWidgetParentGeometry.Size * GetPreviewScale();
 
 			FVector2D FinalPosition(0, 0);
@@ -1510,7 +1510,7 @@ FVector2D SDesignerView::GetExtensionPosition(TSharedRef<FDesignerSurfaceElement
 					break;
 				}
 
-				FVector2D SelectedWidgetScale = SelectedWidgetGeometry.GetAccumulatedRenderTransform().GetMatrix().GetScale().GetVector();
+				FVector2D SelectedWidgetScale = FVector2D(SelectedWidgetGeometry.GetAccumulatedRenderTransform().GetMatrix().GetScale().GetVector());
 
 				FVector2D ApplicationScaledOffset = ExtensionElement->GetOffset() * GetDesignerGeometry().Scale;
 
@@ -2145,17 +2145,17 @@ void SDesignerView::DrawSelectionAndHoverOutline(const FOnPaintHandlerParams& Pa
 			FDesignTimeUtils::GetArrangedWidgetRelativeToWindow(Widget, ArrangedWidget);
 
 			// Draw selection effect
-			const FVector2D OutlinePixelSize = FVector2D(2.0f, 2.0f) / ArrangedWidget.Geometry.GetAccumulatedRenderTransform().GetMatrix().GetScale().GetVector();
+			const FVector2D OutlinePixelSize = FVector2D(2.0f, 2.0f) / FVector2D(ArrangedWidget.Geometry.GetAccumulatedRenderTransform().GetMatrix().GetScale().GetVector());
 			FPaintGeometry SelectionGeometry = ArrangedWidget.Geometry.ToInflatedPaintGeometry(OutlinePixelSize);
 
 			FSlateClippingZone SelectionZone(SelectionGeometry);
 
 			TArray<FVector2D> Points;
-			Points.Add(SelectionZone.TopLeft);
-			Points.Add(SelectionZone.TopRight);
-			Points.Add(SelectionZone.BottomRight);
-			Points.Add(SelectionZone.BottomLeft);
-			Points.Add(SelectionZone.TopLeft);
+			Points.Add(FVector2D(SelectionZone.TopLeft));
+			Points.Add(FVector2D(SelectionZone.TopRight));
+			Points.Add(FVector2D(SelectionZone.BottomRight));
+			Points.Add(FVector2D(SelectionZone.BottomLeft));
+			Points.Add(FVector2D(SelectionZone.TopLeft));
 
 			FSlateDrawElement::MakeLines(
 				PaintArgs.OutDrawElements,
@@ -2183,17 +2183,17 @@ void SDesignerView::DrawSelectionAndHoverOutline(const FOnPaintHandlerParams& Pa
 		FDesignTimeUtils::GetArrangedWidgetRelativeToWindow(Widget, ArrangedWidget);
 
 		// Draw hovered effect
-		const FVector2D OutlinePixelSize = FVector2D(2.0f, 2.0f) / ArrangedWidget.Geometry.GetAccumulatedRenderTransform().GetMatrix().GetScale().GetVector();
+		const FVector2D OutlinePixelSize = FVector2D(2.0f, 2.0f) / FVector2D(ArrangedWidget.Geometry.GetAccumulatedRenderTransform().GetMatrix().GetScale().GetVector());
 		FPaintGeometry HoveredGeometry = ArrangedWidget.Geometry.ToInflatedPaintGeometry(OutlinePixelSize);
 
 		FSlateClippingZone HoveredZone(HoveredGeometry);
 
 		TArray<FVector2D> Points;
-		Points.Add(HoveredZone.TopLeft);
-		Points.Add(HoveredZone.TopRight);
-		Points.Add(HoveredZone.BottomRight);
-		Points.Add(HoveredZone.BottomLeft);
-		Points.Add(HoveredZone.TopLeft);
+		Points.Add(FVector2D(HoveredZone.TopLeft));
+		Points.Add(FVector2D(HoveredZone.TopRight));
+		Points.Add(FVector2D(HoveredZone.BottomRight));
+		Points.Add(FVector2D(HoveredZone.BottomLeft));
+		Points.Add(FVector2D(HoveredZone.TopLeft));
 
 		FSlateDrawElement::MakeLines(
 			PaintArgs.OutDrawElements,
diff --git a/Engine/Source/Editor/UMGEditor/Private/Extensions/CanvasSlotExtension.cpp b/Engine/Source/Editor/UMGEditor/Private/Extensions/CanvasSlotExtension.cpp
--- a/Engine/Source/Editor/UMGEditor/Private/Extensions/CanvasSlotExtension.cpp
+++ b/Engine/Source/Editor/UMGEditor/Private/Extensions/CanvasSlotExtension.cpp
@@ -307,20 +307,20 @@ void FCanvasSlotExtension::GetCollisionSegmentsFromGeometry(FGeometry ArrangedGe
 	Segments.SetNumUninitialized(8);
 
 	// Left Side
-	Segments[0] = ArrangedGeometry.Position;
-	Segments[1] = ArrangedGeometry.Position + FVector2D(0, ArrangedGeometry.GetLocalSize().Y);
+	Segments[0] = FVector2D(ArrangedGeometry.Position);
+	Segments[1] = FVector2D(ArrangedGeometry.Position) + FVector2D(0, ArrangedGeometry.GetLocalSize().Y);
 
 	// Top Side
-	Segments[2] = ArrangedGeometry.Position;
-	Segments[3] = ArrangedGeometry.Position + FVector2D(ArrangedGeometry.GetLocalSize().X, 0);
+	Segments[2] = FVector2D(ArrangedGeometry.Position);
+	Segments[3] = FVector2D(ArrangedGeometry.Position) + FVector2D(ArrangedGeometry.GetLocalSize().X, 0);
 
 	// Right Side
-	Segments[4] = ArrangedGeometry.Position + FVector2D(ArrangedGeometry.GetLocalSize().X, 0);
-	Segments[5] = ArrangedGeometry.Position + ArrangedGeometry.GetLocalSize();
+	Segments[4] = FVector2D(ArrangedGeometry.Position) + FVector2D(ArrangedGeometry.GetLocalSize().X, 0);
+	Segments[5] = FVector2D(ArrangedGeometry.Position) + ArrangedGeometry.GetLocalSize();
 
 	// Bottom Side
-	Segments[6] = ArrangedGeometry.Position + FVector2D(0, ArrangedGeometry.GetLocalSize().Y);
-	Segments[7] = ArrangedGeometry.Position + ArrangedGeometry.GetLocalSize();
+	Segments[6] = FVector2D(ArrangedGeometry.Position) + FVector2D(0, ArrangedGeometry.GetLocalSize().Y);
+	Segments[7] = FVector2D(ArrangedGeometry.Position) + ArrangedGeometry.GetLocalSize();
 }
 
 void FCanvasSlotExtension::Paint(const TSet< FWidgetReference >& Selection, const FGeometry& AllottedGeometry, const FSlateRect& MyCullingRect, FSlateWindowElementList& OutDrawElements, int32 LayerId) const
diff --git a/Engine/Source/Editor/UnrealEd/Private/EditorExporters.cpp b/Engine/Source/Editor/UnrealEd/Private/EditorExporters.cpp
--- a/Engine/Source/Editor/UnrealEd/Private/EditorExporters.cpp
+++ b/Engine/Source/Editor/UnrealEd/Private/EditorExporters.cpp
@@ -977,7 +977,7 @@ static void AddActorToOBJs(AActor* Actor, TArray<FOBJGeom*>& Objects, TSet<UMate
 					// Vertices
 					VerticesOut[i].Vert = LocalToWorld.TransformPosition(RenderData->VertexBuffers.PositionVertexBuffer.VertexPosition(i));
 					// UVs from channel 0
-					VerticesOut[i].UV = RenderData->VertexBuffers.StaticMeshVertexBuffer.GetVertexUV(i, 0);
+					VerticesOut[i].UV = FVector2D(RenderData->VertexBuffers.StaticMeshVertexBuffer.GetVertexUV(i, 0));
 					// Normal
 					VerticesOut[i].Normal = LocalToWorldInverseTranspose.TransformVector((FVector4)RenderData->VertexBuffers.StaticMeshVertexBuffer.VertexTangentZ(i));
 				}
@@ -1856,16 +1856,16 @@ bool UStaticMeshExporterOBJ::ExportText(const FExportObjectInnerContext* Context
 			Verts.Add( Vertex3 );
 
 			// UVs from channel 0
-			UVs.Add( RenderData.VertexBuffers.StaticMeshVertexBuffer.GetVertexUV(Index1, 0) );
-			UVs.Add( RenderData.VertexBuffers.StaticMeshVertexBuffer.GetVertexUV(Index2, 0) );
-			UVs.Add( RenderData.VertexBuffers.StaticMeshVertexBuffer.GetVertexUV(Index3, 0) );
+			UVs.Add( FVector2D(RenderData.VertexBuffers.StaticMeshVertexBuffer.GetVertexUV(Index1, 0)) );
+			UVs.Add( FVector2D(RenderData.VertexBuffers.StaticMeshVertexBuffer.GetVertexUV(Index2, 0)) );
+			UVs.Add( FVector2D(RenderData.VertexBuffers.StaticMeshVertexBuffer.GetVertexUV(Index3, 0)) );
 
 			// UVs from channel 1 (lightmap coords)
 			if (bHasUVLightMap)
 			{
-				UVLMs.Add(RenderData.VertexBuffers.StaticMeshVertexBuffer.GetVertexUV(Index1, 1));
-				UVLMs.Add(RenderData.VertexBuffers.StaticMeshVertexBuffer.GetVertexUV(Index2, 1));
-				UVLMs.Add(RenderData.VertexBuffers.StaticMeshVertexBuffer.GetVertexUV(Index3, 1));
+				UVLMs.Add(FVector2D(RenderData.VertexBuffers.StaticMeshVertexBuffer.GetVertexUV(Index1, 1)));
+				UVLMs.Add(FVector2D(RenderData.VertexBuffers.StaticMeshVertexBuffer.GetVertexUV(Index2, 1)));
+				UVLMs.Add(FVector2D(RenderData.VertexBuffers.StaticMeshVertexBuffer.GetVertexUV(Index3, 1)));
 			}
 
 			// Normals
diff --git a/Engine/Source/Editor/UnrealEd/Private/Fbx/FbxMainExport.cpp b/Engine/Source/Editor/UnrealEd/Private/Fbx/FbxMainExport.cpp
--- a/Engine/Source/Editor/UnrealEd/Private/Fbx/FbxMainExport.cpp
+++ b/Engine/Source/Editor/UnrealEd/Private/Fbx/FbxMainExport.cpp
@@ -1135,7 +1135,7 @@ void FFbxExporter::ExportBSP( UModel* Model, bool bSelectedOnly )
 
 					float U = ((Vertex - TextureBase) | TextureX) / UModel::GetGlobalBSPTexelScale();
 					float V = ((Vertex - TextureBase) | TextureY) / UModel::GetGlobalBSPTexelScale();
-					UVs.Set(VertexInstanceIDs[Corner], 0, FVector2D(U, V));
+					UVs.Set(VertexInstanceIDs[Corner], 0, FVector2f(U, V));
 					//This is not exported when exporting the whole level via ExportModel so leaving out here for now. 
 					//UVs.Set(VertexInstanceIDs[Corner], 1, Vert.ShadowTexCoord);
 					Colors[VertexInstanceIDs[Corner]] = FVector4f(FLinearColor::White);
@@ -3606,7 +3606,7 @@ void DetermineUVsToWeld(TArray<int32>& VertRemap, TArray<int32>& UniqueVerts, co
 	TMap<FVector2D,int32> HashedVerts;
 	for(int32 Vertex=0; Vertex < VertexCount; Vertex++)
 	{
-		const FVector2D& PositionA = VertexBuffer.GetVertexUV(Vertex,TexCoordSourceIndex);
+		const FVector2D& PositionA = FVector2D(VertexBuffer.GetVertexUV(Vertex,TexCoordSourceIndex));
 		const int32* FoundIndex = HashedVerts.Find(PositionA);
 		if ( !FoundIndex )
 		{
@@ -4610,7 +4610,7 @@ FbxNode* FFbxExporter::ExportStaticMeshToFbx(const UStaticMesh* StaticMesh, int3
 			for (int32 FbxVertIndex = 0; FbxVertIndex < UniqueUVs.Num(); FbxVertIndex++)
 			{
 				int32 UnrealVertIndex = UniqueUVs[FbxVertIndex];
-				const FVector2D& TexCoord = RenderMesh.VertexBuffers.StaticMeshVertexBuffer.GetVertexUV(UnrealVertIndex, TexCoordSourceIndex);
+				const FVector2f& TexCoord = RenderMesh.VertexBuffers.StaticMeshVertexBuffer.GetVertexUV(UnrealVertIndex, TexCoordSourceIndex);
 				UVDiffuseLayer->GetDirectArray().Add(FbxVector2(TexCoord.X, -TexCoord.Y + 1.0));
 			}
 
@@ -4932,7 +4932,7 @@ void FFbxExporter::ExportSplineMeshToFbx(const USplineMeshComponent* SplineMeshC
 		// Create the texture coordinate data source.
 		for (int32 UnrealVertIndex : UniqueUVs)
 		{
-			const FVector2D& TexCoord = RenderMesh.VertexBuffers.StaticMeshVertexBuffer.GetVertexUV(UnrealVertIndex, TexCoordSourceIndex);
+			const FVector2f& TexCoord = RenderMesh.VertexBuffers.StaticMeshVertexBuffer.GetVertexUV(UnrealVertIndex, TexCoordSourceIndex);
 			UVDiffuseLayer->GetDirectArray().Add(FbxVector2(TexCoord.X, -TexCoord.Y + 1.0));
 		}
 
diff --git a/Engine/Source/Editor/UnrealEd/Private/Fbx/FbxSkeletalMeshExport.cpp b/Engine/Source/Editor/UnrealEd/Private/Fbx/FbxSkeletalMeshExport.cpp
--- a/Engine/Source/Editor/UnrealEd/Private/Fbx/FbxSkeletalMeshExport.cpp
+++ b/Engine/Source/Editor/UnrealEd/Private/Fbx/FbxSkeletalMeshExport.cpp
@@ -191,7 +191,7 @@ FbxNode* FFbxExporter::CreateMesh(const USkeletalMesh* SkelMesh, const TCHAR* Me
 		// Create the texture coordinate data source.
 		for (int32 TexCoordIndex = 0; TexCoordIndex < VertexCount; ++TexCoordIndex)
 		{
-			const FVector2D& TexCoord = Vertices[TexCoordIndex].UVs[TexCoordSourceIndex];
+			const FVector2D& TexCoord = FVector2D(Vertices[TexCoordIndex].UVs[TexCoordSourceIndex]);
 			UVDiffuseLayer->GetDirectArray().Add(FbxVector2(TexCoord.X, -TexCoord.Y + 1.0));
 		}
 
diff --git a/Engine/Source/Editor/UnrealEd/Private/Fbx/FbxStaticMeshImport.cpp b/Engine/Source/Editor/UnrealEd/Private/Fbx/FbxStaticMeshImport.cpp
--- a/Engine/Source/Editor/UnrealEd/Private/Fbx/FbxStaticMeshImport.cpp
+++ b/Engine/Source/Editor/UnrealEd/Private/Fbx/FbxStaticMeshImport.cpp
@@ -727,7 +727,7 @@ bool UnFbx::FFbxImporter::BuildStaticMeshFromGeometry(FbxNode* Node, UStaticMesh
 				{
 					FUVID UVID = MeshDescription->CreateUV(UVLayerIndex);
 					FbxVector2 UVVector = FBXUVs.LayerElementUV[UVLayerIndex]->GetDirectArray().GetAt(UVIndex);
-					UVCoordinates[UVID] = FVector2D(static_cast<float>(UVVector[0]), 1.0f - static_cast<float>(UVVector[1]));	// flip the Y of UVs for DirectX
+					UVCoordinates[UVID] = FVector2f(static_cast<float>(UVVector[0]), 1.0f - static_cast<float>(UVVector[1]));	// flip the Y of UVs for DirectX
 				}
 			}
 		}
@@ -870,7 +870,7 @@ bool UnFbx::FFbxImporter::BuildStaticMeshFromGeometry(FbxNode* Node, UStaticMesh
 					//UVs attributes
 					for (int32 UVLayerIndex = 0; UVLayerIndex < FBXUVs.UniqueUVCount; UVLayerIndex++)
 					{
-						FVector2D FinalUVVector(0.0f, 0.0f);
+						FVector2f FinalUVVector(0.0f, 0.0f);
 						if (FBXUVs.LayerElementUV[UVLayerIndex] != NULL)
 						{
 							int UVMapIndex = (FBXUVs.UVMappingMode[UVLayerIndex] == FbxLayerElement::eByControlPoint) ? ControlPointIndex : RealFbxVertexIndex;
diff --git a/Engine/Source/Editor/UnrealEd/Private/LevelEditorViewport.cpp b/Engine/Source/Editor/UnrealEd/Private/LevelEditorViewport.cpp
--- a/Engine/Source/Editor/UnrealEd/Private/LevelEditorViewport.cpp
+++ b/Engine/Source/Editor/UnrealEd/Private/LevelEditorViewport.cpp
@@ -4366,13 +4366,13 @@ void FLevelEditorViewportClient::DrawBrushDetails(const FSceneView* View, FPrimi
 					const FVector Vertex0 = Poly->Vertices[0];
 					FVector Vertex1 = Poly->Vertices[1];
 
-					MeshBuilder.AddVertex(Vertex0, FVector2D::ZeroVector, FVector(1, 0, 0), FVector(0, 1, 0), FVector(0, 0, 1), FColor::White);
-					MeshBuilder.AddVertex(Vertex1, FVector2D::ZeroVector, FVector(1, 0, 0), FVector(0, 1, 0), FVector(0, 0, 1), FColor::White);
+					MeshBuilder.AddVertex(Vertex0, FVector2f::ZeroVector, FVector(1, 0, 0), FVector(0, 1, 0), FVector(0, 0, 1), FColor::White);
+					MeshBuilder.AddVertex(Vertex1, FVector2f::ZeroVector, FVector(1, 0, 0), FVector(0, 1, 0), FVector(0, 0, 1), FColor::White);
 
 					for (int32 VertexIdx = 2; VertexIdx < Poly->Vertices.Num(); ++VertexIdx)
 					{
 						const FVector Vertex2 = Poly->Vertices[VertexIdx];
-						MeshBuilder.AddVertex(Vertex2, FVector2D::ZeroVector, FVector(1, 0, 0), FVector(0, 1, 0), FVector(0, 0, 1), FColor::White);
+						MeshBuilder.AddVertex(Vertex2, FVector2f::ZeroVector, FVector(1, 0, 0), FVector(0, 1, 0), FVector(0, 0, 1), FColor::White);
 						MeshBuilder.AddTriangle(VertexOffset, VertexOffset + VertexIdx, VertexOffset + VertexIdx - 1);
 						Vertex1 = Vertex2;
 					}
diff --git a/Engine/Source/Editor/UnrealEd/Private/Lightmass/Lightmass.cpp b/Engine/Source/Editor/UnrealEd/Private/Lightmass/Lightmass.cpp
--- a/Engine/Source/Editor/UnrealEd/Private/Lightmass/Lightmass.cpp
+++ b/Engine/Source/Editor/UnrealEd/Private/Lightmass/Lightmass.cpp
@@ -265,13 +265,13 @@ FORCEINLINE void Copy( const FSplineMeshParams& In, Lightmass::FSplineMeshParams
 {
 	Out.StartPos = In.StartPos;
 	Out.StartTangent = In.StartTangent;
-	Out.StartScale = In.StartScale;
+	Out.StartScale = FVector2f(In.StartScale);	// LWC_TODO: Precision loss
 	Out.StartRoll = In.StartRoll;
-	Out.StartOffset = In.StartOffset;
+	Out.StartOffset = FVector2f(In.StartOffset);	// LWC_TODO: Precision loss
 	Out.EndPos = In.EndPos;
 	Out.EndTangent = In.EndTangent;
-	Out.EndScale = In.EndScale;
-	Out.EndOffset = In.EndOffset;
+	Out.EndScale = FVector2f(In.EndScale);	// LWC_TODO: Precision loss
+	Out.EndOffset = FVector2f(In.EndOffset);	// LWC_TODO: Precision loss
 	Out.EndRoll = In.EndRoll;
 }
 
@@ -1332,10 +1332,10 @@ void FLightmassExporter::WriteStaticMeshes()
 							{
 								Vertex.UVs[UVIndex] = RenderData.VertexBuffers.StaticMeshVertexBuffer.GetVertexUV(VertexIndex, UVIndex);
 							}
-							FVector2D ZeroUV(0.0f, 0.0f);
+
 							for (; UVIndex < MAX_TEXCOORDS; UVIndex++)
 							{
-								Vertex.UVs[UVIndex] = ZeroUV;
+								Vertex.UVs[UVIndex] = FVector2f::ZeroVector;
 							}
 						}
 						Swarm.WriteChannel( Channel, (void*)(LMVertices.GetData()), LMVertices.Num() * sizeof(Lightmass::FStaticMeshVertex));
@@ -1965,7 +1965,7 @@ void FLightmassExporter::WriteMappings( int32 Channel )
 				DstVertex.WorldTangentZ = (FVector3f)SrcVertex.WorldTangentZ;
 				for( int32 CoordIdx=0; CoordIdx < Lightmass::MAX_TEXCOORDS; CoordIdx++ )
 				{	
-					DstVertex.TextureCoordinates[CoordIdx] = SrcVertex.TextureCoordinates[CoordIdx];
+					DstVertex.TextureCoordinates[CoordIdx] = FVector2f(SrcVertex.TextureCoordinates[CoordIdx]);	// LWC_TODO: Precision loss
 				}
 			}
 			Swarm.WriteChannel( Channel, VertexData, VertexDataSize );
diff --git a/Engine/Source/Editor/UnrealEd/Private/Lightmass/LightmassLandscapeRender.cpp b/Engine/Source/Editor/UnrealEd/Private/Lightmass/LightmassLandscapeRender.cpp
--- a/Engine/Source/Editor/UnrealEd/Private/Lightmass/LightmassLandscapeRender.cpp
+++ b/Engine/Source/Editor/UnrealEd/Private/Lightmass/LightmassLandscapeRender.cpp
@@ -71,12 +71,12 @@ void RenderLandscapeMaterialForLightmass(const FLandscapeStaticLightingMesh* Lan
 
 			const FVector2D BasePosition = PatchExpandOffset + FVector2D(SubsectionX, SubsectionY) * PositionScale;
 			const FVector2D BaseLayerCoords = FVector2D(LandscapeComponent->SectionBaseX, LandscapeComponent->SectionBaseY) + FVector2D(UVSubsection) * LayerScale;
-			const FVector2D BaseWeightmapCoords = WeightmapBias + FVector2D(UVSubsection) * WeightmapSubsection;
+			const FVector2f BaseWeightmapCoords = FVector2f(WeightmapBias) + FVector2f(UVSubsection) * WeightmapSubsection;	// LWC_TODO: Precision loss
 
 			int32 Index = Vertices.Add(FDynamicMeshVertex(FVector(BasePosition /*FVector2D(0, 0) * PositionScale*/, 0), FVector(BaseLayerCoords /*FVector2D(0, 0) * UVScale * LayerScale*/, 0), BaseWeightmapCoords /*FVector2D(0, 0) * UVScale * WeightmapScale*/));
-			verifySlow(   Vertices.Add(FDynamicMeshVertex(FVector(BasePosition + FVector2D(1, 0) * PositionScale,   0), FVector(BaseLayerCoords + FVector2D(1, 0) * UVScale * LayerScale,   0), BaseWeightmapCoords + FVector2D(1, 0) * UVScale * WeightmapScale  )) == Index + 1);
-			verifySlow(   Vertices.Add(FDynamicMeshVertex(FVector(BasePosition + FVector2D(0, 1) * PositionScale,   0), FVector(BaseLayerCoords + FVector2D(0, 1) * UVScale * LayerScale,   0), BaseWeightmapCoords + FVector2D(0, 1) * UVScale * WeightmapScale  )) == Index + 2);
-			verifySlow(   Vertices.Add(FDynamicMeshVertex(FVector(BasePosition + FVector2D(1, 1) * PositionScale,   0), FVector(BaseLayerCoords + FVector2D(1, 1) * UVScale * LayerScale,   0), BaseWeightmapCoords + FVector2D(1, 1) * UVScale * WeightmapScale  )) == Index + 3);
+			verifySlow(   Vertices.Add(FDynamicMeshVertex(FVector(BasePosition + FVector2D(1, 0) * PositionScale,   0), FVector(BaseLayerCoords + FVector2D(1, 0) * UVScale * LayerScale,   0), BaseWeightmapCoords + FVector2f(1, 0) * FVector2f(UVScale * WeightmapScale))) == Index + 1);	// LWC_TODO: Precision loss
+			verifySlow(   Vertices.Add(FDynamicMeshVertex(FVector(BasePosition + FVector2D(0, 1) * PositionScale,   0), FVector(BaseLayerCoords + FVector2D(0, 1) * UVScale * LayerScale,   0), BaseWeightmapCoords + FVector2f(0, 1) * FVector2f(UVScale * WeightmapScale))) == Index + 2);	// LWC_TODO: Precision loss
+			verifySlow(   Vertices.Add(FDynamicMeshVertex(FVector(BasePosition + FVector2D(1, 1) * PositionScale,   0), FVector(BaseLayerCoords + FVector2D(1, 1) * UVScale * LayerScale,   0), BaseWeightmapCoords + FVector2f(1, 1) * FVector2f(UVScale * WeightmapScale))) == Index + 3);	// LWC_TODO: Precision loss
 			checkSlow(Index + 3 <= MAX_uint16);
 			Indices.Add(Index);
 			Indices.Add(Index + 3);
diff --git a/Engine/Source/Editor/UnrealEd/Private/StaticLightingSystem/StaticLightingDebug.cpp b/Engine/Source/Editor/UnrealEd/Private/StaticLightingSystem/StaticLightingDebug.cpp
--- a/Engine/Source/Editor/UnrealEd/Private/StaticLightingSystem/StaticLightingDebug.cpp
+++ b/Engine/Source/Editor/UnrealEd/Private/StaticLightingSystem/StaticLightingDebug.cpp
@@ -251,9 +251,9 @@ void SetDebugLightmapSample(TArray<UActorComponent*>* Components, UModel* Model,
 						uint32 Index2 = Indices[ClosestPlaneTriangleIndex + 2];
 
 						// Fetch lightmap UV's
-						FVector2D LightmapUV0 = LODModel.VertexBuffers.StaticMeshVertexBuffer.GetVertexUV(Index0, StaticMesh->GetLightMapCoordinateIndex());
-						FVector2D LightmapUV1 = LODModel.VertexBuffers.StaticMeshVertexBuffer.GetVertexUV(Index1, StaticMesh->GetLightMapCoordinateIndex());
-						FVector2D LightmapUV2 = LODModel.VertexBuffers.StaticMeshVertexBuffer.GetVertexUV(Index2, StaticMesh->GetLightMapCoordinateIndex());
+						FVector2D LightmapUV0 = FVector2D(LODModel.VertexBuffers.StaticMeshVertexBuffer.GetVertexUV(Index0, StaticMesh->GetLightMapCoordinateIndex()));
+						FVector2D LightmapUV1 = FVector2D(LODModel.VertexBuffers.StaticMeshVertexBuffer.GetVertexUV(Index1, StaticMesh->GetLightMapCoordinateIndex()));
+						FVector2D LightmapUV2 = FVector2D(LODModel.VertexBuffers.StaticMeshVertexBuffer.GetVertexUV(Index2, StaticMesh->GetLightMapCoordinateIndex()));
 						// Interpolate lightmap UV's to the click location
 						FVector2D InterpolatedUV = LightmapUV0 * ClosestPlaneBaryCentricWeights.X + LightmapUV1 * ClosestPlaneBaryCentricWeights.Y + LightmapUV2 * ClosestPlaneBaryCentricWeights.Z;
 
@@ -292,9 +292,9 @@ void SetDebugLightmapSample(TArray<UActorComponent*>* Components, UModel* Model,
 					uint32 Index1 = Indices[SelectedTriangle + 1];
 					uint32 Index2 = Indices[SelectedTriangle + 2];
 
-					FVector2D LightmapUV0 = LODModel.VertexBuffers.StaticMeshVertexBuffer.GetVertexUV(Index0, StaticMesh->GetLightMapCoordinateIndex());
-					FVector2D LightmapUV1 = LODModel.VertexBuffers.StaticMeshVertexBuffer.GetVertexUV(Index1, StaticMesh->GetLightMapCoordinateIndex());
-					FVector2D LightmapUV2 = LODModel.VertexBuffers.StaticMeshVertexBuffer.GetVertexUV(Index2, StaticMesh->GetLightMapCoordinateIndex());
+					FVector2D LightmapUV0 = FVector2D(LODModel.VertexBuffers.StaticMeshVertexBuffer.GetVertexUV(Index0, StaticMesh->GetLightMapCoordinateIndex()));
+					FVector2D LightmapUV1 = FVector2D(LODModel.VertexBuffers.StaticMeshVertexBuffer.GetVertexUV(Index1, StaticMesh->GetLightMapCoordinateIndex()));
+					FVector2D LightmapUV2 = FVector2D(LODModel.VertexBuffers.StaticMeshVertexBuffer.GetVertexUV(Index2, StaticMesh->GetLightMapCoordinateIndex()));
 
 					FVector BaryCentricWeights;
 					BaryCentricWeights.X = FMath::FRandRange(0.f, 1.f);
@@ -409,9 +409,9 @@ void SetDebugLightmapSample(TArray<UActorComponent*>* Components, UModel* Model,
 				uint32 Index2 = (*IndexBufferRef)->Indices[ClosestTriangleIndex + 2];
 
 				// Fetch lightmap UV's
-				FVector2D LightmapUV0 = ModelVertices[Index0].ShadowTexCoord;
-				FVector2D LightmapUV1 = ModelVertices[Index1].ShadowTexCoord;
-				FVector2D LightmapUV2 = ModelVertices[Index2].ShadowTexCoord;
+				FVector2D LightmapUV0 = FVector2D(ModelVertices[Index0].ShadowTexCoord);
+				FVector2D LightmapUV1 = FVector2D(ModelVertices[Index1].ShadowTexCoord);
+				FVector2D LightmapUV2 = FVector2D(ModelVertices[Index2].ShadowTexCoord);
 				// Interpolate lightmap UV's to the click location
 				FVector2D InterpolatedUV = LightmapUV0 * ClosestPlaneBaryCentricWeights.X + LightmapUV1 * ClosestPlaneBaryCentricWeights.Y + LightmapUV2 * ClosestPlaneBaryCentricWeights.Z;
 
diff --git a/Engine/Source/Editor/UnrealEd/Private/StaticMeshEdit.cpp b/Engine/Source/Editor/UnrealEd/Private/StaticMeshEdit.cpp
--- a/Engine/Source/Editor/UnrealEd/Private/StaticMeshEdit.cpp
+++ b/Engine/Source/Editor/UnrealEd/Private/StaticMeshEdit.cpp
@@ -232,7 +232,7 @@ void GetBrushMesh(ABrush* Brush, UModel* Model, FMeshDescription& MeshDescriptio
 					VertexPositions[VertexID[CornerIndex]] = Positions[CornerIndex];
 				}
 				VertexInstanceIDs[CornerIndex] = MeshDescription.CreateVertexInstance(VertexID[CornerIndex]);
-				VertexInstanceUVs.Set(VertexInstanceIDs[CornerIndex], 0, FVector2D(
+				VertexInstanceUVs.Set(VertexInstanceIDs[CornerIndex], 0, FVector2f(
 					(Positions[CornerIndex] - TextureBase) | TextureX,
 					(Positions[CornerIndex] - TextureBase) | TextureY));
 			}
diff --git a/Engine/Source/Editor/UnrealEd/Private/UnrealWidgetRender.cpp b/Engine/Source/Editor/UnrealEd/Private/UnrealWidgetRender.cpp
--- a/Engine/Source/Editor/UnrealEd/Private/UnrealWidgetRender.cpp
+++ b/Engine/Source/Editor/UnrealEd/Private/UnrealWidgetRender.cpp
@@ -380,10 +380,10 @@ void DrawCornerHelper( FPrimitiveDrawInterface* PDI, const FMatrix& LocalToWorld
 	// Top
 	{
 		int32 VertexIndices[4];
-		VertexIndices[0] = MeshBuilder.AddVertex( FVector(-TX, -TY, +TZ), FVector2D::ZeroVector, FVector(1,0,0), FVector(0,1,0), FVector(0,0,1), FColor::White );
-		VertexIndices[1] = MeshBuilder.AddVertex( FVector(-TX, +TY, +TZ), FVector2D::ZeroVector, FVector(1,0,0), FVector(0,1,0), FVector(0,0,1), FColor::White );
-		VertexIndices[2] = MeshBuilder.AddVertex( FVector(+TX, +TY, +TZ), FVector2D::ZeroVector, FVector(1,0,0), FVector(0,1,0), FVector(0,0,1), FColor::White );
-		VertexIndices[3] = MeshBuilder.AddVertex( FVector(+TX, -TY, +TZ), FVector2D::ZeroVector, FVector(1,0,0), FVector(0,1,0), FVector(0,0,1), FColor::White );
+		VertexIndices[0] = MeshBuilder.AddVertex( FVector(-TX, -TY, +TZ), FVector2f::ZeroVector, FVector(1,0,0), FVector(0,1,0), FVector(0,0,1), FColor::White );
+		VertexIndices[1] = MeshBuilder.AddVertex( FVector(-TX, +TY, +TZ), FVector2f::ZeroVector, FVector(1,0,0), FVector(0,1,0), FVector(0,0,1), FColor::White );
+		VertexIndices[2] = MeshBuilder.AddVertex( FVector(+TX, +TY, +TZ), FVector2f::ZeroVector, FVector(1,0,0), FVector(0,1,0), FVector(0,0,1), FColor::White );
+		VertexIndices[3] = MeshBuilder.AddVertex( FVector(+TX, -TY, +TZ), FVector2f::ZeroVector, FVector(1,0,0), FVector(0,1,0), FVector(0,0,1), FColor::White );
 
 		MeshBuilder.AddTriangle(VertexIndices[0],VertexIndices[1],VertexIndices[2]);
 		MeshBuilder.AddTriangle(VertexIndices[0],VertexIndices[2],VertexIndices[3]);
@@ -392,10 +392,10 @@ void DrawCornerHelper( FPrimitiveDrawInterface* PDI, const FMatrix& LocalToWorld
 	//Left
 	{
 		int32 VertexIndices[4];
-		VertexIndices[0] = MeshBuilder.AddVertex( FVector(-TX,  -TY, TZ-TH),	FVector2D::ZeroVector, FVector(0,0,1), FVector(0,1,0), FVector(-1,0,0), FColor::White );
-		VertexIndices[1] = MeshBuilder.AddVertex( FVector(-TX, -TY, TZ),		FVector2D::ZeroVector, FVector(0,0,1), FVector(0,1,0), FVector(-1,0,0), FColor::White );
-		VertexIndices[2] = MeshBuilder.AddVertex( FVector(-TX, +TY, TZ),		FVector2D::ZeroVector, FVector(0,0,1), FVector(0,1,0), FVector(-1,0,0), FColor::White );
-		VertexIndices[3] = MeshBuilder.AddVertex( FVector(-TX, +TY, TZ-TH),		FVector2D::ZeroVector, FVector(0,0,1), FVector(0,1,0), FVector(-1,0,0), FColor::White );
+		VertexIndices[0] = MeshBuilder.AddVertex( FVector(-TX,  -TY, TZ-TH),	FVector2f::ZeroVector, FVector(0,0,1), FVector(0,1,0), FVector(-1,0,0), FColor::White );
+		VertexIndices[1] = MeshBuilder.AddVertex( FVector(-TX, -TY, TZ),		FVector2f::ZeroVector, FVector(0,0,1), FVector(0,1,0), FVector(-1,0,0), FColor::White );
+		VertexIndices[2] = MeshBuilder.AddVertex( FVector(-TX, +TY, TZ),		FVector2f::ZeroVector, FVector(0,0,1), FVector(0,1,0), FVector(-1,0,0), FColor::White );
+		VertexIndices[3] = MeshBuilder.AddVertex( FVector(-TX, +TY, TZ-TH),		FVector2f::ZeroVector, FVector(0,0,1), FVector(0,1,0), FVector(-1,0,0), FColor::White );
 
 
 		MeshBuilder.AddTriangle(VertexIndices[0],VertexIndices[1],VertexIndices[2]);
@@ -405,11 +405,11 @@ void DrawCornerHelper( FPrimitiveDrawInterface* PDI, const FMatrix& LocalToWorld
 	// Front
 	{
 		int32 VertexIndices[5];
-		VertexIndices[0] = MeshBuilder.AddVertex( FVector(-TX,	+TY, TZ-TH),	FVector2D::ZeroVector, FVector(1,0,0), FVector(0,0,-1), FVector(0,1,0), FColor::White );
-		VertexIndices[1] = MeshBuilder.AddVertex( FVector(-TX,	+TY, +TZ  ),	FVector2D::ZeroVector, FVector(1,0,0), FVector(0,0,-1), FVector(0,1,0), FColor::White );
-		VertexIndices[2] = MeshBuilder.AddVertex( FVector(+TX-TH, +TY, +TX  ),	FVector2D::ZeroVector, FVector(1,0,0), FVector(0,0,-1), FVector(0,1,0), FColor::White );
-		VertexIndices[3] = MeshBuilder.AddVertex( FVector(+TX,	+TY, +TZ  ),	FVector2D::ZeroVector, FVector(1,0,0), FVector(0,0,-1), FVector(0,1,0), FColor::White );
-		VertexIndices[4] = MeshBuilder.AddVertex( FVector(+TX-TH, +TY, TZ-TH),	FVector2D::ZeroVector, FVector(1,0,0), FVector(0,0,-1), FVector(0,1,0), FColor::White );
+		VertexIndices[0] = MeshBuilder.AddVertex( FVector(-TX,	+TY, TZ-TH),	FVector2f::ZeroVector, FVector(1,0,0), FVector(0,0,-1), FVector(0,1,0), FColor::White );
+		VertexIndices[1] = MeshBuilder.AddVertex( FVector(-TX,	+TY, +TZ  ),	FVector2f::ZeroVector, FVector(1,0,0), FVector(0,0,-1), FVector(0,1,0), FColor::White );
+		VertexIndices[2] = MeshBuilder.AddVertex( FVector(+TX-TH, +TY, +TX  ),	FVector2f::ZeroVector, FVector(1,0,0), FVector(0,0,-1), FVector(0,1,0), FColor::White );
+		VertexIndices[3] = MeshBuilder.AddVertex( FVector(+TX,	+TY, +TZ  ),	FVector2f::ZeroVector, FVector(1,0,0), FVector(0,0,-1), FVector(0,1,0), FColor::White );
+		VertexIndices[4] = MeshBuilder.AddVertex( FVector(+TX-TH, +TY, TZ-TH),	FVector2f::ZeroVector, FVector(1,0,0), FVector(0,0,-1), FVector(0,1,0), FColor::White );
 
 		MeshBuilder.AddTriangle(VertexIndices[0],VertexIndices[1],VertexIndices[2]);
 		MeshBuilder.AddTriangle(VertexIndices[0],VertexIndices[2],VertexIndices[4]);
@@ -419,11 +419,11 @@ void DrawCornerHelper( FPrimitiveDrawInterface* PDI, const FMatrix& LocalToWorld
 	// Back
 	{
 		int32 VertexIndices[5];
-		VertexIndices[0] = MeshBuilder.AddVertex( FVector(-TX,	-TY, TZ-TH),	FVector2D::ZeroVector, FVector(1,0,0), FVector(0,0,1), FVector(0,-1,0), FColor::White );
-		VertexIndices[1] = MeshBuilder.AddVertex( FVector(-TX,	-TY, +TZ),		FVector2D::ZeroVector, FVector(1,0,0), FVector(0,0,1), FVector(0,-1,0), FColor::White );
-		VertexIndices[2] = MeshBuilder.AddVertex( FVector(+TX-TH, -TY, +TX),	FVector2D::ZeroVector, FVector(1,0,0), FVector(0,0,1), FVector(0,-1,0), FColor::White );
-		VertexIndices[3] = MeshBuilder.AddVertex( FVector(+TX,	-TY, +TZ),		FVector2D::ZeroVector, FVector(1,0,0), FVector(0,0,1), FVector(0,-1,0), FColor::White );
-		VertexIndices[4] = MeshBuilder.AddVertex( FVector(+TX-TH, -TY, TZ-TH),	FVector2D::ZeroVector, FVector(1,0,0), FVector(0,0,1), FVector(0,-1,0), FColor::White );
+		VertexIndices[0] = MeshBuilder.AddVertex( FVector(-TX,	-TY, TZ-TH),	FVector2f::ZeroVector, FVector(1,0,0), FVector(0,0,1), FVector(0,-1,0), FColor::White );
+		VertexIndices[1] = MeshBuilder.AddVertex( FVector(-TX,	-TY, +TZ),		FVector2f::ZeroVector, FVector(1,0,0), FVector(0,0,1), FVector(0,-1,0), FColor::White );
+		VertexIndices[2] = MeshBuilder.AddVertex( FVector(+TX-TH, -TY, +TX),	FVector2f::ZeroVector, FVector(1,0,0), FVector(0,0,1), FVector(0,-1,0), FColor::White );
+		VertexIndices[3] = MeshBuilder.AddVertex( FVector(+TX,	-TY, +TZ),		FVector2f::ZeroVector, FVector(1,0,0), FVector(0,0,1), FVector(0,-1,0), FColor::White );
+		VertexIndices[4] = MeshBuilder.AddVertex( FVector(+TX-TH, -TY, TZ-TH),	FVector2f::ZeroVector, FVector(1,0,0), FVector(0,0,1), FVector(0,-1,0), FColor::White );
 
 		MeshBuilder.AddTriangle(VertexIndices[0],VertexIndices[1],VertexIndices[2]);
 		MeshBuilder.AddTriangle(VertexIndices[0],VertexIndices[2],VertexIndices[4]);
@@ -432,10 +432,10 @@ void DrawCornerHelper( FPrimitiveDrawInterface* PDI, const FMatrix& LocalToWorld
 	// Bottom
 	{
 		int32 VertexIndices[4];
-		VertexIndices[0] = MeshBuilder.AddVertex( FVector(-TX, -TY, TZ-TH),		FVector2D::ZeroVector, FVector(1,0,0), FVector(0,0,-1), FVector(0,0,1), FColor::White );
-		VertexIndices[1] = MeshBuilder.AddVertex( FVector(-TX, +TY, TZ-TH),		FVector2D::ZeroVector, FVector(1,0,0), FVector(0,0,-1), FVector(0,0,1), FColor::White );
-		VertexIndices[2] = MeshBuilder.AddVertex( FVector(+TX-TH, +TY, TZ-TH),	FVector2D::ZeroVector, FVector(1,0,0), FVector(0,0,-1), FVector(0,0,1), FColor::White );
-		VertexIndices[3] = MeshBuilder.AddVertex( FVector(+TX-TH, -TY, TZ-TH),	FVector2D::ZeroVector, FVector(1,0,0), FVector(0,0,-1), FVector(0,0,1), FColor::White );
+		VertexIndices[0] = MeshBuilder.AddVertex( FVector(-TX, -TY, TZ-TH),		FVector2f::ZeroVector, FVector(1,0,0), FVector(0,0,-1), FVector(0,0,1), FColor::White );
+		VertexIndices[1] = MeshBuilder.AddVertex( FVector(-TX, +TY, TZ-TH),		FVector2f::ZeroVector, FVector(1,0,0), FVector(0,0,-1), FVector(0,0,1), FColor::White );
+		VertexIndices[2] = MeshBuilder.AddVertex( FVector(+TX-TH, +TY, TZ-TH),	FVector2f::ZeroVector, FVector(1,0,0), FVector(0,0,-1), FVector(0,0,1), FColor::White );
+		VertexIndices[3] = MeshBuilder.AddVertex( FVector(+TX-TH, -TY, TZ-TH),	FVector2f::ZeroVector, FVector(1,0,0), FVector(0,0,-1), FVector(0,0,1), FColor::White );
 
 		MeshBuilder.AddTriangle(VertexIndices[0],VertexIndices[1],VertexIndices[2]);
 		MeshBuilder.AddTriangle(VertexIndices[0],VertexIndices[2],VertexIndices[3]);
@@ -1298,7 +1298,7 @@ void FWidget::DrawThickArc (const FThickArcParams& InParams, const FVector& Axis
 			VertexDir.Normalize();
 
 			float TCAngle = Percent*(PI/2);
-			FVector2D TC(TCRadius*FMath::Cos(Angle), TCRadius*FMath::Sin(Angle));
+			FVector2f TC(TCRadius*FMath::Cos(Angle), TCRadius*FMath::Sin(Angle));
 
 			// Keep the vertices in local space so that we don't lose precision when dealing with LWC
 			// The local-to-world transform is handled in the MeshBuilder.Draw() call at the end of this function
@@ -1388,7 +1388,7 @@ void FWidget::DrawSnapMarker(FPrimitiveDrawInterface* PDI, const FVector& InLoca
 			FDynamicMeshVertex MeshVertex;
 			MeshVertex.Position = LocalVertices[VertexIndex];
 			MeshVertex.Color = InColor;
-			MeshVertex.TextureCoordinate[0] = FVector2D(0.0f, 0.0f);
+			MeshVertex.TextureCoordinate[0] = FVector2f(0.0f, 0.0f);
 			MeshVertex.SetTangents(
 				Axis0,
 				Axis1,
@@ -1447,7 +1447,7 @@ void FWidget::DrawStartStopMarker(FPrimitiveDrawInterface* PDI, const FVector& I
 			FDynamicMeshVertex MeshVertex;
 			MeshVertex.Position = LocalVertices[VertexIndex];
 			MeshVertex.Color = InColor;
-			MeshVertex.TextureCoordinate[0] = FVector2D(0.0f, 0.0f);
+			MeshVertex.TextureCoordinate[0] = FVector2f(0.0f, 0.0f);
 			MeshVertex.SetTangents(
 				RotatedAxis0,
 				RotatedAxis1,
diff --git a/Engine/Source/Editor/WorldBrowser/Private/Tiles/SWorldComposition.cpp b/Engine/Source/Editor/WorldBrowser/Private/Tiles/SWorldComposition.cpp
--- a/Engine/Source/Editor/WorldBrowser/Private/Tiles/SWorldComposition.cpp
+++ b/Engine/Source/Editor/WorldBrowser/Private/Tiles/SWorldComposition.cpp
@@ -879,7 +879,7 @@ protected:
 	/**  Converts cursor absolute position to the world position */
 	FVector2D CursorToWorldPosition(const FGeometry& InGeometry, FVector2D InAbsoluteCursorPosition)
 	{
-		FVector2D ViewSpacePosition = (InAbsoluteCursorPosition - InGeometry.AbsolutePosition)/InGeometry.Scale;
+		FVector2D ViewSpacePosition = (InAbsoluteCursorPosition - FVector2D(InGeometry.AbsolutePosition))/InGeometry.Scale;
 		return PanelCoordToGraphCoord(ViewSpacePosition);
 	}
 
