Cook Package Splitters: Remove support for PackageSplitters to require package population up front. We made a prototype change to move away from supporting that, for performance and robustness, and the prototype was a success.
#rnx

diff --git a/Engine/Source/Editor/UnrealEd/Private/CookOnTheFlyServer.cpp b/Engine/Source/Editor/UnrealEd/Private/CookOnTheFlyServer.cpp
--- a/Engine/Source/Editor/UnrealEd/Private/CookOnTheFlyServer.cpp
+++ b/Engine/Source/Editor/UnrealEd/Private/CookOnTheFlyServer.cpp
@@ -2026,48 +2026,24 @@ void UCookOnTheFlyServer::LoadPackageInQueue(UE::Cook::FPackageData& PackageData
 			return;
 		}
 
-		if (GeneratorStruct->IsDeferredPopulate())
+		FPackageData& OwnerPackageData = const_cast<FPackageData&>(GeneratorStruct->GetOwner());
+		UPackage* OwnerPackage;
+		bool bLoadFullySuccessful = LoadPackageForCooking(OwnerPackageData, OwnerPackage, nullptr, &PackageData);
+		if (!bLoadFullySuccessful)
 		{
-			FPackageData& OwnerPackageData = const_cast<FPackageData&>(GeneratorStruct->GetOwner());
-			UPackage* OwnerPackage;
-			bool bLoadFullySuccessful = LoadPackageForCooking(OwnerPackageData, OwnerPackage, nullptr, &PackageData);
-			if (!bLoadFullySuccessful)
-			{
-				ResultFlags |= COSR_ErrorLoadingPackage;
-				UE_LOG(LogCook, Error, TEXT("Package %s is a generated package and we could not load it's generator package %s. It can not be loaded."),
-					*PackageFileName.ToString(), *OwnerPackageData.GetFileName().ToString());
-				RejectPackageToLoad(PackageData, TEXT("is a generated package which could not load its generator"));
-				return;
-			}
-			FPopulatePackageContext Context(GeneratorStruct, OwnerPackage);
-			Context.GeneratedStruct = GeneratedStruct;
-			LoadedPackage = TryPopulateGeneratedPackage(Context);
-			if (!LoadedPackage)
-			{
-				RejectPackageToLoad(PackageData, TEXT("is a generated package which could not be populated"));
-				return;
-			}
+			ResultFlags |= COSR_ErrorLoadingPackage;
+			UE_LOG(LogCook, Error, TEXT("Package %s is a generated package and we could not load it's generator package %s. It can not be loaded."),
+				*PackageFileName.ToString(), *OwnerPackageData.GetFileName().ToString());
+			RejectPackageToLoad(PackageData, TEXT("is a generated package which could not load its generator"));
+			return;
 		}
-		else
+		FPopulatePackageContext Context(GeneratorStruct, OwnerPackage);
+		Context.GeneratedStruct = GeneratedStruct;
+		LoadedPackage = TryPopulateGeneratedPackage(Context);
+		if (!LoadedPackage)
 		{
-			FString LoadFromFileName = GeneratorStruct->GetIntermediateLocalPath(*GeneratedStruct);
-			bool bLoadFullySuccessful = LoadPackageForCooking(PackageData, LoadedPackage, &LoadFromFileName);
-			if (!bLoadFullySuccessful)
-			{
-				ResultFlags |= COSR_ErrorLoadingPackage;
-				UE_LOG(LogCook, Verbose, TEXT("Not cooking package %s"), *LoadFromFileName);
-				RejectPackageToLoad(PackageData, TEXT("failed to load"));
-				return;
-			}
-			check(LoadedPackage != nullptr && LoadedPackage->IsFullyLoaded());
-			if (LoadedPackage->GetFName() != PackageData.GetPackageName())
-			{
-				ResultFlags |= COSR_ErrorLoadingPackage;
-				UE_LOG(LogCook, Error, TEXT("Generated package loaded for package name %s, but received package name %s. Remapping of names for generated packages is not supported so it can not be loaded."),
-					*PackageData.GetPackageName().ToString(), *LoadedPackage->GetName());
-				RejectPackageToLoad(PackageData, TEXT("had a remapped name"));
-				return;
-			}
+			RejectPackageToLoad(PackageData, TEXT("is a generated package which could not be populated"));
+			return;
 		}
 	}
 
@@ -2291,18 +2267,6 @@ UE::Cook::FGeneratorPackage* UCookOnTheFlyServer::CreateGeneratorPackage(UE::Coo
 	// Create a FGeneratorPackage helper object using this CookPackageSplitter instance and return it
 	PackageData.CreateGeneratorPackage(SplitDataObject, SplitterInstance);
 	UE::Cook::FGeneratorPackage* GeneratorPackage = PackageData.GetGeneratorPackage();
-	if (!SplitterInstance->UseDeferredPopulate())
-	{
-		// Prepare the uncooked intermediate mountpoint directory
-		FString MountPackagePath;
-		FString MountLocalFilePath;
-		GeneratorPackage->GetIntermediateMountPoint(MountPackagePath, MountLocalFilePath);
-		GetAsyncIODelete().DeleteDirectory(*MountLocalFilePath);
-		if (!FPackageName::MountPointExists(MountPackagePath))
-		{
-			FPackageName::RegisterMountPoint(MountPackagePath, MountLocalFilePath);
-		}
-	}
 	return GeneratorPackage;
 }
 
@@ -2362,59 +2326,11 @@ void UCookOnTheFlyServer::SplitPackage(UE::Cook::FGeneratorPackage* GeneratorStr
 	FName OwnerName = Owner->GetFName();
 	if (!GeneratorStruct->HasQueuedGeneratedPackages())
 	{
-		if (Splitter->UseDeferredPopulate())
-		{
-			for (const FGeneratorPackage::FGeneratedStruct& GeneratedStruct : GeneratorStruct->GetPackagesToGenerate())
-			{
-				FPackageData* GeneratedPackageData = GeneratedStruct.PackageData;
-				GeneratedPackageData->ClearCookedPlatformData();
-				QueueDiscoveredPackageData(*GeneratedPackageData, FInstigator(EInstigator::GeneratedPackage, OwnerName));
-			}
-		}
-		else
+		for (const FGeneratorPackage::FGeneratedStruct& GeneratedStruct : GeneratorStruct->GetPackagesToGenerate())
 		{
-			FPopulatePackageContext Context(GeneratorStruct, Owner);
-			TArrayView<FGeneratorPackage::FGeneratedStruct> PackagesToGenerate = GeneratorStruct->GetPackagesToGenerate();
-			// Start Splitting
-			int32& NextIndex = GeneratorStruct->GetNextPopulateIndex();
-			while (NextIndex < PackagesToGenerate.Num())
-			{
-				FGeneratorPackage::FGeneratedStruct& GeneratedStruct = PackagesToGenerate[NextIndex++];
-				FPackageData* GeneratedPackageData = GeneratedStruct.PackageData;
-				GeneratedPackageData->ClearCookedPlatformData();
-
-				Context.GeneratedStruct = &GeneratedStruct;
-				UPackage* GeneratedPackage = TryPopulateGeneratedPackage(Context);
-				if (!GeneratedPackage)
-				{
-					bOutError = true;
-					return;
-				}
-
-				// Save package into the uncooked intermediate directory
-				FString IntermediateFile = GeneratorStruct->GetIntermediateLocalPath(GeneratedStruct);
-				FSavePackageArgs SaveArgs;
-				SaveArgs.TopLevelFlags = RF_Standalone;
-				SaveArgs.SaveFlags = SAVE_KeepGUID;
-				FSavePackageResultStruct SaveResult = GEditor->Save(GeneratedPackage, nullptr, *IntermediateFile,
-					SaveArgs);
-				if (SaveResult.Result != ESavePackageResult::Success)
-				{
-					UE_LOG(LogCook, Error, TEXT("PackageSplitter could not save uncooked generated package. Splitter=%s, Generated=%s."),
-						*GeneratorStruct->GetSplitDataObjectName().ToString(), *IntermediateFile);
-					bOutError = true;
-					return;
-				}
-
-				// Queue the generated package
-				QueueDiscoveredPackageData(*GeneratedPackageData, FInstigator(EInstigator::GeneratedPackage, OwnerName),
-					true /* bLoadReady */);
-
-				if (HasExceededMaxMemory())
-				{
-					return;
-				}
-			}
+			FPackageData* GeneratedPackageData = GeneratedStruct.PackageData;
+			GeneratedPackageData->ClearCookedPlatformData();
+			QueueDiscoveredPackageData(*GeneratedPackageData, FInstigator(EInstigator::GeneratedPackage, OwnerName));
 		}
 		GeneratorStruct->SetQueuedGeneratedPackages();
 	}
@@ -2483,23 +2399,14 @@ UPackage* UCookOnTheFlyServer::TryPopulateGeneratedPackage(UE::Cook::FPopulatePa
 	{
 		if (!GeneratedStruct.bHasCreatedPackage)
 		{
-			if (!Splitter->UseDeferredPopulate())
-			{
-				UE_LOG(LogCook, Error, TEXT("PackageSplitter found an existing copy of a package it was trying to populate;")
-					TEXT("this is unexpected since we previously checked and did not find it. Splitter=%s, Generated=%s."),
-					*GeneratorStruct->GetSplitDataObjectName().ToString(), *GeneratedPackageName);
-			}
-			else
-			{
-				UE_LOG(LogCook, Error, TEXT("PackageSplitter found an existing copy of a package it was trying to populate;")
-					TEXT("this is unexpected since garbage has been collected and the package should have been unreferenced so it should have been collected.")
-					TEXT("Splitter=%s, Generated=%s."),
-					*GeneratorStruct->GetSplitDataObjectName().ToString(), *GeneratedPackageName);
-				EReferenceChainSearchMode SearchMode = EReferenceChainSearchMode::Shortest
-					| EReferenceChainSearchMode::PrintAllResults
-					| EReferenceChainSearchMode::FullChain;
-				FReferenceChainSearch RefChainSearch(GeneratedPackage, SearchMode);
-			}
+			UE_LOG(LogCook, Error, TEXT("PackageSplitter found an existing copy of a package it was trying to populate;")
+				TEXT("this is unexpected since garbage has been collected and the package should have been unreferenced so it should have been collected.")
+				TEXT("Splitter=%s, Generated=%s."),
+				*GeneratorStruct->GetSplitDataObjectName().ToString(), *GeneratedPackageName);
+			EReferenceChainSearchMode SearchMode = EReferenceChainSearchMode::Shortest
+				| EReferenceChainSearchMode::PrintAllResults
+				| EReferenceChainSearchMode::FullChain;
+			FReferenceChainSearch RefChainSearch(GeneratedPackage, SearchMode);
 			return nullptr;
 		}
 		// Otherwise this is the package that was created and passed to presave, and it is still valid because there has not been a GC since
diff --git a/Engine/Source/Editor/UnrealEd/Private/Cooker/CookPackageData.h b/Engine/Source/Editor/UnrealEd/Private/Cooker/CookPackageData.h
--- a/Engine/Source/Editor/UnrealEd/Private/Cooker/CookPackageData.h
+++ b/Engine/Source/Editor/UnrealEd/Private/Cooker/CookPackageData.h
@@ -610,8 +610,6 @@ public:
 	 */
 	FString GetIntermediateLocalPath(const FGeneratorPackage::FGeneratedStruct& GeneratedStruct) const;
 
-	/** Splitter capability for whether generated packages should populate from owner rather than load from disk. */
-	bool IsDeferredPopulate() const { return CookPackageSplitterInstance->UseDeferredPopulate(); }
 	/** Return the FGeneratedStruct with the given PackageData, or null if not found. */
 	FGeneratedStruct* FindGeneratedStruct(FPackageData* PackageData);
 	/**
diff --git a/Engine/Source/Editor/UnrealEd/Public/CookPackageSplitter.h b/Engine/Source/Editor/UnrealEd/Public/CookPackageSplitter.h
--- a/Engine/Source/Editor/UnrealEd/Public/CookPackageSplitter.h
+++ b/Engine/Source/Editor/UnrealEd/Public/CookPackageSplitter.h
@@ -38,14 +38,6 @@ public:
 	 */
 	static bool ShouldSplit(UObject* SplitData) { return false; }
 	
-	/**
-	 * If true, this Splitter can handle populating generated packages long after the generator package is finalized and cooked.
-	 * The cooker will add requests for the packages returned from GetGenerateList and will populate and save them later.
-	 * If false the cooker will populate each package before calling PreSaveGeneratorPackage, and will save them
-	 * without further reference to the splitter after population.
-	 */
-	virtual bool UseDeferredPopulate() { return false; }
-
 	/**
 	 * If true, this splitter forces the Generator package objects it needs to remain referenced, and the cooker
 	 * should expect them to still be in memory after a garbage collect so long as the splitter is alive.
@@ -100,8 +92,7 @@ public:
 		bool bCreatedAsMap = false;
 	};
 	/**
-	 * Called before saving the parent generator package. Depending on the value of UseDeferredPopulate, this may be
-	 * before or after the TryPopulatePackage calls.
+	 * Called before saving the parent generator package, which itself occurs before TryPopulatePackage is called on the generated packages.
 	 * Make any required adjustments to the parent package before it is saved into the target domain.
 	 */
 	virtual void PreSaveGeneratorPackage(UPackage* OwnerPackage, UObject* OwnerObject, const TArray<FGeneratedPackageForPreSave>& PlaceholderPackages) {}
diff --git a/Engine/Source/Runtime/Engine/Private/WorldPartition/WorldPartitionCookPackageSplitter.cpp b/Engine/Source/Runtime/Engine/Private/WorldPartition/WorldPartitionCookPackageSplitter.cpp
--- a/Engine/Source/Runtime/Engine/Private/WorldPartition/WorldPartitionCookPackageSplitter.cpp
+++ b/Engine/Source/Runtime/Engine/Private/WorldPartition/WorldPartitionCookPackageSplitter.cpp
@@ -21,19 +21,6 @@ bool FWorldPartitionCookPackageSplitter::ShouldSplit(UObject* SplitData)
 	return World && World->IsPartitionedWorld();
 }
 
-bool FWorldPartitionCookPackageSplitter::UseDeferredPopulate()
-{
-	static struct FOneTimeRead
-	{
-		FOneTimeRead()
-		{
-			GConfig->GetBool(TEXT("CookSettings"), TEXT("SplitterDeferredPopulate"), Value, GEditorIni);
-		}
-		bool Value = true;
-	} UseDeferredPopulate;
-	return UseDeferredPopulate.Value;
-}
-
 FWorldPartitionCookPackageSplitter::FWorldPartitionCookPackageSplitter()
 {
 	FCoreUObjectDelegates::GetPreGarbageCollectDelegate().AddRaw(this, &FWorldPartitionCookPackageSplitter::PreGarbageCollect);
diff --git a/Engine/Source/Runtime/Engine/Private/WorldPartition/WorldPartitionCookPackageSplitter.h b/Engine/Source/Runtime/Engine/Private/WorldPartition/WorldPartitionCookPackageSplitter.h
--- a/Engine/Source/Runtime/Engine/Private/WorldPartition/WorldPartitionCookPackageSplitter.h
+++ b/Engine/Source/Runtime/Engine/Private/WorldPartition/WorldPartitionCookPackageSplitter.h
@@ -17,8 +17,6 @@ public:
 	FWorldPartitionCookPackageSplitter();
 	virtual ~FWorldPartitionCookPackageSplitter();
 
-	virtual bool UseDeferredPopulate() override;
-
 	virtual TArray<ICookPackageSplitter::FGeneratedPackage> GetGenerateList(const UPackage* OwnerPackage, const UObject* OwnerObject) override;
 	virtual bool TryPopulatePackage(const UPackage* OwnerPackage, const UObject* OwnerObject,
 		const ICookPackageSplitter::FGeneratedPackageForPopulate& GeneratedPackage, bool bWasOwnerReloaded) override;
