CookPackageData: Fix "Unexpected failure to cook filename" warning when CookOnTheFly requests a generated package (so it doesn't exist on disk in editor) with the wrong extension.
#rnx

diff --git a/Engine/Source/Editor/UnrealEd/Private/Cooker/CookPackageData.cpp b/Engine/Source/Editor/UnrealEd/Private/Cooker/CookPackageData.cpp
--- a/Engine/Source/Editor/UnrealEd/Private/Cooker/CookPackageData.cpp
+++ b/Engine/Source/Editor/UnrealEd/Private/Cooker/CookPackageData.cpp
@@ -1743,7 +1743,17 @@ FPackageData* FPackageDatas::TryAddPackageDataByStandardFileName(const FName& Fi
 	}
 	if (ExistingFileName.IsNone())
 	{
-		UE_LOG(LogCook, Warning, TEXT("Unexpected failure to cook filename '%s'. It is mapped to PackageName '%s', but that PackageName does not have a corresponding filename."),
+		if (!bExactMatchRequired)
+		{
+			FReadScopeLock ExistenceReadLock(ExistenceLock);
+			FPackageData** PackageDataMapAddr = PackageNameToPackageData.Find(PackageName);
+			if (PackageDataMapAddr != nullptr)
+			{
+				FoundFileName = (*PackageDataMapAddr)->GetFileName();
+				return *PackageDataMapAddr;
+			}
+		}
+		UE_LOG(LogCook, Warning, TEXT("Unexpected failure to cook filename '%s'. It is mapped to PackageName '%s', but does not exist on disk and we cannot verify the extension."),
 			*FileName.ToString(), *PackageName.ToString());
 		return nullptr;
 	}
