Downgrading level leak asserts to errors when PendingKill is disabled

diff --git a/Engine/Source/Editor/UnrealEd/Private/EditorEngine.cpp b/Engine/Source/Editor/UnrealEd/Private/EditorEngine.cpp
--- a/Engine/Source/Editor/UnrealEd/Private/EditorEngine.cpp
+++ b/Engine/Source/Editor/UnrealEd/Private/EditorEngine.cpp
@@ -7035,7 +7035,7 @@ void UEditorEngine::VerifyLoadMapWorldCleanup()
 				UE_LOG(LogLoad, Error, TEXT("Previously active world %s not cleaned up by garbage collection!"), *World->GetPathName());
 				UE_LOG(LogLoad, Error, TEXT("Once a world has become active, it cannot be reused and must be destroyed and reloaded. World referenced by:"));
 			
-				FindAndPrintStaleReferencesToObject(World, ELogVerbosity::Fatal);
+				FindAndPrintStaleReferencesToObject(World, UObjectBaseUtility::IsPendingKillEnabled() ? ELogVerbosity::Fatal : ELogVerbosity::Error);
 			}
 		}
 	}
diff --git a/Engine/Source/Editor/UnrealEd/Private/EditorLevelUtils.cpp b/Engine/Source/Editor/UnrealEd/Private/EditorLevelUtils.cpp
--- a/Engine/Source/Editor/UnrealEd/Private/EditorLevelUtils.cpp
+++ b/Engine/Source/Editor/UnrealEd/Private/EditorLevelUtils.cpp
@@ -859,7 +859,7 @@ bool UEditorLevelUtils::RemoveLevelsFromWorld(TArray<ULevel*> InLevels, bool bCl
 	// Reset transaction buffer and run GC to clear out the destroyed level
 	GEditor->Cleanse(bClearSelection, false, TransResetText, bResetTransBuffer);
 
-	auto CheckPackage = [](const TArray<FName>& PackageNames, bool bFatal)
+	auto CheckPackage = [](const TArray<FName>& PackageNames, ELogVerbosity::Type Verbosity)
 	{
 		// Check Package no longer exists
 		for (const FName& LevelPackageName : PackageNames)
@@ -871,7 +871,7 @@ bool UEditorLevelUtils::RemoveLevelsFromWorld(TArray<ULevel*> InLevels, bool bCl
 				UWorld* TheWorld = UWorld::FindWorldInPackage(LevelPackage->GetOutermost());
 				if (TheWorld != nullptr)
 				{
-					UEngine::FindAndPrintStaleReferencesToObject(TheWorld, bFatal ? ELogVerbosity::Fatal : ELogVerbosity::Log);
+					UEngine::FindAndPrintStaleReferencesToObject(TheWorld, Verbosity);
 					return false;
 				}
 			}
@@ -880,13 +880,17 @@ bool UEditorLevelUtils::RemoveLevelsFromWorld(TArray<ULevel*> InLevels, bool bCl
 	};
 
 	// Check that packages no longer exist
-	bool bFatal = bResetTransBuffer;
-	bool bFailed = !CheckPackage(PackageNames, bFatal);
-	if (bFailed && !bFatal)
+	ELogVerbosity::Type Verbosity = ELogVerbosity::Log;
+	if (bResetTransBuffer)
+	{
+		Verbosity = UObjectBaseUtility::IsPendingKillEnabled() ? ELogVerbosity::Fatal : ELogVerbosity::Error;
+	}
+	bool bFailed = !CheckPackage(PackageNames, Verbosity);
+	if (bFailed && Verbosity != ELogVerbosity::Fatal)
 	{
 		// We tried avoiding clearing the Transaction buffer but it failed. Plan B.
 		GEditor->Cleanse(bClearSelection, false, TransResetText, true);
-		CheckPackage(PackageNames, true);
+		CheckPackage(PackageNames, Verbosity);
 	}
 
 	return true;
diff --git a/Engine/Source/Runtime/Engine/Private/LevelStreaming.cpp b/Engine/Source/Runtime/Engine/Private/LevelStreaming.cpp
--- a/Engine/Source/Runtime/Engine/Private/LevelStreaming.cpp
+++ b/Engine/Source/Runtime/Engine/Private/LevelStreaming.cpp
@@ -1163,7 +1163,7 @@ bool ULevelStreaming::RequestLevel(UWorld* PersistentWorld, bool bAllowLevelLoad
 				UE_LOG(LogLevelStreaming, Error, TEXT("World exists but PersistentLevel doesn't for %s, most likely caused by reference to world of unloaded level and GC setting reference to null while keeping world object"), *World->GetOutermost()->GetName());
 				UE_LOG(LogLevelStreaming, Error, TEXT("Most likely caused by reference to world of unloaded level and GC setting reference to null while keeping world object. Referenced by:"));
 
-				UEngine::FindAndPrintStaleReferencesToObject(World, ELogVerbosity::Fatal);
+				UEngine::FindAndPrintStaleReferencesToObject(World, UObjectBaseUtility::IsPendingKillEnabled() ? ELogVerbosity::Fatal : ELogVerbosity::Error);
 
 				return false;
 			}
diff --git a/Engine/Source/Runtime/Engine/Private/UnrealEngine.cpp b/Engine/Source/Runtime/Engine/Private/UnrealEngine.cpp
--- a/Engine/Source/Runtime/Engine/Private/UnrealEngine.cpp
+++ b/Engine/Source/Runtime/Engine/Private/UnrealEngine.cpp
@@ -14918,7 +14918,7 @@ void UEngine::VerifyLoadMapWorldCleanup()
 				UE_LOG(LogLoad, Error, TEXT("Previously active world %s not cleaned up by garbage collection!"), *World->GetPathName());
 				UE_LOG(LogLoad, Error, TEXT("Once a world has become active, it cannot be reused and must be destroyed and reloaded. World referenced by:"));
 
-				FindAndPrintStaleReferencesToObject(World, ELogVerbosity::Fatal);
+				FindAndPrintStaleReferencesToObject(World, UObjectBaseUtility::IsPendingKillEnabled() ? ELogVerbosity::Fatal : ELogVerbosity::Error);
 			}
 		}
 	}
