Fix some typos in the code comments for UResavePackagesCommandlet

#rnx

diff --git a/Engine/Source/Editor/UnrealEd/Private/Commandlets/ContentCommandlets.cpp b/Engine/Source/Editor/UnrealEd/Private/Commandlets/ContentCommandlets.cpp
--- a/Engine/Source/Editor/UnrealEd/Private/Commandlets/ContentCommandlets.cpp
+++ b/Engine/Source/Editor/UnrealEd/Private/Commandlets/ContentCommandlets.cpp
@@ -1039,9 +1039,9 @@ int32 UResavePackagesCommandlet::Main( const FString& Params )
 	bIgnoreChangelist = Switches.Contains(TEXT("IgnoreChangelist"));
 	/** whether we should only save packages with changelist zero */
 	bOnlyUnversioned = Switches.Contains(TEXT("OnlyUnversioned"));
-	/** whether we should only save packages saved by licenseed */
+	/** whether we should only save packages saved by licensees */
 	bOnlyLicenseed = Switches.Contains(TEXT("OnlyLicenseed"));
-	/** whether we should only save packages contianing virtualized bulkdata payloads */
+	/** whether we should only save packages containing virtualized bulkdata payloads */
 	bOnlyVirtualized = Switches.Contains(TEXT("OnlyVirtualized"));
 	/** only process packages containing materials */
 	bOnlyMaterials = Switches.Contains(TEXT("onlymaterials"));
