When editor selection changes, refresh cached menus

diff --git a/Engine/Source/Editor/UnrealEd/Private/UnrealEdEngine.cpp b/Engine/Source/Editor/UnrealEd/Private/UnrealEdEngine.cpp
--- a/Engine/Source/Editor/UnrealEd/Private/UnrealEdEngine.cpp
+++ b/Engine/Source/Editor/UnrealEd/Private/UnrealEdEngine.cpp
@@ -70,6 +70,8 @@
 #include "Subsystems/EditorActorSubsystem.h"
 #include "ProfilingDebugging/StallDetector.h"
 #include "GameMapsSettings.h"
+#include "HAL/PlatformApplicationMisc.h"
+
 
 DEFINE_LOG_CATEGORY_STATIC(LogUnrealEdEngine, Log, All);
 
@@ -1490,6 +1492,9 @@ void UUnrealEdEngine::OnEditorSelectionChanged(UObject* SelectionThatChanged)
 			}
 		}
 	}
+#if PLATFORM_MAC
+	FPlatformApplicationMisc::bChachedMacMenuStateNeedsUpdate = true;
+#endif
 }
 
 bool UUnrealEdEngine::HasMountWritePermissionForPackage(const FString& PackageName)
