Fix an issue where alt drag duplication was not duplicating attached and selected actors.

#rnx

diff --git a/Engine/Source/Editor/UnrealEd/Private/LevelEditorViewport.cpp b/Engine/Source/Editor/UnrealEd/Private/LevelEditorViewport.cpp
--- a/Engine/Source/Editor/UnrealEd/Private/LevelEditorViewport.cpp
+++ b/Engine/Source/Editor/UnrealEd/Private/LevelEditorViewport.cpp
@@ -2596,8 +2596,8 @@ bool FLevelEditorViewportClient::InputWidgetDelta(FViewport* InViewport, EAxisLi
 								FTypedElementListConstRef ElementsToManipulate = GetElementsToManipulate();
 								TArray<FTypedElementHandle> DuplicatedElements;
 								{
-									FTypedElementListRef ElementsToDuplicate = ElementsToManipulate->Clone(); // Clone this, as the duplicate may triger construction scripts to modify the selection as we duplicate
-									DuplicatedElements = CommonActions->DuplicateNormalizedElements(ElementsToDuplicate, GetWorld(), FVector::ZeroVector);
+									// Do not used the cached manipulation list here here, as it will have removed attachments, and we do want to duplicate those
+									DuplicatedElements = CommonActions->DuplicateSelectedElements(GetSelectionSet(), GetWorld(), FVector::ZeroVector);
 								}
 
 								// Exclusively select the new elements, so that future gizmo interaction manipulates those items instead
