Modified nullptr check for the model outer in FLevelEditorViewportClient::AttemptApplyObjAsMaterialToSurface to use a weak object pointer instead of directly accessing IsValid on the object.

Speculative fix for crash only happening on one platform for 5.0

#rnx

diff --git a/Engine/Source/Editor/UnrealEd/Private/LevelEditorViewport.cpp b/Engine/Source/Editor/UnrealEd/Private/LevelEditorViewport.cpp
--- a/Engine/Source/Editor/UnrealEd/Private/LevelEditorViewport.cpp
+++ b/Engine/Source/Editor/UnrealEd/Private/LevelEditorViewport.cpp
@@ -799,8 +799,14 @@ bool FLevelEditorViewportClient::AttemptApplyObjAsMaterialToSurface( UObject* Ob
 
 		UModel* Model = ModelHitProxy->GetModel();
 		
-		// If our model doesn't exist or is part of a level that is being destroyed
-		if( !Model || !Model->GetOuter() || !IsValidChecked(Model->GetOuter()) || Model->GetOuter()->IsUnreachable() )
+		// If our model doesn't exist
+		if ( !Model )
+		{
+			return false;
+		}
+		// or is part of an outer that is being destroyed
+		TWeakObjectPtr<UObject> ModelOuter(Model->GetOuter());
+		if ( !ModelOuter.IsValid() )
 		{
 			return false;
 		}
