Add an array of callbacks that can be used to disable CPU throttling for the editor when it doesn't have focus.

diff --git a/Engine/Source/Editor/UnrealEd/Classes/Editor/EditorEngine.h b/Engine/Source/Editor/UnrealEd/Classes/Editor/EditorEngine.h
--- a/Engine/Source/Editor/UnrealEd/Classes/Editor/EditorEngine.h
+++ b/Engine/Source/Editor/UnrealEd/Classes/Editor/EditorEngine.h
@@ -527,6 +527,10 @@ public:
 	DECLARE_MULTICAST_DELEGATE(FPreviewPlatformChanged);
 	FPreviewPlatformChanged PreviewPlatformChanged;
 
+	/** An array of delegates that can force disable throttling cpu usage if any of them return false. */
+	DECLARE_DELEGATE_RetVal(bool, FShouldDisableCPUThrottling);
+	TArray<FShouldDisableCPUThrottling> ShouldDisableCPUThrottlingDelegates;
+
 	/** A delegate that is called when the bugitgo command is used. */
 	DECLARE_MULTICAST_DELEGATE_TwoParams(FPostBugItGoCalled, const FVector& Loc, const FRotator& Rot);
 	FPostBugItGoCalled PostBugItGoCalled;
diff --git a/Engine/Source/Editor/UnrealEd/Private/EditorEngine.cpp b/Engine/Source/Editor/UnrealEd/Private/EditorEngine.cpp
--- a/Engine/Source/Editor/UnrealEd/Private/EditorEngine.cpp
+++ b/Engine/Source/Editor/UnrealEd/Private/EditorEngine.cpp
@@ -6137,6 +6137,16 @@ bool UEditorEngine::ShouldThrottleCPUUsage() const
 		return false;
 	}
 
+	// There might be systems where throttling would cause issues (such as data transfer over the network) - give them
+	// an opportunity to force us to not throttle.
+	for (auto It = ShouldDisableCPUThrottlingDelegates.CreateConstIterator(); It; ++It)
+	{
+		if (It->IsBound() && !It->Execute())
+		{
+			return false;
+		}
+	}
+
 	bool bShouldThrottle = false;
 
 	const bool bRunningCommandlet = IsRunningCommandlet();
