Fix CIS issue by commenting out delegate check causing a shared null ptr

diff --git a/Engine/Source/Editor/UnrealEd/Private/EditorEngine.cpp b/Engine/Source/Editor/UnrealEd/Private/EditorEngine.cpp
--- a/Engine/Source/Editor/UnrealEd/Private/EditorEngine.cpp
+++ b/Engine/Source/Editor/UnrealEd/Private/EditorEngine.cpp
@@ -6137,6 +6137,8 @@ bool UEditorEngine::ShouldThrottleCPUUsage() const
 		return false;
 	}
 
+	/*
+	Commented Out due to stale shared pointer issue.
 	// There might be systems where throttling would cause issues (such as data transfer over the network) - give them
 	// an opportunity to force us to not throttle.
 	for (auto It = ShouldDisableCPUThrottlingDelegates.CreateConstIterator(); It; ++It)
@@ -6146,6 +6148,7 @@ bool UEditorEngine::ShouldThrottleCPUUsage() const
 			return false;
 		}
 	}
+	*/
 
 	bool bShouldThrottle = false;
 
