Variable watch within state machines is broken
- Ensure that any folded AnimNode property gets remapped correctly for pin-to-property map in FBlueprintDebugData
- Ensure  pin value watching for folded properties retrieves the correct value
- Const-ifying some code paths around blueprint pin watching

diff --git a/Engine/Source/Editor/AnimGraph/Private/AnimBlueprintCompiler.cpp b/Engine/Source/Editor/AnimGraph/Private/AnimBlueprintCompiler.cpp
--- a/Engine/Source/Editor/AnimGraph/Private/AnimBlueprintCompiler.cpp
+++ b/Engine/Source/Editor/AnimGraph/Private/AnimBlueprintCompiler.cpp
@@ -1879,15 +1879,16 @@ void FAnimBlueprintCompilerContext::ProcessFoldedPropertyRecords()
 	}
 
 	// Builds a 'data area', returns the total number of properties that were inserted into that area's struct
-	auto BuildDataArea = [this](TArray<TSharedRef<IAnimBlueprintCompilationContext::FFoldedPropertyRecord>>& InRecords, UScriptStruct* InStruct)
+	auto BuildDataArea = [this, NewAnimBlueprintClass](TArray<TSharedRef<IAnimBlueprintCompilationContext::FFoldedPropertyRecord>>& InRecords, UScriptStruct* InStruct)
 	{
 		int32 PropertyIndex = 0;
 
 		if(InStruct)
 		{
 			const UAnimationGraphSchema* AnimationGraphSchema = GetDefault<UAnimationGraphSchema>();
-
-			for(TSharedRef<IAnimBlueprintCompilationContext::FFoldedPropertyRecord>& Record : InRecords)
+			FAnimBlueprintDebugData& AnimBlueprintDebugData = NewAnimBlueprintClass->GetAnimBlueprintDebugData();
+			
+			for(const TSharedRef<IAnimBlueprintCompilationContext::FFoldedPropertyRecord>& Record : InRecords)
 			{
 				// Skip folded records
 				if(Record->FoldIndex == INDEX_NONE)
@@ -1896,7 +1897,7 @@ void FAnimBlueprintCompilerContext::ProcessFoldedPropertyRecords()
 					if(AnimationGraphSchema->ConvertPropertyToPinType(Record->Property, VariableType))
 					{
 						// Patch into sparse class data
-						FName PropertyName = FName(Record->Property->GetClass()->GetFName(), InRecords.Num() - 1 - PropertyIndex);
+						const FName PropertyName = FName(Record->Property->GetClass()->GetFName(), InRecords.Num() - 1 - PropertyIndex);
 						Record->GeneratedProperty = CreateStructVariable(InStruct, PropertyName, VariableType);
 						if (Record->GeneratedProperty == nullptr)
 						{
@@ -1910,6 +1911,9 @@ void FAnimBlueprintCompilerContext::ProcessFoldedPropertyRecords()
 							// Properties need to be BP visible to allow them to be set by the generated exec chains in CreateEvaluationHandlerForNode
 							Record->GeneratedProperty->SetPropertyFlags(CPF_BlueprintVisible);
 							Record->PropertyIndex = PropertyIndex++;
+
+							// Populate the original-to-folded property mapping with the newly generated property
+							AnimBlueprintDebugData.NodeToFoldedPropertyMap.Add(Record->Property, Record->GeneratedProperty);
 						}
 					}
 					else
diff --git a/Engine/Source/Editor/BlueprintGraph/Classes/EdGraphSchema_K2.h b/Engine/Source/Editor/BlueprintGraph/Classes/EdGraphSchema_K2.h
--- a/Engine/Source/Editor/BlueprintGraph/Classes/EdGraphSchema_K2.h
+++ b/Engine/Source/Editor/BlueprintGraph/Classes/EdGraphSchema_K2.h
@@ -953,7 +953,7 @@ public:
 	 *
 	 * @return	The converted type string.
 	 */
-	static FText TypeToText(FProperty* const Property);
+	static FText TypeToText(const FProperty* const Property);
 
 	/**
 	* Converts a terminal type into a fully qualified FText (e.g., object'ObjectName').
diff --git a/Engine/Source/Editor/BlueprintGraph/Private/EdGraphSchema_K2.cpp b/Engine/Source/Editor/BlueprintGraph/Private/EdGraphSchema_K2.cpp
--- a/Engine/Source/Editor/BlueprintGraph/Private/EdGraphSchema_K2.cpp
+++ b/Engine/Source/Editor/BlueprintGraph/Private/EdGraphSchema_K2.cpp
@@ -3573,9 +3573,9 @@ bool UEdGraphSchema_K2::IsWildcardProperty(const FProperty* Property)
 		|| FEdGraphUtilities::IsArrayDependentParam(Function, Property->GetFName()) );
 }
 
-FText UEdGraphSchema_K2::TypeToText(FProperty* const Property)
+FText UEdGraphSchema_K2::TypeToText(const FProperty* const Property)
 {
-	if (FStructProperty* Struct = CastField<FStructProperty>(Property))
+	if (const FStructProperty* Struct = CastField<FStructProperty>(Property))
 	{
 		if (Struct->Struct)
 		{
@@ -3585,7 +3585,7 @@ FText UEdGraphSchema_K2::TypeToText(FProperty* const Property)
 			return TypeToText(PinType);
 		}
 	}
-	else if (FClassProperty* Class = CastField<FClassProperty>(Property))
+	else if (const FClassProperty* Class = CastField<FClassProperty>(Property))
 	{
 		if (Class->MetaClass)
 		{
@@ -3595,7 +3595,7 @@ FText UEdGraphSchema_K2::TypeToText(FProperty* const Property)
 			return TypeToText(PinType);
 		}
 	}
-	else if (FInterfaceProperty* Interface = CastField<FInterfaceProperty>(Property))
+	else if (const FInterfaceProperty* Interface = CastField<FInterfaceProperty>(Property))
 	{
 		if (Interface->InterfaceClass != nullptr)
 		{
@@ -3605,7 +3605,7 @@ FText UEdGraphSchema_K2::TypeToText(FProperty* const Property)
 			return TypeToText(PinType);
 		}
 	}
-	else if (FObjectPropertyBase* Obj = CastField<FObjectPropertyBase>(Property))
+	else if (const FObjectPropertyBase* Obj = CastField<FObjectPropertyBase>(Property))
 	{
 		if( Obj->PropertyClass )
 		{
@@ -3618,7 +3618,7 @@ FText UEdGraphSchema_K2::TypeToText(FProperty* const Property)
 
 		return FText::GetEmpty();
 	}
-	else if (FArrayProperty* Array = CastField<FArrayProperty>(Property))
+	else if (const FArrayProperty* Array = CastField<FArrayProperty>(Property))
 	{
 		if (Array->Inner)
 		{
@@ -3627,7 +3627,7 @@ FText UEdGraphSchema_K2::TypeToText(FProperty* const Property)
 			return FText::Format(LOCTEXT("ArrayPropertyText", "Array of {ArrayType}"), Args); 
 		}
 	}
-	else if (FSetProperty* Set = CastField<FSetProperty>(Property))
+	else if (const FSetProperty* Set = CastField<FSetProperty>(Property))
 	{
 		if (Set->ElementProp)
 		{
@@ -3636,7 +3636,7 @@ FText UEdGraphSchema_K2::TypeToText(FProperty* const Property)
 			return FText::Format(LOCTEXT("SetPropertyText", "Set of {SetType}"), Args);
 		}
 	}
-	else if (FMapProperty* Map = CastField<FMapProperty>(Property))
+	else if (const FMapProperty* Map = CastField<FMapProperty>(Property))
 	{
 		if (Map->KeyProp && Map->ValueProp)
 		{
diff --git a/Engine/Source/Editor/UnrealEd/Private/Kismet2/KismetDebugUtilities.cpp b/Engine/Source/Editor/UnrealEd/Private/Kismet2/KismetDebugUtilities.cpp
--- a/Engine/Source/Editor/UnrealEd/Private/Kismet2/KismetDebugUtilities.cpp
+++ b/Engine/Source/Editor/UnrealEd/Private/Kismet2/KismetDebugUtilities.cpp
@@ -1580,9 +1580,9 @@ UEdGraphPin* FKismetDebugUtilities::FindPinWatchByPredicate(const UBlueprint* Bl
 // Gets the watched tooltip for a specified site
 FKismetDebugUtilities::EWatchTextResult FKismetDebugUtilities::GetWatchText(FString& OutWatchText, UBlueprint* Blueprint, UObject* ActiveObject, const UEdGraphPin* WatchPin)
 {
-	FProperty* PropertyToDebug = nullptr;
-	void* DataPtr = nullptr;
-	void* DeltaPtr = nullptr;
+	const FProperty* PropertyToDebug = nullptr;
+	const void* DataPtr = nullptr;
+	const void* DeltaPtr = nullptr;
 	UObject* ParentObj = nullptr;
 	TArray<UObject*> SeenObjects;
 	bool bIsDirectPtr = false;
@@ -1665,12 +1665,12 @@ bool FKismetDebugUtilities::CanInspectPinValue(const UEdGraphPin* Pin)
 
 FKismetDebugUtilities::EWatchTextResult FKismetDebugUtilities::GetDebugInfo(TSharedPtr<FPropertyInstanceInfo> &OutDebugInfo, UBlueprint* Blueprint, UObject* ActiveObject, const UEdGraphPin* WatchPin)
 {
-	void* DataPtr = nullptr;
-	void* DeltaPtr = nullptr;
+	const void* DataPtr = nullptr;
+	const void* DeltaPtr = nullptr;
 	UObject* ParentObj = nullptr;
 	TArray<UObject*> SeenObjects;
 	bool bIsDirectPtr = false;
-	FProperty* Property = nullptr;
+	const FProperty* Property = nullptr;
 	FKismetDebugUtilities::EWatchTextResult Result = FindDebuggingData(Blueprint, ActiveObject, WatchPin, Property, DataPtr, DeltaPtr, ParentObj, SeenObjects, &bIsDirectPtr);
 
 	if (Result == FKismetDebugUtilities::EWatchTextResult::EWTR_Valid)
@@ -1689,7 +1689,7 @@ FKismetDebugUtilities::EWatchTextResult FKismetDebugUtilities::GetDebugInfo(TSha
 	return Result;
 }
 
-FKismetDebugUtilities::EWatchTextResult FKismetDebugUtilities::FindDebuggingData(UBlueprint* Blueprint, UObject* ActiveObject, const UEdGraphPin* WatchPin, FProperty*& OutProperty, void*& OutData, void*& OutDelta, UObject*& OutParent, TArray<UObject*>& SeenObjects, bool* bOutIsDirectPtr /* = nullptr */)
+FKismetDebugUtilities::EWatchTextResult FKismetDebugUtilities::FindDebuggingData(UBlueprint* Blueprint, UObject* ActiveObject, const UEdGraphPin* WatchPin, const FProperty*& OutProperty, const void*& OutData, const void*& OutDelta, UObject*& OutParent, TArray<UObject*>& SeenObjects, bool* bOutIsDirectPtr /* = nullptr */)
 {
 	FKismetDebugUtilitiesData& Data = FKismetDebugUtilitiesData::Get();
 
@@ -1804,11 +1804,11 @@ FKismetDebugUtilities::EWatchTextResult FKismetDebugUtilities::FindDebuggingData
 #endif // USE_UBER_GRAPH_PERSISTENT_FRAME
 
 			// see if our WatchPin is on a animation node & if so try to get its property info
-			UAnimBlueprintGeneratedClass* AnimBlueprintGeneratedClass = Cast<UAnimBlueprintGeneratedClass>(Blueprint->GeneratedClass);
+			const UAnimBlueprintGeneratedClass* AnimBlueprintGeneratedClass = Cast<UAnimBlueprintGeneratedClass>(Blueprint->GeneratedClass);
 			if (!PropertyBase && AnimBlueprintGeneratedClass)
 			{
 				// are we linked to an anim graph node?
-				FProperty* LinkedProperty = Property;
+				const FProperty* LinkedProperty = Property;
 				const UAnimGraphNode_Base* Node = Cast<UAnimGraphNode_Base>(WatchPin->GetOuter());
 				if (Node == nullptr && WatchPin->LinkedTo.Num() > 0)
 				{
@@ -1820,14 +1820,39 @@ FKismetDebugUtilities::EWatchTextResult FKismetDebugUtilities::FindDebuggingData
 
 				if (Node && LinkedProperty)
 				{
-					FStructProperty* NodeStructProperty = CastField<FStructProperty>(FKismetDebugUtilities::FindClassPropertyForNode(Blueprint, Node));
-					if (NodeStructProperty)
+					// In case the property was folded its value has to be retrieved from the Mutable data struct on the instance rather than from the Anim Node itself
+					if (const TFieldPath<const FProperty>* FoldedPropertyPathPtr = AnimBlueprintGeneratedClass->AnimBlueprintDebugData.NodeToFoldedPropertyMap.Find(LinkedProperty))
 					{
+						const FProperty* FoldedProperty = FoldedPropertyPathPtr->Get();						
+						const UStruct* FoldedPropertyStruct = FoldedProperty ? FoldedProperty->GetOwnerStruct() : nullptr;
+
+						if(FoldedPropertyStruct && FoldedPropertyStruct->IsChildOf(FAnimBlueprintMutableData::StaticStruct()))
+						{
+							const FAnimBlueprintMutableData* MutableData = AnimBlueprintGeneratedClass->GetMutableNodeData(Cast<const UObject>(ActiveObject));
+						
+							if(bOutIsDirectPtr && MutableData)
+							{
+								const void* ValuePtr = FoldedProperty->ContainerPtrToValuePtr<void>(MutableData);
+								OutProperty = FoldedProperty;
+								OutData = ValuePtr;
+								OutDelta = ValuePtr;
+								OutParent = nullptr;
+										
+								// Flag to caller that OutData points directly at the data, not at the
+								// base of a property container (so don't apply property's Offset_Internal)
+								*bOutIsDirectPtr = true;
+								
+								return EWTR_Valid;
+							}
+						}
+					}
+					else if (FStructProperty* NodeStructProperty = CastField<FStructProperty>(FKismetDebugUtilities::FindClassPropertyForNode(Blueprint, Node)))
+					{						
 						for (const FStructPropertyPath& NodeProperty : AnimBlueprintGeneratedClass->GetAnimNodeProperties())
 						{
 							if (NodeProperty.Get() == NodeStructProperty)
 							{
-								void* NodePtr = NodeProperty->ContainerPtrToValuePtr<void>(ActiveObject);
+								const void* NodePtr = NodeProperty->ContainerPtrToValuePtr<void>(ActiveObject);
 								OutProperty = LinkedProperty;
 								OutData = NodePtr;
 								OutDelta = NodePtr;
@@ -1849,16 +1874,16 @@ FKismetDebugUtilities::EWatchTextResult FKismetDebugUtilities::FindDebuggingData
 					UEdGraphPin* SelfPin = WatchNode->FindPin(TEXT("self"));
 					if (SelfPin && SelfPin != WatchPin)
 					{
-						FProperty* SelfPinProperty = nullptr;
-						void* SelfPinData = nullptr;
-						void* SelfPinDelta = nullptr;
+						const FProperty* SelfPinProperty = nullptr;
+						const void* SelfPinData = nullptr;
+						const void* SelfPinDelta = nullptr;
 						UObject* SelfPinParent = nullptr;
 						SeenObjects.AddUnique(ActiveObject);
 						FKismetDebugUtilities::EWatchTextResult Result = FindDebuggingData(Blueprint, ActiveObject, SelfPin, SelfPinProperty, SelfPinData, SelfPinDelta, SelfPinParent, SeenObjects);
-						FObjectPropertyBase* SelfPinPropertyBase = CastField<FObjectPropertyBase>(SelfPinProperty);
+						const FObjectPropertyBase* SelfPinPropertyBase = CastField<const FObjectPropertyBase>(SelfPinProperty);
 						if (Result == EWTR_Valid && SelfPinPropertyBase != nullptr)
 						{
-							void* PropertyValue = SelfPinProperty->ContainerPtrToValuePtr<void>(SelfPinData);
+							const void* PropertyValue = SelfPinProperty->ContainerPtrToValuePtr<void>(SelfPinData);
 							UObject* TempActiveObject = SelfPinPropertyBase->GetObjectPropertyValue(PropertyValue);
 							if (TempActiveObject && TempActiveObject != ActiveObject)
 							{
@@ -1897,12 +1922,12 @@ FKismetDebugUtilities::EWatchTextResult FKismetDebugUtilities::FindDebuggingData
 	}
 }
 
-void FKismetDebugUtilities::GetDebugInfo_InContainer(int32 Index, TSharedPtr<FPropertyInstanceInfo> &DebugInfo, FProperty* Property, const void* Data)
+void FKismetDebugUtilities::GetDebugInfo_InContainer(int32 Index, TSharedPtr<FPropertyInstanceInfo> &DebugInfo, const FProperty* Property, const void* Data)
 {
 	GetDebugInfoInternal(DebugInfo, Property, Property->ContainerPtrToValuePtr<void>(Data, Index));
 }
 
-void FKismetDebugUtilities::GetDebugInfoInternal(TSharedPtr<FPropertyInstanceInfo> &DebugInfo, FProperty* Property, const void* PropertyValue)
+void FKismetDebugUtilities::GetDebugInfoInternal(TSharedPtr<FPropertyInstanceInfo> &DebugInfo, const FProperty* Property, const void* PropertyValue)
 {
 	TMap<FPropertyInstanceInfo::FPropertyInstance, TSharedPtr<FPropertyInstanceInfo>> VisitedNodes;
 	DebugInfo = FPropertyInstanceInfo::FindOrMake({Property, PropertyValue}, VisitedNodes);
@@ -1914,14 +1939,14 @@ FPropertyInstanceInfo::FPropertyInstanceInfo(FPropertyInstance PropertyInstance)
 	Type(UEdGraphSchema_K2::TypeToText(PropertyInstance.Property)),
 	Property(PropertyInstance.Property)
 {
-	FProperty* ResolvedProperty = Property.Get();
+	const FProperty* ResolvedProperty = Property.Get();
 	check(ResolvedProperty);
 	if (PropertyInstance.Value == nullptr)
 	{
 		return;
 	}
 
-	if (FByteProperty* ByteProperty = CastField<FByteProperty>(ResolvedProperty))
+	if (const FByteProperty* ByteProperty = CastField<FByteProperty>(ResolvedProperty))
 	{
 		UEnum* Enum = ByteProperty->GetIntPropertyEnum();
 		if (Enum)
@@ -1941,34 +1966,34 @@ FPropertyInstanceInfo::FPropertyInstanceInfo(FPropertyInstance PropertyInstance)
 		// if there is no Enum we need to fall through and treat this as a FNumericProperty
 	}
 
-	if (FNumericProperty* NumericProperty = CastField<FNumericProperty>(ResolvedProperty))
+	if (const FNumericProperty* NumericProperty = CastField<FNumericProperty>(ResolvedProperty))
 	{
 		Value = FText::FromString(NumericProperty->GetNumericPropertyValueToString(PropertyInstance.Value));
 		return;
 	}
-	else if (FBoolProperty* BoolProperty = CastField<FBoolProperty>(ResolvedProperty))
+	else if (const FBoolProperty* BoolProperty = CastField<FBoolProperty>(ResolvedProperty))
 	{
 		const FCoreTexts& CoreTexts = FCoreTexts::Get();
 
 		Value = BoolProperty->GetPropertyValue(PropertyInstance.Value) ? CoreTexts.True : CoreTexts.False;
 		return;
 	}
-	else if (FNameProperty* NameProperty = CastField<FNameProperty>(ResolvedProperty))
+	else if (const FNameProperty* NameProperty = CastField<FNameProperty>(ResolvedProperty))
 	{
 		Value = FText::FromName(*(FName*)PropertyInstance.Value);
 		return;
 	}
-	else if (FTextProperty* TextProperty = CastField<FTextProperty>(ResolvedProperty))
+	else if (const FTextProperty* TextProperty = CastField<FTextProperty>(ResolvedProperty))
 	{
 		Value = TextProperty->GetPropertyValue(PropertyInstance.Value);
 		return;
 	}
-	else if (FStrProperty* StringProperty = CastField<FStrProperty>(ResolvedProperty))
+	else if (const FStrProperty* StringProperty = CastField<FStrProperty>(ResolvedProperty))
 	{
 		Value = FText::FromString(StringProperty->GetPropertyValue(PropertyInstance.Value));
 		return;
 	}
-	else if (FArrayProperty* ArrayProperty = CastField<FArrayProperty>(ResolvedProperty))
+	else if (const FArrayProperty* ArrayProperty = CastField<FArrayProperty>(ResolvedProperty))
 	{
 		checkSlow(ArrayProperty->Inner);
 
@@ -1977,14 +2002,14 @@ FPropertyInstanceInfo::FPropertyInstanceInfo(FPropertyInstance PropertyInstance)
 		Value = FText::Format(LOCTEXT("ArraySize", "Num={0}"), FText::AsNumber(ArrayHelper.Num()));
 		return;
 	}
-	else if (FStructProperty* StructProperty = CastField<FStructProperty>(ResolvedProperty))
+	else if (const FStructProperty* StructProperty = CastField<FStructProperty>(ResolvedProperty))
 	{
 		FString WatchText;
 		StructProperty->ExportTextItem(WatchText, PropertyInstance.Value, PropertyInstance.Value, nullptr, PPF_PropertyWindow | PPF_BlueprintDebugView, nullptr);
 		Value = FText::FromString(WatchText);
 		return;
 	}
-	else if (FEnumProperty* EnumProperty = CastField<FEnumProperty>(ResolvedProperty))
+	else if (const FEnumProperty* EnumProperty = CastField<FEnumProperty>(ResolvedProperty))
 	{
 		FNumericProperty* LocalUnderlyingProp = EnumProperty->GetUnderlyingProperty();
 		UEnum* Enum = EnumProperty->GetEnum();
@@ -2011,19 +2036,19 @@ FPropertyInstanceInfo::FPropertyInstanceInfo(FPropertyInstance PropertyInstance)
 
 		return;
 	}
-	else if (FMapProperty* MapProperty = CastField<FMapProperty>(ResolvedProperty))
+	else if (const FMapProperty* MapProperty = CastField<FMapProperty>(ResolvedProperty))
 	{
 		FScriptMapHelper MapHelper(MapProperty, PropertyInstance.Value);
 		Value = FText::Format(LOCTEXT("MapSize", "Num={0}"), FText::AsNumber(MapHelper.Num()));
 		return;
 	}
-	else if (FSetProperty* SetProperty = CastField<FSetProperty>(ResolvedProperty))
+	else if (const FSetProperty* SetProperty = CastField<FSetProperty>(ResolvedProperty))
 	{
 		FScriptSetHelper SetHelper(SetProperty, PropertyInstance.Value);
 		Value = FText::Format(LOCTEXT("SetSize", "Num={0}"), FText::AsNumber(SetHelper.Num()));
 		return;
 	}
-	else if (FObjectPropertyBase* ObjectPropertyBase = CastField<FObjectPropertyBase>(ResolvedProperty))
+	else if (const FObjectPropertyBase* ObjectPropertyBase = CastField<FObjectPropertyBase>(ResolvedProperty))
 	{
 		Object = ObjectPropertyBase->GetObjectPropertyValue(PropertyInstance.Value);
 		if (Object.IsValid())
@@ -2037,7 +2062,7 @@ FPropertyInstanceInfo::FPropertyInstanceInfo(FPropertyInstance PropertyInstance)
 
 		return;
 	}
-	else if (FDelegateProperty* DelegateProperty = CastField<FDelegateProperty>(ResolvedProperty))
+	else if (const FDelegateProperty* DelegateProperty = CastField<FDelegateProperty>(ResolvedProperty))
 	{
 		if (DelegateProperty->SignatureFunction)
 		{
@@ -2050,7 +2075,7 @@ FPropertyInstanceInfo::FPropertyInstanceInfo(FPropertyInstance PropertyInstance)
 
 		return;
 	}
-	else if (FMulticastDelegateProperty* MulticastDelegateProperty = CastField<FMulticastDelegateProperty>(ResolvedProperty))
+	else if (const FMulticastDelegateProperty* MulticastDelegateProperty = CastField<FMulticastDelegateProperty>(ResolvedProperty))
 	{
 		if (MulticastDelegateProperty->SignatureFunction)
 		{
@@ -2063,7 +2088,7 @@ FPropertyInstanceInfo::FPropertyInstanceInfo(FPropertyInstance PropertyInstance)
 
 		return;
 	}
-	else if (FInterfaceProperty* InterfaceProperty = CastField<FInterfaceProperty>(ResolvedProperty))
+	else if (const FInterfaceProperty* InterfaceProperty = CastField<FInterfaceProperty>(ResolvedProperty))
 	{
 		const FScriptInterface* InterfaceData = StaticCast<const FScriptInterface*>(PropertyInstance.Value);
 		Object = InterfaceData->GetObject();
@@ -2111,7 +2136,7 @@ void FPropertyInstanceInfo::PopulateChildren(FPropertyInstance PropertyInstance,
 		return;
 	}
 
-	if (FArrayProperty* ArrayProperty = CastField<FArrayProperty>(Property.Get()))
+	if (const FArrayProperty* ArrayProperty = CastField<FArrayProperty>(Property.Get()))
 	{
 		checkSlow(ArrayProperty->Inner);
 
@@ -2130,7 +2155,7 @@ void FPropertyInstanceInfo::PopulateChildren(FPropertyInstance PropertyInstance,
 			Children.Add(ChildInfo);
 		}
 	}
-	else if (FStructProperty* StructProperty = CastField<FStructProperty>(Property.Get()))
+	else if (const FStructProperty* StructProperty = CastField<FStructProperty>(Property.Get()))
 	{
 		for (TFieldIterator<FProperty> It(StructProperty->Struct); It; ++It)
 		{
@@ -2143,7 +2168,7 @@ void FPropertyInstanceInfo::PopulateChildren(FPropertyInstance PropertyInstance,
 			Children.Add(ChildInfo);
 		}
 	}
-	else if (FMapProperty* MapProperty = CastField<FMapProperty>(Property.Get()))
+	else if (const FMapProperty* MapProperty = CastField<FMapProperty>(Property.Get()))
 	{
 		FScriptMapHelper MapHelper(MapProperty, PropertyInstance.Value);
 		uint8* PropData = MapHelper.GetPairPtr(0);
@@ -2177,7 +2202,7 @@ void FPropertyInstanceInfo::PopulateChildren(FPropertyInstance PropertyInstance,
 			}
 		}
 	}
-	else if (FSetProperty* SetProperty = CastField<FSetProperty>(Property.Get()))
+	else if (const FSetProperty* SetProperty = CastField<FSetProperty>(Property.Get()))
 	{
 		FScriptSetHelper SetHelper(SetProperty, PropertyInstance.Value);
 		uint8* PropData = SetHelper.GetElementPtr(0);
@@ -2202,7 +2227,7 @@ void FPropertyInstanceInfo::PopulateChildren(FPropertyInstance PropertyInstance,
 			}
 		}
 	}
-	else if (FObjectPropertyBase* ObjectPropertyBase = CastField<FObjectPropertyBase>(Property.Get()))
+	else if (const FObjectPropertyBase* ObjectPropertyBase = CastField<FObjectPropertyBase>(Property.Get()))
 	{
 		if (UObject* Obj = ObjectPropertyBase->GetObjectPropertyValue(PropertyInstance.Value))
 		{
@@ -2230,7 +2255,7 @@ TSharedPtr<FPropertyInstanceInfo> FPropertyInstanceInfo::ResolvePathToProperty(c
 		if (InPropertyInfo)
 		{
 			const FString ChildNameStr = InChildName.ToString();
-			FProperty* Prop = InPropertyInfo->Property.Get();
+			const FProperty* Prop = InPropertyInfo->Property.Get();
 
 			if (Prop->IsA<FSetProperty>() || Prop->IsA<FArrayProperty>() || Prop->IsA<FMapProperty>())
 			{
@@ -2275,7 +2300,7 @@ TSharedPtr<FPropertyInstanceInfo> FPropertyInstanceInfo::ResolvePathToProperty(c
 
 FString FPropertyInstanceInfo::GetWatchText() const
 {
-	FProperty* Prop = Property.Get();
+	const FProperty* Prop = Property.Get();
 	if (Prop)
 	{
 		if (Prop->IsA<FSetProperty>() || Prop->IsA<FArrayProperty>() || Prop->IsA<FMapProperty>())
diff --git a/Engine/Source/Editor/UnrealEd/Public/Kismet2/KismetDebugUtilities.h b/Engine/Source/Editor/UnrealEd/Public/Kismet2/KismetDebugUtilities.h
--- a/Engine/Source/Editor/UnrealEd/Public/Kismet2/KismetDebugUtilities.h
+++ b/Engine/Source/Editor/UnrealEd/Public/Kismet2/KismetDebugUtilities.h
@@ -55,7 +55,7 @@ struct FPropertyInstanceInfo
 	 */
 	struct FPropertyInstance
 	{
-		FProperty* Property;
+		const FProperty* Property;
 		const void* Value;
 	};
 	
@@ -89,7 +89,7 @@ struct FPropertyInstanceInfo
 	FText Value;
 	FText Type;
 	TWeakObjectPtr<UObject> Object = nullptr; // only filled if property is a UObject
-	TFieldPath<FProperty> Property;
+	TFieldPath<const FProperty> Property;
 	TArray<TSharedPtr<FPropertyInstanceInfo>> Children;
 	bool bIsInContainer = false;
 };
@@ -429,7 +429,7 @@ public:
 	static EWatchTextResult GetDebugInfo(TSharedPtr<FPropertyInstanceInfo> &OutDebugInfo, UBlueprint* Blueprint, UObject* ActiveObject, const UEdGraphPin* WatchPin);
 
 	// Retrieves Debug info from a Property and a pointer to it's associated data
-	static void GetDebugInfoInternal(TSharedPtr<FPropertyInstanceInfo> &DebugInfo, FProperty* Property, const void* PropertyValue);
+	static void GetDebugInfoInternal(TSharedPtr<FPropertyInstanceInfo> &DebugInfo, const FProperty* Property, const void* PropertyValue);
 
 	//@TODO: Pretty lame way to handle this messaging, ideally the entire Info object gets pushed into the editor when intraframe debugging is triggered!
 	// This doesn't work properly if there is more than one blueprint editor open at once either (one will consume it, the others will be left in the cold)
@@ -438,7 +438,7 @@ protected:
 	static void CheckBreakConditions(UEdGraphNode* NodeStoppedAt, bool bHitBreakpoint, int32 BreakpointOffset, bool& InOutBreakExecution);
 	static void AttemptToBreakExecution(UBlueprint* BlueprintObj, const UObject* ActiveObject, const FFrame& StackFrame, const FBlueprintExceptionInfo& Info, UEdGraphNode* NodeStoppedAt, int32 DebugOpcodeOffset);
 
-	static void GetDebugInfo_InContainer(int32 Index, TSharedPtr<FPropertyInstanceInfo> &DebugInfo, FProperty* Property, const void* Data);
+	static void GetDebugInfo_InContainer(int32 Index, TSharedPtr<FPropertyInstanceInfo> &DebugInfo, const FProperty* Property, const void* Data);
 
 	/**
 	* @brief	Helper function for converting between blueprint and debuggable data
@@ -455,7 +455,7 @@ protected:
 	* @param 	bOutIsDirectPtr	True if OutData/OutDelta point directly at the property rather than its base address
 	* @return	EWTR_Valid if the debug data could be found, otherwise an appropriate error code
 	*/
-	static EWatchTextResult FindDebuggingData(UBlueprint* Blueprint, UObject* ActiveObject, const UEdGraphPin* WatchPin, FProperty*& OutProperty, void*& OutData, void*& OutDelta, UObject*& OutParent, TArray<UObject*>& SeenObjects, bool* bOutIsDirectPtr = nullptr);
+	static EWatchTextResult FindDebuggingData(UBlueprint* Blueprint, UObject* ActiveObject, const UEdGraphPin* WatchPin, const FProperty*& OutProperty, const void*& OutData, const void*& OutDelta, UObject*& OutParent, TArray<UObject*>& SeenObjects, bool* bOutIsDirectPtr = nullptr);
 
 	/**	Retrieve the user settings associated with a blueprint.
 	*	returns null if the blueprint has default settings (no breakpoints and no watches) */
diff --git a/Engine/Source/Runtime/CoreUObject/Public/UObject/FieldPath.h b/Engine/Source/Runtime/CoreUObject/Public/UObject/FieldPath.h
--- a/Engine/Source/Runtime/CoreUObject/Public/UObject/FieldPath.h
+++ b/Engine/Source/Runtime/CoreUObject/Public/UObject/FieldPath.h
@@ -415,7 +415,7 @@ public:
 	template <typename OtherPropertyType>
 	FORCEINLINE bool operator==(const TFieldPath<OtherPropertyType> &Other) const
 	{
-		static_assert(TPointerIsConvertibleFromTo<OtherPropertyType, FField>::Value, "TFieldPath can only be compared with FField types");
+		static_assert(TPointerIsConvertibleFromTo<OtherPropertyType, const FField>::Value, "TFieldPath can only be compared with FField types");
 		static_assert(TPointerIsConvertibleFromTo<PropertyType, OtherPropertyType>::Value, "Unable to compare TFieldPath with raw pointer - types are incompatible");
 
 		return FFieldPath::operator==(Other);
@@ -428,7 +428,7 @@ public:
 	template <typename OtherPropertyType>
 	FORCEINLINE bool operator!=(const TFieldPath<OtherPropertyType> &Other) const
 	{
-		static_assert(TPointerIsConvertibleFromTo<OtherPropertyType, FField>::Value, "TFieldPath can only be compared with FField types");
+		static_assert(TPointerIsConvertibleFromTo<OtherPropertyType, const FField>::Value, "TFieldPath can only be compared with FField types");
 		static_assert(TPointerIsConvertibleFromTo<PropertyType, OtherPropertyType>::Value, "Unable to compare TFieldPath with raw pointer - types are incompatible");
 
 		return FFieldPath::operator!=(Other);
@@ -441,7 +441,7 @@ public:
 	template <typename OtherPropertyType>
 	FORCEINLINE bool operator==(const OtherPropertyType* Other) const
 	{
-		static_assert(TPointerIsConvertibleFromTo<OtherPropertyType, FField>::Value, "TFieldPath can only be compared with FField types");
+		static_assert(TPointerIsConvertibleFromTo<OtherPropertyType, const FField>::Value, "TFieldPath can only be compared with FField types");
 		static_assert(TPointerIsConvertibleFromTo<PropertyType, OtherPropertyType>::Value, "Unable to compare TFieldPath with raw pointer - types are incompatible");
 
 		return Get() == Other;
@@ -454,7 +454,7 @@ public:
 	template <typename OtherPropertyType>
 	FORCEINLINE bool operator!=(const OtherPropertyType* Other) const
 	{
-		static_assert(TPointerIsConvertibleFromTo<OtherPropertyType, FField>::Value, "TFieldPath can only be compared with FField types");
+		static_assert(TPointerIsConvertibleFromTo<OtherPropertyType, const FField>::Value, "TFieldPath can only be compared with FField types");
 		static_assert(TPointerIsConvertibleFromTo<PropertyType, OtherPropertyType>::Value, "Unable to compare TFieldPath with raw pointer - types are incompatible");
 
 		return Get() != Other;
@@ -473,7 +473,7 @@ FORCENOINLINE bool operator==(const LhsT* Lhs, const TFieldPath<RhsT>& Rhs)
 {
 	// It's also possible that these static_asserts may fail for valid conversions because
 	// one or both of the types have only been forward-declared.
-	static_assert(TPointerIsConvertibleFromTo<LhsT, FField>::Value, "TFieldPath can only be compared with FField types");
+	static_assert(TPointerIsConvertibleFromTo<LhsT, const FField>::Value, "TFieldPath can only be compared with FField types");
 	static_assert(TPointerIsConvertibleFromTo<LhsT, RhsT>::Value || TPointerIsConvertibleFromTo<RhsT, LhsT>::Value, "Unable to compare TFieldPath with raw pointer - types are incompatible");
 
 	return Rhs == Lhs;
@@ -496,7 +496,7 @@ FORCENOINLINE bool operator!=(const LhsT* Lhs, const TFieldPath<RhsT>& Rhs)
 {
 	// It's also possible that these static_asserts may fail for valid conversions because
 	// one or both of the types have only been forward-declared.
-	static_assert(TPointerIsConvertibleFromTo<LhsT, FField>::Value, "TFieldPath can only be compared with FField types");
+	static_assert(TPointerIsConvertibleFromTo<LhsT, const FField>::Value, "TFieldPath can only be compared with FField types");
 	static_assert(TPointerIsConvertibleFromTo<LhsT, RhsT>::Value || TPointerIsConvertibleFromTo<RhsT, LhsT>::Value, "Unable to compare TFieldPath with raw pointer - types are incompatible");
 
 	return Rhs != Lhs;
diff --git a/Engine/Source/Runtime/Engine/Classes/Animation/AnimBlueprintGeneratedClass.h b/Engine/Source/Runtime/Engine/Classes/Animation/AnimBlueprintGeneratedClass.h
--- a/Engine/Source/Runtime/Engine/Classes/Animation/AnimBlueprintGeneratedClass.h
+++ b/Engine/Source/Runtime/Engine/Classes/Animation/AnimBlueprintGeneratedClass.h
@@ -176,6 +176,9 @@ public:
 	// History of snapshots of animation data
 	TSimpleRingBuffer<FAnimationFrameSnapshot>* SnapshotBuffer;
 
+	// Mapping from animation node properties to folded versions
+	TMap<TFieldPath<const FProperty>, TFieldPath<const FProperty>> NodeToFoldedPropertyMap;
+
 	// Node visit structure
 	struct FNodeVisit
 	{
@@ -356,6 +359,7 @@ class ENGINE_API UAnimBlueprintGeneratedClass : public UBlueprintGeneratedClass,
 
 	friend class FAnimBlueprintCompilerContext;
 	friend class FAnimBlueprintGeneratedClassCompiledData;
+	friend class FKismetDebugUtilities;
 
 	// List of state machines present in this blueprint class
 	UPROPERTY()
