Volume Factories: Avoid creating incompatible VolumeFactoryClass <-> VolumeClass pairings

diff --git a/Engine/Source/Editor/UnrealEd/Private/EditorEngine.cpp b/Engine/Source/Editor/UnrealEd/Private/EditorEngine.cpp
--- a/Engine/Source/Editor/UnrealEd/Private/EditorEngine.cpp
+++ b/Engine/Source/Editor/UnrealEd/Private/EditorEngine.cpp
@@ -760,8 +760,15 @@ void UEditorEngine::InitEditor(IEngineLoop* InEngineLoop)
 		ActorFactories.Reserve(ActorFactories.Num() + (VolumeFactoryClasses.Num() * VolumeClasses.Num()));
 		for (UClass* VolumeFactoryClass : VolumeFactoryClasses)
 		{
+			// Use NewActorClass of Factory CDO as the supported base class for VolumeClasses
+			const UClass* DefaultActorClass = VolumeFactoryClass->GetDefaultObject<UActorFactory>()->NewActorClass;
 			for (UClass* VolumeClass : VolumeClasses)
 			{
+				if (DefaultActorClass && !VolumeClass->IsChildOf(DefaultActorClass))
+				{
+					continue;
+				}
+
 				UActorFactory* NewFactory = NewObject<UActorFactory>(GetTransientPackage(), VolumeFactoryClass);
 				check(NewFactory);
 				NewFactory->NewActorClass = VolumeClass;
