Do not throttle cpu usage when benchmarking or automation testing

diff --git a/Engine/Source/Editor/UnrealEd/Private/EditorEngine.cpp b/Engine/Source/Editor/UnrealEd/Private/EditorEngine.cpp
--- a/Engine/Source/Editor/UnrealEd/Private/EditorEngine.cpp
+++ b/Engine/Source/Editor/UnrealEd/Private/EditorEngine.cpp
@@ -6159,11 +6159,7 @@ bool UEditorEngine::ShouldThrottleCPUUsage() const
 
 	bool bShouldThrottle = false;
 
-	const bool bRunningCommandlet = IsRunningCommandlet();
-
-	const bool bHasFocus = FApp::HasFocus();
-
-	if( !bHasFocus && !bRunningCommandlet )
+	if (!FApp::HasFocus() && !IsRunningCommandlet() && !GIsAutomationTesting && !FApp::IsBenchmarking())
 	{
 		const UEditorPerformanceSettings* Settings = GetDefault<UEditorPerformanceSettings>();
 		bShouldThrottle = Settings->bThrottleCPUWhenNotForeground;
