Re-enable delegate calling in EEditorEngine::ShouldThrottleCPUUsage and reverse the action taken by Execute to make better sense with the delegate name

diff --git a/Engine/Source/Editor/UnrealEd/Private/EditorEngine.cpp b/Engine/Source/Editor/UnrealEd/Private/EditorEngine.cpp
--- a/Engine/Source/Editor/UnrealEd/Private/EditorEngine.cpp
+++ b/Engine/Source/Editor/UnrealEd/Private/EditorEngine.cpp
@@ -6144,18 +6144,16 @@ bool UEditorEngine::ShouldThrottleCPUUsage() const
 		return false;
 	}
 
-	/*
-	Commented Out due to stale shared pointer issue.
 	// There might be systems where throttling would cause issues (such as data transfer over the network) - give them
 	// an opportunity to force us to not throttle.
 	for (auto It = ShouldDisableCPUThrottlingDelegates.CreateConstIterator(); It; ++It)
 	{
-		if (It->IsBound() && !It->Execute())
+		if (It->IsBound() && It->Execute())
 		{
 			return false;
 		}
 	}
-	*/
+	
 
 	bool bShouldThrottle = false;
 
