Option to disable C++ creation/editing from the editor
Only hooked on "New C++ Component" entry of the "Add Component" list for now (because of immediate time constraints, but will add checks in more places in the editor later)

diff --git a/Engine/Source/Editor/UnrealEd/Classes/Preferences/UnrealEdOptions.h b/Engine/Source/Editor/UnrealEd/Classes/Preferences/UnrealEdOptions.h
--- a/Engine/Source/Editor/UnrealEd/Classes/Preferences/UnrealEdOptions.h
+++ b/Engine/Source/Editor/UnrealEd/Classes/Preferences/UnrealEdOptions.h
@@ -117,6 +117,13 @@ public:
 		return GetNewAssetDefaultClassesDelegate.IsBound() ? GetNewAssetDefaultClassesDelegate.Execute() : NewAssetDefaultClasses;
 	}
 
+	/* Delegate to disallow C++ creation/editing from the editor */
+	DECLARE_DELEGATE_RetVal(bool, FIsCPPAllowed);
+	FIsCPPAllowed& OnIsCPPAllowed() { return IsCPPAllowedDelegate; }
+
+	/** Returns whether C++ creation/editing from the editor is allowed */
+	bool IsCPPAllowed() { return IsCPPAllowedDelegate.IsBound() ? IsCPPAllowedDelegate.Execute() : true; }
+
 	/** Mapping of command name's to array index. */
 	TMap<FName, int32>	CommandMap;
 
@@ -144,7 +151,7 @@ private:
 
 	/* Delegate to override NewAssetDefaultClasses */
 	FGetNewAssetDefaultClasses GetNewAssetDefaultClassesDelegate;
-};
-
-
 
+	/** Delegate to disallow C++ creation/editing from the editor */
+	FIsCPPAllowed IsCPPAllowedDelegate;
+};
diff --git a/Engine/Source/Editor/UnrealEd/Private/ComponentTypeRegistry.cpp b/Engine/Source/Editor/UnrealEd/Private/ComponentTypeRegistry.cpp
--- a/Engine/Source/Editor/UnrealEd/Private/ComponentTypeRegistry.cpp
+++ b/Engine/Source/Editor/UnrealEd/Private/ComponentTypeRegistry.cpp
@@ -12,7 +12,6 @@
 #include "Materials/Material.h"
 #include "Engine/StaticMesh.h"
 #include "AssetData.h"
-
 #include "AssetRegistryModule.h"
 #include "ClassIconFinder.h"
 #include "EdGraphSchema_K2.h"
@@ -21,6 +20,9 @@
 #include "Settings/ClassViewerSettings.h"
 #include "ClassViewerFilter.h"
 #include "EditorClassUtils.h"
+#include "Editor/UnrealEdEngine.h"
+#include "Preferences/UnrealEdOptions.h"
+#include "UnrealEdGlobals.h"
 
 #define LOCTEXT_NAMESPACE "ComponentTypeRegistry"
 
@@ -483,8 +485,11 @@ void FComponentTypeRegistryData::ForceRefreshComponentList()
 		FComponentClassComboEntryPtr NewBPClass = MakeShareable(new FComponentClassComboEntry(NewComponentsHeading, UActorComponent::StaticClass(), true, EComponentCreateAction::CreateNewBlueprintClass));
 		ComponentClassList.Add(NewBPClass);
 
-		FComponentClassComboEntryPtr NewCPPClass = MakeShareable(new FComponentClassComboEntry(NewComponentsHeading, UActorComponent::StaticClass(), true, EComponentCreateAction::CreateNewCPPClass));
-		ComponentClassList.Add(NewCPPClass);
+		if (ensure(GUnrealEd) && GUnrealEd->GetUnrealEdOptions()->IsCPPAllowed())
+		{
+			FComponentClassComboEntryPtr NewCPPClass = MakeShareable(new FComponentClassComboEntry(NewComponentsHeading, UActorComponent::StaticClass(), true, EComponentCreateAction::CreateNewCPPClass));
+			ComponentClassList.Add(NewCPPClass);
+		}
 
 		FComponentClassComboEntryPtr NewSeparator(new FComponentClassComboEntry());
 		ComponentClassList.Add(NewSeparator);
