Remove Flush/WaitForPendingFlush calls when running COTF, they are no longer needed

diff --git a/Engine/Source/Developer/IoStoreUtilities/Private/ZenStoreWriter.cpp b/Engine/Source/Developer/IoStoreUtilities/Private/ZenStoreWriter.cpp
--- a/Engine/Source/Developer/IoStoreUtilities/Private/ZenStoreWriter.cpp
+++ b/Engine/Source/Developer/IoStoreUtilities/Private/ZenStoreWriter.cpp
@@ -516,7 +516,8 @@ void FZenStoreWriter::BeginCook()
 
 void FZenStoreWriter::EndCook()
 {
-	Flush();
+	UE_LOG(LogZenStoreWriter, Display, TEXT("Flushing..."));
+	CommitQueue->WaitUntilEmpty();
 	
 	CommitQueue->CompleteAdding();
 	
@@ -814,14 +815,6 @@ void FZenStoreWriter::GetEntries(TFunction<void(TArrayView<const FPackageStoreEn
 	Callback(PackageStoreEntries, CookedPackagesInfo);
 }
 
-void FZenStoreWriter::Flush()
-{
-	TRACE_CPUPROFILER_EVENT_SCOPE(FZenStoreWriter::Flush);
-
-	UE_LOG(LogZenStoreWriter, Display, TEXT("Flushing..."));
-	CommitQueue->WaitUntilEmpty();
-}
-
 TUniquePtr<FAssetRegistryState> FZenStoreWriter::LoadPreviousAssetRegistry() 
 {
 	// Load the previous asset registry to return to CookOnTheFlyServer, and set the packages enumerated in both *this and
diff --git a/Engine/Source/Developer/IoStoreUtilities/Public/ZenStoreWriter.h b/Engine/Source/Developer/IoStoreUtilities/Public/ZenStoreWriter.h
--- a/Engine/Source/Developer/IoStoreUtilities/Public/ZenStoreWriter.h
+++ b/Engine/Source/Developer/IoStoreUtilities/Public/ZenStoreWriter.h
@@ -65,9 +65,6 @@ public:
 		return MarkUpToDateEvent;
 	}
 
-
-	virtual void Flush() override;
-
 	IOSTOREUTILITIES_API void WriteIoStorePackageData(const FPackageInfo& Info, const FIoBuffer& PackageData, const FPackageStoreEntryResource& PackageStoreEntry, const TArray<FFileRegion>& FileRegions);
 
 	IOSTOREUTILITIES_API TUniquePtr<FAssetRegistryState> LoadPreviousAssetRegistry() override;
diff --git a/Engine/Source/Editor/UnrealEd/Private/CookOnTheFlyServer.cpp b/Engine/Source/Editor/UnrealEd/Private/CookOnTheFlyServer.cpp
--- a/Engine/Source/Editor/UnrealEd/Private/CookOnTheFlyServer.cpp
+++ b/Engine/Source/Editor/UnrealEd/Private/CookOnTheFlyServer.cpp
@@ -267,14 +267,11 @@ class UCookOnTheFlyServer::FCookOnTheFlyServerInterface final
 public:
 	FCookOnTheFlyServerInterface(UCookOnTheFlyServer& InCooker)
 		: Cooker(InCooker)
-		, FlushCompletedEvent(FPlatformProcess::GetSynchEventFromPool(true))
 	{
-		FlushCompletedEvent->Trigger();
 	}
 
 	virtual ~FCookOnTheFlyServerInterface()
 	{
-		FPlatformProcess::ReturnSynchEventToPool(FlushCompletedEvent);
 	}
 
 	virtual FString GetSandboxDirectory() const override
@@ -395,33 +392,6 @@ public:
 		return Cooker.FindOrCreatePackageWriter(TargetPlatform);
 	}
 
-	virtual double WaitForPendingFlush() override
-	{
-		const double StartTime = bIsFlushing ? FPlatformTime::Seconds() : 0.0;
-
-		FlushCompletedEvent->Wait();
-
-		return StartTime > 0.0 ? FPlatformTime::Seconds() - StartTime : 0.0;
-	}
-
-	void Flush()
-	{
-		bIsFlushing = true;
-		FlushCompletedEvent->Reset();
-
-		UE_LOG(LogCook, Log, TEXT("Flushing..."));
-
-		for (UE::Cook::FCookSavePackageContext* Context : Cooker.SavePackageContexts)
-		{
-			Context->PackageWriter->Flush();
-		}
-
-		UE_LOG(LogCook, Log, TEXT("Flush completed"));
-
-		bIsFlushing = false;
-		FlushCompletedEvent->Trigger();
-	}
-
 private:
 	struct FRequest
 	{
@@ -454,12 +424,9 @@ private:
 	}
 
 	UCookOnTheFlyServer& Cooker;
-	FEvent* FlushCompletedEvent;
-	FCriticalSection FlushCriticalSection;
 	FCriticalSection RequestsCriticalSection;
 	TChunkedArray<FRequest> Requests;
 	FRequest* FirstFree = nullptr;
-	TAtomic<bool> bIsFlushing = {false};
 };
 
 /* UCookOnTheFlyServer functions
@@ -1339,11 +1306,6 @@ uint32 UCookOnTheFlyServer::TickCookOnTheSide(const float TimeSlice, uint32 &Coo
 		CookByTheBookFinished();
 	}
 
-	if (IsCookOnTheFlyMode() && bCookComplete)
-	{
-		CookOnTheFlyServerInterface->Flush();
-	}
-
 	CookedPackageCount += StackData.CookedPackageCount;
 	return StackData.ResultFlags;
 }
diff --git a/Engine/Source/Editor/UnrealEd/Private/Cooker/CookOnTheFlyServerInterface.h b/Engine/Source/Editor/UnrealEd/Private/Cooker/CookOnTheFlyServerInterface.h
--- a/Engine/Source/Editor/UnrealEd/Private/Cooker/CookOnTheFlyServerInterface.h
+++ b/Engine/Source/Editor/UnrealEd/Private/Cooker/CookOnTheFlyServerInterface.h
@@ -100,13 +100,6 @@ public:
 	 * Returns the package store writer for the specified platform.
 	 */
 	virtual ICookedPackageWriter& GetPackageWriter(const ITargetPlatform* TargetPlatform) = 0;
-
-	/**
-	 * Wait until any pending flush request is completed.
-	 *
-	 * @return Duration in seconds.
-	 */
-	virtual double WaitForPendingFlush() = 0;
 };
 
 /**
diff --git a/Engine/Source/Editor/UnrealEd/Private/Cooker/DiffPackageWriter.h b/Engine/Source/Editor/UnrealEd/Private/Cooker/DiffPackageWriter.h
--- a/Engine/Source/Editor/UnrealEd/Private/Cooker/DiffPackageWriter.h
+++ b/Engine/Source/Editor/UnrealEd/Private/Cooker/DiffPackageWriter.h
@@ -70,10 +70,6 @@ public:
 	{
 		Inner->EndCook();
 	}
-	virtual void Flush() override
-	{
-		Inner->Flush();
-	}
 	virtual TUniquePtr<FAssetRegistryState> LoadPreviousAssetRegistry() override
 	{
 		return Inner->LoadPreviousAssetRegistry();
@@ -196,10 +192,6 @@ public:
 	{
 		Inner->EndCook();
 	}
-	virtual void Flush() override
-	{
-		Inner->Flush();
-	}
 	virtual TUniquePtr<FAssetRegistryState> LoadPreviousAssetRegistry() override
 	{
 		return Inner->LoadPreviousAssetRegistry();
diff --git a/Engine/Source/Editor/UnrealEd/Private/Cooker/IoStoreCookOnTheFlyRequestManager.cpp b/Engine/Source/Editor/UnrealEd/Private/Cooker/IoStoreCookOnTheFlyRequestManager.cpp
--- a/Engine/Source/Editor/UnrealEd/Private/Cooker/IoStoreCookOnTheFlyRequestManager.cpp
+++ b/Engine/Source/Editor/UnrealEd/Private/Cooker/IoStoreCookOnTheFlyRequestManager.cpp
@@ -387,10 +387,6 @@ private:
 
 		TRACE_CPUPROFILER_EVENT_SCOPE(CookOnTheFly::HandleCookPackageRequest);
 
-		const double FlushDurationInSeconds = CookOnTheFlyServer.WaitForPendingFlush();
-		UE_CLOG(FlushDurationInSeconds > 0.0, LogCookOnTheFly, Log, TEXT("Waited '%.2llf's for cooker to flush pending package(s), Client='%s (%u)'"),
-			FlushDurationInSeconds, *Client.PlatformName.ToString(), Client.ClientId);
-
 		FPlatformContext& Context = GetContext(Client.PlatformName);
 		FScopeLock _(&Context.CriticalSection);
 
diff --git a/Engine/Source/Editor/UnrealEd/Private/Cooker/LooseCookedPackageWriter.cpp b/Engine/Source/Editor/UnrealEd/Private/Cooker/LooseCookedPackageWriter.cpp
--- a/Engine/Source/Editor/UnrealEd/Private/Cooker/LooseCookedPackageWriter.cpp
+++ b/Engine/Source/Editor/UnrealEd/Private/Cooker/LooseCookedPackageWriter.cpp
@@ -397,11 +397,6 @@ void FLooseCookedPackageWriter::EndCook()
 	PackageStoreManifest.Save(*(MetadataDirectoryPath / TEXT("packagestore.manifest")));
 }
 
-void FLooseCookedPackageWriter::Flush()
-{
-	UPackage::WaitForAsyncFileWrites();
-}
-
 TUniquePtr<FAssetRegistryState> FLooseCookedPackageWriter::LoadPreviousAssetRegistry()
 {
 	// Report files from the shared build if the option is set
diff --git a/Engine/Source/Editor/UnrealEd/Private/Cooker/LooseCookedPackageWriter.h b/Engine/Source/Editor/UnrealEd/Private/Cooker/LooseCookedPackageWriter.h
--- a/Engine/Source/Editor/UnrealEd/Private/Cooker/LooseCookedPackageWriter.h
+++ b/Engine/Source/Editor/UnrealEd/Private/Cooker/LooseCookedPackageWriter.h
@@ -40,7 +40,6 @@ public:
 	virtual void Initialize(const FCookInfo& Info) override;
 	virtual void BeginCook() override;
 	virtual void EndCook() override;
-	virtual void Flush() override;
 	virtual TUniquePtr<FAssetRegistryState> LoadPreviousAssetRegistry() override;
 	virtual FCbObject GetOplogAttachment(FName PackageName, FUtf8StringView AttachmentKey) override;
 	virtual void RemoveCookedPackages(TArrayView<const FName> PackageNamesToRemove) override;
diff --git a/Engine/Source/Runtime/Core/Public/Serialization/PackageWriter.h b/Engine/Source/Runtime/Core/Public/Serialization/PackageWriter.h
--- a/Engine/Source/Runtime/Core/Public/Serialization/PackageWriter.h
+++ b/Engine/Source/Runtime/Core/Public/Serialization/PackageWriter.h
@@ -238,11 +238,6 @@ public:
 	  */
 	virtual void EndCook() = 0;
 
-	/**
-	 * Flush any outstanding writes.
-	 */
-	virtual void Flush() = 0;
-
 	struct FCookedPackageInfo
 	{
 		FName PackageName;
