Curve Editor: Empty selection when the curve is cleared. This fixes a longstanding bug that occurs in the curve editor when deleting the curve

diff --git a/Engine/Source/Editor/UnrealEd/Private/SCurveEditor.cpp b/Engine/Source/Editor/UnrealEd/Private/SCurveEditor.cpp
--- a/Engine/Source/Editor/UnrealEd/Private/SCurveEditor.cpp
+++ b/Engine/Source/Editor/UnrealEd/Private/SCurveEditor.cpp
@@ -3804,7 +3804,14 @@ void SCurveEditor::OnObjectPropertyChanged(UObject* Object, FPropertyChangedEven
 {
 	if ( CurveOwner && CurveOwner->GetOwners().Contains(Object) )
 	{
-		ValidateSelection();
+		if (PropertyChangedEvent.ChangeType == EPropertyChangeType::ArrayClear)
+		{
+			EmptyAllSelection();
+		}
+		else
+		{
+			ValidateSelection();
+		}
 	}
 }
 
