- Preload the DevAR based on a #define (ASSETREGISTRY_ENABLE_PREMADE_REGISTRY_IN_EDITOR) as well as REquiresCookedData check
- Allow for skipping of saving the DevelopmentAssetREgistry.bin (useful when cooking DLC)
- Don't initialize the TemporaryState in AssetRegistry when cooking DLC
- Allow for some delaying of blocking on AsyncAR

diff --git a/Engine/Source/Editor/UnrealEd/Private/Commandlets/AssetRegistryGenerator.cpp b/Engine/Source/Editor/UnrealEd/Private/Commandlets/AssetRegistryGenerator.cpp
--- a/Engine/Source/Editor/UnrealEd/Private/Commandlets/AssetRegistryGenerator.cpp
+++ b/Engine/Source/Editor/UnrealEd/Private/Commandlets/AssetRegistryGenerator.cpp
@@ -1054,7 +1054,7 @@ void FAssetRegistryGenerator::UpdateCollectionAssetData()
 	}
 }
 
-void FAssetRegistryGenerator::Initialize(const TArray<FName> &InStartupPackages)
+void FAssetRegistryGenerator::Initialize(const TArray<FName> &InStartupPackages, bool bInitializeFromExisting)
 {
 	StartupPackages.Append(InStartupPackages);
 
@@ -1072,7 +1072,10 @@ void FAssetRegistryGenerator::Initialize(const TArray<FName> &InStartupPackages)
 
 	AssetRegistry.InitializeSerializationOptions(SaveOptions, TargetPlatform->IniPlatformName());
 
-	AssetRegistry.InitializeTemporaryAssetRegistryState(State, SaveOptions);
+	if (bInitializeFromExisting)
+	{
+		AssetRegistry.InitializeTemporaryAssetRegistryState(State, SaveOptions);
+	}
 
 	FGameDelegates::Get().GetAssignLayerChunkDelegate() = FAssignLayerChunkDelegate::CreateStatic(AssignLayerChunkDelegate);
 }
diff --git a/Engine/Source/Editor/UnrealEd/Private/CookOnTheFlyServer.cpp b/Engine/Source/Editor/UnrealEd/Private/CookOnTheFlyServer.cpp
--- a/Engine/Source/Editor/UnrealEd/Private/CookOnTheFlyServer.cpp
+++ b/Engine/Source/Editor/UnrealEd/Private/CookOnTheFlyServer.cpp
@@ -1108,9 +1108,6 @@ static FDelayedAutoRegisterHelper GPreloadARHelper(EDelayedRegisterRunPhase::Ear
 				if (FParse::Value(FCommandLine::Get(), TEXT("BasedOnReleaseVersion="), BasedOnReleaseVersion) &&
 					FParse::Value(FCommandLine::Get(), TEXT("DevelopmentAssetRegistryPlatformOverride="), DevelopmentAssetRegistryPlatformOverride))
 				{
-					const bool bIsCookingAgainstFixedBase = FParse::Param(FCommandLine::Get(), TEXT("CookAgainstFixedBase "));
-					const bool bVerifyPackagesExist = !bIsCookingAgainstFixedBase;
-
 					// get the AR file path and see if it exists
 					GPreloadedARPath = GetBasedOnReleaseVersionAssetRegistryPath(BasedOnReleaseVersion, DevelopmentAssetRegistryPlatformOverride) / TEXT("Metadata") / GetDevelopmentAssetRegistryFilename();
 
@@ -5699,7 +5696,11 @@ void UCookOnTheFlyServer::RefreshPlatformAssetRegistries(const TArrayView<const
 			PlatformData->RegistryGenerator = TUniquePtr<FAssetRegistryGenerator>(RegistryGenerator);
 			RegistryGenerator->CleanManifestDirectories();
 		}
-		RegistryGenerator->Initialize(CookByTheBookOptions ? CookByTheBookOptions->StartupPackages : TArray<FName>());
+
+		// if we are cooking DLC, we will just spend a lot of time removing the shipped packages from the AR,
+		// so we don't bother copying them over. can easily save 10 seconds on a large project
+		bool bInitalizeFromExisting = !IsCookingDLC();
+		RegistryGenerator->Initialize(CookByTheBookOptions ? CookByTheBookOptions->StartupPackages : TArray<FName>(), bInitalizeFromExisting);
 	}
 }
 
@@ -6678,6 +6679,12 @@ void UCookOnTheFlyServer::CookByTheBookFinished()
 		if (!FParse::Param(FCommandLine::Get(), TEXT("SkipSaveAssetRegistry")))
 		{
 			UE_SCOPED_HIERARCHICAL_COOKTIMER(SavingAssetRegistry);
+			SCOPED_BOOT_TIMING("SavingAssetRegistry");
+
+			// if we are cooking DLC, the DevelopmentAR isn't needed - it's used when making DLC against shipping, so there's no need to make it
+			// again, as we don't make DLC against DLC (but allow an override just in case)
+			bool bSaveDevelopmentAssetRegistry = !FParse::Param(FCommandLine::Get(), TEXT("NoSaveDevAR"));
+
 			for (const ITargetPlatform* TargetPlatform : PlatformManager->GetSessionPlatforms())
 			{
 				FPlatformData* PlatformData = PlatformManager->GetPlatformData(TargetPlatform);
@@ -6737,17 +6744,20 @@ void UCookOnTheFlyServer::CookByTheBookFinished()
 				}
 
 				TSet<FName> IgnorePackageNames;
-				for (FPackageData* PackageData : IgnorePackageDatas)
+				if (bSaveDevelopmentAssetRegistry)
 				{
-					IgnorePackageNames.Add(PackageData->GetPackageName());
-				}
+					for (FPackageData* PackageData : IgnorePackageDatas)
+					{
+						IgnorePackageNames.Add(PackageData->GetPackageName());
+					}
 
-				// ignore packages that weren't cooked because they were only referenced by editor-only properties
-				TSet<FName> UncookedEditorOnlyPackageNames;
-				PackageTracker->UncookedEditorOnlyPackages.GetValues(UncookedEditorOnlyPackageNames);
-				for (FName UncookedEditorOnlyPackage : UncookedEditorOnlyPackageNames)
-				{
-					IgnorePackageNames.Add(UncookedEditorOnlyPackage);
+					// ignore packages that weren't cooked because they were only referenced by editor-only properties
+					TSet<FName> UncookedEditorOnlyPackageNames;
+					PackageTracker->UncookedEditorOnlyPackages.GetValues(UncookedEditorOnlyPackageNames);
+					for (FName UncookedEditorOnlyPackage : UncookedEditorOnlyPackageNames)
+					{
+						IgnorePackageNames.Add(UncookedEditorOnlyPackage);
+					}
 				}
 				{
 					Generator.PreSave(CookedPackageNames);
@@ -6778,7 +6788,7 @@ void UCookOnTheFlyServer::CookByTheBookFinished()
 				}
 				{
 					UE_SCOPED_HIERARCHICAL_COOKTIMER(SaveRealAssetRegistry);
-					Generator.SaveAssetRegistry(SandboxRegistryFilename, true, bForceNoFilterAssetsFromAssetRegistry);
+					Generator.SaveAssetRegistry(SandboxRegistryFilename, bSaveDevelopmentAssetRegistry, bForceNoFilterAssetsFromAssetRegistry);
 				}
 				{
 					Generator.PostSave();
@@ -7473,7 +7483,7 @@ void UCookOnTheFlyServer::StartCookByTheBook( const FCookByTheBookStartupOptions
 	CookByTheBookOptions->bErrorOnEngineContentUse = CookByTheBookStartupOptions.bErrorOnEngineContentUse;
 
 	// if we are going to change the state of dlc, we need to clean out our package filename cache (the generated filename cache is dependent on this key). This has to happen later on, but we want to set the DLC State earlier.
-	const bool bDlcStateChanged = CookByTheBookOptions->DlcName != DLCName;
+	const bool bDlcStateChanged = CookByTheBookOptions->DlcName != TEXT("") && CookByTheBookOptions->DlcName != DLCName;
 	CookByTheBookOptions->DlcName = DLCName;
 	if (CookByTheBookOptions->bSkipHardReferences && !CookByTheBookOptions->bSkipSoftReferences)
 	{
diff --git a/Engine/Source/Editor/UnrealEd/Public/Commandlets/AssetRegistryGenerator.h b/Engine/Source/Editor/UnrealEd/Public/Commandlets/AssetRegistryGenerator.h
--- a/Engine/Source/Editor/UnrealEd/Public/Commandlets/AssetRegistryGenerator.h
+++ b/Engine/Source/Editor/UnrealEd/Public/Commandlets/AssetRegistryGenerator.h
@@ -37,7 +37,7 @@ public:
 	/**
 	 * Initializes manifest generator - creates manifest lists, hooks up delegates.
 	 */
-	void Initialize(const TArray<FName> &StartupPackages);
+	void Initialize(const TArray<FName> &StartupPackages, bool bInitializeFromExisting);
 
 	const ITargetPlatform* GetTargetPlatform() const { return TargetPlatform; }
 
diff --git a/Engine/Source/Runtime/AssetRegistry/Private/AssetRegistry.cpp b/Engine/Source/Runtime/AssetRegistry/Private/AssetRegistry.cpp
--- a/Engine/Source/Runtime/AssetRegistry/Private/AssetRegistry.cpp
+++ b/Engine/Source/Runtime/AssetRegistry/Private/AssetRegistry.cpp
@@ -165,7 +165,7 @@ static class FCookedAssetRegistryPreloader
 {
 public:
 	FCookedAssetRegistryPreloader()
-		: bLoadOnce(FPlatformProperties::RequiresCookedData())
+		: bLoadOnce(FPlatformProperties::RequiresCookedData() || ASSETREGISTRY_ENABLE_PREMADE_REGISTRY_IN_EDITOR)
 	{
 		if (bLoadOnce)
 		{
@@ -197,6 +197,8 @@ public:
 
 	bool Consume(FAssetRegistryState& Out)
 	{
+		SCOPED_BOOT_TIMING("FCookedAssetRegistryPreloader::Consume");
+
 		if (StateReady.IsValid())
 		{
 			StateReady.Wait();
@@ -318,6 +320,8 @@ struct FInitializeContext
 UAssetRegistryImpl::UAssetRegistryImpl(const FObjectInitializer& ObjectInitializer)
 	: Super(ObjectInitializer)
 {
+	SCOPED_BOOT_TIMING("UAssetRegistryImpl::UAssetRegistryImpl");
+
 	UE::AssetRegistry::Impl::FInitializeContext Context;
 
 	{
@@ -404,26 +408,30 @@ void FAssetRegistryImpl::Initialize(Impl::FInitializeContext& Context)
 #if ASSETREGISTRY_ENABLE_PREMADE_REGISTRY_IN_EDITOR 
 	if (GIsEditor && !!LoadPremadeAssetRegistryInEditor)
 	{
-		FAssetRegistryLoadOptions Options;
-		const int32 ThreadReduction = 2; // This thread + main thread already has work to do 
-		const bool bCanLoadAsync = FPlatformProcess::SupportsMultithreading() && FTaskGraphInterface::IsRunning();
-		int32 MaxWorkers = bCanLoadAsync ? FPlatformMisc::NumberOfCoresIncludingHyperthreads() - ThreadReduction : 0;
-		Options.ParallelWorkers = FMath::Clamp(MaxWorkers, 0, 16);
-
-		if (LoadAssetRegistry(*(FPaths::ProjectDir() / TEXT("EditorClientAssetRegistry.bin")), Options, State))
-		{
-			UE_LOG(LogAssetRegistry, Log, TEXT("Loaded premade editor client asset registry"));
-			CachePathsFromState(Context.Events, State);
-		}
-		else if (LoadAssetRegistry(*(FPaths::ProjectDir() / TEXT("AssetRegistry.bin")), Options, State))
-		{
-			UE_LOG(LogAssetRegistry, Log, TEXT("Loaded premade asset registry"));
-			CachePathsFromState(Context.Events, State);
-		}
-		else
+		// if the preload failed, load one now
+		if (GCookedAssetRegistryPreloader.Consume(/* Out */ State) == false)
 		{
-			UE_LOG(LogAssetRegistry, Log, TEXT("Failed to load premade asset registry"));
+			SCOPED_BOOT_TIMING("ConsumeFailed");
+			FAssetRegistryLoadOptions Options;
+			const int32 ThreadReduction = 2; // This thread + main thread already has work to do 
+			const bool bCanLoadAsync = FPlatformProcess::SupportsMultithreading() && FTaskGraphInterface::IsRunning();
+			int32 MaxWorkers = bCanLoadAsync ? FPlatformMisc::NumberOfCoresIncludingHyperthreads() - ThreadReduction : 0;
+			Options.ParallelWorkers = FMath::Clamp(MaxWorkers, 0, 16);
+
+			if (LoadAssetRegistry(*(FPaths::ProjectDir() / TEXT("EditorClientAssetRegistry.bin")), Options, State))
+			{
+				UE_LOG(LogAssetRegistry, Log, TEXT("Loaded premade editor client asset registry"));
+			}
+			else if (LoadAssetRegistry(*(FPaths::ProjectDir() / TEXT("AssetRegistry.bin")), Options, State))
+			{
+				UE_LOG(LogAssetRegistry, Log, TEXT("Loaded premade asset registry"));
+			}
+			else
+			{
+				UE_LOG(LogAssetRegistry, Log, TEXT("Failed to load premade asset registry"));
+			}
 		}
+		CachePathsFromState(Context.Events, State);
 #else
 	else if (FPlatformProperties::RequiresCookedData())
 	{
@@ -3172,6 +3180,8 @@ void FAssetRegistryImpl::AppendState(Impl::FEventContext& EventContext, const FA
 
 void FAssetRegistryImpl::CachePathsFromState(Impl::FEventContext& EventContext, const FAssetRegistryState& InState)
 {
+	SCOPED_BOOT_TIMING("FAssetRegistryImpl::CachePathsFromState");
+
 	LLM_SCOPE(ELLMTag::AssetRegistry);
 
 	// Refreshes ClassGeneratorNames if out of date due to module load
