- Workaround for LaunchOn with projects that need a temp target (C++ thinks its content only, UBT makes a temp target, and they get out of sync with the new -target option)

diff --git a/Engine/Source/Editor/UnrealEd/Private/PlayLevelViaLauncher.cpp b/Engine/Source/Editor/UnrealEd/Private/PlayLevelViaLauncher.cpp
--- a/Engine/Source/Editor/UnrealEd/Private/PlayLevelViaLauncher.cpp
+++ b/Engine/Source/Editor/UnrealEd/Private/PlayLevelViaLauncher.cpp
@@ -266,7 +266,13 @@ void UEditorEngine::StartPlayUsingLauncherSession(FRequestPlaySessionParams& InR
 		}
 	}
 
-	LauncherProfile->SetBuildTarget(GetDefault<UProjectPackagingSettings>()->GetLaunchOnTargetInfo()->Name);
+	// content only projects won't have multiple targets to pick from, and pasing -target=UnrealGame will fail if what C++ thinks
+	// is a content only project needs a temporary target.cs file in UBT, 
+	// only set the BuildTarget in code-based projects
+	if (LauncherSessionInfo->bPlayUsingLauncherHasCode)
+	{
+		LauncherProfile->SetBuildTarget(GetDefault<UProjectPackagingSettings>()->GetLaunchOnTargetInfo()->Name);
+	}
 	LauncherProfile->SetCookMode(CurrentLauncherCookMode);
 	LauncherProfile->SetUnversionedCooking(!bIncrimentalCooking);
 	LauncherProfile->SetIncrementalCooking(bIncrimentalCooking);
