Add a new switch '-OnlyPayloadTrailers' to the ResavePackages commandlet. When used, only packages containing FPackageTrailers will be resaved.

#rnx

- This is intended to aid in debugging. I do not expect it to be used as part of a projects normal workflow.

diff --git a/Engine/Source/Editor/UnrealEd/Classes/Commandlets/ResavePackagesCommandlet.h b/Engine/Source/Editor/UnrealEd/Classes/Commandlets/ResavePackagesCommandlet.h
--- a/Engine/Source/Editor/UnrealEd/Classes/Commandlets/ResavePackagesCommandlet.h
+++ b/Engine/Source/Editor/UnrealEd/Classes/Commandlets/ResavePackagesCommandlet.h
@@ -105,6 +105,9 @@ protected:
 	/** Only save packages containing bulkdata payloads that have been virtualized **/
 	bool bOnlyVirtualized;
 
+	/** Only save packages containing FPayloadTrailers **/
+	bool bOnlyPayloadTrailers;
+
 	/** Should we build navigation data for the packages we are saving? **/
 	bool bShouldBuildNavigationData;
 
diff --git a/Engine/Source/Editor/UnrealEd/Private/Commandlets/ContentCommandlets.cpp b/Engine/Source/Editor/UnrealEd/Private/Commandlets/ContentCommandlets.cpp
--- a/Engine/Source/Editor/UnrealEd/Private/Commandlets/ContentCommandlets.cpp
+++ b/Engine/Source/Editor/UnrealEd/Private/Commandlets/ContentCommandlets.cpp
@@ -1043,6 +1043,8 @@ int32 UResavePackagesCommandlet::Main( const FString& Params )
 	bOnlyLicenseed = Switches.Contains(TEXT("OnlyLicenseed"));
 	/** whether we should only save packages containing virtualized bulkdata payloads */
 	bOnlyVirtualized = Switches.Contains(TEXT("OnlyVirtualized"));
+	/** whether we should only save packages containing FPayloadTrailers */
+	bOnlyPayloadTrailers = Switches.Contains(TEXT("OnlyPayloadTrailers"));
 	/** only process packages containing materials */
 	bOnlyMaterials = Switches.Contains(TEXT("onlymaterials"));
 	/** determine if we are building navigation data for the map packages on the pass. **/
@@ -1442,6 +1444,17 @@ void UResavePackagesCommandlet::PerformPreloadOperations( FLinkerLoad* PackageLi
 		}
 	}
 
+	// Check if the package contains a FPackageTrailer or not
+	if (bOnlyPayloadTrailers)
+	{
+		const UE::FPackageTrailer* Trailer = PackageLinker->GetPackageTrailer();
+		if (Trailer == nullptr)
+		{
+			bSavePackage = false;
+			return;
+		}
+	}
+
 	// Check if the package contains any instances of the class that needs to be resaved.
 	if (ResaveClasses.Num() > 0)
 	{
