Change to make size of the arcball larger  so that it works more like Maya. It now deafults to the same size as Maya and is controllabe by a cvar.

diff --git a/Engine/Source/Editor/UnrealEd/Private/UnrealWidget.cpp b/Engine/Source/Editor/UnrealEd/Private/UnrealWidget.cpp
--- a/Engine/Source/Editor/UnrealEd/Private/UnrealWidget.cpp
+++ b/Engine/Source/Editor/UnrealEd/Private/UnrealWidget.cpp
@@ -22,7 +22,6 @@ constexpr float FWidget::INNER_AXIS_CIRCLE_RADIUS;
 constexpr float FWidget::OUTER_AXIS_CIRCLE_RADIUS;
 constexpr float FWidget::ROTATION_TEXT_RADIUS;
 constexpr int32 FWidget::AXIS_CIRCLE_SIDES;
-constexpr float FWidget::ARCALL_RELATIVE_INNER_SIZE;
 constexpr float FWidget::AXIS_LENGTH_SCALE_OFFSET;
 
 FWidget::FWidget()
diff --git a/Engine/Source/Editor/UnrealEd/Private/UnrealWidgetRender.cpp b/Engine/Source/Editor/UnrealEd/Private/UnrealWidgetRender.cpp
--- a/Engine/Source/Editor/UnrealEd/Private/UnrealWidgetRender.cpp
+++ b/Engine/Source/Editor/UnrealEd/Private/UnrealWidgetRender.cpp
@@ -647,6 +647,14 @@ void FWidget::DrawColoredSphere(FPrimitiveDrawInterface* PDI, const FVector& Cen
 	}
 	MeshBuilder.Draw(PDI, FScaleMatrix(Radii) * FRotationMatrix(Orientation) * FTranslationMatrix(Center), MaterialRenderProxy, DepthPriority, bDisableBackfaceCulling);
 }
+//CVAR for the arcball size, so animators can adjust it
+static TAutoConsoleVariable<float> CVarArcballSize(
+	TEXT("r.Editor.ArcballSize"),
+	1.0,
+	TEXT("Custom Size of the Arcball, size of 1.0 means same size as gizmo)"),
+	ECVF_RenderThreadSafe
+);
+
 /**
  * Draws the rotation widget.
  */
@@ -662,15 +670,6 @@ void FWidget::Render_Rotate( const FSceneView* View,FPrimitiveDrawInterface* PDI
 
 	EAxisList::Type DrawAxis = GetAxisToDraw(ViewportClient->GetWidgetMode());
 
-	//if drawing arcball, and dragging make sure axis is showing local
-	if (DrawAxis&EAxisList::XYZ  && CurrentAxis == EAxisList::XYZ && bDragging)
-	{
-		FMatrix LocalCoordinate = ViewportClient->GetLocalCoordinateSystem();
-		XAxis = LocalCoordinate.TransformVector(FVector(1, 0, 0));
-		YAxis = LocalCoordinate.TransformVector(FVector(0, 1, 0));
-		ZAxis = LocalCoordinate.TransformVector(FVector(0, 0, 1));
-	}
-
 	FVector DirectionToWidget = View->IsPerspectiveProjection() ? (InLocation - View->ViewMatrices.GetViewOrigin()) : -View->GetViewDirection();
 	DirectionToWidget.Normalize();
 
@@ -678,27 +677,34 @@ void FWidget::Render_Rotate( const FSceneView* View,FPrimitiveDrawInterface* PDI
 	if (bDrawWidget || bDragging)
 	{
 		bIsOrthoDrawingFullRing = false;
-		//now draw the arc segments
-		if (DrawAxis&EAxisList::X)
+		//now draw the arc segments, only if we are not dragging and moving the arcball per animators request
+		if ((DrawAxis & EAxisList::XYZ && CurrentAxis == EAxisList::XYZ && bDragging) == false) 
 		{
-			DrawRotationArc(View, PDI, EAxisList::X, InLocation, ZAxis, YAxis, DirectionToWidget, AxisColorX.ToFColor(true), Scale, XAxisDir);
-		}
 
-		if (DrawAxis&EAxisList::Y)
-		{
-			DrawRotationArc(View, PDI, EAxisList::Y, InLocation, XAxis, ZAxis, DirectionToWidget, AxisColorY.ToFColor(true), Scale, YAxisDir);
-		}
+			if (DrawAxis & EAxisList::X)
+			{
+				DrawRotationArc(View, PDI, EAxisList::X, InLocation, ZAxis, YAxis, DirectionToWidget, AxisColorX.ToFColor(true), Scale, XAxisDir);
+			}
 
-		if (DrawAxis&EAxisList::Z)
-		{
-			DrawRotationArc(View, PDI, EAxisList::Z, InLocation, XAxis, YAxis, DirectionToWidget, AxisColorZ.ToFColor(true), Scale, ZAxisDir);
+			if (DrawAxis & EAxisList::Y)
+			{
+				DrawRotationArc(View, PDI, EAxisList::Y, InLocation, XAxis, ZAxis, DirectionToWidget, AxisColorY.ToFColor(true), Scale, YAxisDir);
+			}
+
+			if (DrawAxis & EAxisList::Z)
+			{
+				DrawRotationArc(View, PDI, EAxisList::Z, InLocation, XAxis, YAxis, DirectionToWidget, AxisColorZ.ToFColor(true), Scale, ZAxisDir);
+			}
+			
 		}
 
 		if (DrawAxis&EAxisList::XYZ  && (!bDragging || CurrentAxis == EAxisList::XYZ) && GetDefault<ULevelEditorViewportSettings>()->bAllowArcballRotate)
 		{
 			FVector Center = InLocation;
 			FRotator Orientation = FRotator::ZeroRotator;
-			const float InnerDistance = (INNER_AXIS_CIRCLE_RADIUS * ARCALL_RELATIVE_INNER_SIZE * Scale) + GetDefault<ULevelEditorViewportSettings>()->TransformWidgetSizeAdjustment;
+			const float ArcballSize = CVarArcballSize.GetValueOnGameThread();
+
+			const float InnerDistance = (INNER_AXIS_CIRCLE_RADIUS * ArcballSize * Scale) + GetDefault<ULevelEditorViewportSettings>()->TransformWidgetSizeAdjustment;
 			FVector Radii(InnerDistance, InnerDistance, InnerDistance);
 			const bool bDisabled = IsWidgetDisabled();
 			FColor ArcColor = ArcBallColor;
diff --git a/Engine/Source/Editor/UnrealEd/Public/UnrealWidget.h b/Engine/Source/Editor/UnrealEd/Public/UnrealWidget.h
--- a/Engine/Source/Editor/UnrealEd/Public/UnrealWidget.h
+++ b/Engine/Source/Editor/UnrealEd/Public/UnrealWidget.h
@@ -26,7 +26,6 @@ public:
 	constexpr static float OUTER_AXIS_CIRCLE_RADIUS            = 56.0f;
 	constexpr static float ROTATION_TEXT_RADIUS                = 75.0f;
 	constexpr static int32 AXIS_CIRCLE_SIDES                   = 24;
-	constexpr static float ARCALL_RELATIVE_INNER_SIZE          = 0.75f;
 	constexpr static float AXIS_LENGTH_SCALE_OFFSET            = 5.0f;
 
 	FWidget();
