Fix spurious display message from SourceControl when config class settings are changed; changes in script packages should not be treated as SourceControl on the package file.
#rnx

diff --git a/Engine/Source/Editor/UnrealEd/Private/UnrealEdEngine.cpp b/Engine/Source/Editor/UnrealEd/Private/UnrealEdEngine.cpp
--- a/Engine/Source/Editor/UnrealEd/Private/UnrealEdEngine.cpp
+++ b/Engine/Source/Editor/UnrealEd/Private/UnrealEdEngine.cpp
@@ -513,7 +513,8 @@ void UUnrealEdEngine::OnPackageDirtyStateUpdated( UPackage* Pkg)
 
 		if( !bIsAutoSaving && 
 			!GIsEditorLoadingPackage && // Don't ask if the package was modified as a result of a load
-			!GIsCookerLoadingPackage)   // don't ask if the package was modified as a result of a cooker load
+			!GIsCookerLoadingPackage && // don't ask if the package was modified as a result of a cooker load
+			!(Package->GetPackageFlags() & PKG_CompiledIn)) // don't ask if the package is a script package (changes are saved elsewhere via config files)
 		{
 			PackagesDirtiedThisTick.Add(Package);
 
