Fixed issue with packages not cooking again after being marked dirty

diff --git a/Engine/Source/Editor/UnrealEd/Private/CookOnTheFlyServer.cpp b/Engine/Source/Editor/UnrealEd/Private/CookOnTheFlyServer.cpp
--- a/Engine/Source/Editor/UnrealEd/Private/CookOnTheFlyServer.cpp
+++ b/Engine/Source/Editor/UnrealEd/Private/CookOnTheFlyServer.cpp
@@ -3533,7 +3533,7 @@ void UCookOnTheFlyServer::MarkPackageDirtyForCookerFromSchedulerThread(const FNa
 	{
 		check(IsInGameThread()); // We're editing scheduler data, which is only allowable from the scheduler thread
 		bool bHadCookedPlatforms = PackageData->HasAnyCookedPlatform();
-		PackageData->SetPlatformsNotCooked();
+		PackageData->ClearCookProgress();
 		if (PackageData->IsInProgress())
 		{
 			PackageData->SendToState(UE::Cook::EPackageState::Request, UE::Cook::ESendFlags::QueueAddAndRemove);
diff --git a/Engine/Source/Editor/UnrealEd/Private/Cooker/CookPackageData.cpp b/Engine/Source/Editor/UnrealEd/Private/Cooker/CookPackageData.cpp
--- a/Engine/Source/Editor/UnrealEd/Private/Cooker/CookPackageData.cpp
+++ b/Engine/Source/Editor/UnrealEd/Private/Cooker/CookPackageData.cpp
@@ -67,7 +67,7 @@ FPackageData::~FPackageData()
 	// ClearReferences should have been called earlier, but call it here in case it was missed
 	ClearReferences();
 	// We need to send OnLastCookedPlatformRemoved message to the monitor, so call SetPlatformsNotCooked
-	SetPlatformsNotCooked();
+	ClearCookProgress();
 	// Update the monitor's counters and call exit functions
 	SendToState(EPackageState::Idle, ESendFlags::QueueNone);
 }
@@ -307,47 +307,49 @@ void FPackageData::SetPlatformCooked(const ITargetPlatform* TargetPlatform, bool
 	}
 }
 
-void FPackageData::SetPlatformsNotCooked(const TConstArrayView<const ITargetPlatform*> TargetPlatforms)
+void FPackageData::ClearCookProgress(const TConstArrayView<const ITargetPlatform*> TargetPlatforms)
 {
 	for (const ITargetPlatform* TargetPlatform : TargetPlatforms)
 	{
-		SetPlatformNotCooked(TargetPlatform);
+		ClearCookProgress(TargetPlatform);
 	}
 }
 
-void FPackageData::SetPlatformsNotCooked()
+void FPackageData::ClearCookProgress()
 {
-	bool bModified = false;
+	bool bModifiedCookAttempted = false;
 	for (TPair<const ITargetPlatform*, FPlatformData>& Pair : PlatformDatas)
 	{
-		bModified = bModified | (Pair.Value.bCookAttempted != false);
+		bModifiedCookAttempted = bModifiedCookAttempted | (Pair.Value.bCookAttempted != false);
 		Pair.Value.bCookAttempted = false;
 		Pair.Value.bCookSucceeded = false;
+		Pair.Value.bExplored = false;
 	}
-	if (bModified)
+	if (bModifiedCookAttempted)
 	{
 		PackageDatas.GetMonitor().OnLastCookedPlatformRemoved(*this);
 	}
 }
 
-void FPackageData::SetPlatformNotCooked(const ITargetPlatform* TargetPlatform)
+void FPackageData::ClearCookProgress(const ITargetPlatform* TargetPlatform)
 {
 	bool bHasAnyOthers = false;
-	bool bModified = false;
+	bool bModifiedCookAttempted = false;
 	for (TPair<const ITargetPlatform*, FPlatformData>& Pair : PlatformDatas)
 	{
 		if (Pair.Key == TargetPlatform)
 		{
-			bModified = bModified | (Pair.Value.bCookAttempted != false);
+			bModifiedCookAttempted = bModifiedCookAttempted | (Pair.Value.bCookAttempted != false);
 			Pair.Value.bCookAttempted = false;
 			Pair.Value.bCookSucceeded = false;
+			Pair.Value.bExplored = false;
 		}
 		else
 		{
 			bHasAnyOthers = bHasAnyOthers | (Pair.Value.bCookAttempted != false);
 		}
 	}
-	if (bModified && !bHasAnyOthers)
+	if (bModifiedCookAttempted && !bHasAnyOthers)
 	{
 		PackageDatas.GetMonitor().OnLastCookedPlatformRemoved(*this);
 	}
@@ -2091,7 +2093,7 @@ void FPackageDatas::ClearCookedPlatforms()
 {
 	for (FPackageData* PackageData : PackageDatas)
 	{
-		PackageData->SetPlatformsNotCooked();
+		PackageData->ClearCookProgress();
 	}
 }
 
diff --git a/Engine/Source/Editor/UnrealEd/Private/Cooker/CookPackageData.h b/Engine/Source/Editor/UnrealEd/Private/Cooker/CookPackageData.h
--- a/Engine/Source/Editor/UnrealEd/Private/Cooker/CookPackageData.h
+++ b/Engine/Source/Editor/UnrealEd/Private/Cooker/CookPackageData.h
@@ -229,9 +229,9 @@ public:
 	 * FindOrAdd each TargetPlatform and set its flags: CookAttempted=false.
 	 * In Version that takes no TargetPlatform, CookAttempted is cleared from all existing platforms.
 	 */
-	void SetPlatformsNotCooked(const TConstArrayView<const ITargetPlatform*> TargetPlatforms);
-	void SetPlatformsNotCooked();
-	void SetPlatformNotCooked(const ITargetPlatform* TargetPlatform);
+	void ClearCookProgress(const TConstArrayView<const ITargetPlatform*> TargetPlatforms);
+	void ClearCookProgress();
+	void ClearCookProgress(const ITargetPlatform* TargetPlatform);
 
 	/** Access the information about platforms interacted with by *this. */
 	const TSortedMap<const ITargetPlatform*, FPlatformData>& GetPlatformDatas() const;
