ZenCOTF: Notify request manager of packages generated during the cook
#rnx

diff --git a/Engine/Source/Editor/UnrealEd/Private/CookOnTheFlyServer.cpp b/Engine/Source/Editor/UnrealEd/Private/CookOnTheFlyServer.cpp
--- a/Engine/Source/Editor/UnrealEd/Private/CookOnTheFlyServer.cpp
+++ b/Engine/Source/Editor/UnrealEd/Private/CookOnTheFlyServer.cpp
@@ -2075,6 +2075,11 @@ void UCookOnTheFlyServer::QueueDiscoveredPackageData(UE::Cook::FPackageData& Pac
 		return;
 	}
 
+	if (CookOnTheFlyRequestManager && PackageData.IsGenerated())
+	{
+		CookOnTheFlyRequestManager->OnPackageGenerated(PackageData.GetPackageName());
+	}
+	
 	if (!PackageData.IsInProgress() &&
 		(PackageData.IsGenerated() || !IsCookByTheBookMode() || !CookByTheBookOptions->bSkipHardReferences))
 	{
diff --git a/Engine/Source/Editor/UnrealEd/Private/Cooker/CookOnTheFlyServerInterface.h b/Engine/Source/Editor/UnrealEd/Private/Cooker/CookOnTheFlyServerInterface.h
--- a/Engine/Source/Editor/UnrealEd/Private/Cooker/CookOnTheFlyServerInterface.h
+++ b/Engine/Source/Editor/UnrealEd/Private/Cooker/CookOnTheFlyServerInterface.h
@@ -117,6 +117,9 @@ public:
 
 	/** Shutdown the reques tmanager. */
 	virtual void Shutdown() = 0;
+
+	/** Called when a new package is generated */
+	virtual void OnPackageGenerated(const FName& PackageName) = 0;
 };
 
 }} // namespace UE::Cook
diff --git a/Engine/Source/Editor/UnrealEd/Private/Cooker/IoStoreCookOnTheFlyRequestManager.cpp b/Engine/Source/Editor/UnrealEd/Private/Cooker/IoStoreCookOnTheFlyRequestManager.cpp
--- a/Engine/Source/Editor/UnrealEd/Private/Cooker/IoStoreCookOnTheFlyRequestManager.cpp
+++ b/Engine/Source/Editor/UnrealEd/Private/Cooker/IoStoreCookOnTheFlyRequestManager.cpp
@@ -267,10 +267,19 @@ private:
 		ConnectionServer.Reset();
 	}
 
+	virtual void OnPackageGenerated(const FName& PackageName)
+	{
+		FPackageId PackageId = FPackageId::FromName(PackageName);
+		UE_LOG(LogCookOnTheFly, Verbose, TEXT("Package 0x%llX '%s' generated"), PackageId.ValueForDebugging(), *PackageName.ToString());
+
+		FScopeLock _(&AllKnownPackagesCriticalSection);
+		AllKnownPackagesMap.Add(PackageId, PackageName);
+	}
+
 private:
 	bool HandleClientConnection(UE::Cook::FCookOnTheFlyClient Client, UE::Cook::ECookOnTheFlyConnectionStatus ConnectionStatus)
 	{
-		FScopeLock _(&CriticalSection);
+		FScopeLock _(&ContextsCriticalSection);
 
 		if (ConnectionStatus == UE::Cook::ECookOnTheFlyConnectionStatus::Connected)
 		{
@@ -395,7 +404,11 @@ private:
 		const bool bIsCooked = Context.PackageTracker.IsCooked(CookRequest.PackageId);
 		if (!bIsCooked)
 		{
-			FName PackageName = AllKnownPackagesMap.FindRef(CookRequest.PackageId);
+			FName PackageName;
+			{
+				FScopeLock LockPackageNames(&AllKnownPackagesCriticalSection);
+				PackageName = AllKnownPackagesMap.FindRef(CookRequest.PackageId);
+			}
 			if (PackageName.IsNone())
 			{
 				UE_LOG(LogCookOnTheFly, Log, TEXT("Received cook request for unknown package 0x%llX"), CookRequest.PackageId.ValueForDebugging());
@@ -574,7 +587,7 @@ private:
 
 	FPlatformContext& GetContext(const FName& PlatformName)
 	{
-		FScopeLock _(&CriticalSection);
+		FScopeLock _(&ContextsCriticalSection);
 		TUniquePtr<FPlatformContext>& Ctx = PlatformContexts.FindChecked(PlatformName);
 		check(Ctx.IsValid());
 		return *Ctx;
@@ -583,8 +596,9 @@ private:
 	UE::Cook::ICookOnTheFlyServer& CookOnTheFlyServer;
 	UE::Cook::FIoStoreCookOnTheFlyServerOptions Options;
 	TUniquePtr<UE::Cook::ICookOnTheFlyConnectionServer> ConnectionServer;
-	FCriticalSection CriticalSection;
+	FCriticalSection ContextsCriticalSection;
 	TMap<FName, TUniquePtr<FPlatformContext>> PlatformContexts;
+	FCriticalSection AllKnownPackagesCriticalSection;
 	TMap<FPackageId, FName> AllKnownPackagesMap;
 };
 
@@ -598,4 +612,4 @@ TUniquePtr<ICookOnTheFlyRequestManager> MakeIoStoreCookOnTheFlyRequestManager(IC
 
 }}
 
-#endif
\ No newline at end of file
+#endif
diff --git a/Engine/Source/Editor/UnrealEd/Private/Cooker/NetworkFileCookOnTheFlyRequestManager.cpp b/Engine/Source/Editor/UnrealEd/Private/Cooker/NetworkFileCookOnTheFlyRequestManager.cpp
--- a/Engine/Source/Editor/UnrealEd/Private/Cooker/NetworkFileCookOnTheFlyRequestManager.cpp
+++ b/Engine/Source/Editor/UnrealEd/Private/Cooker/NetworkFileCookOnTheFlyRequestManager.cpp
@@ -43,6 +43,10 @@ public:
 		}
 	}
 
+	virtual void OnPackageGenerated(const FName& PackageName) override
+	{
+	}
+
 private:
 	bool OnNewConnection(const FString& VersionInfo, const FString& PlatformName)
 	{
