Made sure that view extensions are gathered and that PreRenderView_RenderThread/PreRenderViewFamily_RenderThread and SetupView/SetupViewFamily are called on all of them for all code paths triggering a scene render.
* For thumbnail renderers, the prototype of RenderViewFamily (which is called by all the types of thumbnail renderers) has been changed to take a non-const FSceneView* parameter (so that SetupView can be called on it) and the various GetView functions have been renamed CreateView since it more accurately depicts what the functions do (plus they return the non-const view they've created)
* UGameViewportClient::Draw was calling SetupViewFamily but not SetupView
* Made sure FSceneRenderer::ViewExtensionPreRender_RenderThread is called right before RenderThreadBegin for all code paths
* Made sure view extensions are gathered for all code paths rendering a view family

Note: BeginRenderViewFamily (i.e. game-thread), PostRenderView_RenderThread and PostRenderViewFamily_RenderThread are still not called on the non-standard code paths

#tests EngineTests (screenshots), editor

diff --git a/Engine/Plugins/Animation/ControlRig/Source/ControlRigEditor/Private/Tools/ControlRigPoseThumbnailRenderer.cpp b/Engine/Plugins/Animation/ControlRig/Source/ControlRigEditor/Private/Tools/ControlRigPoseThumbnailRenderer.cpp
--- a/Engine/Plugins/Animation/ControlRig/Source/ControlRigEditor/Private/Tools/ControlRigPoseThumbnailRenderer.cpp
+++ b/Engine/Plugins/Animation/ControlRig/Source/ControlRigEditor/Private/Tools/ControlRigPoseThumbnailRenderer.cpp
@@ -217,8 +217,7 @@ void UControlRigPoseThumbnailRenderer::Draw(UObject* Object, int32 X, int32 Y, u
 		ViewFamily.EngineShowFlags.MotionBlur = 0;
 		ViewFamily.EngineShowFlags.LOD = 0;
 
-		ThumbnailScene->GetView(&ViewFamily, X, Y, Width, Height);
-		RenderViewFamily(Canvas, &ViewFamily);
+		RenderViewFamily(Canvas, &ViewFamily, ThumbnailScene->CreateView(&ViewFamily, X, Y, Width, Height));
 		ThumbnailScene->SetControlRigPoseAsset(nullptr);
 	}
 }
diff --git a/Engine/Plugins/Experimental/GeometryCollectionPlugin/Source/GeometryCollectionEditor/Private/GeometryCollection/GeometryCollectionThumbnailRenderer.cpp b/Engine/Plugins/Experimental/GeometryCollectionPlugin/Source/GeometryCollectionEditor/Private/GeometryCollection/GeometryCollectionThumbnailRenderer.cpp
--- a/Engine/Plugins/Experimental/GeometryCollectionPlugin/Source/GeometryCollectionEditor/Private/GeometryCollection/GeometryCollectionThumbnailRenderer.cpp
+++ b/Engine/Plugins/Experimental/GeometryCollectionPlugin/Source/GeometryCollectionEditor/Private/GeometryCollection/GeometryCollectionThumbnailRenderer.cpp
@@ -34,8 +34,7 @@ void UGeometryCollectionThumbnailRenderer::Draw(UObject* Object, int32 X, int32
 		ViewFamily.EngineShowFlags.MotionBlur = 0;
 		ViewFamily.EngineShowFlags.LOD = 0;
 
-		ThumbnailScene->GetView(&ViewFamily, X, Y, Width, Height);
-		RenderViewFamily(Canvas, &ViewFamily);
+		RenderViewFamily(Canvas, &ViewFamily, ThumbnailScene->CreateView(&ViewFamily, X, Y, Width, Height));
 		ThumbnailScene->SetGeometryCollection(nullptr);
 	}
 }
diff --git a/Engine/Plugins/Experimental/GeometryCollectionPlugin/Source/GeometryCollectionEditor/Public/GeometryCollection/GeometryCollectionThumbnailScene.h b/Engine/Plugins/Experimental/GeometryCollectionPlugin/Source/GeometryCollectionEditor/Public/GeometryCollection/GeometryCollectionThumbnailScene.h
--- a/Engine/Plugins/Experimental/GeometryCollectionPlugin/Source/GeometryCollectionEditor/Public/GeometryCollection/GeometryCollectionThumbnailScene.h
+++ b/Engine/Plugins/Experimental/GeometryCollectionPlugin/Source/GeometryCollectionEditor/Public/GeometryCollection/GeometryCollectionThumbnailScene.h
@@ -14,7 +14,7 @@ public:
 	/** Constructor */
 	FGeometryCollectionThumbnailScene();
 
-	/** Sets the geometry collection to use in the next GetView() */
+	/** Sets the geometry collection to use in the next CreateView() */
 	void SetGeometryCollection(UGeometryCollection* GeometryCollection);
 
 protected:
diff --git a/Engine/Plugins/Runtime/ApexDestruction/Source/ApexDestructionEditor/Private/DestructibleMeshThumbnailRenderer.cpp b/Engine/Plugins/Runtime/ApexDestruction/Source/ApexDestructionEditor/Private/DestructibleMeshThumbnailRenderer.cpp
--- a/Engine/Plugins/Runtime/ApexDestruction/Source/ApexDestructionEditor/Private/DestructibleMeshThumbnailRenderer.cpp
+++ b/Engine/Plugins/Runtime/ApexDestruction/Source/ApexDestructionEditor/Private/DestructibleMeshThumbnailRenderer.cpp
@@ -36,8 +36,7 @@ PRAGMA_DISABLE_DEPRECATION_WARNINGS
 		ViewFamily.EngineShowFlags.MotionBlur = 0;
 		ViewFamily.EngineShowFlags.LOD = 0;
 
-		ThumbnailScene->GetView(&ViewFamily, X, Y, Width, Height);
-		RenderViewFamily(Canvas,&ViewFamily);
+		RenderViewFamily(Canvas, &ViewFamily, ThumbnailScene->CreateView(&ViewFamily, X, Y, Width, Height));
 		ThumbnailScene->SetDestructibleMesh(nullptr);
 	}
 PRAGMA_ENABLE_DEPRECATION_WARNINGS
diff --git a/Engine/Plugins/Runtime/ApexDestruction/Source/ApexDestructionEditor/Private/DestructibleMeshThumbnailScene.h b/Engine/Plugins/Runtime/ApexDestruction/Source/ApexDestructionEditor/Private/DestructibleMeshThumbnailScene.h
--- a/Engine/Plugins/Runtime/ApexDestruction/Source/ApexDestructionEditor/Private/DestructibleMeshThumbnailScene.h
+++ b/Engine/Plugins/Runtime/ApexDestruction/Source/ApexDestructionEditor/Private/DestructibleMeshThumbnailScene.h
@@ -17,7 +17,7 @@ public:
 	FDestructibleMeshThumbnailScene();
 
 PRAGMA_DISABLE_DEPRECATION_WARNINGS
-	/** Sets the skeletal mesh to use in the next GetView() */
+	/** Sets the skeletal mesh to use in the next CreateView() */
 	void SetDestructibleMesh(class UDestructibleMesh* InMesh);
 PRAGMA_ENABLE_DEPRECATION_WARNINGS
 
diff --git a/Engine/Plugins/Runtime/GeometryCache/Source/GeometryCacheEd/Classes/GeometryCacheThumbnailScene.h b/Engine/Plugins/Runtime/GeometryCache/Source/GeometryCacheEd/Classes/GeometryCacheThumbnailScene.h
--- a/Engine/Plugins/Runtime/GeometryCache/Source/GeometryCacheEd/Classes/GeometryCacheThumbnailScene.h
+++ b/Engine/Plugins/Runtime/GeometryCache/Source/GeometryCacheEd/Classes/GeometryCacheThumbnailScene.h
@@ -14,7 +14,7 @@ public:
 	/** Constructor */
 	FGeometryCacheThumbnailScene();
 
-	/** Sets the static mesh to use in the next GetView() */
+	/** Sets the static mesh to use in the next CreateView() */
 	void SetGeometryCache(UGeometryCache* GeometryCache);
 
 protected:
diff --git a/Engine/Plugins/Runtime/GeometryCache/Source/GeometryCacheEd/Private/GeometryCacheThumbnailRenderer.cpp b/Engine/Plugins/Runtime/GeometryCache/Source/GeometryCacheEd/Private/GeometryCacheThumbnailRenderer.cpp
--- a/Engine/Plugins/Runtime/GeometryCache/Source/GeometryCacheEd/Private/GeometryCacheThumbnailRenderer.cpp
+++ b/Engine/Plugins/Runtime/GeometryCache/Source/GeometryCacheEd/Private/GeometryCacheThumbnailRenderer.cpp
@@ -34,8 +34,7 @@ void UGeometryCacheThumbnailRenderer::Draw(UObject* Object, int32 X, int32 Y, ui
 		ViewFamily.EngineShowFlags.MotionBlur = 0;
 		ViewFamily.EngineShowFlags.LOD = 0;
 
-		ThumbnailScene->GetView(&ViewFamily, X, Y, Width, Height);
-		RenderViewFamily(Canvas, &ViewFamily);
+		RenderViewFamily(Canvas, &ViewFamily, ThumbnailScene->CreateView(&ViewFamily, X, Y, Width, Height));
 		ThumbnailScene->SetGeometryCache(nullptr);
 	}
 }
diff --git a/Engine/Source/Editor/FoliageEdit/Private/FoliageType_ISMThumbnailRenderer.cpp b/Engine/Source/Editor/FoliageEdit/Private/FoliageType_ISMThumbnailRenderer.cpp
--- a/Engine/Source/Editor/FoliageEdit/Private/FoliageType_ISMThumbnailRenderer.cpp
+++ b/Engine/Source/Editor/FoliageEdit/Private/FoliageType_ISMThumbnailRenderer.cpp
@@ -36,8 +36,7 @@ void UFoliageType_ISMThumbnailRenderer::Draw(UObject* Object, int32 X, int32 Y,
 		ViewFamily.EngineShowFlags.MotionBlur = 0;
 		ViewFamily.EngineShowFlags.LOD = 0;
 
-		ThumbnailScene->GetView(&ViewFamily, X, Y, Width, Height);
-		RenderViewFamily(Canvas, &ViewFamily);
+		RenderViewFamily(Canvas, &ViewFamily, ThumbnailScene->CreateView(&ViewFamily, X, Y, Width, Height));
 		ThumbnailScene->SetStaticMesh(nullptr);
 		ThumbnailScene->SetOverrideMaterials(TArray<class UMaterialInterface*>());
 	}
diff --git a/Engine/Source/Editor/UnrealEd/Classes/ThumbnailRendering/LevelThumbnailRenderer.h b/Engine/Source/Editor/UnrealEd/Classes/ThumbnailRendering/LevelThumbnailRenderer.h
--- a/Engine/Source/Editor/UnrealEd/Classes/ThumbnailRendering/LevelThumbnailRenderer.h
+++ b/Engine/Source/Editor/UnrealEd/Classes/ThumbnailRendering/LevelThumbnailRenderer.h
@@ -29,6 +29,6 @@ class ULevelThumbnailRenderer : public UDefaultSizedThumbnailRenderer
 	// End UThumbnailRenderer Object
 
 private:
-	FSceneView* GetView(ULevel* Level, FSceneViewFamily* ViewFamily, int32 X, int32 Y, uint32 SizeX, uint32 SizeY) const;
+	FSceneView* CreateView(ULevel* Level, FSceneViewFamily* ViewFamily, int32 X, int32 Y, uint32 SizeX, uint32 SizeY) const; 
 };
 
diff --git a/Engine/Source/Editor/UnrealEd/Classes/ThumbnailRendering/ThumbnailRenderer.h b/Engine/Source/Editor/UnrealEd/Classes/ThumbnailRendering/ThumbnailRenderer.h
--- a/Engine/Source/Editor/UnrealEd/Classes/ThumbnailRendering/ThumbnailRenderer.h
+++ b/Engine/Source/Editor/UnrealEd/Classes/ThumbnailRendering/ThumbnailRenderer.h
@@ -83,6 +83,9 @@ public:
 
 protected:
 	/** Renders the thumbnail's view family. */
+	UNREALED_API static void RenderViewFamily(FCanvas* Canvas, class FSceneViewFamily* ViewFamily, class FSceneView* View);
+
+	UE_DEPRECATED(5.0, "Please use the prototype of RenderViewFamily which takes a non-const view in parameter : like the view family the view is still not quite final until SetupView is called.")
 	UNREALED_API static void RenderViewFamily(FCanvas* Canvas, class FSceneViewFamily* ViewFamily);
 
 	UNREALED_API static struct FGameTime GetTime();
diff --git a/Engine/Source/Editor/UnrealEd/Classes/ThumbnailRendering/WorldThumbnailRenderer.h b/Engine/Source/Editor/UnrealEd/Classes/ThumbnailRendering/WorldThumbnailRenderer.h
--- a/Engine/Source/Editor/UnrealEd/Classes/ThumbnailRendering/WorldThumbnailRenderer.h
+++ b/Engine/Source/Editor/UnrealEd/Classes/ThumbnailRendering/WorldThumbnailRenderer.h
@@ -27,7 +27,8 @@ class UWorldThumbnailRenderer : public UDefaultSizedThumbnailRenderer
 	// End UThumbnailRenderer Object
 
 private:
-	void GetView(UWorld* World, FSceneViewFamily* ViewFamily, int32 X, int32 Y, uint32 SizeX, uint32 SizeY) const; 
+	/** Allocates then adds an FSceneView to the ViewFamily. */
+	FSceneView* CreateView(UWorld* World, FSceneViewFamily* ViewFamily, int32 X, int32 Y, uint32 SizeX, uint32 SizeY) const;
 
 private:
 	/** Offset used to orient all worlds to show a more vertical camera, if necessary. Individual thumbnail infos can provide additional offset. */
diff --git a/Engine/Source/Editor/UnrealEd/Private/ThumbnailHelpers.cpp b/Engine/Source/Editor/UnrealEd/Private/ThumbnailHelpers.cpp
--- a/Engine/Source/Editor/UnrealEd/Private/ThumbnailHelpers.cpp
+++ b/Engine/Source/Editor/UnrealEd/Private/ThumbnailHelpers.cpp
@@ -67,7 +67,14 @@ FThumbnailPreviewScene::FThumbnailPreviewScene()
 	FPreviewScene::AddComponent(FloorPlaneComponent, FloorPlaneTransform);
 }
 
+PRAGMA_DISABLE_DEPRECATION_WARNINGS
 void FThumbnailPreviewScene::GetView(FSceneViewFamily* ViewFamily, int32 X, int32 Y, uint32 SizeX, uint32 SizeY) const
+{
+	CreateView(ViewFamily, X, Y, SizeX, SizeY);
+}
+PRAGMA_ENABLE_DEPRECATION_WARNINGS
+
+FSceneView* FThumbnailPreviewScene::CreateView(FSceneViewFamily * ViewFamily, int32 X, int32 Y, uint32 SizeX, uint32 SizeY) const
 {
 	check(ViewFamily);
 
@@ -77,65 +84,69 @@ void FThumbnailPreviewScene::GetView(FSceneViewFamily* ViewFamily, int32 X, int3
 		FMath::Max<int32>(X+SizeX,0),
 		FMath::Max<int32>(Y+SizeY,0));
 
-	if (ViewRect.Width() > 0 && ViewRect.Height() > 0)
+	if (ViewRect.Area() <= 0)
 	{
-		const float FOVDegrees = 30.f;
-		const float HalfFOVRadians = FMath::DegreesToRadians<float>(FOVDegrees) * 0.5f;
-		static_assert((int32)ERHIZBuffer::IsInverted != 0, "Check NearPlane and Projection Matrix");
-		const float NearPlane = 1.0f;
-		FMatrix ProjectionMatrix = FReversedZPerspectiveMatrix(
-			HalfFOVRadians,
-			1.0f,
-			1.0f,
-			NearPlane
-			);
-
-		FVector Origin(0);
-		float OrbitPitch = 0;
-		float OrbitYaw = 0;
-		float OrbitZoom = 0;
-		GetViewMatrixParameters(FOVDegrees, Origin, OrbitPitch, OrbitYaw, OrbitZoom);
-
-		// Ensure a minimum camera distance to prevent problems with really small objects
-		const float MinCameraDistance = 48;
-		OrbitZoom = FMath::Max<float>(MinCameraDistance, OrbitZoom);
-
-		const FRotator RotationOffsetToViewCenter(0.f, 90.f, 0.f);
-		FMatrix ViewRotationMatrix = FRotationMatrix( FRotator(0, OrbitYaw, 0) ) * 
-			FRotationMatrix( FRotator(0, 0, OrbitPitch) ) *
-			FTranslationMatrix( FVector(0, OrbitZoom, 0) ) *
-			FInverseRotationMatrix( RotationOffsetToViewCenter );
-
-		ViewRotationMatrix = ViewRotationMatrix * FMatrix(
-			FPlane(0,	0,	1,	0),
-			FPlane(1,	0,	0,	0),
-			FPlane(0,	1,	0,	0),
-			FPlane(0,	0,	0,	1));
-
-		Origin -= ViewRotationMatrix.InverseTransformPosition( FVector::ZeroVector );
-		ViewRotationMatrix = ViewRotationMatrix.RemoveTranslation();
-
-		FSceneViewInitOptions ViewInitOptions;
-		ViewInitOptions.ViewFamily = ViewFamily;
-		ViewInitOptions.SetViewRectangle(ViewRect);
-		ViewInitOptions.ViewOrigin = -Origin;
-		ViewInitOptions.ViewRotationMatrix = ViewRotationMatrix;
-		ViewInitOptions.ProjectionMatrix = ProjectionMatrix;
-		ViewInitOptions.BackgroundColor = FLinearColor::Black;
-
-		FSceneView* NewView = new FSceneView(ViewInitOptions);
-
-		ViewFamily->Views.Add(NewView);
-
-		NewView->StartFinalPostprocessSettings( ViewInitOptions.ViewOrigin );
-		NewView->EndFinalPostprocessSettings(ViewInitOptions);
-		
-		// Tell the texture streaming system about this thumbnail view, so the textures will stream in as needed
-		// NOTE: Sizes may not actually be in screen space depending on how the thumbnail ends up stretched by the UI.  Not a big deal though.
-		// NOTE: Textures still take a little time to stream if the view has not been re-rendered recently, so they may briefly appear blurry while mips are prepared
-		// NOTE: Content Browser only renders thumbnails for loaded assets, and only when the mouse is over the panel. They'll be frozen in their last state while the mouse cursor is not over the panel.  This is for performance reasons
-		IStreamingManager::Get().AddViewInformation( Origin, SizeX, SizeX / FMath::Tan( FOVDegrees ) );
+		return nullptr;
 	}
+		
+	const float FOVDegrees = 30.f;
+	const float HalfFOVRadians = FMath::DegreesToRadians<float>(FOVDegrees) * 0.5f;
+	static_assert((int32)ERHIZBuffer::IsInverted != 0, "Check NearPlane and Projection Matrix");
+	const float NearPlane = 1.0f;
+	FMatrix ProjectionMatrix = FReversedZPerspectiveMatrix(
+		HalfFOVRadians,
+		1.0f,
+		1.0f,
+		NearPlane
+		);
+
+	FVector Origin(0);
+	float OrbitPitch = 0;
+	float OrbitYaw = 0;
+	float OrbitZoom = 0;
+	GetViewMatrixParameters(FOVDegrees, Origin, OrbitPitch, OrbitYaw, OrbitZoom);
+
+	// Ensure a minimum camera distance to prevent problems with really small objects
+	const float MinCameraDistance = 48;
+	OrbitZoom = FMath::Max<float>(MinCameraDistance, OrbitZoom);
+
+	const FRotator RotationOffsetToViewCenter(0.f, 90.f, 0.f);
+	FMatrix ViewRotationMatrix = FRotationMatrix( FRotator(0, OrbitYaw, 0) ) * 
+		FRotationMatrix( FRotator(0, 0, OrbitPitch) ) *
+		FTranslationMatrix( FVector(0, OrbitZoom, 0) ) *
+		FInverseRotationMatrix( RotationOffsetToViewCenter );
+
+	ViewRotationMatrix = ViewRotationMatrix * FMatrix(
+		FPlane(0,	0,	1,	0),
+		FPlane(1,	0,	0,	0),
+		FPlane(0,	1,	0,	0),
+		FPlane(0,	0,	0,	1));
+
+	Origin -= ViewRotationMatrix.InverseTransformPosition( FVector::ZeroVector );
+	ViewRotationMatrix = ViewRotationMatrix.RemoveTranslation();
+
+	FSceneViewInitOptions ViewInitOptions;
+	ViewInitOptions.ViewFamily = ViewFamily;
+	ViewInitOptions.SetViewRectangle(ViewRect);
+	ViewInitOptions.ViewOrigin = -Origin;
+	ViewInitOptions.ViewRotationMatrix = ViewRotationMatrix;
+	ViewInitOptions.ProjectionMatrix = ProjectionMatrix;
+	ViewInitOptions.BackgroundColor = FLinearColor::Black;
+
+	FSceneView* NewView = new FSceneView(ViewInitOptions);
+
+	ViewFamily->Views.Add(NewView);
+
+	NewView->StartFinalPostprocessSettings( ViewInitOptions.ViewOrigin );
+	NewView->EndFinalPostprocessSettings(ViewInitOptions);
+		
+	// Tell the texture streaming system about this thumbnail view, so the textures will stream in as needed
+	// NOTE: Sizes may not actually be in screen space depending on how the thumbnail ends up stretched by the UI.  Not a big deal though.
+	// NOTE: Textures still take a little time to stream if the view has not been re-rendered recently, so they may briefly appear blurry while mips are prepared
+	// NOTE: Content Browser only renders thumbnails for loaded assets, and only when the mouse is over the panel. They'll be frozen in their last state while the mouse cursor is not over the panel.  This is for performance reasons
+	IStreamingManager::Get().AddViewInformation( Origin, SizeX, SizeX / FMath::Tan( FOVDegrees ) );
+
+	return NewView;
 }
 
 void FThumbnailPreviewScene::Tick(float DeltaTime)
diff --git a/Engine/Source/Editor/UnrealEd/Private/ThumbnailRendering/AnimBlueprintThumbnailRenderer.cpp b/Engine/Source/Editor/UnrealEd/Private/ThumbnailRendering/AnimBlueprintThumbnailRenderer.cpp
--- a/Engine/Source/Editor/UnrealEd/Private/ThumbnailRendering/AnimBlueprintThumbnailRenderer.cpp
+++ b/Engine/Source/Editor/UnrealEd/Private/ThumbnailRendering/AnimBlueprintThumbnailRenderer.cpp
@@ -34,8 +34,7 @@ void UAnimBlueprintThumbnailRenderer::Draw(UObject* Object, int32 X, int32 Y, ui
 			ViewFamily.EngineShowFlags.MotionBlur = 0;
 			ViewFamily.EngineShowFlags.LOD = 0;
 
-			ThumbnailScene->GetView(&ViewFamily, X, Y, Width, Height);
-			RenderViewFamily(Canvas, &ViewFamily);
+			RenderViewFamily(Canvas, &ViewFamily, ThumbnailScene->CreateView(&ViewFamily, X, Y, Width, Height));
 		}
 	}
 }
diff --git a/Engine/Source/Editor/UnrealEd/Private/ThumbnailRendering/AnimSequenceThumbnailRenderer.cpp b/Engine/Source/Editor/UnrealEd/Private/ThumbnailRendering/AnimSequenceThumbnailRenderer.cpp
--- a/Engine/Source/Editor/UnrealEd/Private/ThumbnailRendering/AnimSequenceThumbnailRenderer.cpp
+++ b/Engine/Source/Editor/UnrealEd/Private/ThumbnailRendering/AnimSequenceThumbnailRenderer.cpp
@@ -33,8 +33,7 @@ void UAnimSequenceThumbnailRenderer::Draw(UObject* Object, int32 X, int32 Y, uin
 			ViewFamily.EngineShowFlags.MotionBlur = 0;
 			ViewFamily.EngineShowFlags.LOD = 0;
 
-			ThumbnailScene->GetView(&ViewFamily, X, Y, Width, Height);
-			RenderViewFamily(Canvas, &ViewFamily);
+			RenderViewFamily(Canvas, &ViewFamily, ThumbnailScene->CreateView(&ViewFamily, X, Y, Width, Height));
 			ThumbnailScene->SetAnimation(nullptr);
 		}
 	}
diff --git a/Engine/Source/Editor/UnrealEd/Private/ThumbnailRendering/BlendSpaceThumbnailRenderer.cpp b/Engine/Source/Editor/UnrealEd/Private/ThumbnailRendering/BlendSpaceThumbnailRenderer.cpp
--- a/Engine/Source/Editor/UnrealEd/Private/ThumbnailRendering/BlendSpaceThumbnailRenderer.cpp
+++ b/Engine/Source/Editor/UnrealEd/Private/ThumbnailRendering/BlendSpaceThumbnailRenderer.cpp
@@ -34,8 +34,7 @@ void UBlendSpaceThumbnailRenderer::Draw(UObject* Object, int32 X, int32 Y, uint3
 			ViewFamily.EngineShowFlags.MotionBlur = 0;
 			ViewFamily.EngineShowFlags.LOD = 0;
 
-			ThumbnailScene->GetView(&ViewFamily, X, Y, Width, Height);
-			RenderViewFamily(Canvas, &ViewFamily);
+			RenderViewFamily(Canvas, &ViewFamily, ThumbnailScene->CreateView(&ViewFamily, X, Y, Width, Height));
 			ThumbnailScene->SetBlendSpace(nullptr);
 		}
 	}
diff --git a/Engine/Source/Editor/UnrealEd/Private/ThumbnailRendering/BlueprintThumbnailRenderer.cpp b/Engine/Source/Editor/UnrealEd/Private/ThumbnailRendering/BlueprintThumbnailRenderer.cpp
--- a/Engine/Source/Editor/UnrealEd/Private/ThumbnailRendering/BlueprintThumbnailRenderer.cpp
+++ b/Engine/Source/Editor/UnrealEd/Private/ThumbnailRendering/BlueprintThumbnailRenderer.cpp
@@ -111,8 +111,7 @@ void UBlueprintThumbnailRenderer::Draw(UObject* Object, int32 X, int32 Y, uint32
 		ViewFamily.EngineShowFlags.DisableAdvancedFeatures();
 		ViewFamily.EngineShowFlags.MotionBlur = 0;
 
-		ThumbnailScene->GetView(&ViewFamily, X, Y, Width, Height);
-		RenderViewFamily(Canvas,&ViewFamily);
+		RenderViewFamily(Canvas, &ViewFamily, ThumbnailScene->CreateView(&ViewFamily, X, Y, Width, Height));
 	}
 }
 
diff --git a/Engine/Source/Editor/UnrealEd/Private/ThumbnailRendering/ClassThumbnailRenderer.cpp b/Engine/Source/Editor/UnrealEd/Private/ThumbnailRendering/ClassThumbnailRenderer.cpp
--- a/Engine/Source/Editor/UnrealEd/Private/ThumbnailRendering/ClassThumbnailRenderer.cpp
+++ b/Engine/Source/Editor/UnrealEd/Private/ThumbnailRendering/ClassThumbnailRenderer.cpp
@@ -47,8 +47,7 @@ void UClassThumbnailRenderer::Draw(UObject* Object, int32 X, int32 Y, uint32 Wid
 		ViewFamily.EngineShowFlags.DisableAdvancedFeatures();
 		ViewFamily.EngineShowFlags.MotionBlur = 0;
 
-		ThumbnailScene->GetView(&ViewFamily, X, Y, Width, Height);
-		RenderViewFamily(Canvas,&ViewFamily);
+		RenderViewFamily(Canvas, &ViewFamily, ThumbnailScene->CreateView(&ViewFamily, X, Y, Width, Height));
 	}
 }
 
diff --git a/Engine/Source/Editor/UnrealEd/Private/ThumbnailRendering/LevelThumbnailRenderer.cpp b/Engine/Source/Editor/UnrealEd/Private/ThumbnailRendering/LevelThumbnailRenderer.cpp
--- a/Engine/Source/Editor/UnrealEd/Private/ThumbnailRendering/LevelThumbnailRenderer.cpp
+++ b/Engine/Source/Editor/UnrealEd/Private/ThumbnailRendering/LevelThumbnailRenderer.cpp
@@ -28,23 +28,11 @@ void ULevelThumbnailRenderer::Draw(UObject* Object, int32 X, int32 Y, uint32 Wid
 		ViewFamily.EngineShowFlags.SetDistanceCulledPrimitives(true); // show distance culled objects
 		ViewFamily.EngineShowFlags.SetPostProcessing(false);
 
-		FSceneView* NewView = GetView(Level, &ViewFamily, X, Y, Width, Height);
-
-		ViewFamily.ViewExtensions = GEngine->ViewExtensions->GatherActiveExtensions(FSceneViewExtensionContext(Level->OwningWorld->Scene));
-		for (const FSceneViewExtensionRef& Extension : ViewFamily.ViewExtensions)
-		{
-			Extension->SetupViewFamily(ViewFamily);
-			Extension->SetupView(ViewFamily, *NewView);
-		}
-
-		if (NewView)
-		{
-			RenderViewFamily(Canvas, &ViewFamily);
-		}
+		RenderViewFamily(Canvas, &ViewFamily, CreateView(Level, &ViewFamily, X, Y, Width, Height));
 	}
 }
 
-FSceneView* ULevelThumbnailRenderer::GetView(ULevel* Level, FSceneViewFamily* ViewFamily, int32 X, int32 Y, uint32 SizeX, uint32 SizeY) const
+FSceneView* ULevelThumbnailRenderer::CreateView(ULevel* Level, FSceneViewFamily* ViewFamily, int32 X, int32 Y, uint32 SizeX, uint32 SizeY) const 
 {
 	check(ViewFamily);
 
@@ -55,7 +43,6 @@ FSceneView* ULevelThumbnailRenderer::GetView(ULevel* Level, FSceneViewFamily* Vi
 		FMath::Max<int32>(Y+SizeY,0));
 
 	FBox LevelBox(ForceInit);
-	FSceneView* NewView = nullptr;
 
 	if (Level->LevelBoundsActor.IsValid())
 	{
@@ -66,33 +53,33 @@ FSceneView* ULevelThumbnailRenderer::GetView(ULevel* Level, FSceneViewFamily* Vi
 		LevelBox = ALevelBounds::CalculateLevelBounds(Level);
 	}
 
-	if (ViewRect.Width() > 0 && ViewRect.Height() > 0)
+	if (ViewRect.Area() <= 0)
 	{
-		FSceneViewInitOptions ViewInitOptions;
-
-		ViewInitOptions.SetViewRectangle(ViewRect);
-		ViewInitOptions.ViewFamily = ViewFamily;
-
-		const FVector ViewPoint = LevelBox.GetCenter();
-		ViewInitOptions.ViewOrigin = FVector(ViewPoint.X, ViewPoint.Y, 0);
-		ViewInitOptions.ViewRotationMatrix = FMatrix(
-			FPlane(1,				0,				0,		0),
-			FPlane(0,				-1,				0,		0),
-			FPlane(0,				0,				-1,		0),
-			FPlane(0,				0,				0,		1));
-
-		const float ZOffset = WORLD_MAX;
-		ViewInitOptions.ProjectionMatrix =  FReversedZOrthoMatrix(
-			LevelBox.GetSize().X/2.f,
-			LevelBox.GetSize().Y/2.f,
-			0.5f / ZOffset,
-			ZOffset
-			);
-
-		NewView = new FSceneView(ViewInitOptions);
-
-		ViewFamily->Views.Add(NewView);
+		return nullptr;
 	}
 
+	FSceneViewInitOptions ViewInitOptions;
+
+	ViewInitOptions.SetViewRectangle(ViewRect);
+	ViewInitOptions.ViewFamily = ViewFamily;
+
+	const FVector ViewPoint = LevelBox.GetCenter();
+	ViewInitOptions.ViewOrigin = FVector(ViewPoint.X, ViewPoint.Y, 0);
+	ViewInitOptions.ViewRotationMatrix = FMatrix(
+		FPlane(1,				0,				0,		0),
+		FPlane(0,				-1,				0,		0),
+		FPlane(0,				0,				-1,		0),
+		FPlane(0,				0,				0,		1));
+
+	const float ZOffset = WORLD_MAX;
+	ViewInitOptions.ProjectionMatrix =  FReversedZOrthoMatrix(
+		LevelBox.GetSize().X/2.f,
+		LevelBox.GetSize().Y/2.f,
+		0.5f / ZOffset,
+		ZOffset
+		);
+
+	FSceneView* NewView = new FSceneView(ViewInitOptions);
+	ViewFamily->Views.Add(NewView);
 	return NewView;
 }
diff --git a/Engine/Source/Editor/UnrealEd/Private/ThumbnailRendering/MaterialFunctionThumbnailRenderer.cpp b/Engine/Source/Editor/UnrealEd/Private/ThumbnailRendering/MaterialFunctionThumbnailRenderer.cpp
--- a/Engine/Source/Editor/UnrealEd/Private/ThumbnailRendering/MaterialFunctionThumbnailRenderer.cpp
+++ b/Engine/Source/Editor/UnrealEd/Private/ThumbnailRendering/MaterialFunctionThumbnailRenderer.cpp
@@ -57,12 +57,7 @@ void UMaterialFunctionThumbnailRenderer::Draw(UObject* Object, int32 X, int32 Y,
 			ViewFamily.EngineShowFlags.DisableAdvancedFeatures();
 			ViewFamily.EngineShowFlags.MotionBlur = 0;
 
-			ThumbnailScene->GetView(&ViewFamily, X, Y, Width, Height);
-
-			if (ViewFamily.Views.Num() > 0)
-			{
-				RenderViewFamily(Canvas,&ViewFamily);
-			}
+			RenderViewFamily(Canvas, &ViewFamily, ThumbnailScene->CreateView(&ViewFamily, X, Y, Width, Height));
 
 			ThumbnailScene->SetMaterialInterface(nullptr);
 		}
diff --git a/Engine/Source/Editor/UnrealEd/Private/ThumbnailRendering/MaterialInstanceThumbnailRenderer.cpp b/Engine/Source/Editor/UnrealEd/Private/ThumbnailRendering/MaterialInstanceThumbnailRenderer.cpp
--- a/Engine/Source/Editor/UnrealEd/Private/ThumbnailRendering/MaterialInstanceThumbnailRenderer.cpp
+++ b/Engine/Source/Editor/UnrealEd/Private/ThumbnailRendering/MaterialInstanceThumbnailRenderer.cpp
@@ -90,17 +90,12 @@ void UMaterialInstanceThumbnailRenderer::Draw(UObject* Object, int32 X, int32 Y,
 				.SetTime(UThumbnailRenderer::GetTime())
 				.SetAdditionalViewFamily(bAdditionalViewFamily));
 
-
 			ViewFamily.EngineShowFlags.DisableAdvancedFeatures();
 			ViewFamily.EngineShowFlags.SetSeparateTranslucency(ThumbnailScene->ShouldSetSeparateTranslucency(MatInst));
 			ViewFamily.EngineShowFlags.MotionBlur = 0;
 			ViewFamily.EngineShowFlags.AntiAliasing = 0;
-			ThumbnailScene->GetView(&ViewFamily, X, Y, Width, Height);
 
-			if (ViewFamily.Views.Num() > 0)
-			{
-				RenderViewFamily(Canvas, &ViewFamily);
-			}
+			RenderViewFamily(Canvas, &ViewFamily, ThumbnailScene->CreateView(&ViewFamily, X, Y, Width, Height));
 
 			ThumbnailScene->SetMaterialInterface(nullptr);
 		}
diff --git a/Engine/Source/Editor/UnrealEd/Private/ThumbnailRendering/ParticleSystemThumbnailRenderer.cpp b/Engine/Source/Editor/UnrealEd/Private/ThumbnailRendering/ParticleSystemThumbnailRenderer.cpp
--- a/Engine/Source/Editor/UnrealEd/Private/ThumbnailRendering/ParticleSystemThumbnailRenderer.cpp
+++ b/Engine/Source/Editor/UnrealEd/Private/ThumbnailRendering/ParticleSystemThumbnailRenderer.cpp
@@ -83,8 +83,7 @@ void UParticleSystemThumbnailRenderer::Draw(UObject* Object, int32 X, int32 Y, u
 				ViewFamily.EngineShowFlags.DisableAdvancedFeatures();
 				ViewFamily.EngineShowFlags.MotionBlur = 0;
 
-				ThumbnailScene->GetView(&ViewFamily, X, Y, Width, Height);
-				RenderViewFamily(Canvas, &ViewFamily);
+				RenderViewFamily(Canvas, &ViewFamily, ThumbnailScene->CreateView(&ViewFamily, X, Y, Width, Height));
 				ThumbnailScene->SetParticleSystem(nullptr);
 			}
 			else if (ParticleSystem->ThumbnailImage)
diff --git a/Engine/Source/Editor/UnrealEd/Private/ThumbnailRendering/PhysicsAssetThumbnailRenderer.cpp b/Engine/Source/Editor/UnrealEd/Private/ThumbnailRendering/PhysicsAssetThumbnailRenderer.cpp
--- a/Engine/Source/Editor/UnrealEd/Private/ThumbnailRendering/PhysicsAssetThumbnailRenderer.cpp
+++ b/Engine/Source/Editor/UnrealEd/Private/ThumbnailRendering/PhysicsAssetThumbnailRenderer.cpp
@@ -32,8 +32,7 @@ void UPhysicsAssetThumbnailRenderer::Draw(UObject* Object, int32 X, int32 Y, uin
 		ViewFamily.EngineShowFlags.MotionBlur = 0;
 		ViewFamily.EngineShowFlags.LOD = 0;
 
-		ThumbnailScene->GetView(&ViewFamily, X, Y, Width, Height);
-		RenderViewFamily(Canvas,&ViewFamily);
+		RenderViewFamily(Canvas, &ViewFamily, ThumbnailScene->CreateView(&ViewFamily, X, Y, Width, Height));
 		ThumbnailScene->SetPhysicsAsset(nullptr);
 	}
 }
diff --git a/Engine/Source/Editor/UnrealEd/Private/ThumbnailRendering/SkeletalMeshThumbnailRenderer.cpp b/Engine/Source/Editor/UnrealEd/Private/ThumbnailRendering/SkeletalMeshThumbnailRenderer.cpp
--- a/Engine/Source/Editor/UnrealEd/Private/ThumbnailRendering/SkeletalMeshThumbnailRenderer.cpp
+++ b/Engine/Source/Editor/UnrealEd/Private/ThumbnailRendering/SkeletalMeshThumbnailRenderer.cpp
@@ -31,8 +31,7 @@ void USkeletalMeshThumbnailRenderer::Draw(UObject* Object, int32 X, int32 Y, uin
 	ViewFamily.EngineShowFlags.MotionBlur = 0;
 	ViewFamily.EngineShowFlags.LOD = 0;
 
-	ThumbnailScene->GetView(&ViewFamily, X, Y, Width, Height);
-	RenderViewFamily(Canvas,&ViewFamily);
+	RenderViewFamily(Canvas, &ViewFamily, ThumbnailScene->CreateView(&ViewFamily, X, Y, Width, Height));
 	ThumbnailScene->SetSkeletalMesh(nullptr);
 }
 
diff --git a/Engine/Source/Editor/UnrealEd/Private/ThumbnailRendering/SkeletonThumbnailRenderer.cpp b/Engine/Source/Editor/UnrealEd/Private/ThumbnailRendering/SkeletonThumbnailRenderer.cpp
--- a/Engine/Source/Editor/UnrealEd/Private/ThumbnailRendering/SkeletonThumbnailRenderer.cpp
+++ b/Engine/Source/Editor/UnrealEd/Private/ThumbnailRendering/SkeletonThumbnailRenderer.cpp
@@ -42,8 +42,7 @@ void USkeletonThumbnailRenderer::Draw(UObject* Object, int32 X, int32 Y, uint32
 	ViewFamily.EngineShowFlags.MotionBlur = 0;
 	ViewFamily.EngineShowFlags.LOD = 0;
 
-	ThumbnailScene->GetView(&ViewFamily, X, Y, Width, Height);
-	RenderViewFamily(Canvas,&ViewFamily);
+	RenderViewFamily(Canvas, &ViewFamily, ThumbnailScene->CreateView(&ViewFamily, X, Y, Width, Height));
 	ThumbnailScene->SetSkeletalMesh(nullptr);
 }
 
diff --git a/Engine/Source/Editor/UnrealEd/Private/ThumbnailRendering/StaticMeshThumbnailRenderer.cpp b/Engine/Source/Editor/UnrealEd/Private/ThumbnailRendering/StaticMeshThumbnailRenderer.cpp
--- a/Engine/Source/Editor/UnrealEd/Private/ThumbnailRendering/StaticMeshThumbnailRenderer.cpp
+++ b/Engine/Source/Editor/UnrealEd/Private/ThumbnailRendering/StaticMeshThumbnailRenderer.cpp
@@ -39,8 +39,7 @@ void UStaticMeshThumbnailRenderer::Draw(UObject* Object, int32 X, int32 Y, uint3
 		ViewFamily.EngineShowFlags.MotionBlur = 0;
 		ViewFamily.EngineShowFlags.LOD = 0;
 
-		ThumbnailScene->GetView(&ViewFamily, X, Y, Width, Height);
-		RenderViewFamily(Canvas,&ViewFamily);
+		RenderViewFamily(Canvas, &ViewFamily, ThumbnailScene->CreateView(&ViewFamily, X, Y, Width, Height));
 		ThumbnailScene->SetStaticMesh(nullptr);
 	}
 }
diff --git a/Engine/Source/Editor/UnrealEd/Private/ThumbnailRendering/ThumbnailRenderer.cpp b/Engine/Source/Editor/UnrealEd/Private/ThumbnailRendering/ThumbnailRenderer.cpp
--- a/Engine/Source/Editor/UnrealEd/Private/ThumbnailRendering/ThumbnailRenderer.cpp
+++ b/Engine/Source/Editor/UnrealEd/Private/ThumbnailRendering/ThumbnailRenderer.cpp
@@ -5,6 +5,7 @@
 #include "RendererInterface.h"
 #include "EngineModule.h"
 #include "SceneView.h"
+#include "SceneViewExtension.h"
 #include "LegacyScreenPercentageDriver.h"
 
 UThumbnailRenderer::UThumbnailRenderer(const FObjectInitializer& ObjectInitializer)
@@ -13,12 +14,43 @@ UThumbnailRenderer::UThumbnailRenderer(const FObjectInitializer& ObjectInitializ
 }
 
 // static
+PRAGMA_DISABLE_DEPRECATION_WARNINGS
 void UThumbnailRenderer::RenderViewFamily(FCanvas* Canvas, FSceneViewFamily* ViewFamily)
 {
+	if ((ViewFamily == nullptr) || ViewFamily->Views.IsEmpty() || (ViewFamily->Views[0] == nullptr))
+	{
+		return;
+	}
+
+	// The new prototype of RenderViewFamily takes a non-const FSceneView in parameter since view extensions can still modify it as well as as the view family before rendering : 
+	FSceneView* View = const_cast<FSceneView*>(ViewFamily->Views[0]);
+	UThumbnailRenderer::RenderViewFamily(Canvas, ViewFamily, View);
+}
+PRAGMA_ENABLE_DEPRECATION_WARNINGS
+
+// static
+void UThumbnailRenderer::RenderViewFamily(FCanvas* Canvas, FSceneViewFamily* ViewFamily, FSceneView* View)
+{
+	if ((ViewFamily == nullptr) || (View == nullptr))
+	{ 
+		return;
+	}
+
+	check(ViewFamily->Views.Num() == 1 && (ViewFamily->Views[0] == View));
+
 	ViewFamily->EngineShowFlags.ScreenPercentage = false;
 	ViewFamily->SetScreenPercentageInterface(new FLegacyScreenPercentageDriver(
 		*ViewFamily, /* GlobalResolutionFraction = */ 1.0f));
 
+	ViewFamily->ViewExtensions = GEngine->ViewExtensions->GatherActiveExtensions(FSceneViewExtensionContext(ViewFamily->Scene));
+
+	// View extensions should have a chance at changing ViewFamily and View (while it's still mutable, since ViewFamily->Views contains const FSceneView pointers) before rendering :
+	for (const FSceneViewExtensionRef& Extension : ViewFamily->ViewExtensions)
+	{
+		Extension->SetupViewFamily(*ViewFamily);
+		Extension->SetupView(*ViewFamily, *View);
+	}
+
 	GetRendererModule().BeginRenderingViewFamily(Canvas, ViewFamily);
 }
 
diff --git a/Engine/Source/Editor/UnrealEd/Private/ThumbnailRendering/VolumeTextureThumbnailRenderer.cpp b/Engine/Source/Editor/UnrealEd/Private/ThumbnailRendering/VolumeTextureThumbnailRenderer.cpp
--- a/Engine/Source/Editor/UnrealEd/Private/ThumbnailRendering/VolumeTextureThumbnailRenderer.cpp
+++ b/Engine/Source/Editor/UnrealEd/Private/ThumbnailRendering/VolumeTextureThumbnailRenderer.cpp
@@ -31,7 +31,7 @@ public:
 	/** Constructor */
 	FVolumeTextureThumbnailScene();
 
-	/** Sets the material to use in the next GetView() */
+	/** Sets the material to use in the next CreateView() */
 	void SetMaterialInterface(UMaterialInstance* InMaterial);
 
 protected:
@@ -169,12 +169,7 @@ void UVolumeTextureThumbnailRenderer::Draw(UObject* Object, int32 X, int32 Y, ui
 			ViewFamily.EngineShowFlags.MotionBlur = 0;
 			ViewFamily.EngineShowFlags.AntiAliasing = 0;
 
-			ThumbnailScene->GetView(&ViewFamily, X, Y, Width, Height);
-
-			if (ViewFamily.Views.Num() > 0)
-			{
-				RenderViewFamily(Canvas, &ViewFamily);
-			}
+			RenderViewFamily(Canvas, &ViewFamily, ThumbnailScene->CreateView(&ViewFamily, X, Y, Width, Height));
 		}
 
 		ThumbnailScene->SetMaterialInterface(nullptr);
diff --git a/Engine/Source/Editor/UnrealEd/Private/ThumbnailRendering/WorldThumbnailRenderer.cpp b/Engine/Source/Editor/UnrealEd/Private/ThumbnailRendering/WorldThumbnailRenderer.cpp
--- a/Engine/Source/Editor/UnrealEd/Private/ThumbnailRendering/WorldThumbnailRenderer.cpp
+++ b/Engine/Source/Editor/UnrealEd/Private/ThumbnailRendering/WorldThumbnailRenderer.cpp
@@ -45,8 +45,6 @@ void UWorldThumbnailRenderer::Draw(UObject* Object, int32 X, int32 Y, uint32 Wid
 			.SetTime(UThumbnailRenderer::GetTime())
 			.SetAdditionalViewFamily(bAdditionalViewFamily));
 
-		ViewFamily.ViewExtensions = GEngine->ViewExtensions->GatherActiveExtensions(FSceneViewExtensionContext(World->Scene));
-
 		ViewFamily.EngineShowFlags.SetDiffuse(true);
 		ViewFamily.EngineShowFlags.SetSkeletalMeshes(true);
 		ViewFamily.EngineShowFlags.SetTranslucency(true);
@@ -77,16 +75,11 @@ void UWorldThumbnailRenderer::Draw(UObject* Object, int32 X, int32 Y, uint32 Wid
 			ViewFamily.EngineShowFlags.SetReflectionEnvironment(true);
 		}
 
-		GetView(World, &ViewFamily, X, Y, Width, Height);
-
-		if (ViewFamily.Views.Num() > 0)
-		{
-			RenderViewFamily(Canvas, &ViewFamily);
-		}
+		RenderViewFamily(Canvas, &ViewFamily, CreateView(World, &ViewFamily, X, Y, Width, Height));
 	}
 }
 
-void UWorldThumbnailRenderer::GetView(UWorld* World, FSceneViewFamily* ViewFamily, int32 X, int32 Y, uint32 SizeX, uint32 SizeY) const 
+FSceneView* UWorldThumbnailRenderer::CreateView(UWorld* World, FSceneViewFamily* ViewFamily, int32 X, int32 Y, uint32 SizeX, uint32 SizeY) const 
 {
 	check(ViewFamily);
 	check(World);
@@ -98,216 +91,220 @@ void UWorldThumbnailRenderer::GetView(UWorld* World, FSceneViewFamily* ViewFamil
 		FMath::Max<int32>(X + SizeX, 0),
 		FMath::Max<int32>(Y + SizeY, 0));
 
-	if (ViewRect.Width() > 0 && ViewRect.Height() > 0)
+	if (ViewRect.Area() <= 0)
+	{
+		return nullptr;
+	}
+
+	FBox WorldBox(ForceInit);
+	TArray<ULevel*> LevelsToRender = World->GetLevels();
+	for ( ULevel* Level : LevelsToRender )
 	{
-		FBox WorldBox(ForceInit);
-		TArray<ULevel*> LevelsToRender = World->GetLevels();
-		for ( ULevel* Level : LevelsToRender )
+		if (Level && Level->bIsVisible)
 		{
-			if (Level && Level->bIsVisible)
-			{
-				ALevelBounds* LevelBounds = Level->LevelBoundsActor.Get();
+			ALevelBounds* LevelBounds = Level->LevelBoundsActor.Get();
 				
-				if (!LevelBounds && !Level->HasAllFlags(RF_Transient))
-				{
-					// Ensure a Level Bounds Actor exists for future renders
-					FActorSpawnParameters SpawnParameters;
-					SpawnParameters.OverrideLevel = Level;
-					LevelBounds = World->SpawnActor<ALevelBounds>(SpawnParameters);
-					LevelBounds->UpdateLevelBoundsImmediately();
-					Level->LevelBoundsActor = LevelBounds;
-				}
+			if (!LevelBounds && !Level->HasAllFlags(RF_Transient))
+			{
+				// Ensure a Level Bounds Actor exists for future renders
+				FActorSpawnParameters SpawnParameters;
+				SpawnParameters.OverrideLevel = Level;
+				LevelBounds = World->SpawnActor<ALevelBounds>(SpawnParameters);
+				LevelBounds->UpdateLevelBoundsImmediately();
+				Level->LevelBoundsActor = LevelBounds;
+			}
 
-				if (LevelBounds)
+			if (LevelBounds)
+			{
+				if(!LevelBounds->IsUsingDefaultBounds())
 				{
-					if(!LevelBounds->IsUsingDefaultBounds())
-					{
-						WorldBox += LevelBounds->GetComponentsBoundingBox();
-					}
+					WorldBox += LevelBounds->GetComponentsBoundingBox();
 				}
-				else
+			}
+			else
+			{
+				FBox LevelBox = ALevelBounds::CalculateLevelBounds(Level);
+				if (LevelBox.IsValid)
 				{
-					FBox LevelBox = ALevelBounds::CalculateLevelBounds(Level);
-					if (LevelBox.IsValid)
-					{
-						WorldBox += LevelBox;
-					}
+					WorldBox += LevelBox;
 				}
 			}
 		}
+	}
 
-		UWorldThumbnailInfo* ThumbnailInfo = Cast<UWorldThumbnailInfo>(World->ThumbnailInfo);
-		if (!ThumbnailInfo)
-		{
-			ThumbnailInfo = UWorldThumbnailInfo::StaticClass()->GetDefaultObject<UWorldThumbnailInfo>();
-		}
-
-		const FVector Origin = WorldBox.GetCenter();
-		FMatrix ViewRotationMatrix;
-		FMatrix ProjectionMatrix;
-		float FOVScreenSize = 0; // Screen size taking FOV into account
-		if (ThumbnailInfo->CameraMode == ECameraProjectionMode::Perspective)
-		{
-			const float FOVDegrees = 30.f;
-			const float HalfFOVRadians = FMath::DegreesToRadians<float>(FOVDegrees) * 0.5f;
-			const float WorldRadius = WorldBox.GetSize().Size() / 2.f;
-			float TargetDistance = WorldRadius / FMath::Tan(HalfFOVRadians);
-
-			if (TargetDistance + ThumbnailInfo->OrbitZoom < 0)
-			{
-				ThumbnailInfo->OrbitZoom = -TargetDistance;
-			}
+	UWorldThumbnailInfo* ThumbnailInfo = Cast<UWorldThumbnailInfo>(World->ThumbnailInfo);
+	if (!ThumbnailInfo)
+	{
+		ThumbnailInfo = UWorldThumbnailInfo::StaticClass()->GetDefaultObject<UWorldThumbnailInfo>();
+	}
 
-			float OrbitPitch = GlobalOrbitPitchOffset + ThumbnailInfo->OrbitPitch;
-			float OrbitYaw = GlobalOrbitYawOffset + ThumbnailInfo->OrbitYaw;
-			float OrbitZoom = TargetDistance + ThumbnailInfo->OrbitZoom;
-
-			// Ensure a minimum camera distance to prevent problems with really small objects
-			const float MinCameraDistance = 48;
-			OrbitZoom = FMath::Max<float>(MinCameraDistance, OrbitZoom);
-
-			const FRotator RotationOffsetToViewCenter(0.f, 90.f, 0.f);
-			ViewRotationMatrix = FRotationMatrix(FRotator(0, OrbitYaw, 0)) *
-				FRotationMatrix(FRotator(0, 0, OrbitPitch)) *
-				FTranslationMatrix(FVector(0, OrbitZoom, 0)) *
-				FInverseRotationMatrix(RotationOffsetToViewCenter);
-
-			ViewRotationMatrix = ViewRotationMatrix * FMatrix(
-				FPlane(0, 0, 1, 0),
-				FPlane(1, 0, 0, 0),
-				FPlane(0, 1, 0, 0),
-				FPlane(0, 0, 0, 1));
-
-			const float NearPlane = 1.0f;
-			ProjectionMatrix = FReversedZPerspectiveMatrix(
-				HalfFOVRadians,
-				1.0f,
-				1.0f,
-				NearPlane
-				);
-
-			FOVScreenSize = SizeX / FMath::Tan(FOVDegrees);
-		}
-		else if (ThumbnailInfo->CameraMode == ECameraProjectionMode::Orthographic)
-		{
-			FVector2D WorldSizeMin2D;
-			FVector2D WorldSizeMax2D;
-			switch (ThumbnailInfo->OrthoDirection)
-			{
-				case EOrthoThumbnailDirection::Top:
-					ViewRotationMatrix = FMatrix(
-						FPlane(1, 0, 0, 0),
-						FPlane(0, -1, 0, 0),
-						FPlane(0, 0, -1, 0),
-						FPlane(0, 0, Origin.Z, 1));
-					WorldSizeMin2D = FVector2D(WorldBox.Min.X,WorldBox.Min.Y);
-					WorldSizeMax2D = FVector2D(WorldBox.Max.X,WorldBox.Max.Y);
-					break;
-				case EOrthoThumbnailDirection::Bottom:
-					ViewRotationMatrix = FMatrix(
-						FPlane(1, 0, 0, 0),
-						FPlane(0, -1, 0, 0),
-						FPlane(0, 0, 1, 0),
-						FPlane(0, 0, Origin.Z, 1));
-					WorldSizeMin2D = FVector2D(WorldBox.Min.X, WorldBox.Min.Y);
-					WorldSizeMax2D = FVector2D(WorldBox.Max.X, WorldBox.Max.Y);
-					break;
-				case EOrthoThumbnailDirection::Front:
-					ViewRotationMatrix = FMatrix(
-						FPlane(1, 0, 0, 0),
-						FPlane(0, 0, -1, 0),
-						FPlane(0, 1, 0, 0),
-						FPlane(0, 0, Origin.Y, 1));
-					WorldSizeMin2D = FVector2D(WorldBox.Min.X, WorldBox.Min.Z);
-					WorldSizeMax2D = FVector2D(WorldBox.Max.X, WorldBox.Max.Z);
-					break;
-				case EOrthoThumbnailDirection::Back:
-					ViewRotationMatrix = FMatrix(
-						FPlane(-1, 0, 0, 0),
-						FPlane(0, 0, 1, 0),
-						FPlane(0, 1, 0, 0),
-						FPlane(0, 0, Origin.Y, 1));
-					WorldSizeMin2D = FVector2D(WorldBox.Min.X, WorldBox.Min.Z);
-					WorldSizeMax2D = FVector2D(WorldBox.Max.X, WorldBox.Max.Z);
-					break;
-				case EOrthoThumbnailDirection::Left:
-					ViewRotationMatrix = FMatrix(
-						FPlane(0, 0, -1, 0),
-						FPlane(-1, 0, 0, 0),
-						FPlane(0, 1, 0, 0),
-						FPlane(0, 0, Origin.X, 1));
-					WorldSizeMin2D = FVector2D(WorldBox.Min.Y, WorldBox.Min.Z);
-					WorldSizeMax2D = FVector2D(WorldBox.Max.Y, WorldBox.Max.Z);
-					break;
-				case EOrthoThumbnailDirection::Right:
-					ViewRotationMatrix = FMatrix(
-						FPlane(0, 0, 1, 0),
-						FPlane(1, 0, 0, 0),
-						FPlane(0, 1, 0, 0),
-						FPlane(0, 0, Origin.X, 1));
-					WorldSizeMin2D = FVector2D(WorldBox.Min.Y, WorldBox.Min.Z);
-					WorldSizeMax2D = FVector2D(WorldBox.Max.Y, WorldBox.Max.Z);
-					break;
-				default:
-					// Unknown OrthoDirection
-					ensureMsgf(false, TEXT("Unknown thumbnail OrthoDirection"));
-					break;
-			}
+	const FVector Origin = WorldBox.GetCenter();
+	FMatrix ViewRotationMatrix;
+	FMatrix ProjectionMatrix;
+	float FOVScreenSize = 0; // Screen size taking FOV into account
+	if (ThumbnailInfo->CameraMode == ECameraProjectionMode::Perspective)
+	{
+		const float FOVDegrees = 30.f;
+		const float HalfFOVRadians = FMath::DegreesToRadians<float>(FOVDegrees) * 0.5f;
+		const float WorldRadius = WorldBox.GetSize().Size() / 2.f;
+		float TargetDistance = WorldRadius / FMath::Tan(HalfFOVRadians);
 
-			FVector2D WorldSize2D = (WorldSizeMax2D - WorldSizeMin2D);
-			WorldSize2D.X = FMath::Abs(WorldSize2D.X);
-			WorldSize2D.Y = FMath::Abs(WorldSize2D.Y);
-			const bool bUseXAxis = (WorldSize2D.X / WorldSize2D.Y) > 1.f;
-			const float WorldAxisSize = bUseXAxis ? WorldSize2D.X : WorldSize2D.Y;
-			const uint32 ViewportAxisSize = bUseXAxis ? SizeX : SizeY;
-			const float OrthoZoom = WorldAxisSize / ViewportAxisSize / 2.f;
-			const float OrthoWidth = FMath::Max(1.f, SizeX * OrthoZoom);
-			const float OrthoHeight = FMath::Max(1.f, SizeY * OrthoZoom);
-
-			const float ZOffset = HALF_WORLD_MAX;
-			ProjectionMatrix = FReversedZOrthoMatrix(
-				OrthoWidth,
-				OrthoHeight,
-				0.5f / ZOffset,
-				ZOffset
-				);
-
-			FOVScreenSize = SizeX;
-		}
-		else
+		if (TargetDistance + ThumbnailInfo->OrbitZoom < 0)
 		{
-			// Unknown CameraMode
-			ensureMsgf(false, TEXT("Unknown thumbnail CameraMode"));
+			ThumbnailInfo->OrbitZoom = -TargetDistance;
 		}
 
-		ensureMsgf(!ViewRotationMatrix.ContainsNaN(), TEXT("Nans found on ViewRotationMatrix"));
-		if (ViewRotationMatrix.ContainsNaN())
+		float OrbitPitch = GlobalOrbitPitchOffset + ThumbnailInfo->OrbitPitch;
+		float OrbitYaw = GlobalOrbitYawOffset + ThumbnailInfo->OrbitYaw;
+		float OrbitZoom = TargetDistance + ThumbnailInfo->OrbitZoom;
+
+		// Ensure a minimum camera distance to prevent problems with really small objects
+		const float MinCameraDistance = 48;
+		OrbitZoom = FMath::Max<float>(MinCameraDistance, OrbitZoom);
+
+		const FRotator RotationOffsetToViewCenter(0.f, 90.f, 0.f);
+		ViewRotationMatrix = FRotationMatrix(FRotator(0, OrbitYaw, 0)) *
+			FRotationMatrix(FRotator(0, 0, OrbitPitch)) *
+			FTranslationMatrix(FVector(0, OrbitZoom, 0)) *
+			FInverseRotationMatrix(RotationOffsetToViewCenter);
+
+		ViewRotationMatrix = ViewRotationMatrix * FMatrix(
+			FPlane(0, 0, 1, 0),
+			FPlane(1, 0, 0, 0),
+			FPlane(0, 1, 0, 0),
+			FPlane(0, 0, 0, 1));
+
+		const float NearPlane = 1.0f;
+		ProjectionMatrix = FReversedZPerspectiveMatrix(
+			HalfFOVRadians,
+			1.0f,
+			1.0f,
+			NearPlane
+			);
+
+		FOVScreenSize = SizeX / FMath::Tan(FOVDegrees);
+	}
+	else if (ThumbnailInfo->CameraMode == ECameraProjectionMode::Orthographic)
+	{
+		FVector2D WorldSizeMin2D;
+		FVector2D WorldSizeMax2D;
+		switch (ThumbnailInfo->OrthoDirection)
 		{
-			ViewRotationMatrix.SetIdentity();
+			case EOrthoThumbnailDirection::Top:
+				ViewRotationMatrix = FMatrix(
+					FPlane(1, 0, 0, 0),
+					FPlane(0, -1, 0, 0),
+					FPlane(0, 0, -1, 0),
+					FPlane(0, 0, Origin.Z, 1));
+				WorldSizeMin2D = FVector2D(WorldBox.Min.X,WorldBox.Min.Y);
+				WorldSizeMax2D = FVector2D(WorldBox.Max.X,WorldBox.Max.Y);
+				break;
+			case EOrthoThumbnailDirection::Bottom:
+				ViewRotationMatrix = FMatrix(
+					FPlane(1, 0, 0, 0),
+					FPlane(0, -1, 0, 0),
+					FPlane(0, 0, 1, 0),
+					FPlane(0, 0, Origin.Z, 1));
+				WorldSizeMin2D = FVector2D(WorldBox.Min.X, WorldBox.Min.Y);
+				WorldSizeMax2D = FVector2D(WorldBox.Max.X, WorldBox.Max.Y);
+				break;
+			case EOrthoThumbnailDirection::Front:
+				ViewRotationMatrix = FMatrix(
+					FPlane(1, 0, 0, 0),
+					FPlane(0, 0, -1, 0),
+					FPlane(0, 1, 0, 0),
+					FPlane(0, 0, Origin.Y, 1));
+				WorldSizeMin2D = FVector2D(WorldBox.Min.X, WorldBox.Min.Z);
+				WorldSizeMax2D = FVector2D(WorldBox.Max.X, WorldBox.Max.Z);
+				break;
+			case EOrthoThumbnailDirection::Back:
+				ViewRotationMatrix = FMatrix(
+					FPlane(-1, 0, 0, 0),
+					FPlane(0, 0, 1, 0),
+					FPlane(0, 1, 0, 0),
+					FPlane(0, 0, Origin.Y, 1));
+				WorldSizeMin2D = FVector2D(WorldBox.Min.X, WorldBox.Min.Z);
+				WorldSizeMax2D = FVector2D(WorldBox.Max.X, WorldBox.Max.Z);
+				break;
+			case EOrthoThumbnailDirection::Left:
+				ViewRotationMatrix = FMatrix(
+					FPlane(0, 0, -1, 0),
+					FPlane(-1, 0, 0, 0),
+					FPlane(0, 1, 0, 0),
+					FPlane(0, 0, Origin.X, 1));
+				WorldSizeMin2D = FVector2D(WorldBox.Min.Y, WorldBox.Min.Z);
+				WorldSizeMax2D = FVector2D(WorldBox.Max.Y, WorldBox.Max.Z);
+				break;
+			case EOrthoThumbnailDirection::Right:
+				ViewRotationMatrix = FMatrix(
+					FPlane(0, 0, 1, 0),
+					FPlane(1, 0, 0, 0),
+					FPlane(0, 1, 0, 0),
+					FPlane(0, 0, Origin.X, 1));
+				WorldSizeMin2D = FVector2D(WorldBox.Min.Y, WorldBox.Min.Z);
+				WorldSizeMax2D = FVector2D(WorldBox.Max.Y, WorldBox.Max.Z);
+				break;
+			default:
+				// Unknown OrthoDirection
+				ensureMsgf(false, TEXT("Unknown thumbnail OrthoDirection"));
+				break;
 		}
 
-		ensureMsgf(!ProjectionMatrix.ContainsNaN(), TEXT("Nans found on ProjectionMatrix"));
-		if (ProjectionMatrix.ContainsNaN())
-		{
-			ProjectionMatrix.SetIdentity();
-		}
+		FVector2D WorldSize2D = (WorldSizeMax2D - WorldSizeMin2D);
+		WorldSize2D.X = FMath::Abs(WorldSize2D.X);
+		WorldSize2D.Y = FMath::Abs(WorldSize2D.Y);
+		const bool bUseXAxis = (WorldSize2D.X / WorldSize2D.Y) > 1.f;
+		const float WorldAxisSize = bUseXAxis ? WorldSize2D.X : WorldSize2D.Y;
+		const uint32 ViewportAxisSize = bUseXAxis ? SizeX : SizeY;
+		const float OrthoZoom = WorldAxisSize / ViewportAxisSize / 2.f;
+		const float OrthoWidth = FMath::Max(1.f, SizeX * OrthoZoom);
+		const float OrthoHeight = FMath::Max(1.f, SizeY * OrthoZoom);
+
+		const float ZOffset = HALF_WORLD_MAX;
+		ProjectionMatrix = FReversedZOrthoMatrix(
+			OrthoWidth,
+			OrthoHeight,
+			0.5f / ZOffset,
+			ZOffset
+			);
+
+		FOVScreenSize = SizeX;
+	}
+	else
+	{
+		// Unknown CameraMode
+		ensureMsgf(false, TEXT("Unknown thumbnail CameraMode"));
+	}
 
-		FSceneViewInitOptions ViewInitOptions;
-		ViewInitOptions.ViewFamily = ViewFamily;
-		ViewInitOptions.SetViewRectangle(ViewRect);
-		ViewInitOptions.BackgroundColor = FLinearColor::Black;
-		ViewInitOptions.ViewOrigin = Origin;
-		ViewInitOptions.ViewRotationMatrix = ViewRotationMatrix;
-		ViewInitOptions.ProjectionMatrix = ProjectionMatrix;
+	ensureMsgf(!ViewRotationMatrix.ContainsNaN(), TEXT("Nans found on ViewRotationMatrix"));
+	if (ViewRotationMatrix.ContainsNaN())
+	{
+		ViewRotationMatrix.SetIdentity();
+	}
 
-		FSceneView* NewView = new FSceneView(ViewInitOptions);
+	ensureMsgf(!ProjectionMatrix.ContainsNaN(), TEXT("Nans found on ProjectionMatrix"));
+	if (ProjectionMatrix.ContainsNaN())
+	{
+		ProjectionMatrix.SetIdentity();
+	}
 
-		ViewFamily->Views.Add(NewView);
+	FSceneViewInitOptions ViewInitOptions;
+	ViewInitOptions.ViewFamily = ViewFamily;
+	ViewInitOptions.SetViewRectangle(ViewRect);
+	ViewInitOptions.BackgroundColor = FLinearColor::Black;
+	ViewInitOptions.ViewOrigin = Origin;
+	ViewInitOptions.ViewRotationMatrix = ViewRotationMatrix;
+	ViewInitOptions.ProjectionMatrix = ProjectionMatrix;
 
-		// Tell the texture streaming system about this thumbnail view, so the textures will stream in as needed
-		// NOTE: Sizes may not actually be in screen space depending on how the thumbnail ends up stretched by the UI.  Not a big deal though.
-		// NOTE: Textures still take a little time to stream if the view has not been re-rendered recently, so they may briefly appear blurry while mips are prepared
-		// NOTE: Content Browser only renders thumbnails for loaded assets, and only when the mouse is over the panel. They'll be frozen in their last state while the mouse cursor is not over the panel.  This is for performance reasons
-		IStreamingManager::Get().AddViewInformation(Origin, SizeX, FOVScreenSize);
-	}
+	FSceneView* NewView = new FSceneView(ViewInitOptions);
+
+	ViewFamily->Views.Add(NewView);
+
+	// Tell the texture streaming system about this thumbnail view, so the textures will stream in as needed
+	// NOTE: Sizes may not actually be in screen space depending on how the thumbnail ends up stretched by the UI.  Not a big deal though.
+	// NOTE: Textures still take a little time to stream if the view has not been re-rendered recently, so they may briefly appear blurry while mips are prepared
+	// NOTE: Content Browser only renders thumbnails for loaded assets, and only when the mouse is over the panel. They'll be frozen in their last state while the mouse cursor is not over the panel.  This is for performance reasons
+	IStreamingManager::Get().AddViewInformation(Origin, SizeX, FOVScreenSize);
+
+	return NewView;
 }
diff --git a/Engine/Source/Editor/UnrealEd/Public/ThumbnailHelpers.h b/Engine/Source/Editor/UnrealEd/Public/ThumbnailHelpers.h
--- a/Engine/Source/Editor/UnrealEd/Public/ThumbnailHelpers.h
+++ b/Engine/Source/Editor/UnrealEd/Public/ThumbnailHelpers.h
@@ -21,9 +21,12 @@ public:
 	/** Constructor */
 	FThumbnailPreviewScene();
 
-	/** Allocates then adds an FSceneView to the ViewFamily. */
+	UE_DEPRECATED(5.0, "Use CreateView")
 	void GetView(FSceneViewFamily* ViewFamily, int32 X, int32 Y, uint32 SizeX, uint32 SizeY) const;
 
+	/** Allocates then adds an FSceneView to the ViewFamily. */
+	FSceneView* CreateView(FSceneViewFamily* ViewFamily, int32 X, int32 Y, uint32 SizeX, uint32 SizeY) const;
+
 	/* Begin FTickableEditorObject */
 	virtual void Tick(float DeltaTime) override;
 	virtual ETickableTickType GetTickableTickType() const override { return ETickableTickType::Always; }
@@ -35,7 +38,7 @@ protected:
 	float GetBoundsZOffset(const FBoxSphereBounds& Bounds) const;
 
 	/**
-	  * Gets parameters to create a view matrix to be used by GetView(). Implemented in children classes.
+	  * Gets parameters to create a view matrix to be used by CreateView(). Implemented in children classes.
 	  * @param InFOVDegrees  The FOV used to display the thumbnail. Often used to calculate the output parameters.
 	  * @param OutOrigin	 The origin of the orbit view. Typically the center of the bounds of the target object.
 	  * @param OutOrbitPitch The pitch of the orbit cam around the object.
@@ -52,7 +55,7 @@ public:
 	FParticleSystemThumbnailScene();
 	virtual ~FParticleSystemThumbnailScene();
 
-	/** Sets the particle system to use in the next GetView() */
+	/** Sets the particle system to use in the next CreateView() */
 	void SetParticleSystem(class UParticleSystem* ParticleSystem);
 
 protected:
@@ -73,7 +76,7 @@ public:
 	/** Constructor */
 	FMaterialThumbnailScene();
 
-	/** Sets the material to use in the next GetView() */
+	/** Sets the material to use in the next CreateView() */
 	void SetMaterialInterface(class UMaterialInterface* InMaterial);
 
 	bool ShouldSetSeparateTranslucency(class UMaterialInterface* InMaterial) const;
@@ -95,7 +98,7 @@ public:
 	/** Constructor */
 	FSkeletalMeshThumbnailScene();
 
-	/** Sets the skeletal mesh to use in the next GetView() */
+	/** Sets the skeletal mesh to use in the next CreateView() */
 	void SetSkeletalMesh(class USkeletalMesh* InSkeletalMesh);
 
 	/** Set whether to draw debug skeleton */
@@ -125,7 +128,7 @@ public:
 	/** Constructor */
 	FStaticMeshThumbnailScene();
 
-	/** Sets the static mesh to use in the next GetView() */
+	/** Sets the static mesh to use in the next CreateView() */
 	void SetStaticMesh(class UStaticMesh* StaticMesh);
 
 	/** Sets override materials for the static mesh  */
@@ -152,7 +155,7 @@ public:
 	/** Constructor */
 	FAnimationSequenceThumbnailScene();
 
-	/** Sets the animation to use in the next GetView() */
+	/** Sets the animation to use in the next CreateView() */
 	bool SetAnimation(class UAnimSequenceBase* InAnimation);
 
 protected:
@@ -176,7 +179,7 @@ public:
 	/** Constructor */
 	FBlendSpaceThumbnailScene();
 
-	/** Sets the animation to use in the next GetView() */
+	/** Sets the animation to use in the next CreateView() */
 	bool SetBlendSpace(class UBlendSpace* InBlendSpace);
 
 protected:
@@ -200,7 +203,7 @@ public:
 	/** Constructor */
 	FAnimBlueprintThumbnailScene();
 
-	/** Sets the animation blueprint to use in the next GetView() */
+	/** Sets the animation blueprint to use in the next CreateView() */
 	bool SetAnimBlueprint(class UAnimBlueprint* InBlueprint);
 
 protected:
@@ -224,7 +227,7 @@ public:
 	/** Constructor */
 	FPhysicsAssetThumbnailScene();
 
-	/** Sets the skeletal mesh to use in the next GetView() */
+	/** Sets the skeletal mesh to use in the next CreateView() */
 	void SetPhysicsAsset(class UPhysicsAsset* InPhysicsAsset);
 
 protected:
@@ -251,7 +254,7 @@ protected:
 	// FThumbnailPreviewScene implementation
 	virtual void GetViewMatrixParameters(const float InFOVDegrees, FVector& OutOrigin, float& OutOrbitPitch, float& OutOrbitYaw, float& OutOrbitZoom) const override;
 
-	/** Sets the object (class or blueprint) used in the next GetView() */
+	/** Sets the object (class or blueprint) used in the next CreateView() */
 	void SpawnPreviewActor(class UClass* Obj);
 
 	/** Get the scene thumbnail info to use for the object currently being rendered */
@@ -274,7 +277,7 @@ public:
 	/** Constructor/Destructor */
 	FBlueprintThumbnailScene();
 
-	/** Sets the static mesh to use in the next GetView() */
+	/** Sets the static mesh to use in the next CreateView() */
 	void SetBlueprint(class UBlueprint* Blueprint);
 
 	/** Refreshes components for the specified blueprint */
@@ -296,7 +299,7 @@ public:
 	/** Constructor/Destructor */
 	FClassThumbnailScene();
 
-	/** Sets the class use in the next GetView() */
+	/** Sets the class use in the next CreateView() */
 	void SetClass(class UClass* Class);
 
 protected:
diff --git a/Engine/Source/Runtime/Engine/Private/GameViewportClient.cpp b/Engine/Source/Runtime/Engine/Private/GameViewportClient.cpp
--- a/Engine/Source/Runtime/Engine/Private/GameViewportClient.cpp
+++ b/Engine/Source/Runtime/Engine/Private/GameViewportClient.cpp
@@ -1694,7 +1694,16 @@ void UGameViewportClient::Draw(FViewport* InViewport, FCanvas* SceneCanvas)
 	// Draw the player views.
 	if (!bDisableWorldRendering && PlayerViewMap.Num() > 0 && FSlateApplication::Get().GetPlatformApplication()->IsAllowedToRender()) //-V560
 	{
-		GetRendererModule().BeginRenderingViewFamily(SceneCanvas,&ViewFamily);
+		// Views 
+		for (auto ViewExt : ViewFamily.ViewExtensions)
+		{
+			for (FSceneView* View : Views)
+			{
+				ViewExt->SetupView(ViewFamily, *View);
+			}
+		}
+
+		GetRendererModule().BeginRenderingViewFamily(SceneCanvas, &ViewFamily);
 	}
 	else
 	{
diff --git a/Engine/Source/Runtime/Renderer/Private/PlanarReflectionRendering.cpp b/Engine/Source/Runtime/Renderer/Private/PlanarReflectionRendering.cpp
--- a/Engine/Source/Runtime/Renderer/Private/PlanarReflectionRendering.cpp
+++ b/Engine/Source/Runtime/Renderer/Private/PlanarReflectionRendering.cpp
@@ -32,6 +32,7 @@
 #include "PipelineStateCache.h"
 #include "ClearQuad.h"
 #include "SceneTextureParameters.h"
+#include "SceneViewExtension.h"
 
 void SetupPlanarReflectionUniformParameters(const class FSceneView& View, const FPlanarReflectionSceneProxy* ReflectionSceneProxy, FPlanarReflectionUniformParameters& OutParameters)
 {
@@ -281,6 +282,9 @@ static void UpdatePlanarReflectionContents_RenderThread(
 {
 	QUICK_SCOPE_CYCLE_COUNTER(STAT_RenderPlanarReflection);
 
+	// We need to execute the pre-render view extensions before we do any view dependent work.
+	FSceneRenderer::ViewExtensionPreRender_RenderThread(RHICmdList, SceneRenderer);
+
 	SceneRenderer->RenderThreadBegin(RHICmdList);
 
 	// Make sure we render to the same set of GPUs as the main scene renderer.
@@ -428,6 +432,10 @@ static void UpdatePlanarReflectionContentsWithoutRendering_RenderThread(
 	const FName OwnerName)
 {
 	QUICK_SCOPE_CYCLE_COUNTER(STAT_RenderPlanarReflection);
+
+	// We need to execute the pre-render view extensions before we do any view dependent work.
+	FSceneRenderer::ViewExtensionPreRender_RenderThread(RHICmdList, SceneRenderer);
+
 	SceneRenderer->RenderThreadBegin(RHICmdList);
 
 	FBox PlanarReflectionBounds = SceneProxy->WorldBounds;
@@ -512,13 +520,13 @@ void FScene::UpdatePlanarReflectionContents(UPlanarReflectionComponent* CaptureC
 
 		const bool bIsMobilePixelProjectedReflectionEnabled = IsMobilePixelProjectedReflectionEnabled(GetShaderPlatform());
 
-		const bool bIsRenderTargetValid = CaptureComponent->RenderTarget != NULL
+		const bool bIsRenderTargetValid = CaptureComponent->RenderTarget != nullptr
 									&& CaptureComponent->RenderTarget->GetSizeXY() == DesiredPlanarReflectionTextureSize
 									// The RenderTarget's TextureRHI could be nullptr if it is used for mobile pixel projected reflection.
 									&& (bIsMobilePixelProjectedReflectionEnabled || CaptureComponent->RenderTarget->TextureRHI.IsValid());
 		
 
-		if (CaptureComponent->RenderTarget != NULL && !bIsRenderTargetValid)
+		if (CaptureComponent->RenderTarget != nullptr && !bIsRenderTargetValid)
 		{
 			FPlanarReflectionRenderTarget* RenderTarget = CaptureComponent->RenderTarget;
 			ENQUEUE_RENDER_COMMAND(ReleaseRenderTargetCommand)(
@@ -528,10 +536,10 @@ void FScene::UpdatePlanarReflectionContents(UPlanarReflectionComponent* CaptureC
 					delete RenderTarget;
 				});
 
-			CaptureComponent->RenderTarget = NULL;
+			CaptureComponent->RenderTarget = nullptr;
 		}
 
-		if (CaptureComponent->RenderTarget == NULL)
+		if (CaptureComponent->RenderTarget == nullptr)
 		{
 			CaptureComponent->RenderTarget = new FPlanarReflectionRenderTarget(DesiredPlanarReflectionTextureSize);
 
@@ -618,6 +626,8 @@ void FScene::UpdatePlanarReflectionContents(UPlanarReflectionComponent* CaptureC
 		// Uses the exact same secondary view fraction on the planar reflection as the main viewport.
 		ViewFamily.SecondaryViewFraction = MainSceneRenderer.ViewFamily.SecondaryViewFraction;
 
+		ViewFamily.ViewExtensions = GEngine->ViewExtensions->GatherActiveExtensions(FSceneViewExtensionContext(this));
+
 		SetupViewFamilyForSceneCapture(
 			ViewFamily,
 			CaptureComponent,
@@ -636,13 +646,24 @@ void FScene::UpdatePlanarReflectionContents(UPlanarReflectionComponent* CaptureC
 		// Disable screen percentage on planar reflection renderer if main one has screen percentage disabled.
 		SceneRenderer->ViewFamily.EngineShowFlags.ScreenPercentage = MainSceneRenderer.ViewFamily.EngineShowFlags.ScreenPercentage;
 
+		for (const FSceneViewExtensionRef& Extension : ViewFamily.ViewExtensions)
+		{
+			Extension->SetupViewFamily(ViewFamily);
+		}
+
 		for (int32 ViewIndex = 0; ViewIndex < SceneCaptureViewInfo.Num(); ++ViewIndex)
 		{
-			SceneRenderer->Views[ViewIndex].GlobalClippingPlane = MirrorPlane;
+			FViewInfo& ViewInfo = SceneRenderer->Views[ViewIndex];
+			ViewInfo.GlobalClippingPlane = MirrorPlane;
 			// Jitter can't be removed completely due to the clipping plane
 			// Also, this prevents the prefilter pass, which reads from jittered depth, from having to do special handling of it's depth-dependent input
-			SceneRenderer->Views[ViewIndex].bAllowTemporalJitter = false;
-			SceneRenderer->Views[ViewIndex].bRenderSceneTwoSided = CaptureComponent->bRenderSceneTwoSided;
+			ViewInfo.bAllowTemporalJitter = false;
+			ViewInfo.bRenderSceneTwoSided = CaptureComponent->bRenderSceneTwoSided;
+
+			for (const FSceneViewExtensionRef& Extension : ViewFamily.ViewExtensions)
+			{
+				Extension->SetupView(ViewFamily, ViewInfo);
+			}
 
 			CaptureComponent->ProjectionWithExtraFOV[ViewIndex] = SceneCaptureViewInfo[ViewIndex].ProjectionMatrix;
 
diff --git a/Engine/Source/Runtime/Renderer/Private/ReflectionEnvironmentCapture.cpp b/Engine/Source/Runtime/Renderer/Private/ReflectionEnvironmentCapture.cpp
--- a/Engine/Source/Runtime/Renderer/Private/ReflectionEnvironmentCapture.cpp
+++ b/Engine/Source/Runtime/Renderer/Private/ReflectionEnvironmentCapture.cpp
@@ -14,6 +14,7 @@
 #include "RendererInterface.h"
 #include "RHIStaticStates.h"
 #include "SceneView.h"
+#include "SceneViewExtension.h"
 #include "LegacyScreenPercentageDriver.h"
 #include "Shader.h"
 #include "TextureResource.h"
@@ -728,6 +729,9 @@ void ClearScratchCubemaps(FRHICommandListImmediate& RHICmdList, int32 TargetSize
 /** Captures the scene for a reflection capture by rendering the scene multiple times and copying into a cubemap texture. */
 void CaptureSceneToScratchCubemap(FRHICommandListImmediate& RHICmdList, FSceneRenderer* SceneRenderer, ECubeFace CubeFace, int32 CubemapSize, bool bCapturingForSkyLight, bool bLowerHemisphereIsBlack, const FLinearColor& LowerHemisphereColor, bool bCapturingForMobile)
 {
+	// We need to execute the pre-render view extensions before we do any view dependent work.
+	FSceneRenderer::ViewExtensionPreRender_RenderThread(RHICmdList, SceneRenderer);
+
 	SceneRenderer->RenderThreadBegin(RHICmdList);
 
 	// update any resources that needed a deferred update
@@ -1463,7 +1467,15 @@ void CaptureSceneIntoScratchCubemap(
 		ViewFamily.SetScreenPercentageInterface(new FLegacyScreenPercentageDriver(
 			ViewFamily, /* GlobalResolutionFraction = */ 1.0f));
 
-		FSceneRenderer* SceneRenderer = FSceneRenderer::CreateSceneRenderer(&ViewFamily, NULL);
+		ViewFamily.ViewExtensions = GEngine->ViewExtensions->GatherActiveExtensions(FSceneViewExtensionContext(Scene));
+
+		for (const FSceneViewExtensionRef& Extension : ViewFamily.ViewExtensions)
+		{
+			Extension->SetupViewFamily(ViewFamily);
+			Extension->SetupView(ViewFamily, *View);
+		}
+
+		FSceneRenderer* SceneRenderer = FSceneRenderer::CreateSceneRenderer(&ViewFamily, nullptr);
 
 		ENQUEUE_RENDER_COMMAND(CaptureCommand)(
 			[SceneRenderer, CubeFace, CubemapSize, bCapturingForSkyLight, bLowerHemisphereIsBlack, LowerHemisphereColor, bCapturingForMobile](FRHICommandListImmediate& RHICmdList)
diff --git a/Engine/Source/Runtime/Renderer/Private/SceneCaptureRendering.cpp b/Engine/Source/Runtime/Renderer/Private/SceneCaptureRendering.cpp
--- a/Engine/Source/Runtime/Renderer/Private/SceneCaptureRendering.cpp
+++ b/Engine/Source/Runtime/Renderer/Private/SceneCaptureRendering.cpp
@@ -278,6 +278,9 @@ static void UpdateSceneCaptureContentDeferred_RenderThread(
 	bool bOrthographicCamera
 	)
 {
+	// We need to execute the pre-render view extensions before we do any view dependent work.
+	FSceneRenderer::ViewExtensionPreRender_RenderThread(RHICmdList, SceneRenderer);
+
 	SceneRenderer->RenderThreadBegin(RHICmdList);
 
 	// update any resources that needed a deferred update
@@ -360,6 +363,9 @@ void UpdateSceneCaptureContentMobile_RenderThread(
 	const FGenerateMipsParams& GenerateMipsParams,
 	bool bDisableFlipCopyGLES)
 {
+	// We need to execute the pre-render view extensions before we do any view dependent work.
+	FSceneRenderer::ViewExtensionPreRender_RenderThread(RHICmdList, SceneRenderer);
+
 	SceneRenderer->RenderThreadBegin(RHICmdList);
 
 	// update any resources that needed a deferred update
@@ -1089,9 +1095,6 @@ void FScene::UpdateSceneCaptureContents(USceneCaptureComponent2D* CaptureCompone
 					TextureRenderTargetResource->SetActiveGPUMask(FRHIGPUMask::All());
 				}
 
-				// We need to execute the pre-render view extensions before we do any view dependent work.
-				FSceneRenderer::ViewExtensionPreRender_RenderThread(RHICmdList, SceneRenderer);
-
 				FResolveParams ResolveParams;
 
 				if (bEnableOrthographicTiling)
@@ -1215,9 +1218,9 @@ void FScene::UpdateSceneCaptureContents(USceneCaptureComponentCube* CaptureCompo
 
 				bool bCaptureSceneColor = CaptureNeedsSceneColor(CaptureComponent->CaptureSource);
 
-				FSceneRenderer* SceneRenderer = CreateSceneRendererForSceneCapture(this, CaptureComponent, 
-				    TextureTarget->GameThread_GetRenderTargetResource(), CaptureSize, ViewRotationMatrix, 
-					Location, ProjectionMatrix, false, CaptureComponent->MaxViewDistanceOverride, 
+				FSceneRenderer* SceneRenderer = CreateSceneRendererForSceneCapture(this, CaptureComponent,
+					TextureTarget->GameThread_GetRenderTargetResource(), CaptureSize, ViewRotationMatrix,
+					Location, ProjectionMatrix, false, CaptureComponent->MaxViewDistanceOverride,
 					bCaptureSceneColor, &PostProcessSettings, 0, CaptureComponent->GetViewOwner(), StereoIPD);
 
 				SceneRenderer->ViewFamily.SceneCaptureSource = CaptureComponent->CaptureSource;
@@ -1234,9 +1237,9 @@ void FScene::UpdateSceneCaptureContents(USceneCaptureComponentCube* CaptureCompo
 				}
 				ENQUEUE_RENDER_COMMAND(CaptureCommand)(
 					[SceneRenderer, TextureRenderTarget, EventName, TargetFace](FRHICommandListImmediate& RHICmdList)
-				{
+					{
 						UpdateSceneCaptureContent_RenderThread(RHICmdList, SceneRenderer, TextureRenderTarget, TextureRenderTarget, EventName, FResolveParams(FResolveRect(), TargetFace), false, FGenerateMipsParams(), false, true, false);
-				}
+					}
 				);
 			}
 		}
diff --git a/Engine/Source/Runtime/Renderer/Private/SceneRendering.cpp b/Engine/Source/Runtime/Renderer/Private/SceneRendering.cpp
--- a/Engine/Source/Runtime/Renderer/Private/SceneRendering.cpp
+++ b/Engine/Source/Runtime/Renderer/Private/SceneRendering.cpp
@@ -4066,6 +4066,9 @@ static void RenderViewFamily_RenderThread(FRHICommandListImmediate& RHICmdList,
 {
 	LLM_SCOPE(ELLMTag::SceneRender);
 
+	// We need to execute the pre-render view extensions before we do any view dependent work.
+	FSceneRenderer::ViewExtensionPreRender_RenderThread(RHICmdList, SceneRenderer);
+
 	SceneRenderer->RenderThreadBegin(RHICmdList);
 
 	// update any resources that needed a deferred update
@@ -4334,14 +4337,6 @@ void FRendererModule::BeginRenderingViewFamily(FCanvas* Canvas, FSceneViewFamily
 			USceneCaptureComponent::UpdateDeferredCaptures(Scene);
 		}
 
-		// We need to execute the pre-render view extensions before we do any view dependent work.
-		// Anything between here and FDrawSceneCommand will add to HMD view latency
-		ENQUEUE_RENDER_COMMAND(FViewExtensionPreDrawCommand)(
-			[SceneRenderer](FRHICommandListImmediate& RHICmdList)
-			{
-				FSceneRenderer::ViewExtensionPreRender_RenderThread(RHICmdList, SceneRenderer);
-			});
-
 		if (!SceneRenderer->ViewFamily.EngineShowFlags.HitProxies)
 		{
 			for (int32 ReflectionIndex = 0; ReflectionIndex < SceneRenderer->Scene->PlanarReflections_GameThread.Num(); ReflectionIndex++)
