Adding Global Shader Types to the DumpMaterialShaderTypes commandlet.

- Sort global shaders alphabetically so diffs between output is stable.
- Adding stats for scanning the asset registry and dumping the stats.
- Added a separate global shaders histogram, but currently not dumping this to the output.
- Report full asset path for materials and MICs to make them easier to find in the editor.

diff --git a/Engine/Source/Editor/UnrealEd/Private/Commandlets/DumpMaterialShaderTypes.cpp b/Engine/Source/Editor/UnrealEd/Private/Commandlets/DumpMaterialShaderTypes.cpp
--- a/Engine/Source/Editor/UnrealEd/Private/Commandlets/DumpMaterialShaderTypes.cpp
+++ b/Engine/Source/Editor/UnrealEd/Private/Commandlets/DumpMaterialShaderTypes.cpp
@@ -3,6 +3,7 @@
 #include "Commandlets/DumpMaterialShaderTypes.h"
 #include "AssetData.h"
 #include "AssetRegistryModule.h"
+#include "GlobalShader.h"
 #include "HAL/FileManager.h"
 #include "Interfaces/ITargetPlatform.h"
 #include "Interfaces/ITargetPlatformManagerModule.h"
@@ -29,6 +30,19 @@ public:
 		delete DebugWriter;
 	}
 
+	void AddToGlobalShaderTypeHistogram(const TCHAR* GlobalShaderName)
+	{
+		FString Name(GlobalShaderName);
+		if (int32* Existing = GlobalShaderTypeHistogram.Find(Name))
+		{
+			++(*Existing);
+		}
+		else
+		{
+			GlobalShaderTypeHistogram.FindOrAdd(Name, 1);
+		}
+	}
+
 	void AddToHistogram(const TCHAR* VertexFactoryName, const TCHAR* ShaderPipelineName, const TCHAR* ShaderTypeName)
 	{
 		FString ShaderType(ShaderTypeName);
@@ -194,6 +208,9 @@ private:
 
 	/** Map of vertex factory display names to their counts. */
 	TMap<FString, int32>	VertexFactoryTypeHistogram;
+
+	/** Map of global shader type display names to their counts. */
+	TMap<FString, int32>	GlobalShaderTypeHistogram;
 };
 
 UDumpMaterialShaderTypesCommandlet::UDumpMaterialShaderTypesCommandlet(const FObjectInitializer& ObjectInitializer)
@@ -302,7 +319,7 @@ int ProcessMaterials(const ITargetPlatform* TargetPlatform, const EShaderPlatfor
 			TotalShaders += TotalShadersForMaterial;
 
 			Output.Log(TEXT(""));
-			Output.Log(FString::Printf(TEXT("Material: %s - %d shaders"), *AssetData.AssetName.ToString(), TotalShadersForMaterial));
+			Output.Log(FString::Printf(TEXT("Material: %s - %d shaders"), *AssetData.ObjectPath.ToString(), TotalShadersForMaterial));
 
 			PrintDebugShaderInfo(Output, OutShaderInfo);
 		}
@@ -348,7 +365,7 @@ int ProcessMaterialInstances(const ITargetPlatform* TargetPlatform, const EShade
 			}
 
 			Output.Log(TEXT(""));
-			Output.Log(FString::Printf(TEXT("Material Instance: %s - %d shaders"), *AssetData.AssetName.ToString(), TotalShadersForMaterial));
+			Output.Log(FString::Printf(TEXT("Material Instance: %s - %d shaders"), *AssetData.ObjectPath.ToString(), TotalShadersForMaterial));
 			Output.Log(FString::Printf(TEXT("Static Parameter %s"), *StaticParameterString));
 			Output.Log(FString::Printf(TEXT("Parent: %s"), MaterialInstance->Parent ? *MaterialInstance->Parent->GetName() : TEXT("NO PARENT")));
 
@@ -370,6 +387,56 @@ int ProcessMaterialInstances(const ITargetPlatform* TargetPlatform, const EShade
 	return TotalShaders;
 }
 
+int ProcessGlobalShaders(const ITargetPlatform* TargetPlatform, const EShaderPlatform ShaderPlatform, FShaderStatsGatheringContext& Output)
+{
+	Output.Log(TEXT(""));
+	Output.Log(TEXT("Global Shaders"));
+
+	TArray<const FGlobalShaderType*> GlobalShaderTypes;
+	for (TLinkedList<FShaderType*>::TIterator ShaderTypeIt(FShaderType::GetTypeList()); ShaderTypeIt; ShaderTypeIt.Next())
+	{
+		FGlobalShaderType* GlobalShaderType = ShaderTypeIt->GetGlobalShaderType();
+		if (!GlobalShaderType)
+		{
+			continue;
+		}
+
+		GlobalShaderTypes.Add(GlobalShaderType);
+	}
+	Algo::SortBy(GlobalShaderTypes, [](const FGlobalShaderType* GS) { return GS->GetName(); }, FNameLexicalLess());
+
+	int TotalShaders = 0;
+
+	FPlatformTypeLayoutParameters LayoutParams;
+	LayoutParams.InitializeForPlatform(TargetPlatform);
+	EShaderPermutationFlags PermutationFlags = GetShaderPermutationFlags(LayoutParams);
+
+	for (const FGlobalShaderType* GS : GlobalShaderTypes)
+	{
+		int PermutationCount = 0;
+		for (int32 id = 0; id < GS->GetPermutationCount(); id++)
+		{
+			if (GS->ShouldCompilePermutation(ShaderPlatform, id, PermutationFlags))
+			{
+				PermutationCount++;
+				TotalShaders++;
+				Output.AddToGlobalShaderTypeHistogram(GS->GetName());
+			}
+		}
+
+		if (PermutationCount)
+		{
+			Output.Log(FString::Printf(TEXT("%s - %d permutations"), GS->GetName(), PermutationCount));
+		}
+	}
+
+	Output.Log(TEXT(""));
+	Output.Log(TEXT("Global Shaders Summary"));
+	Output.Log(FString::Printf(TEXT("Total Shaders: %d"), TotalShaders));
+
+	return TotalShaders;
+}
+
 void ProcessForTargetAndShaderPlatform(const ITargetPlatform* TargetPlatform, const EShaderPlatform ShaderPlatform, const FString& Params, TArray<FAssetData>& MaterialList, TArray<FAssetData> MaterialInstanceList)
 {
 	const double StartTime = FPlatformTime::Seconds();
@@ -389,6 +456,8 @@ void ProcessForTargetAndShaderPlatform(const ITargetPlatform* TargetPlatform, co
 	TotalShaders += ProcessMaterialInstances(TargetPlatform, ShaderPlatform, Output, MaterialInstanceList);
 	TotalAssets += MaterialInstanceList.Num();
 
+	TotalShaders += ProcessGlobalShaders(TargetPlatform, ShaderPlatform, Output);
+
 	Output.Log(TEXT(""));
 	Output.Log(TEXT("Summary"));
 	Output.Log(FString::Printf(TEXT("Total Assets: %d"), TotalAssets));
@@ -422,6 +491,8 @@ int32 UDumpMaterialShaderTypesCommandlet::Main(const FString& Params)
 		return 0;
 	}
 
+	const double AssetRegistryStart = FPlatformTime::Seconds();
+
 	UE_LOG(LogDumpMaterialShaderTypesCommandlet, Display, TEXT("Searching the asset registry for all assets..."));
 	IAssetRegistry& AssetRegistry = FModuleManager::Get().LoadModuleChecked<FAssetRegistryModule>("AssetRegistry").Get();
 	AssetRegistry.SearchAllAssets(true);
@@ -462,6 +533,10 @@ int32 UDumpMaterialShaderTypesCommandlet::Main(const FString& Params)
 		}
 	}
 
+
+	const double AssetRegistryEnd = FPlatformTime::Seconds();
+	UE_LOG(LogDumpMaterialShaderTypesCommandlet, Display, TEXT("Asset scan took: %.3f"), AssetRegistryEnd - AssetRegistryStart);
+
 	// For all active platforms
 	ITargetPlatformManagerModule* TPM = GetTargetPlatformManager();
 	const TArray<ITargetPlatform*>& Platforms = TPM->GetActiveTargetPlatforms();
@@ -479,5 +554,8 @@ int32 UDumpMaterialShaderTypesCommandlet::Main(const FString& Params)
 			ProcessForTargetAndShaderPlatform(Platforms[Index], ShaderPlatform, Params, MaterialList, MaterialInstanceList);
 		}
 	}
+
+	UE_LOG(LogDumpMaterialShaderTypesCommandlet, Display, TEXT("Dumping stats took: %.3f"), FPlatformTime::Seconds() - AssetRegistryEnd);
+
 	return 0;
 }
