[WP Navmesh] Fix warning "Can't load actor..." occurring when building a WP navmesh.
It was caused by actors being deleted but WP not being notified and trying to load deleted actors.

diff --git a/Engine/Source/Editor/UnrealEd/Private/WorldPartition/WorldPartitionNavigationDataBuilder.cpp b/Engine/Source/Editor/UnrealEd/Private/WorldPartition/WorldPartitionNavigationDataBuilder.cpp
--- a/Engine/Source/Editor/UnrealEd/Private/WorldPartition/WorldPartitionNavigationDataBuilder.cpp
+++ b/Engine/Source/Editor/UnrealEd/Private/WorldPartition/WorldPartitionNavigationDataBuilder.cpp
@@ -101,7 +101,11 @@ bool UWorldPartitionNavigationDataBuilder::RunInternal(UWorld* World, const FCel
 			}
 			
 			UE_LOG(LogWorldPartitionNavigationDataBuilder, Verbose, TEXT("   Destroy actor %s in package %s."), *Actor->GetName(), *Actor->GetPackage()->GetName());
-			World->DestroyActor(Actor);
+			const bool bDestroyed = World->DestroyActor(Actor);
+			if (!bDestroyed)
+			{
+				UE_LOG(LogWorldPartitionNavigationDataBuilder, Warning, TEXT("      Could not be destroy %s."), *Actor->GetName());
+			}
 		}
 	}
 	UE_LOG(LogWorldPartitionNavigationDataBuilder, Verbose, TEXT("   Number of ANavigationDataChunkActor: %i"), Count);
@@ -122,6 +126,13 @@ bool UWorldPartitionNavigationDataBuilder::RunInternal(UWorld* World, const FCel
 			return true;
 		}
 
+		// If we had packages to delete we need to notify the delete after the save. Else WP might try to load the deleted descriptors on the next iteration.
+		// Note: this notification is expected to be done by the SavePackages()
+		for (UPackage* Package : PackagesToClean)
+		{
+			WorldPartition->OnPackageDeleted(Package);
+		}		
+
 		return true;
 	}
 
@@ -221,6 +232,13 @@ bool UWorldPartitionNavigationDataBuilder::RunInternal(UWorld* World, const FCel
 		{
 			return true;
 		}
+
+		// If we had packages to delete we need to notify the delete after the save. Else WP might try to load the deleted descriptors on the next iteration.
+		// Note: this notification is expected to be done by the SavePackages()
+		for (UPackage* Package : PackagesToDelete)
+		{
+			WorldPartition->OnPackageDeleted(Package);
+		}
 		
 		{
 			// Add new packages to source control
