Fixed memory leak each time we drag something in the editor viewport

diff --git a/Engine/Source/Editor/UnrealEd/Private/EditorViewportClient.cpp b/Engine/Source/Editor/UnrealEd/Private/EditorViewportClient.cpp
--- a/Engine/Source/Editor/UnrealEd/Private/EditorViewportClient.cpp
+++ b/Engine/Source/Editor/UnrealEd/Private/EditorViewportClient.cpp
@@ -2230,7 +2230,7 @@ void FEditorViewportClient::UpdateMouseDelta()
 				if (DragStartView == nullptr)
 				{
 					// Compute a view.
-					DragStartViewFamily = new FSceneViewFamily(FSceneViewFamily::ConstructionValues(
+					DragStartViewFamily = new FSceneViewFamilyContext(FSceneViewFamily::ConstructionValues(
 						Viewport,
 						GetScene(),
 						EngineShowFlags)
diff --git a/Engine/Source/Runtime/Engine/Public/SceneView.h b/Engine/Source/Runtime/Engine/Public/SceneView.h
--- a/Engine/Source/Runtime/Engine/Public/SceneView.h
+++ b/Engine/Source/Runtime/Engine/Public/SceneView.h
@@ -1796,7 +1796,7 @@ public:
 
 	/** Initialization constructor. */
 	FSceneViewFamily( const ConstructionValues& CVS );
-	~FSceneViewFamily();
+	virtual ~FSceneViewFamily();
 
 	ERHIFeatureLevel::Type GetFeatureLevel() const;
 
@@ -1952,5 +1952,5 @@ public:
 	{}
 
 	/** Destructor. */
-	ENGINE_API ~FSceneViewFamilyContext();
+	ENGINE_API virtual ~FSceneViewFamilyContext();
 };
\ No newline at end of file
