Improvement to Localization Text Gatherer in code to:
- Support arguments on multiple lines for localization macros.
- Fix to a lot of Localization String that were missing the last character of missing a closing ".

diff --git a/Engine/Source/Editor/UnrealEd/Classes/Commandlets/GatherTextFromSourceCommandlet.h b/Engine/Source/Editor/UnrealEd/Classes/Commandlets/GatherTextFromSourceCommandlet.h
--- a/Engine/Source/Editor/UnrealEd/Classes/Commandlets/GatherTextFromSourceCommandlet.h
+++ b/Engine/Source/Editor/UnrealEd/Classes/Commandlets/GatherTextFromSourceCommandlet.h
@@ -130,6 +130,8 @@ private:
 		//Discovered string table data from all files
 		TMap<FName, FParsedStringTable> ParsedStringTables;
 
+		TArray<FString> TextLines;
+
 		FSourceFileParseContext(UGatherTextFromSourceCommandlet* InOwnerCommandlet)
 			: FileTypes(EGatherTextSourceFileTypes::None)
 			, Filename()
@@ -254,36 +256,50 @@ private:
 	public:
 		static const FString TextMacroString;
 
-		FMacroDescriptor(FString InName)
+		FMacroDescriptor(FString InName, int32 InMinArgumentNumber)
 			: Name(MoveTemp(InName))
+			, MinArgumentNumber(InMinArgumentNumber)
 		{
 			ApplicableFileTypes = EGatherTextSourceFileTypes::Cpp;
 		}
 
-		virtual const FString& GetToken() const override { return Name; }
+		virtual const FString& GetToken() const override
+		{
+			return Name;
+		}
+
+		int32 GetMinNumberOfArgument() const
+		{
+			return MinArgumentNumber;
+		}
 
 	protected:
 		bool ParseArgsFromMacro(const FString& Text, TArray<FString>& Args, FSourceFileParseContext& Context) const;
+		bool ParseArgsFromNextLines(TArray<FString>& Args, int32& BracketStack, FSourceFileParseContext& Context) const;
+		bool ParseArgumentString(const FString& Text, const int32 OpenBracketIdx, int32& BracketStack, const FSourceFileParseContext& Context, TArray<FString>& Args) const;
 
 		static bool PrepareArgument(FString& Argument, bool IsAutoText, const FString& IdentForLogging, bool& OutHasQuotes);
 
 	private:
-		FString Name;
+		const FString Name;
+
+		// Minimum number argument for that Macro.
+		const int32 MinArgumentNumber;
 	};
 
 	class FUICommandMacroDescriptor : public FMacroDescriptor
 	{
 	public:
 		FUICommandMacroDescriptor()
-			: FMacroDescriptor(TEXT("UI_COMMAND"))
+			: FMacroDescriptor(TEXT("UI_COMMAND"), 5)
 		{
 		}
 
 		virtual void TryParse(const FString& Text, FSourceFileParseContext& Context) const override;
 
 	protected:
-		FUICommandMacroDescriptor(FString InName)
-			: FMacroDescriptor(MoveTemp(InName))
+		FUICommandMacroDescriptor(FString InName, int32 InMinNumberOfArgument)
+			: FMacroDescriptor(MoveTemp(InName), InMinNumberOfArgument)
 		{
 		}
 
@@ -294,7 +310,7 @@ private:
 	{
 	public:
 		FUICommandExtMacroDescriptor()
-			: FUICommandMacroDescriptor(TEXT("UI_COMMAND_EXT"))
+			: FUICommandMacroDescriptor(TEXT("UI_COMMAND_EXT"), 5)
 		{
 		}
 
@@ -319,7 +335,7 @@ private:
 			FMacroArg(EMacroArgSemantic InSema, bool InIsAutoText) : Semantic(InSema), IsAutoText(InIsAutoText) {}
 		};
 
-		FStringMacroDescriptor(FString InName, FMacroArg Arg0, FMacroArg Arg1, FMacroArg Arg2) : FMacroDescriptor(InName)
+		FStringMacroDescriptor(FString InName, FMacroArg Arg0, FMacroArg Arg1, FMacroArg Arg2) : FMacroDescriptor(InName, 3)
 		{
 			ApplicableFileTypes = EGatherTextSourceFileTypes::Cpp | EGatherTextSourceFileTypes::Ini;
 			Arguments.Add(Arg0);
@@ -327,14 +343,14 @@ private:
 			Arguments.Add(Arg2);
 		}
 
-		FStringMacroDescriptor(FString InName, FMacroArg Arg0, FMacroArg Arg1) : FMacroDescriptor(InName)
+		FStringMacroDescriptor(FString InName, FMacroArg Arg0, FMacroArg Arg1) : FMacroDescriptor(InName, 2)
 		{
 			ApplicableFileTypes = EGatherTextSourceFileTypes::Cpp | EGatherTextSourceFileTypes::Ini;
 			Arguments.Add(Arg0);
 			Arguments.Add(Arg1);
 		}
 
-		FStringMacroDescriptor(FString InName, FMacroArg Arg0) : FMacroDescriptor(InName)
+		FStringMacroDescriptor(FString InName, FMacroArg Arg0) : FMacroDescriptor(InName, 1)
 		{
 			ApplicableFileTypes = EGatherTextSourceFileTypes::Cpp | EGatherTextSourceFileTypes::Ini;
 			Arguments.Add(Arg0);
@@ -350,7 +366,7 @@ private:
 	{
 	public:
 		FStringTableMacroDescriptor()
-			: FMacroDescriptor(TEXT("LOCTABLE_NEW"))
+			: FMacroDescriptor(TEXT("LOCTABLE_NEW"), 2)
 		{
 		}
 
@@ -360,7 +376,7 @@ private:
 	class FStringTableFromFileMacroDescriptor : public FMacroDescriptor
 	{
 	public:
-		FStringTableFromFileMacroDescriptor(FString InName, FString InRootPath) : FMacroDescriptor(MoveTemp(InName)), RootPath(MoveTemp(InRootPath)) {}
+		FStringTableFromFileMacroDescriptor(FString InName, FString InRootPath) : FMacroDescriptor(MoveTemp(InName), 2), RootPath(MoveTemp(InRootPath)) {}
 
 		virtual void TryParse(const FString& Text, FSourceFileParseContext& Context) const override;
 
@@ -372,7 +388,7 @@ private:
 	{
 	public:
 		FStringTableEntryMacroDescriptor()
-			: FMacroDescriptor(TEXT("LOCTABLE_SETSTRING"))
+			: FMacroDescriptor(TEXT("LOCTABLE_SETSTRING"), 3)
 		{
 		}
 
@@ -383,7 +399,7 @@ private:
 	{
 	public:
 		FStringTableEntryMetaDataMacroDescriptor()
-			: FMacroDescriptor(TEXT("LOCTABLE_SETMETA"))
+			: FMacroDescriptor(TEXT("LOCTABLE_SETMETA"), 4)
 		{
 		}
 
diff --git a/Engine/Source/Editor/UnrealEd/Private/Commandlets/GatherTextFromSourceCommandlet.cpp b/Engine/Source/Editor/UnrealEd/Private/Commandlets/GatherTextFromSourceCommandlet.cpp
--- a/Engine/Source/Editor/UnrealEd/Private/Commandlets/GatherTextFromSourceCommandlet.cpp
+++ b/Engine/Source/Editor/UnrealEd/Private/Commandlets/GatherTextFromSourceCommandlet.cpp
@@ -302,6 +302,7 @@ int32 UGatherTextFromSourceCommandlet::Main( const FString& Params )
 		ParseCtxt.WithinStringLiteral = false;
 		ParseCtxt.WithinNamespaceDefineLineNumber = INDEX_NONE;
 		ParseCtxt.WithinStartingLine = nullptr;
+		ParseCtxt.TextLines.Reset();
 		ParseCtxt.FlushMacroStack();
 
 		FString SourceFileText;
@@ -618,7 +619,6 @@ FString UGatherTextFromSourceCommandlet::RemoveStringFromTextMacro(const FString
 
 FString UGatherTextFromSourceCommandlet::StripCommentsFromToken(const FString& InToken, FSourceFileParseContext& Context)
 {
-	check(!Context.WithinBlockComment);
 	check(!Context.WithinLineComment);
 	check(!Context.WithinStringLiteral);
 
@@ -745,16 +745,17 @@ bool UGatherTextFromSourceCommandlet::ParseSourceText(const FString& Text, const
 	const int32 BlockCommentEndLen = BlockCommentEnd ? FCString::Strlen(BlockCommentEnd) : 0;
 	checkf((BlockCommentStartLen == 0 && BlockCommentEndLen == 0) || (BlockCommentStartLen > 0 && BlockCommentEndLen > 0), TEXT("Block comments require both a start and an end marker!"));
 
-	// Split the file into lines 
-	TArray<FString> TextLines;
+	// Split the file into lines
+	TArray<FString>& TextLines = ParseCtxt.TextLines;
 	Text.ParseIntoArrayLines(TextLines, false);
 
-	// Move through the text lines looking for the tokens that denote the items in the Parsables list
+	// Move through the text lines looking for the tokens that denote the items in the Parsables list 
 	for (int32 LineIdx = 0; LineIdx < TextLines.Num(); LineIdx++)
 	{
+		// Remove spaces at the end of the line.
 		TextLines[LineIdx].TrimEndInline();
 		const FString& Line = TextLines[LineIdx];
-		if( Line.IsEmpty() )
+		if ( Line.IsEmpty() )
 			continue;
 
 		// Use these pending vars to defer parsing a token hit until longer tokens can't hit too
@@ -776,6 +777,7 @@ bool UGatherTextFromSourceCommandlet::ParseSourceText(const FString& Text, const
 
 			if (!ParseCtxt.WithinLineComment && !ParseCtxt.WithinBlockComment && !ParseCtxt.WithinStringLiteral)
 			{
+				// Detect that the line starts with a Line comment.
 				if (LineCommentLen > 0 && FCString::Strncmp(Cursor, LineComment, LineCommentLen) == 0)
 				{
 					ParseCtxt.WithinLineComment = true;
@@ -793,6 +795,7 @@ bool UGatherTextFromSourceCommandlet::ParseSourceText(const FString& Text, const
 				}
 			}
 
+			// We are not in a comment (line or block) and we are not parsing a string.
 			if (!ParseCtxt.WithinLineComment && !ParseCtxt.WithinBlockComment && !ParseCtxt.WithinStringLiteral)
 			{
 				if (*Cursor == TEXT('\"'))
@@ -964,6 +967,7 @@ bool UGatherTextFromSourceCommandlet::ParseSourceText(const FString& Text, const
 				continue;
 			}
 
+			// Go through all the Parserables to find matches
 			for (int32 ParIdx = 0; ParIdx < ParsablesForFile.Num(); ++ParIdx)
 			{
 				FParsableDescriptor* Parsable = ParsablesForFile[ParIdx];
@@ -1443,107 +1447,198 @@ void UGatherTextFromSourceCommandlet::FEndIfDescriptor::TryParse(const FString&
 	}
 }
 
+bool UGatherTextFromSourceCommandlet::FMacroDescriptor::ParseArgumentString(const FString& Text, const int32 CursorOffset, int32& BracketStack, const FSourceFileParseContext& Context, TArray<FString>& Args) const
+{
+	bool bInDblQuotes = false;
+	bool bInSglQuotes = false;
+	bool bEscapeNextChar = false;
+
+	const TCHAR* ArgStart = *Text + CursorOffset;
+	const TCHAR* Cursor = ArgStart;
+	for (; 0 < BracketStack; ++Cursor)
+	{
+		// Skip this character
+		if (bEscapeNextChar)
+		{
+			bEscapeNextChar = false;
+		}
+		else if ((bInDblQuotes || bInSglQuotes) && !bEscapeNextChar && '\\' == *Cursor)
+		{
+			bEscapeNextChar = true;
+		}
+
+		// We are closing "
+		else if (bInDblQuotes)
+		{
+			if ('\"' == *Cursor)
+			{
+				bInDblQuotes = false;
+			}
+		}
+		// We are closing a '
+		else if (bInSglQuotes)
+		{
+			if ('\'' == *Cursor)
+			{
+				bInSglQuotes = false;
+			}
+		}
+		// We are opening a "
+		else if ('\"' == *Cursor)
+		{
+			bInDblQuotes = true;
+		}
+		// We are opening a '
+		else if ('\'' == *Cursor)
+		{
+			bInSglQuotes = true;
+		}
+		// We found an opening bracket '(', increment the stack
+		else if ('(' == *Cursor)
+		{
+			++BracketStack;
+		}
+		// We found the closing bracket ')' decrement the stack
+		else if (')' == *Cursor)
+		{
+			--BracketStack;
+
+			if (0 > BracketStack)
+			{
+				UE_LOG(LogGatherTextFromSourceCommandlet, Warning, TEXT("Unexpected bracket ')' in %s macro while parsing %s:%d. %s"), *GetToken(), *Context.Filename, Context.LineNumber, *FLocTextHelper::SanitizeLogOutput(Context.LineText.TrimStartAndEnd()));
+				return false;
+			}
+		}
+		// We have a single bracket open and we found a ',' this is the end of the argument
+		else if (1 == BracketStack && ',' == *Cursor)
+		{
+			// create argument from ArgStart to Cursor and set Start next char
+			Args.Add(FString(Cursor - ArgStart, ArgStart));
+			ArgStart = Cursor + 1;
+		}
+		// We just closed the last ')' this is the end of all args for this macro 
+		if (0 == BracketStack)
+		{
+			// If the arg is empty it means we found a closing braket after a ',' or at the begining of a line
+			FString NewArg = FString(Cursor - ArgStart, ArgStart);
+			if (!NewArg.IsEmpty())
+			{
+				Args.Add(NewArg);
+			}
+		}
+		// We are at the end of the line and this is the last argument on that line.
+		if (*Cursor == '\0')
+		{
+			FString NewArg = FString(Cursor - ArgStart - 1, ArgStart);
+			if (!NewArg.IsEmpty())
+			{
+				Args.Add(NewArg);
+			}
+			break;
+		}
+	}
+	return true;
+}
+
 bool UGatherTextFromSourceCommandlet::FMacroDescriptor::ParseArgsFromMacro(const FString& Text, TArray<FString>& Args, FSourceFileParseContext& Context) const
 {
 	// Attempt to parse something of the format
 	// NAME(param0, param1, param2, etc)
 
-	bool Success = false;
-
 	// Step over the token name and any whitespace after it
 	FString RemainingText = Text.RightChop(GetToken().Len());
+
+	//RemainingText could be empty if the bracket is at the begining of the next line
 	RemainingText.TrimStartInline();
 
+	// Find the Opening bracket.
 	const int32 OpenBracketIdx = RemainingText.Find(TEXT("("), ESearchCase::CaseSensitive);
-	if (OpenBracketIdx == INDEX_NONE)
+
+	// If we find a bracket it must be the first character of the remaining text
+	if (OpenBracketIdx > 0)
 	{
-		// No opening bracket; warn about this, but don't consider it an error as we're likely parsing something we shouldn't be
-		UE_LOG(LogGatherTextFromSourceCommandlet, Warning, TEXT("Missing bracket '(' in %s macro at %s:%d. %s"), *GetToken(), *Context.Filename, Context.LineNumber, *FLocTextHelper::SanitizeLogOutput(Context.LineText.TrimStartAndEnd()));
+		// We stepped over the whitespace when building RemainingText, so if the bracket isn't the first character in the text then it means we only partially matched a longer token and shouldn't parse it
 		return false;
 	}
-	else if (OpenBracketIdx > 0)
+
+	Args.Empty();
+
+	int32 BracketStack = OpenBracketIdx == INDEX_NONE ? 0:1;
+
+	// If we found a bracket, we can start parsing argument on this line
+	if (BracketStack > 0)
 	{
-		// We stepped over the whitespace when building RemainingText, so if the bracket isn't the first character in the text then it means we only partially matched a longer token and shouldn't parse it
-		return false;
+		// Parse the argument that are on the same line as the macro.
+		if (!ParseArgumentString(RemainingText, OpenBracketIdx + 1, BracketStack, Context, Args))
+		{
+			return false;
+		}
 	}
-	else
+
+	// Not all arguments were found, look on the next lines.
+	if (Args.Num() < GetMinNumberOfArgument())
 	{
-		Args.Empty();
+		if (!ParseArgsFromNextLines(Args, BracketStack, Context))
+		{
+			return false;
+		}
+	}
+
+	return 0 < Args.Num() ? true : false;
+}
 
-		bool bInDblQuotes = false;
-		bool bInSglQuotes = false;
-		int32 BracketStack = 1;
-		bool bEscapeNextChar = false;
+bool UGatherTextFromSourceCommandlet::FMacroDescriptor::ParseArgsFromNextLines(TArray<FString>& Args, int32& BracketStack, FSourceFileParseContext& Context) const
+{
+	// We allow one line for the macro,that line may contain one or multiple arguments.
+	// So we parse a line by arguments left and a line for the end bracket.
+	int32 MaxLineToParse = Context.LineNumber + GetMinNumberOfArgument() - Args.Num() + 1;
 
-		const TCHAR* ArgStart = *RemainingText + OpenBracketIdx + 1;
-		const TCHAR* Cursor = ArgStart;
-		for (; 0 < BracketStack && '\0' != *Cursor; ++Cursor)
+	// Loop until we have all arguments and the closing bracket.
+	for (int32 i = Context.LineNumber; Args.Num() < GetMinNumberOfArgument() || BracketStack > 0; i++)
+	{
+		if(i >= Context.TextLines.Num())
 		{
-			if (bEscapeNextChar)
-			{
-				bEscapeNextChar = false;
-			}
-			else if ((bInDblQuotes || bInSglQuotes) && !bEscapeNextChar && '\\' == *Cursor)
-			{
-				bEscapeNextChar = true;
-			}
-			else if (bInDblQuotes)
-			{
-				if ('\"' == *Cursor)
-				{
-					bInDblQuotes = false;
-				}
-			}
-			else if (bInSglQuotes)
-			{
-				if ('\'' == *Cursor)
-				{
-					bInSglQuotes = false;
-				}
-			}
-			else if ('\"' == *Cursor)
-			{
-				bInDblQuotes = true;
-			}
-			else if ('\'' == *Cursor)
-			{
-				bInSglQuotes = true;
-			}
-			else if ('(' == *Cursor)
-			{
-				++BracketStack;
-			}
-			else if (')' == *Cursor)
-			{
-				--BracketStack;
+			UE_LOG(LogGatherTextFromSourceCommandlet, Warning, TEXT("We reached end of file while parsing specified %s macro for arguments. %s:%d. %s"), *GetToken(), *Context.Filename, Context.LineNumber, *FLocTextHelper::SanitizeLogOutput(Context.LineText.TrimStartAndEnd()));
+			return false;
+		}
 
-				if (0 > BracketStack)
-				{
-					UE_LOG(LogGatherTextFromSourceCommandlet, Warning, TEXT("Unexpected bracket ')' in %s macro while parsing %s:%d. %s"), *GetToken(), *Context.Filename, Context.LineNumber, *FLocTextHelper::SanitizeLogOutput(Context.LineText.TrimStartAndEnd()));
-					return false;
-				}
-			}
-			else if (1 == BracketStack && ',' == *Cursor)
-			{
-				// create argument from ArgStart to Cursor and set Start next char
-				Args.Add(FString(Cursor - ArgStart, ArgStart));
-				ArgStart = Cursor + 1;
-			}
+		if (i > MaxLineToParse)
+		{
+			UE_LOG(LogGatherTextFromSourceCommandlet, Warning, TEXT("The specified %s macro exceeds maximum allowed lines for arguments. %s:%d. %s"), *GetToken(), *Context.Filename, Context.LineNumber, *FLocTextHelper::SanitizeLogOutput(Context.LineText.TrimStartAndEnd()));
+			return false;
 		}
 
-		if (0 == BracketStack)
+		FString LineText = StripCommentsFromToken(Context.TextLines[i], Context);
+		// Ignore empty lines.
+		if (LineText.IsEmpty())
 		{
-			Args.Add(FString(Cursor - ArgStart - 1, ArgStart));
+			MaxLineToParse++;
+			continue;
 		}
-		//else
-		//{
-		//	Args.Add(FString(ArgStart));
-		//}
 
-		Success = 0 < Args.Num() ? true : false;	
-	}
+		int32 ParsingOffset = 0;
 
-	return Success;
+		// We are not in an opening bracket yet, look at the beginning of the line.
+		if (BracketStack == 0)
+		{
+			int32 OpenBracketIdx = LineText.Find(TEXT("("), ESearchCase::CaseSensitive);
+			// We did not find the opening bracket on the first line and on the second non empty line, we give up
+			if (OpenBracketIdx == INDEX_NONE)
+			{
+				UE_LOG(LogGatherTextFromSourceCommandlet, Warning, TEXT("Openaning bracket not found while parsing specified %s macro for arguments. %s:%d. %s"), *GetToken(), *Context.Filename, Context.LineNumber, *FLocTextHelper::SanitizeLogOutput(Context.LineText.TrimStartAndEnd()));
+				return false;
+			}
+			BracketStack = 1;
+			// Start parsing right after the bracket
+			ParsingOffset = OpenBracketIdx + 1;
+		}
+		
+		if (!ParseArgumentString(LineText, ParsingOffset, BracketStack, Context, Args))
+		{
+			return false;
+		}	
+	}
+	return true;
 }
 
 bool UGatherTextFromSourceCommandlet::FMacroDescriptor::PrepareArgument(FString& Argument, bool IsAutoText, const FString& IdentForLogging, bool& OutHasQuotes)
@@ -1565,7 +1660,7 @@ bool UGatherTextFromSourceCommandlet::FMacroDescriptor::PrepareArgument(FString&
 void UGatherTextFromSourceCommandlet::FUICommandMacroDescriptor::TryParseArgs(const FString& Text, FSourceFileParseContext& Context, const TArray<FString>& Arguments, const int32 ArgIndexOffset) const
 {
 	FString Identifier = Arguments[ArgIndexOffset];
-	Identifier.TrimStartInline();
+	Identifier.TrimStartInline(); // Remove whitespaces at the start of the line.
 
 	// Identifier may optionally be in quotes, as it's sometimes a string literal (in UE_COMMAND_EXT), and sometimes stringified by the macro (in UI_COMMAND)
 	// Because this is optional, we don't care if this processing fails
@@ -1631,15 +1726,15 @@ void UGatherTextFromSourceCommandlet::FUICommandMacroDescriptor::TryParse(const
 		TArray<FString> Arguments;
 		if (ParseArgsFromMacro(StripCommentsFromToken(Text, Context), Arguments, Context))
 		{
-			// Need at least 3 arguments
-			if (Arguments.Num() < 3)
+			// Validate that we got the rightnumber of Arguments
+			if (Arguments.Num() < GetMinNumberOfArgument())
 			{
-				UE_LOG(LogGatherTextFromSourceCommandlet, Warning, TEXT("Expected at least 3 arguments for %s macro, but got %d while parsing %s:%d. %s"), *GetToken(), Arguments.Num(), *Context.Filename, Context.LineNumber, *FLocTextHelper::SanitizeLogOutput(Context.LineText.TrimStartAndEnd()));
-			}
-			else
-			{
-				TryParseArgs(Text, Context, Arguments, 0);
+				UE_LOG(LogGatherTextFromSourceCommandlet, Warning, TEXT("Expected at least %d arguments for %s macro, but got %d while parsing %s:%d. %s"), GetMinNumberOfArgument(), *GetToken(), Arguments.Num(), *Context.Filename, Context.LineNumber, *FLocTextHelper::SanitizeLogOutput(Context.LineText.TrimStartAndEnd()));
+				return;
 			}
+			
+			// Parse all arguments found
+			TryParseArgs(Text, Context, Arguments, 0);
 		}
 	}
 }
@@ -1655,14 +1750,12 @@ void UGatherTextFromSourceCommandlet::FUICommandExtMacroDescriptor::TryParse(con
 		if (ParseArgsFromMacro(StripCommentsFromToken(Text, Context), Arguments, Context))
 		{
 			// Need at least 5 arguments
-			if (Arguments.Num() < 5)
-			{
-				UE_LOG(LogGatherTextFromSourceCommandlet, Warning, TEXT("Expected at least 5 arguments for %s macro, but got %d while parsing %s:%d. %s"), *GetToken(), Arguments.Num(), *Context.Filename, Context.LineNumber, *FLocTextHelper::SanitizeLogOutput(Context.LineText.TrimStartAndEnd()));
-			}
-			else
+			if (Arguments.Num() < GetMinNumberOfArgument())
 			{
-				TryParseArgs(Text, Context, Arguments, 2);
+				UE_LOG(LogGatherTextFromSourceCommandlet, Warning, TEXT("Expected at least %d arguments for %s macro, but got %d while parsing %s:%d. %s"), GetMinNumberOfArgument(), *GetToken(), Arguments.Num(), *Context.Filename, Context.LineNumber, *FLocTextHelper::SanitizeLogOutput(Context.LineText.TrimStartAndEnd()));
+				return;
 			}
+			TryParseArgs(Text, Context, Arguments, 2);
 		}
 	}
 }
@@ -1773,9 +1866,9 @@ void UGatherTextFromSourceCommandlet::FStringTableMacroDescriptor::TryParse(cons
 		TArray<FString> Arguments;
 		if (ParseArgsFromMacro(StripCommentsFromToken(Text, Context), Arguments, Context))
 		{
-			if (Arguments.Num() != 2)
+			if (Arguments.Num() != GetMinNumberOfArgument())
 			{
-				UE_LOG(LogGatherTextFromSourceCommandlet, Warning, TEXT("Expected 2 arguments for %s macro, but got %d while parsing %s:%d. %s"), *GetToken(), Arguments.Num(), *Context.Filename, Context.LineNumber, *FLocTextHelper::SanitizeLogOutput(Context.LineText.TrimStartAndEnd()));
+				UE_LOG(LogGatherTextFromSourceCommandlet, Warning, TEXT("Expected %d arguments for %s macro, but got %d while parsing %s:%d. %s"), GetMinNumberOfArgument(), *GetToken(), Arguments.Num(), *Context.Filename, Context.LineNumber, *FLocTextHelper::SanitizeLogOutput(Context.LineText.TrimStartAndEnd()));
 			}
 			else
 			{
@@ -1816,9 +1909,9 @@ void UGatherTextFromSourceCommandlet::FStringTableFromFileMacroDescriptor::TryPa
 		TArray<FString> Arguments;
 		if (ParseArgsFromMacro(StripCommentsFromToken(Text, Context), Arguments, Context))
 		{
-			if (Arguments.Num() != 3)
+			if (Arguments.Num() != GetMinNumberOfArgument())
 			{
-				UE_LOG(LogGatherTextFromSourceCommandlet, Warning, TEXT("Expected 3 arguments for %s macro, but got %d while parsing %s:%d. %s"), *GetToken(), Arguments.Num(), *Context.Filename, Context.LineNumber, *FLocTextHelper::SanitizeLogOutput(Context.LineText.TrimStartAndEnd()));
+				UE_LOG(LogGatherTextFromSourceCommandlet, Warning, TEXT("Expected %d arguments for %s macro, but got %d while parsing %s:%d. %s"), GetMinNumberOfArgument(), *GetToken(), Arguments.Num(), *Context.Filename, Context.LineNumber, *FLocTextHelper::SanitizeLogOutput(Context.LineText.TrimStartAndEnd()));
 			}
 			else
 			{
@@ -1862,9 +1955,9 @@ void UGatherTextFromSourceCommandlet::FStringTableEntryMacroDescriptor::TryParse
 		TArray<FString> Arguments;
 		if (ParseArgsFromMacro(StripCommentsFromToken(Text, Context), Arguments, Context))
 		{
-			if (Arguments.Num() != 3)
+			if (Arguments.Num() != GetMinNumberOfArgument())
 			{
-				UE_LOG(LogGatherTextFromSourceCommandlet, Warning, TEXT("Expected 3 arguments for %s macro, but got %d while parsing %s:%d. %s"), *GetToken(), Arguments.Num(), *Context.Filename, Context.LineNumber, *FLocTextHelper::SanitizeLogOutput(Context.LineText.TrimStartAndEnd()));
+				UE_LOG(LogGatherTextFromSourceCommandlet, Warning, TEXT("Expected %d arguments for %s macro, but got %d while parsing %s:%d. %s"), GetMinNumberOfArgument(), *GetToken(), Arguments.Num(), *Context.Filename, Context.LineNumber, *FLocTextHelper::SanitizeLogOutput(Context.LineText.TrimStartAndEnd()));
 			}
 			else
 			{
@@ -1908,9 +2001,9 @@ void UGatherTextFromSourceCommandlet::FStringTableEntryMetaDataMacroDescriptor::
 		TArray<FString> Arguments;
 		if (ParseArgsFromMacro(StripCommentsFromToken(Text, Context), Arguments, Context))
 		{
-			if (Arguments.Num() != 4)
+			if (Arguments.Num() != GetMinNumberOfArgument())
 			{
-				UE_LOG(LogGatherTextFromSourceCommandlet, Warning, TEXT("Expected 4 arguments for %s macro, but got %d while parsing %s:%d. %s"), *GetToken(), Arguments.Num(), *Context.Filename, Context.LineNumber, *FLocTextHelper::SanitizeLogOutput(Context.LineText.TrimStartAndEnd()));
+				UE_LOG(LogGatherTextFromSourceCommandlet, Warning, TEXT("Expected %d arguments for %s macro, but got %d while parsing %s:%d. %s"), GetMinNumberOfArgument(), *GetToken(), Arguments.Num(), *Context.Filename, Context.LineNumber, *FLocTextHelper::SanitizeLogOutput(Context.LineText.TrimStartAndEnd()));
 			}
 			else
 			{
