COTF: Made the callback optional when cooking packages as it prevents multiple platforms from being cooked at the same time
#rnx

diff --git a/Engine/Source/Editor/UnrealEd/Private/CookOnTheFlyServer.cpp b/Engine/Source/Editor/UnrealEd/Private/CookOnTheFlyServer.cpp
--- a/Engine/Source/Editor/UnrealEd/Private/CookOnTheFlyServer.cpp
+++ b/Engine/Source/Editor/UnrealEd/Private/CookOnTheFlyServer.cpp
@@ -333,34 +333,10 @@ public:
 		if (const ITargetPlatform* TargetPlatform = Cooker.AddCookOnTheFlyPlatform(CookPackageRequest.PlatformName))
 		{
 			Cooker.CookOnTheFlyDeferredInitialize();
-			Cooker.PlatformManager->AddRefCookOnTheFlyPlatform(CookPackageRequest.PlatformName, Cooker);
 			FName StandardFileName(*FPaths::CreateStandardFilename(CookPackageRequest.Filename));
-			FRequest* Request = AllocRequest();
-			Request->CookPackageRequest = MoveTemp(CookPackageRequest);
-
-			FCompletionCallback CookCompleted = [this, Request](FPackageData* PackageData)
-			{
-				const FPlatformData* PlatformData = Cooker.PlatformManager->GetPlatformDataByName(Request->CookPackageRequest.PlatformName);
-				check(PlatformData);
-				const ECookResult CookResult = PackageData
-					? PackageData->GetCookResults(PlatformData->TargetPlatform)
-					: ECookResult::Failed;
-				
-				UE_LOG(LogCook, Verbose, TEXT("Cook request completed, Filename='%s', Platform='%s', CookResult='%d'"),
-					*Request->CookPackageRequest.Filename, *Request->CookPackageRequest.PlatformName.ToString(), uint32(CookResult));
-				
-				if (Request->CookPackageRequest.CompletionCallback)
-				{
-					Request->CookPackageRequest.CompletionCallback(CookResult);
-				}
-				
-				Cooker.PlatformManager->ReleaseCookOnTheFlyPlatform(Request->CookPackageRequest.PlatformName);
-				FreeRequest(Request);
-			};
-
-			UE_LOG(LogCook, Verbose, TEXT("Enqueing cook request, Filename='%s', Platform='%s'"), *Request->CookPackageRequest.Filename, *Request->CookPackageRequest.PlatformName.ToString());
+			UE_LOG(LogCook, Verbose, TEXT("Enqueing cook request, Filename='%s', Platform='%s'"), *CookPackageRequest.Filename, *CookPackageRequest.PlatformName.ToString());
 			Cooker.ExternalRequests->EnqueueUnique(
-				FFilePlatformRequest(StandardFileName, EInstigator::CookOnTheFly, TargetPlatform, MoveTemp(CookCompleted)),
+				FFilePlatformRequest(StandardFileName, EInstigator::CookOnTheFly, TargetPlatform, MoveTemp(CookPackageRequest.CompletionCallback)),
 				true);
 			if (Cooker.ExternalRequests->CookRequestEvent)
 			{
@@ -393,40 +369,7 @@ public:
 	}
 
 private:
-	struct FRequest
-	{
-		FRequest* NextFree = nullptr;
-		UE::Cook::FCookPackageRequest CookPackageRequest;
-	};
-
-	FRequest* AllocRequest()
-	{
-		FScopeLock _(&RequestsCriticalSection);
-		if (FRequest* Request = FirstFree)
-		{
-			FirstFree = Request->NextFree;
-			Request->NextFree = nullptr;
-			return Request;
-		}
-		else
-		{
-			const int32 RequestIndex = Requests.Add();
-			return &Requests[RequestIndex];
-		}
-	}
-
-	void FreeRequest(FRequest* Request)
-	{
-		FScopeLock _(&RequestsCriticalSection);
-		check(Request->NextFree == nullptr);
-		Request->NextFree = FirstFree;
-		FirstFree = Request;
-	}
-
 	UCookOnTheFlyServer& Cooker;
-	FCriticalSection RequestsCriticalSection;
-	TChunkedArray<FRequest> Requests;
-	FRequest* FirstFree = nullptr;
 };
 
 /* UCookOnTheFlyServer functions
diff --git a/Engine/Source/Editor/UnrealEd/Private/Cooker/CookOnTheFlyServerInterface.h b/Engine/Source/Editor/UnrealEd/Private/Cooker/CookOnTheFlyServerInterface.h
--- a/Engine/Source/Editor/UnrealEd/Private/Cooker/CookOnTheFlyServerInterface.h
+++ b/Engine/Source/Editor/UnrealEd/Private/Cooker/CookOnTheFlyServerInterface.h
@@ -13,8 +13,6 @@ class ITargetPlatform;
 namespace UE { namespace Cook
 {
 
-using FCookRequestCompletedCallback = TFunction<void(const ECookResult)>;
-
 using FRecompileShaderCompletedCallback = TFunction<void()>;
 
 using FPrecookedFileList = TMap<FString, FDateTime>;
@@ -29,7 +27,7 @@ struct FCookPackageRequest
 	/* Asset filename to cook. */
 	FString Filename;
 	/** Completion callback. */
-	FCookRequestCompletedCallback CompletionCallback;
+	FCompletionCallback CompletionCallback;
 };
 
 /**
diff --git a/Engine/Source/Editor/UnrealEd/Private/Cooker/NetworkFileCookOnTheFlyRequestManager.cpp b/Engine/Source/Editor/UnrealEd/Private/Cooker/NetworkFileCookOnTheFlyRequestManager.cpp
--- a/Engine/Source/Editor/UnrealEd/Private/Cooker/NetworkFileCookOnTheFlyRequestManager.cpp
+++ b/Engine/Source/Editor/UnrealEd/Private/Cooker/NetworkFileCookOnTheFlyRequestManager.cpp
@@ -63,7 +63,7 @@ private:
 		if (bIsCookable)
 		{
 			FEvent* CookCompletedEvent = FPlatformProcess::GetSynchEventFromPool();
-			FCookRequestCompletedCallback CookRequestCompleted = [this, CookCompletedEvent](const UE::Cook::ECookResult)
+			FCompletionCallback CookRequestCompleted = [this, CookCompletedEvent](UE::Cook::FPackageData*)
 			{
 				CookCompletedEvent->Trigger();
 			};
