ZenCOTF: Refactored request manager to allow packages being recooked
#rnx

diff --git a/Engine/Source/Editor/UnrealEd/Private/Cooker/IoStoreCookOnTheFlyRequestManager.cpp b/Engine/Source/Editor/UnrealEd/Private/Cooker/IoStoreCookOnTheFlyRequestManager.cpp
--- a/Engine/Source/Editor/UnrealEd/Private/Cooker/IoStoreCookOnTheFlyRequestManager.cpp
+++ b/Engine/Source/Editor/UnrealEd/Private/Cooker/IoStoreCookOnTheFlyRequestManager.cpp
@@ -12,226 +12,167 @@
 #include "UObject/SavePackage.h"
 #include "ProfilingDebugging/CountersTrace.h"
 #include "AssetRegistry/IAssetRegistry.h"
+#include "Templates/Function.h"
 
 #if WITH_COTF
 
 DEFINE_LOG_CATEGORY_STATIC(LogCookOnTheFlyTracker, Log, All);
 
-class FCookOnTheFlyPackageTracker
+class FIoStoreCookOnTheFlyRequestManager final
+	: public UE::Cook::ICookOnTheFlyRequestManager
 {
-	enum class EPackageStatus
-	{
-		None,
-		Cooking,
-		Cooked,
-		Failed
-	};
-
-	inline const TCHAR* LexToString(EPackageStatus Status)
-	{
-		switch (Status)
-		{
-			case EPackageStatus::None:
-				return TEXT("None");
-			case EPackageStatus::Cooking:
-				return TEXT("Cooking");
-			case EPackageStatus::Cooked:
-				return TEXT("Cooked");
-			case EPackageStatus::Failed:
-				return TEXT("Failed");
-			default:
-				return TEXT("Unknown");
-		};
-	};
-
-	struct FPackage
-	{
-		FPackageId PackageId;
-		FName PackageName;
-		EPackageStatus Status = EPackageStatus::None;
-		int32 EntryIndex = INDEX_NONE;
-	};
-
 public:
-
-	struct FCompletedPackages
-	{
-		TArray<int32> CookedEntryIndices;
-		TArray<FPackageId> FailedPackages;
-		int32 TotalCooked = 0;
-		int32 TotalFailed = 0;
-		int32 TotalTracked = 0;
-	};
-
-	TArrayView<const int32> GetCookedEntryIndices() const
-	{
-		return CookedEntries;
-	}
-
-	TArrayView<const FName> GetFailedPackages() const
-	{
-		return FailedPackages; 
-	}
-
-	int32 TotalTracked() const
-	{
-		return Packages.Num();
-	}
-
-	EPackageStoreEntryStatus GetStatus(FPackageId PackageId)
+	FIoStoreCookOnTheFlyRequestManager(UE::Cook::ICookOnTheFlyServer& InCookOnTheFlyServer, const IAssetRegistry* AssetRegistry, UE::Cook::FIoStoreCookOnTheFlyServerOptions InOptions)
+		: CookOnTheFlyServer(InCookOnTheFlyServer)
+		, Options(MoveTemp(InOptions))
 	{
-		FPackage& Package = GetPackage(PackageId);
-
-		switch(Package.Status)
+		TArray<FAssetData> AllAssets;
+		AssetRegistry->GetAllAssets(AllAssets, true);
+		for (const FAssetData& AssetData : AllAssets)
 		{
-			case EPackageStatus::None:
-				return EPackageStoreEntryStatus::None;
-			case EPackageStatus::Cooking:
-				return EPackageStoreEntryStatus::Pending;
-			case EPackageStatus::Cooked:
-				return EPackageStoreEntryStatus::Ok;
-			default:
-				return EPackageStoreEntryStatus::Missing;
+			AllKnownPackagesMap.Add(FPackageId::FromName(AssetData.PackageName), AssetData.PackageName);
 		}
 	}
 
-	bool IsCooked(FPackageId PackageId)
-	{
-		FPackage& Package = GetPackage(PackageId);
-		return Package.Status != EPackageStatus::None;
-	}
-
-	void MarkAsRequested(FPackageId PackageId)
-	{
-		FPackage& Package = GetPackage(PackageId);
-		check(Package.Status == EPackageStatus::None);
-		Package.Status = EPackageStatus::Cooking;
-	}
-
-	FCompletedPackages MarkAsFailed(const FName& PackageName)
-	{
-		FCompletedPackages ResultPackages;
-		MarkAsCompleted(PackageName, INDEX_NONE, TArrayView<const FPackageStoreEntryResource>(), ResultPackages);
-		return ResultPackages;
-	}
-
-	void MarkAsCompleted(const FName& PackageName, const int32 EntryIndex, TArrayView<const FPackageStoreEntryResource> Entries, FCompletedPackages& OutPackages)
+private:
+	class FPlatformContext
 	{
-		const int32 NumCookedPackages = CookedEntries.Num();
-		const int32 NumFailedPackages = FailedPackages.Num();
-
-		const FPackageId PackageId = FPackageId::FromName(PackageName);
-
-		FPackage& Package	= GetPackage(PackageId);
-		check(Package.Status == EPackageStatus::None || Package.Status == EPackageStatus::Cooking);
+	public:
+		enum class EPackageStatus
+		{
+			None,
+			Cooking,
+			Cooked,
+			Failed,
+		};
 
-		Package.PackageId	= PackageId;
-		Package.PackageName = PackageName;
-		Package.Status		= EntryIndex != INDEX_NONE ? EPackageStatus::Cooked : EPackageStatus::Failed;
-		Package.EntryIndex	= EntryIndex;
+		struct FPackage
+		{
+			EPackageStatus Status = EPackageStatus::None;
+			FPackageStoreEntryResource Entry;
+		};
 
-		AddCompletedPackage(Package);
-		
-		OutPackages.TotalCooked = CookedEntries.Num();
-		OutPackages.TotalFailed = FailedPackages.Num();
-		OutPackages.TotalTracked = Packages.Num();
+		FPlatformContext(FName InPlatformName, IPackageStoreWriter* InPackageWriter)
+			: PlatformName(InPlatformName)
+			, PackageWriter(InPackageWriter)
+		{
+		}
 
-		if (CookedEntries.Num() > NumCookedPackages || FailedPackages.Num() > NumFailedPackages)
+		FCriticalSection& GetLock()
 		{
-			for (int32 Idx = NumCookedPackages; Idx < CookedEntries.Num(); ++Idx)
-			{
-				OutPackages.CookedEntryIndices.Add(CookedEntries[Idx]);
-			}
+			return CriticalSection;
+		}
 
-			for (int32 Idx = NumFailedPackages; Idx < FailedPackages.Num(); ++Idx)
+		void GetCompletedPackages(UE::ZenCookOnTheFly::Messaging::FCompletedPackages& OutCompletedPackages)
+		{
+			OutCompletedPackages.CookedPackages.Reserve(OutCompletedPackages.CookedPackages.Num() + Packages.Num());
+			for (const auto& KV : Packages)
 			{
-				OutPackages.FailedPackages.Add(FPackageId::FromName(FailedPackages[Idx]));
+				if (KV.Value->Status == EPackageStatus::Cooked)
+				{
+					OutCompletedPackages.CookedPackages.Add(KV.Value->Entry);
+				}
+				else if (KV.Value->Status == EPackageStatus::Failed)
+				{
+					OutCompletedPackages.FailedPackages.Add(KV.Key);
+				}
 			}
 		}
-	}
 
-	void AddExistingPackages(TArrayView<const FPackageStoreEntryResource> Entries, TArrayView<const IPackageStoreWriter::FOplogCookInfo> CookInfos)
-	{
-		Packages.Reserve(Entries.Num());
-		CookedEntries.Reserve(Entries.Num());
-
-		for (int32 EntryIndex = 0, Num = Entries.Num(); EntryIndex < Num; ++EntryIndex)
+		EPackageStoreEntryStatus RequestCook(UE::Cook::ICookOnTheFlyServer& InCookOnTheFlyServer, const FPackageId& PackageId, TFunctionRef<FName()> GetPackageName, FPackageStoreEntryResource& OutEntry)
 		{
-			const FPackageStoreEntryResource& Entry = Entries[EntryIndex];
-			const IPackageStoreWriter::FOplogCookInfo& CookInfo = CookInfos[EntryIndex];
-			const FPackageId PackageId = Entry.GetPackageId();
-
-			FPackage& Package	= GetPackage(PackageId);
-			Package.PackageId	= PackageId;
-			Package.PackageName = Entry.PackageName;
-			Package.EntryIndex	= EntryIndex;
-
-			if (CookInfo.bUpToDate)
+			FPackage& Package = GetPackage(PackageId);
+			if (Package.Status == EPackageStatus::Cooked)
+			{
+				UE_LOG(LogCookOnTheFly, Verbose, TEXT("0x%llX was already cooked"), PackageId.ValueForDebugging());
+				OutEntry = Package.Entry;
+				return EPackageStoreEntryStatus::Ok;
+			}
+			else if (Package.Status == EPackageStatus::Failed)
+			{
+				UE_LOG(LogCookOnTheFly, Verbose, TEXT("0x%llX was already failed"), PackageId.ValueForDebugging());
+				return EPackageStoreEntryStatus::Missing;
+			}
+			else if (Package.Status == EPackageStatus::Cooking)
+			{
+				UE_LOG(LogCookOnTheFly, Verbose, TEXT("0x%llX was already cooking"), PackageId.ValueForDebugging());
+				return EPackageStoreEntryStatus::Pending;
+			}
+			FName PackageName = GetPackageName();
+			if (PackageName.IsNone())
+			{
+				UE_LOG(LogCookOnTheFly, Warning, TEXT("Received cook request for unknown package 0x%llX"), PackageId.ValueForDebugging());
+				return EPackageStoreEntryStatus::Missing;
+			}
+			FString Filename;
+			if (FPackageName::TryConvertLongPackageNameToFilename(PackageName.ToString(), Filename))
 			{
-				CookedEntries.Add(Package.EntryIndex);
-				Package.Status = EPackageStatus::Cooked;
+				UE_LOG(LogCookOnTheFly, Verbose, TEXT("Cooking package 0x%llX '%s'"), PackageId.ValueForDebugging(), *PackageName.ToString());
+				Package.Status = EPackageStatus::Cooking;
+				const bool bEnqueued = InCookOnTheFlyServer.EnqueueCookRequest(UE::Cook::FCookPackageRequest{ PlatformName, Filename });
+				check(bEnqueued);
+				return EPackageStoreEntryStatus::Pending;
 			}
 			else
 			{
-				Package.Status = EPackageStatus::None;
+				UE_LOG(LogCookOnTheFly, Warning, TEXT("Failed to cook package 0x%llX '%s' (File not found)"), PackageId.ValueForDebugging(), *PackageName.ToString());
+				Package.Status = EPackageStatus::Failed;
+				return EPackageStoreEntryStatus::Missing;
 			}
 		}
-	}
-
-private:
 
-	FPackage& GetPackage(FPackageId PackageId)
-	{
-		TUniquePtr<FPackage>& Package = Packages.FindOrAdd(PackageId, MakeUnique<FPackage>());
-		check(Package.IsValid());
-		return *Package;
-	}
-	
-	void AddCompletedPackage(FPackage& Package)
-	{
-		check(Package.Status == EPackageStatus::Cooked || Package.Status == EPackageStatus::Failed);
-
-		if (Package.Status == EPackageStatus::Cooked)
+		void MarkAsFailed(FPackageId PackageId, UE::ZenCookOnTheFly::Messaging::FCompletedPackages& OutCompletedPackages)
 		{
-			CookedEntries.Add(Package.EntryIndex);
+			UE_LOG(LogCookOnTheFly, Warning, TEXT("0x%llX failed"), PackageId.ValueForDebugging());
+			FPackage& Package = GetPackage(PackageId);
+			Package.Status = EPackageStatus::Failed;
+			OutCompletedPackages.FailedPackages.Add(PackageId);
 		}
-		else
+
+		void MarkAsCooked(FPackageId PackageId, const FPackageStoreEntryResource& Entry, UE::ZenCookOnTheFly::Messaging::FCompletedPackages& OutCompletedPackages)
 		{
-			FailedPackages.Add(Package.PackageName);
+			UE_LOG(LogCookOnTheFly, Verbose, TEXT("0x%llX cooked"), PackageId.ValueForDebugging());
+			FPackage& Package = GetPackage(PackageId);
+			Package.Status = EPackageStatus::Cooked;
+			Package.Entry = Entry;
+			OutCompletedPackages.CookedPackages.Add(Entry);
 		}
 
-		UE_LOG(LogCookOnTheFlyTracker, Log, TEXT("'%s' (0x%llX) [%s], Cooked/Failed='%d/%d'"),
-			*Package.PackageName.ToString(), Package.PackageId.Value(), LexToString(Package.Status),
-			CookedEntries.Num(), FailedPackages.Num());
-	}
+		void AddExistingPackages(TArrayView<const FPackageStoreEntryResource> Entries, TArrayView<const IPackageStoreWriter::FOplogCookInfo> CookInfos)
+		{
+			Packages.Reserve(Entries.Num());
 
-	TMap<FPackageId, TUniquePtr<FPackage>> Packages;
-	TArray<int32> CookedEntries;
-	TArray<FName> FailedPackages;
-};
+			for (int32 EntryIndex = 0, Num = Entries.Num(); EntryIndex < Num; ++EntryIndex)
+			{
+				const FPackageStoreEntryResource& Entry = Entries[EntryIndex];
+				const IPackageStoreWriter::FOplogCookInfo& CookInfo = CookInfos[EntryIndex];
+				const FPackageId PackageId = Entry.GetPackageId();
 
-class FIoStoreCookOnTheFlyRequestManager final
-	: public UE::Cook::ICookOnTheFlyRequestManager
-{
-public:
-	FIoStoreCookOnTheFlyRequestManager(UE::Cook::ICookOnTheFlyServer& InCookOnTheFlyServer, const IAssetRegistry* AssetRegistry, UE::Cook::FIoStoreCookOnTheFlyServerOptions InOptions)
-		: CookOnTheFlyServer(InCookOnTheFlyServer)
-		, Options(MoveTemp(InOptions))
-	{
-		TArray<FAssetData> AllAssets;
-		AssetRegistry->GetAllAssets(AllAssets, true);
-		for (const FAssetData& AssetData : AllAssets)
+				FPackage& Package = GetPackage(PackageId);
+
+				if (CookInfo.bUpToDate)
+				{
+					Package.Status = EPackageStatus::Cooked;
+					Package.Entry = Entry;
+				}
+				else
+				{
+					Package.Status = EPackageStatus::None;
+				}
+			}
+		}
+
+		FPackage& GetPackage(FPackageId PackageId)
 		{
-			AllKnownPackagesMap.Add(FPackageId::FromName(AssetData.PackageName), AssetData.PackageName);
+			TUniquePtr<FPackage>& Package = Packages.FindOrAdd(PackageId, MakeUnique<FPackage>());
+			check(Package.IsValid());
+			return *Package;
 		}
-	}
 
-private:
-	struct FPlatformContext
-	{
+	private:
 		FCriticalSection CriticalSection;
 		FName PlatformName;
-		FCookOnTheFlyPackageTracker PackageTracker;
+		TMap<FPackageId, TUniquePtr<FPackage>> Packages;
 		IPackageStoreWriter* PackageWriter = nullptr;
 	};
 
@@ -288,16 +229,14 @@ private:
 			{
 				if (!PlatformContexts.Contains(Client.PlatformName))
 				{
-					TUniquePtr<FPlatformContext>& Context = PlatformContexts.Add(Client.PlatformName, MakeUnique<FPlatformContext>());
-					Context->PlatformName = Client.PlatformName;
 					IPackageStoreWriter* PackageWriter = CookOnTheFlyServer.GetPackageWriter(TargetPlatform).AsPackageStoreWriter();
 					check(PackageWriter); // This class should not be used except when COTFS is using an IPackageStoreWriter
-					Context->PackageWriter = PackageWriter;
+					TUniquePtr<FPlatformContext>& Context = PlatformContexts.Add(Client.PlatformName, MakeUnique<FPlatformContext>(Client.PlatformName, PackageWriter));
 					
 					PackageWriter->GetEntries([&Context](TArrayView<const FPackageStoreEntryResource> Entries,
 						TArrayView<const IPackageStoreWriter::FOplogCookInfo> CookInfos)
 					{
-						Context->PackageTracker.AddExistingPackages(Entries, CookInfos);
+						Context->AddExistingPackages(Entries, CookInfos);
 					});
 
 					PackageWriter->OnCommit().AddRaw(this, &FIoStoreCookOnTheFlyRequestManager::OnPackageCooked);
@@ -356,35 +295,14 @@ private:
 		using namespace UE::Cook;
 		using namespace UE::ZenCookOnTheFly::Messaging;
 
-		FPlatformContext& Context = GetContext(Client.PlatformName);
-		
-		UE::ZenCookOnTheFly::Messaging::FPackageStoreData PackageStoreData;
-		
-		Context.PackageWriter->GetEntries(
-			[&Context, &PackageStoreData](TArrayView<const FPackageStoreEntryResource> Entries,
-				TArrayView<const IPackageStoreWriter::FOplogCookInfo> CookInfos)
+		FCompletedPackages CompletedPackages;
 		{
-			FScopeLock _(&Context.CriticalSection);
-			
-			const TArrayView<const int32> CookedEntryIndices = Context.PackageTracker.GetCookedEntryIndices();
-			const TArrayView<const FName> FailedPackages = Context.PackageTracker.GetFailedPackages();
-			
-			for (int32 EntryIndex : CookedEntryIndices)
-			{
-				check(EntryIndex != INDEX_NONE);
-				PackageStoreData.CookedPackages.Add(Entries[EntryIndex]);
-			}
-			
-			for (const FName& FailedPackageName : FailedPackages)
-			{
-				PackageStoreData.FailedPackages.Add(FPackageId::FromName(FailedPackageName));
-			}
-			
-			PackageStoreData.TotalCookedPackages	= CookedEntryIndices.Num();
-			PackageStoreData.TotalFailedPackages	= FailedPackages.Num();
-		});
-		
-		Response.SetBodyTo(FGetCookedPackagesResponse { MoveTemp(PackageStoreData) });
+			FPlatformContext& Context = GetContext(Client.PlatformName);
+			FScopeLock _(&Context.GetLock());
+			Context.GetCompletedPackages(CompletedPackages);
+		}
+
+		Response.SetBodyTo( MoveTemp(CompletedPackages) );
 		Response.SetStatus(ECookOnTheFlyMessageStatus::Ok);
 
 		return true;
@@ -396,44 +314,20 @@ private:
 
 		TRACE_CPUPROFILER_EVENT_SCOPE(CookOnTheFly::HandleCookPackageRequest);
 
-		FPlatformContext& Context = GetContext(Client.PlatformName);
-		FScopeLock _(&Context.CriticalSection);
-
 		FCookPackageRequest CookRequest = Request.GetBodyAs<FCookPackageRequest>();
-		
-		const bool bIsCooked = Context.PackageTracker.IsCooked(CookRequest.PackageId);
-		if (!bIsCooked)
+		UE_LOG(LogCookOnTheFly, Verbose, TEXT("Received cook request 0x%llX"), CookRequest.PackageId.ValueForDebugging());
+		FPlatformContext& Context = GetContext(Client.PlatformName);
 		{
-			FName PackageName;
+			auto GetPackageNameFunc = [this, &CookRequest]()
 			{
-				FScopeLock LockPackageNames(&AllKnownPackagesCriticalSection);
-				PackageName = AllKnownPackagesMap.FindRef(CookRequest.PackageId);
-			}
-			if (PackageName.IsNone())
-			{
-				UE_LOG(LogCookOnTheFly, Log, TEXT("Received cook request for unknown package 0x%llX"), CookRequest.PackageId.ValueForDebugging());
-				Response.SetBodyTo(FCookPackageResponse { EPackageStoreEntryStatus::Missing });
-			}
-			else
-			{
-				UE_LOG(LogCookOnTheFly, Log, TEXT("Received cook request, PackageName='%s'"), *PackageName.ToString());
-			
-				Context.PackageTracker.MarkAsRequested(CookRequest.PackageId);
-
-				FString Filename;
-				if (FPackageName::TryConvertLongPackageNameToFilename(PackageName.ToString(), Filename))
-				{
-					const bool bEnqueued = CookOnTheFlyServer.EnqueueCookRequest(UE::Cook::FCookPackageRequest{ Client.PlatformName, Filename });
-					check(bEnqueued);
-				}
-				else
-				{
-					UE_LOG(LogCookOnTheFly, Warning, TEXT("Failed to cook package '%s' (File not found)"), *PackageName.ToString());
-					Context.PackageTracker.MarkAsFailed(PackageName);
-				}
-			}
-
-			Response.SetBodyTo(FCookPackageResponse{ Context.PackageTracker.GetStatus(CookRequest.PackageId) });
+				FScopeLock _(&AllKnownPackagesCriticalSection);
+				return AllKnownPackagesMap.FindRef(CookRequest.PackageId);
+			};
+
+			FScopeLock _(&Context.GetLock());
+			FPackageStoreEntryResource Entry;
+			EPackageStoreEntryStatus PackageStatus = Context.RequestCook(CookOnTheFlyServer, CookRequest.PackageId, GetPackageNameFunc, Entry);
+			Response.SetBodyTo(FCookPackageResponse{ PackageStatus, MoveTemp(Entry) });
 		}
 		Response.SetStatus(UE::Cook::ECookOnTheFlyMessageStatus::Ok);
 
@@ -469,68 +363,60 @@ private:
 			ConnectionServer->BroadcastMessage(Message, EventArgs.PlatformName);
 		}
 
-		FCookOnTheFlyPackageTracker::FCompletedPackages NewCompletedPackages;
+		FCompletedPackages NewCompletedPackages;
 		{
 			TRACE_CPUPROFILER_EVENT_SCOPE(CookOnTheFly::MarkAsCooked);
 
 			FPlatformContext& Context = GetContext(EventArgs.PlatformName);
-			FScopeLock _(&Context.CriticalSection);
+			FScopeLock _(&Context.GetLock());
 
-			Context.PackageTracker.MarkAsCompleted(EventArgs.PackageName, EventArgs.EntryIndex, EventArgs.Entries, NewCompletedPackages);
+			if (EventArgs.EntryIndex >= 0)
+			{
+				Context.MarkAsCooked(FPackageId::FromName(EventArgs.PackageName), EventArgs.Entries[EventArgs.EntryIndex], NewCompletedPackages);
+			}
+			else
+			{
+				Context.MarkAsFailed(FPackageId::FromName(EventArgs.PackageName), NewCompletedPackages);
+			}
 		}
-		BroadcastCompletedPackages(EventArgs.PlatformName, MoveTemp(NewCompletedPackages), EventArgs.Entries);
+		BroadcastCompletedPackages(EventArgs.PlatformName, MoveTemp(NewCompletedPackages));
 	}
 
 	void OnPackagesMarkedUpToDate(const IPackageStoreWriter::FMarkUpToDateEventArgs& EventArgs)
 	{
 		TRACE_CPUPROFILER_EVENT_SCOPE(CookOnTheFly::OnPackagesMarkedUpToDate);
 
-		FCookOnTheFlyPackageTracker::FCompletedPackages NewCompletedPackages;
+		UE::ZenCookOnTheFly::Messaging::FCompletedPackages NewCompletedPackages;
 		{
 			TRACE_CPUPROFILER_EVENT_SCOPE(CookOnTheFly::MarkAsCooked);
 
 			FPlatformContext& Context = GetContext(EventArgs.PlatformName);
-			FScopeLock _(&Context.CriticalSection);
+			FScopeLock _(&Context.GetLock());
 
 			for (int32 EntryIndex : EventArgs.PackageIndexes)
 			{
 				FName PackageName = EventArgs.Entries[EntryIndex].PackageName;
-				Context.PackageTracker.MarkAsCompleted(PackageName, EntryIndex, EventArgs.Entries, NewCompletedPackages);
+				Context.MarkAsCooked(FPackageId::FromName(PackageName), EventArgs.Entries[EntryIndex], NewCompletedPackages);
 			}
 		}
-		BroadcastCompletedPackages(EventArgs.PlatformName, MoveTemp(NewCompletedPackages), EventArgs.Entries);
+		BroadcastCompletedPackages(EventArgs.PlatformName, MoveTemp(NewCompletedPackages));
 	}
 
-	void BroadcastCompletedPackages(FName PlatformName,
-		FCookOnTheFlyPackageTracker::FCompletedPackages&& NewCompletedPackages, TArrayView<const FPackageStoreEntryResource> Entries)
+	void BroadcastCompletedPackages(FName PlatformName, UE::ZenCookOnTheFly::Messaging::FCompletedPackages&& NewCompletedPackages)
 	{
 		using namespace UE::Cook;
 		using namespace UE::ZenCookOnTheFly::Messaging;
-		if (NewCompletedPackages.CookedEntryIndices.Num() || NewCompletedPackages.FailedPackages.Num())
+		if (NewCompletedPackages.CookedPackages.Num() || NewCompletedPackages.FailedPackages.Num())
 		{
 			TRACE_CPUPROFILER_EVENT_SCOPE(CookOnTheFly::SendCookedPackagesMessage);
 
-			FPackageStoreData PackageStoreData;
-			PackageStoreData.TotalCookedPackages = NewCompletedPackages.TotalCooked;
-			PackageStoreData.TotalFailedPackages = NewCompletedPackages.TotalFailed;
-			PackageStoreData.FailedPackages = MoveTemp(NewCompletedPackages.FailedPackages);
-
-			for (int32 EntryIndex : NewCompletedPackages.CookedEntryIndices)
-			{
-				check(EntryIndex != INDEX_NONE);
-				PackageStoreData.CookedPackages.Add(Entries[EntryIndex]);
-			}
-
-			UE_LOG(LogCookOnTheFly, Verbose, TEXT("Sending '%s' message, Cooked='%d', Failed='%d', Server total='%d/%d/%d' (Tracked/Cooked/Failed)"),
+			UE_LOG(LogCookOnTheFly, Verbose, TEXT("Sending '%s' message, Cooked='%d', Failed='%d'"),
 				LexToString(ECookOnTheFlyMessage::PackagesCooked),
-				PackageStoreData.CookedPackages.Num(),
-				PackageStoreData.FailedPackages.Num(),
-				NewCompletedPackages.TotalTracked,
-				PackageStoreData.TotalCookedPackages,
-				PackageStoreData.TotalFailedPackages);
+				NewCompletedPackages.CookedPackages.Num(),
+				NewCompletedPackages.FailedPackages.Num());
 
 			FCookOnTheFlyMessage Message(ECookOnTheFlyMessage::PackagesCooked);
-			Message.SetBodyTo<FPackagesCookedMessage>(FPackagesCookedMessage { MoveTemp(PackageStoreData) });
+			Message.SetBodyTo<FCompletedPackages>(MoveTemp(NewCompletedPackages));
 			ConnectionServer->BroadcastMessage(Message, PlatformName);
 		}
 	}
diff --git a/Engine/Source/Runtime/CookOnTheFly/Internal/CookOnTheFlyMessages.h b/Engine/Source/Runtime/CookOnTheFly/Internal/CookOnTheFlyMessages.h
--- a/Engine/Source/Runtime/CookOnTheFly/Internal/CookOnTheFlyMessages.h
+++ b/Engine/Source/Runtime/CookOnTheFly/Internal/CookOnTheFlyMessages.h
@@ -11,14 +11,12 @@
 namespace UE { namespace ZenCookOnTheFly { namespace Messaging
 {
 
-struct FPackageStoreData
+struct FCompletedPackages
 {
 	TArray<FPackageStoreEntryResource> CookedPackages;
 	TArray<FPackageId> FailedPackages;
-	int32 TotalCookedPackages = 0;
-	int32 TotalFailedPackages = 0;
 
-	COOKONTHEFLY_API friend FArchive& operator<<(FArchive& Ar, FPackageStoreData& PackageStoreData);
+	COOKONTHEFLY_API friend FArchive& operator<<(FArchive& Ar, FCompletedPackages& CompletedPackages);
 };
 
 struct FCookPackageRequest
@@ -30,31 +28,12 @@ struct FCookPackageRequest
 
 struct FCookPackageResponse
 {
-	EPackageStoreEntryStatus Status = EPackageStoreEntryStatus::None;
+	EPackageStoreEntryStatus Status;
+	FPackageStoreEntryResource CookedEntry;
 
 	COOKONTHEFLY_API friend FArchive& operator<<(FArchive& Ar, FCookPackageResponse& Response);
 };
 
-struct FPackagesCookedMessage
-{
-	FPackageStoreData PackageStoreData;
-
-	friend FArchive& operator<<(FArchive& Ar, FPackagesCookedMessage& PackagesCookedMessage)
-	{
-		return Ar << PackagesCookedMessage.PackageStoreData;
-	}
-};
-
-struct FGetCookedPackagesResponse
-{
-	FPackageStoreData PackageStoreData;
-	
-	friend FArchive& operator<<(FArchive& Ar, FGetCookedPackagesResponse& GetCookedPackagesResponse)
-	{
-		return Ar << GetCookedPackagesResponse.PackageStoreData;
-	}
-};
-
 }}} // namesapce UE::ZenCookOnTheFly::Messaging
 
 #endif // WITH_COTF
diff --git a/Engine/Source/Runtime/CookOnTheFly/Private/CookOnTheFlyMessages.cpp b/Engine/Source/Runtime/CookOnTheFly/Private/CookOnTheFlyMessages.cpp
--- a/Engine/Source/Runtime/CookOnTheFly/Private/CookOnTheFlyMessages.cpp
+++ b/Engine/Source/Runtime/CookOnTheFly/Private/CookOnTheFlyMessages.cpp
@@ -7,13 +7,11 @@
 namespace UE { namespace ZenCookOnTheFly { namespace Messaging
 {
 
-FArchive& operator<<(FArchive& Ar, FPackageStoreData& PackageStoreData)
+FArchive& operator<<(FArchive& Ar, FCompletedPackages& CompletedPackages)
 {
-	Ar << PackageStoreData.CookedPackages;
-	Ar << PackageStoreData.FailedPackages;
-	Ar << PackageStoreData.TotalCookedPackages;
-	Ar << PackageStoreData.TotalFailedPackages;
-
+	Ar << CompletedPackages.CookedPackages;
+	Ar << CompletedPackages.FailedPackages;
+	
 	return Ar;
 }
 
@@ -26,12 +24,10 @@ FArchive& operator<<(FArchive& Ar, FCookPackageRequest& Request)
 
 FArchive& operator<<(FArchive& Ar, FCookPackageResponse& Response)
 {
-	uint32 Status = static_cast<uint32>(Response.Status);
-	Ar << Status;
-
-	if (Ar.IsLoading())
+	Ar << Response.Status;
+	if (Response.Status == EPackageStoreEntryStatus::Ok)
 	{
-		Response.Status = static_cast<EPackageStoreEntryStatus>(Status);
+		Ar << Response.CookedEntry;
 	}
 
 	return Ar;
diff --git a/Engine/Source/Runtime/StorageServerClient/Private/CookOnTheFlyPackageStore.cpp b/Engine/Source/Runtime/StorageServerClient/Private/CookOnTheFlyPackageStore.cpp
--- a/Engine/Source/Runtime/StorageServerClient/Private/CookOnTheFlyPackageStore.cpp
+++ b/Engine/Source/Runtime/StorageServerClient/Private/CookOnTheFlyPackageStore.cpp
@@ -24,19 +24,18 @@ FCookOnTheFlyPackageStore::FCookOnTheFlyPackageStore(UE::Cook::ICookOnTheFlyServ
 
 	if (Response.IsOk())
 	{
-		FGetCookedPackagesResponse GetCookedPackagesResponse = Response.GetBodyAs<FGetCookedPackagesResponse>();
-		FPackageStoreData& PackageStoreData = GetCookedPackagesResponse.PackageStoreData;
-
+		FCompletedPackages GetCookedPackagesResponse = Response.GetBodyAs<FCompletedPackages>();
+		
 		UE_LOG(LogCookOnTheFly, Log, TEXT("Got '%d' cooked and '%d' failed packages from server"),
-			PackageStoreData.CookedPackages.Num(), PackageStoreData.FailedPackages.Num());
+			GetCookedPackagesResponse.CookedPackages.Num(), GetCookedPackagesResponse.FailedPackages.Num());
 
-		AddPackages(MoveTemp(PackageStoreData.CookedPackages), MoveTemp(PackageStoreData.FailedPackages));
+		AddPackages(MoveTemp(GetCookedPackagesResponse.CookedPackages), MoveTemp(GetCookedPackagesResponse.FailedPackages));
 
 		LastServerActivtyTime = FPlatformTime::Seconds();
 	}
 	else
 	{
-		UE_LOG(LogCookOnTheFly, Warning, TEXT("Failed to send '%s' request"), LexToString(Request.GetHeader().MessageType));
+		UE_LOG(LogCookOnTheFly, Warning, TEXT("Failed to send 'GetCookedPackages' request"));
 	}
 }
 
@@ -52,6 +51,7 @@ EPackageStoreEntryStatus FCookOnTheFlyPackageStore::GetPackageStoreEntry(FPackag
 	using namespace UE::Cook;
 	using namespace UE::ZenCookOnTheFly::Messaging;
 
+	EPackageStoreEntryStatus OriginalStatus;
 	{
 		FScopeLock _(&CriticalSection);
 		FEntryInfo& EntryInfo = PackageIdToEntryInfo.FindOrAdd(PackageId, FEntryInfo());
@@ -75,6 +75,8 @@ EPackageStoreEntryStatus FCookOnTheFlyPackageStore::GetPackageStoreEntry(FPackag
 
 		// The package hasn't been requested, set the status to pending.
 		EntryInfo.Status = EPackageStoreEntryStatus::Pending;
+
+		OriginalStatus = EntryInfo.Status;
 	}
 
 	LastClientActivtyTime = FPlatformTime::Seconds();
@@ -86,7 +88,7 @@ EPackageStoreEntryStatus FCookOnTheFlyPackageStore::GetPackageStoreEntry(FPackag
 
 	if (!Response.IsOk())
 	{
-		UE_LOG(LogCookOnTheFly, Warning, TEXT("Failed to send '%s' request"), LexToString(Request.GetHeader().MessageType));
+		UE_LOG(LogCookOnTheFly, Warning, TEXT("Failed to send 'CookPackage' request"));
 		FScopeLock _(&CriticalSection);
 		FEntryInfo& EntryInfo = PackageIdToEntryInfo.FindChecked(PackageId);
 		EntryInfo.Status = EPackageStoreEntryStatus::Missing;
@@ -94,20 +96,25 @@ EPackageStoreEntryStatus FCookOnTheFlyPackageStore::GetPackageStoreEntry(FPackag
 	}
 
 	FCookPackageResponse CookPackageResponse = Response.GetBodyAs<FCookPackageResponse>();
-
 	{
 		FScopeLock _(&CriticalSection);
 
 		FEntryInfo& EntryInfo = PackageIdToEntryInfo.FindChecked(PackageId);
-		if (CookPackageResponse.Status == EPackageStoreEntryStatus::Missing)
-		{
-			EntryInfo.Status = EPackageStoreEntryStatus::Missing;
-			return EPackageStoreEntryStatus::Missing;
-		}
-		else
+
+		// We might have received a PackagesCooked message while processing the request and in that case the response could be outdated so ignore it
+		if (EntryInfo.Status == OriginalStatus)
 		{
-			return CreatePackageStoreEntry(EntryInfo, OutPackageStoreEntry);
+			EntryInfo.Status = CookPackageResponse.Status;
+			if (EntryInfo.Status == EPackageStoreEntryStatus::Ok)
+			{
+				if (EntryInfo.EntryIndex == INDEX_NONE)
+				{
+					EntryInfo.EntryIndex = PackageEntries.Add();
+				}
+				PackageEntries[EntryInfo.EntryIndex] = MoveTemp(CookPackageResponse.CookedEntry);
+			}
 		}
+		return CreatePackageStoreEntry(EntryInfo, OutPackageStoreEntry);
 	}
 }
 
@@ -132,7 +139,7 @@ void FCookOnTheFlyPackageStore::AddPackages(TArray<FPackageStoreEntryResource> E
 		
 	for (FPackageId FailedPackageId : FailedPackageIds)
 	{
-		UE_LOG(LogCookOnTheFly, Verbose, TEXT("'0x%llX' [Failed]"), FailedPackageId.Value());
+		UE_LOG(LogCookOnTheFly, Warning, TEXT("0x%llX [Failed]"), FailedPackageId.ValueForDebugging());
 		FEntryInfo& EntryInfo = PackageIdToEntryInfo.FindOrAdd(FailedPackageId, FEntryInfo());
 		EntryInfo.Status = EPackageStoreEntryStatus::Missing;
 		PackageStats.Failed++;
@@ -145,16 +152,16 @@ void FCookOnTheFlyPackageStore::AddPackages(TArray<FPackageStoreEntryResource> E
 		const FPackageId PackageId = Entry.GetPackageId();
 		FEntryInfo& EntryInfo = PackageIdToEntryInfo.FindOrAdd(PackageId, FEntryInfo());
 
+		EntryInfo.Status = EPackageStoreEntryStatus::Ok;
 		if (EntryInfo.EntryIndex == INDEX_NONE)
 		{
-			EntryInfo.Status = EPackageStoreEntryStatus::Ok;
 			EntryInfo.EntryIndex = PackageEntries.Add();
-			PackageEntries[EntryInfo.EntryIndex] = MoveTemp(Entry);
-			PackageStats.Cooked++;
-
-			UE_LOG(LogCookOnTheFly, Verbose, TEXT("'%s' [OK] (Cooked/Failed='%d/%d')"),
-				*PackageEntries[EntryInfo.EntryIndex].PackageName.ToString(), PackageStats.Cooked.Load(), PackageStats.Failed.Load());
 		}
+		PackageEntries[EntryInfo.EntryIndex] = MoveTemp(Entry);
+		PackageStats.Cooked++;
+
+		UE_LOG(LogCookOnTheFly, Verbose, TEXT("0x%llX '%s' [OK] (Cooked/Failed='%d/%d')"),
+			PackageId.ValueForDebugging(), *PackageEntries[EntryInfo.EntryIndex].PackageName.ToString(), PackageStats.Cooked.Load(), PackageStats.Failed.Load());
 	}
 }
 
@@ -166,21 +173,13 @@ void FCookOnTheFlyPackageStore::OnCookOnTheFlyMessage(const UE::Cook::FCookOnThe
 	{
 		case UE::Cook::ECookOnTheFlyMessage::PackagesCooked:
 		{
-			FPackagesCookedMessage PackagesCookedMessage = Message.GetBodyAs<FPackagesCookedMessage>();
-			FPackageStoreData& PackageStoreData = PackagesCookedMessage.PackageStoreData;
-
-			UE_LOG(LogCookOnTheFly, Verbose, TEXT("Received '%s' message, Cooked='%d', Failed='%d', Server total='%d/%d' (Cooked/Failed)"),
-				LexToString(Message.GetHeader().MessageType),
-				PackageStoreData.CookedPackages.Num(),
-				PackageStoreData.FailedPackages.Num(),
-				PackageStoreData.TotalCookedPackages,
-				PackageStoreData.TotalFailedPackages);
+			FCompletedPackages PackagesCookedMessage = Message.GetBodyAs<FCompletedPackages>();
 
-			AddPackages(MoveTemp(PackageStoreData.CookedPackages), MoveTemp(PackageStoreData.FailedPackages));
+			UE_LOG(LogCookOnTheFly, Verbose, TEXT("Received 'PackagesCooked' message, Cooked='%d', Failed='%d'"),
+				PackagesCookedMessage.CookedPackages.Num(),
+				PackagesCookedMessage.FailedPackages.Num());
 
-			UE_CLOG(PackageStoreData.TotalCookedPackages != PackageStats.Cooked.Load() || PackageStoreData.TotalFailedPackages != PackageStats.Failed.Load(),
-				LogCookOnTheFly, Warning, TEXT("Client/Server package mismatch, Cooked='%d/%d', Failed='%d/%d' (Client/Server)"),
-				PackageStats.Cooked.Load(), PackageStoreData.TotalCookedPackages, PackageStats.Failed.Load(), PackageStoreData.TotalFailedPackages);
+			AddPackages(MoveTemp(PackagesCookedMessage.CookedPackages), MoveTemp(PackagesCookedMessage.FailedPackages));
 
 			PendingEntriesAdded.Broadcast();
 
