LiveCoding will no longer crash the editor when code changes are detected in the H file but no structural changes are made to a class.

#rnx

diff --git a/Engine/Source/Editor/UnrealEd/Private/Kismet2/KismetReinstanceUtilities.cpp b/Engine/Source/Editor/UnrealEd/Private/Kismet2/KismetReinstanceUtilities.cpp
--- a/Engine/Source/Editor/UnrealEd/Private/Kismet2/KismetReinstanceUtilities.cpp
+++ b/Engine/Source/Editor/UnrealEd/Private/Kismet2/KismetReinstanceUtilities.cpp
@@ -2217,9 +2217,13 @@ void FBlueprintCompileReinstancer::ReplaceInstancesOfClass_Inner(TMap<UClass*, U
 		AllClassInstances.Reserve(InOldToNewClassMap.Num());
 		for (TPair<UClass*, UClass*> OldToNewClass : InOldToNewClassMap)
 		{
-			// We depend on the proper states of this flag in the algorithm that follows
-			OldToNewClass.Key->ClassFlags |= CLASS_NewerVersionExists; // This flag should be set at this point so make sure it is
-			ensure(!OldToNewClass.Value->HasAnyClassFlags(CLASS_NewerVersionExists));
+			// We depend on the proper states of this flag in the algorithm that follows.  Reload system will
+			// have old == new if a code change has been detected, but no structural changes were made. 
+			if (OldToNewClass.Key != OldToNewClass.Value)
+			{
+				OldToNewClass.Key->ClassFlags |= CLASS_NewerVersionExists; // This flag should be set at this point so make sure it is
+				ensure(!OldToNewClass.Value->HasAnyClassFlags(CLASS_NewerVersionExists));
+			}
 
 			// Cache all instances of each class being reinstanced
 			FClassInstances ClassInstances;
