ZenCOTF: Privatized network server and moved it out of runtime code
#rnx

diff --git a/Engine/Source/Editor/UnrealEd/Private/Cooker/IoStoreCookOnTheFlyRequestManager.cpp b/Engine/Source/Editor/UnrealEd/Private/Cooker/IoStoreCookOnTheFlyRequestManager.cpp
--- a/Engine/Source/Editor/UnrealEd/Private/Cooker/IoStoreCookOnTheFlyRequestManager.cpp
+++ b/Engine/Source/Editor/UnrealEd/Private/Cooker/IoStoreCookOnTheFlyRequestManager.cpp
@@ -13,10 +13,431 @@
 #include "ProfilingDebugging/CountersTrace.h"
 #include "AssetRegistry/IAssetRegistry.h"
 #include "Templates/Function.h"
+#include "IPAddress.h"
+#include "SocketSubsystem.h"
+#include "Sockets.h"
+#include "Async/Async.h"
+#include "Serialization/BufferArchive.h"
+#include "NetworkMessage.h"
+
+class FIoStoreCookOnTheFlyNetworkServer
+{
+	static constexpr uint32 ServerSenderId = ~uint32(0);
+	static constexpr double HeartbeatTimeoutInSeconds = 60 * 5;
+
+public:
+	/**
+	 * Connection status
+	 */
+	enum class EConnectionStatus
+	{
+		Disconnected,
+		Connected
+	};
+
+	using FRequestHandler = TFunction<bool(const FName&, const UE::Cook::FCookOnTheFlyRequest&, UE::Cook::FCookOnTheFlyResponse&)>;
+
+	using FClientConnectionHandler = TFunction<bool(const FName&, EConnectionStatus)>;
+
+	using FFillRequest = TFunction<void(FArchive&)>;
+
+	using FProcessResponse = TFunction<bool(FArchive&)>;
+
+	/**
+	 * Cook-on-the-fly connection server options.
+	 */
+	struct FServerOptions
+	{
+		/** The port to listen for new connections. */
+		int32 Port = UE::Cook::DefaultCookOnTheFlyServingPort;
+
+		/** Callback invoked when a client has connected or disconnected. */
+		FClientConnectionHandler HandleClientConnection;
+
+		/** Callback invoked when the server receives a new request. */
+		FRequestHandler HandleRequest;
+	};
+
+	FIoStoreCookOnTheFlyNetworkServer(FServerOptions InOptions)
+		: Options(MoveTemp(InOptions))
+	{
+	}
+
+	~FIoStoreCookOnTheFlyNetworkServer()
+	{
+		StopServer();
+	}
+
+	bool StartServer()
+	{
+		const int32 Port = Options.Port > 0 ? Options.Port : UE::Cook::DefaultCookOnTheFlyServingPort;
+		UE_LOG(LogCookOnTheFly, Log, TEXT("Starting COTF server on port '%d'"), Port);
+
+		check(!bIsRunning);
+		bIsRunning = false;
+		bStopRequested = false;
+
+		ISocketSubsystem& SocketSubsystem = *ISocketSubsystem::Get();
+
+		ListenAddr = SocketSubsystem.GetLocalBindAddr(*GLog);
+		ListenAddr->SetPort(Port);
+
+		// create a server TCP socket
+		Socket = SocketSubsystem.CreateSocket(NAME_Stream, TEXT("COTF-Server"), ListenAddr->GetProtocolType());
+
+		if (!Socket)
+		{
+			UE_LOG(LogCookOnTheFly, Error, TEXT("Could not create listen socket"));
+			return false;
+		}
+
+		Socket->SetReuseAddr();
+
+		if (!Socket->Bind(*ListenAddr))
+		{
+			UE_LOG(LogCookOnTheFly, Error, TEXT("Failed to bind socket to address '%s'"), *ListenAddr->ToString(true));
+			return false;
+		}
+
+		if (!Socket->Listen(16))
+		{
+			UE_LOG(LogCookOnTheFly, Warning, TEXT("Failed to listen on address '%s'"), *ListenAddr->ToString(true));
+			return false;
+		}
+
+		ListenAddr->SetPort(Socket->GetPortNo());
+
+		ServerThread = AsyncThread([this] { return ServerThreadEntry(); }, 8 * 1024, TPri_AboveNormal);
+
+		UE_LOG(LogCookOnTheFly, Display, TEXT("COTF server is ready for client(s) on '%s'!"), *ListenAddr->ToString(true));
+
+		return true;
+	}
+
+	void StopServer()
+	{
+		if (bIsRunning && !bStopRequested)
+		{
+			bStopRequested = true;
+
+			ISocketSubsystem& SocketSubsystem = *ISocketSubsystem::Get();
+
+			while (bIsRunning)
+			{
+				FPlatformProcess::Sleep(0.25f);
+			}
+
+			for (TUniquePtr<FClient>& Client : Clients)
+			{
+				Client->bStopRequested = true;
+				Client->Socket->Close();
+				Client->Thread.Wait();
+				SocketSubsystem.DestroySocket(Client->Socket);
+			}
+
+			Clients.Empty();
+		}
+	}
+
+	bool BroadcastMessage(const UE::Cook::FCookOnTheFlyMessage& Message, const FName& PlatformName = NAME_None)
+	{
+		using namespace UE::Cook;
+
+		FCookOnTheFlyMessageHeader Header = Message.GetHeader();
+
+		Header.MessageType = Header.MessageType | ECookOnTheFlyMessage::Message;
+		Header.MessageStatus = ECookOnTheFlyMessageStatus::Ok;
+		Header.SenderId = ServerSenderId;
+		Header.CorrelationId = NextCorrelationId++;
+		Header.Timestamp = FDateTime::UtcNow().GetTicks();
+
+		FBufferArchive MessagePayload;
+		MessagePayload.Reserve(Message.TotalSize());
+
+		MessagePayload << Header;
+		MessagePayload << const_cast<TArray<uint8>&>(Message.GetBody());
+
+		UE_LOG(LogCookOnTheFly, Verbose, TEXT("Sending: %s, Size='%lld'"), *Header.ToString(), Message.TotalSize());
+
+		TArray<FClient*, TInlineAllocator<4>> ClientsToBroadcast;
+		{
+			FScopeLock _(&ClintsCriticalSection);
+
+			for (TUniquePtr<FClient>& Client : Clients)
+			{
+				if (PlatformName.IsNone() || Client->PlatformName == PlatformName)
+				{
+					ClientsToBroadcast.Add(Client.Get());
+				}
+			}
+		}
+
+		bool bBroadcasted = true;
+		for (FClient* Client : ClientsToBroadcast)
+		{
+			if (!FNFSMessageHeader::WrapAndSendPayload(MessagePayload, FSimpleAbstractSocket_FSocket(Client->Socket)))
+			{
+				UE_LOG(LogCookOnTheFly, Warning, TEXT("Failed to send message '%s' to client '%s' (Id='%d', Platform='%s')"),
+					LexToString(Message.GetHeader().MessageType), *Client->PeerAddr->ToString(true), Client->ClientId, *Client->PlatformName.ToString());
+
+				Client->bIsRunning = false;
+				bBroadcasted = false;
+			}
+
+			Client->LastActivityTime = FPlatformTime::Seconds();
+		}
+
+		return bBroadcasted;
+	}
+
+private:
+	struct FClient
+	{
+		FSocket* Socket = nullptr;
+		TSharedPtr<FInternetAddr> Addr;
+		TSharedPtr<FInternetAddr> PeerAddr;
+		TFuture<void> Thread;
+		TAtomic<bool> bIsRunning{ false };
+		TAtomic<bool> bStopRequested{ false };
+		TAtomic<bool> bIsProcessingRequest{ false };
+		TAtomic<double> LastActivityTime{ 0 };
+		uint32 ClientId = 0;
+		FName PlatformName;
+	};
+
+	void ServerThreadEntry()
+	{
+		using namespace UE::Cook;
+
+		ISocketSubsystem& SocketSubsystem = *ISocketSubsystem::Get();
+		bIsRunning = true;
+
+		while (!bStopRequested)
+		{
+			bool bIsReady = false;
+			if (Socket->WaitForPendingConnection(bIsReady, FTimespan::FromSeconds(0.25f)))
+			{
+				if (bIsReady)
+				{
+					if (FSocket* ClientSocket = Socket->Accept(TEXT("COTF-Client")))
+					{
+						TUniquePtr<FClient> Client = MakeUnique<FClient>();
+						Client->Socket = ClientSocket;
+						Client->Addr = SocketSubsystem.CreateInternetAddr();
+						Client->PeerAddr = SocketSubsystem.CreateInternetAddr();
+						ClientSocket->GetAddress(*Client->Addr);
+						ClientSocket->GetPeerAddress(*Client->PeerAddr);
+						Client->ClientId = NextClientId++;
+						Client->bIsRunning = true;
+						Client->LastActivityTime = FPlatformTime::Seconds();
+						Client->Thread = AsyncThread([this, ConnectedClient = Client.Get()]()
+						{
+							ClientThreadEntry(ConnectedClient);
+						});
+
+						UE_LOG(LogCookOnTheFly, Display, TEXT("New client connected from address '%s' (ID='%d')"), *Client->PeerAddr->ToString(true), Client->ClientId);
+						{
+							FScopeLock _(&ClintsCriticalSection);
+							Clients.Add(MoveTemp(Client));
+						}
+					}
+				}
+			}
+			else
+			{
+				FPlatformProcess::Sleep(0.25f);
+			}
+
+			{
+				FScopeLock _(&ClintsCriticalSection);
+				for (auto It = Clients.CreateIterator(); It; ++It)
+				{
+					TUniquePtr<FClient>& Client = *It;
+
+					const double SecondsSinceLastActivity = FPlatformTime::Seconds() - Client->LastActivityTime.Load();
+
+					if (SecondsSinceLastActivity > HeartbeatTimeoutInSeconds && !Client->bIsProcessingRequest)
+					{
+						Client->LastActivityTime = FPlatformTime::Seconds();
+
+						UE_LOG(LogCookOnTheFly, Display, TEXT("Sending hearbeat message, ClientId='%d', Platform='%s', Address='%s', IdleTime='%.2llf's"),
+							Client->ClientId, *Client->PlatformName.ToString(), *Client->PeerAddr->ToString(true), HeartbeatTimeoutInSeconds);
+
+						FCookOnTheFlyMessage HeartbeatRequest(ECookOnTheFlyMessage::Heartbeat | ECookOnTheFlyMessage::Request);
+						FCookOnTheFlyMessageHeader& Header = HeartbeatRequest.GetHeader();
+
+						Header.MessageStatus = ECookOnTheFlyMessageStatus::Ok;
+						Header.SenderId = ServerSenderId;
+						Header.CorrelationId = Client->ClientId;
+						Header.Timestamp = FDateTime::UtcNow().GetTicks();
+
+						FBufferArchive RequestPayload;
+						RequestPayload.Reserve(HeartbeatRequest.TotalSize());
+						RequestPayload << HeartbeatRequest;
+
+						if (!FNFSMessageHeader::WrapAndSendPayload(RequestPayload, FSimpleAbstractSocket_FSocket(Client->Socket)))
+						{
+							Client->bIsRunning = false;
+							UE_LOG(LogCookOnTheFly, Display, TEXT("Heartbeat [Failed]"));
+						}
+					}
+
+					if (!Client->bIsRunning)
+					{
+						UE_LOG(LogCookOnTheFly, Display, TEXT("Closing connection to client on address '%s' (Id='%d', Platform='%s')"),
+							*Client->PeerAddr->ToString(true), Client->ClientId, *Client->PlatformName.ToString());
+
+						Options.HandleClientConnection(Client->PlatformName, EConnectionStatus::Disconnected);
+
+						Client->Socket->Close();
+						Client->Thread.Wait();
+						SocketSubsystem.DestroySocket(Client->Socket);
+						It.RemoveCurrent();
+					}
+				}
+			}
+		}
+
+		bIsRunning = false;
+	}
+
+	void ClientThreadEntry(FClient* Client)
+	{
+		while (!bStopRequested && !Client->bStopRequested)
+		{
+			Client->LastActivityTime = FPlatformTime::Seconds();
+			if (!ProcesseRequest(*Client))
+			{
+				break;
+			}
+		}
+
+		Client->bIsRunning = false;
+	}
+
+	bool ProcesseRequest(FClient& Client)
+	{
+		using namespace UE::Cook;
+
+		Client.bIsProcessingRequest = false;
+
+		FArrayReader RequestPayload;
+		if (!FNFSMessageHeader::ReceivePayload(RequestPayload, FSimpleAbstractSocket_FSocket(Client.Socket)))
+		{
+			UE_LOG(LogCookOnTheFly, Warning, TEXT("Unable to receive request from client"));
+			return false;
+		}
+
+		Client.bIsProcessingRequest = true;
+
+		FCookOnTheFlyRequest Request;
+		RequestPayload << Request;
+
+		UE_LOG(LogCookOnTheFly, Verbose, TEXT("Received: %s, Size='%lld'"), *Request.GetHeader().ToString(), Request.TotalSize());
 
-#if WITH_COTF
+		EnumRemoveFlags(Request.GetHeader().MessageType, ECookOnTheFlyMessage::TypeFlags);
 
-DEFINE_LOG_CATEGORY_STATIC(LogCookOnTheFlyTracker, Log, All);
+		FCookOnTheFlyResponse Response;
+		bool bRequestOk = false;
+		bool bIsResponse = false;
+
+		switch (Request.GetHeader().MessageType)
+		{
+		case ECookOnTheFlyMessage::Handshake:
+		{
+			ProcessHandshake(Client, Request, Response);
+			bRequestOk = Options.HandleClientConnection(Client.PlatformName, EConnectionStatus::Connected);
+			break;
+		}
+		case ECookOnTheFlyMessage::Heartbeat:
+		{
+			const bool bHeartbeatOk = Request.GetHeader().CorrelationId == Client.ClientId;
+
+			UE_LOG(LogCookOnTheFly, Display, TEXT("Heartbeat [%s], ClientId='%d', Platform='%s', Address='%s'"),
+				bHeartbeatOk ? TEXT("Ok") : TEXT("Failed"), Client.ClientId, *Client.PlatformName.ToString(), *Client.PeerAddr->ToString(true));
+
+			bRequestOk = bHeartbeatOk;
+			bIsResponse = true;
+			break;
+		}
+		default:
+		{
+			bRequestOk = Options.HandleRequest(Client.PlatformName, Request, Response);
+			break;
+		}
+		}
+
+		if (bRequestOk && !bIsResponse)
+		{
+			FCookOnTheFlyMessageHeader& ResponseHeader = Response.GetHeader();
+
+			ResponseHeader.MessageType = Request.GetHeader().MessageType | ECookOnTheFlyMessage::Response;
+			ResponseHeader.SenderId = ServerSenderId;
+			ResponseHeader.CorrelationId = Request.GetHeader().CorrelationId;
+			ResponseHeader.Timestamp = FDateTime::UtcNow().GetTicks();
+
+			Response.SetHeader(ResponseHeader);
+
+			FBufferArchive ResponsePayload;
+			ResponsePayload.Reserve(Response.TotalSize());
+
+			ResponsePayload << Response;
+			bRequestOk = FNFSMessageHeader::WrapAndSendPayload(ResponsePayload, FSimpleAbstractSocket_FSocket(Client.Socket));
+		}
+
+		return bRequestOk;
+	}
+
+	void ProcessHandshake(FClient& Client, UE::Cook::FCookOnTheFlyRequest& HandshakeRequest, UE::Cook::FCookOnTheFlyResponse& Response)
+	{
+		FString PlatformName;
+		FString ProjectName;
+
+		{
+			TUniquePtr<FArchive> Ar = HandshakeRequest.ReadBody();
+			*Ar << PlatformName;
+			*Ar << ProjectName;
+		}
+
+		if (PlatformName.Len())
+		{
+			Client.PlatformName = FName(*PlatformName);
+
+			Response.SetBodyTo(Client.ClientId);
+			Response.SetStatus(UE::Cook::ECookOnTheFlyMessageStatus::Ok);
+		}
+		else
+		{
+			Response.SetStatus(UE::Cook::ECookOnTheFlyMessageStatus::Error);
+		}
+	}
+
+	FClient* GetClientById(uint32 ClientId)
+	{
+		FScopeLock _(&ClintsCriticalSection);
+
+		for (TUniquePtr<FClient>& Client : Clients)
+		{
+			if (Client->ClientId == ClientId)
+			{
+				return Client.Get();
+			}
+		}
+
+		return nullptr;
+	}
+
+	FIoStoreCookOnTheFlyNetworkServer::FServerOptions Options;
+	TFuture<void> ServerThread;
+	TSharedPtr<FInternetAddr> ListenAddr;
+	FSocket* Socket = nullptr;
+	FCriticalSection ClintsCriticalSection;
+	TArray<TUniquePtr<FClient>> Clients;
+	TAtomic<bool> bIsRunning{ false };
+	TAtomic<bool> bStopRequested{ false };
+	uint32 NextClientId = 1;
+	TAtomic<uint32> NextCorrelationId{ 1 };
+};
 
 class FIoStoreCookOnTheFlyRequestManager final
 	: public UE::Cook::ICookOnTheFlyRequestManager
@@ -184,16 +605,16 @@ private:
 
 		const int32 Port = Options.Port > 0 ? Options.Port : UE::Cook::DefaultCookOnTheFlyServingPort;
 
-		ConnectionServer = CookOnTheFlyModule.CreateConnectionServer(FCookOnTheFlyServerOptions
+		ConnectionServer = MakeUnique<FIoStoreCookOnTheFlyNetworkServer>(FIoStoreCookOnTheFlyNetworkServer::FServerOptions
 		{
 			Port,
-			[this](FCookOnTheFlyClient Client, ECookOnTheFlyConnectionStatus ConnectionStatus)
+			[this](const FName& PlatformName, FIoStoreCookOnTheFlyNetworkServer::EConnectionStatus ConnectionStatus)
 			{
-				return HandleClientConnection(Client, ConnectionStatus);
+				return HandleClientConnection(PlatformName, ConnectionStatus);
 			},
-			[this](FCookOnTheFlyClient Client, const FCookOnTheFlyRequest& Request, FCookOnTheFlyResponse& Response)
+			[this](const FName& PlatformName, const FCookOnTheFlyRequest& Request, FCookOnTheFlyResponse& Response)
 			{ 
-				return HandleClientRequest(Client, Request, Response);
+				return HandleClientRequest(PlatformName, Request, Response);
 			}
 		});
 
@@ -218,20 +639,20 @@ private:
 	}
 
 private:
-	bool HandleClientConnection(UE::Cook::FCookOnTheFlyClient Client, UE::Cook::ECookOnTheFlyConnectionStatus ConnectionStatus)
+	bool HandleClientConnection(const FName& PlatformName, FIoStoreCookOnTheFlyNetworkServer::EConnectionStatus ConnectionStatus)
 	{
 		FScopeLock _(&ContextsCriticalSection);
 
-		if (ConnectionStatus == UE::Cook::ECookOnTheFlyConnectionStatus::Connected)
+		if (ConnectionStatus == FIoStoreCookOnTheFlyNetworkServer::EConnectionStatus::Connected)
 		{
-			const ITargetPlatform* TargetPlatform = CookOnTheFlyServer.AddPlatform(Client.PlatformName);
+			const ITargetPlatform* TargetPlatform = CookOnTheFlyServer.AddPlatform(PlatformName);
 			if (TargetPlatform)
 			{
-				if (!PlatformContexts.Contains(Client.PlatformName))
+				if (!PlatformContexts.Contains(PlatformName))
 				{
 					IPackageStoreWriter* PackageWriter = CookOnTheFlyServer.GetPackageWriter(TargetPlatform).AsPackageStoreWriter();
 					check(PackageWriter); // This class should not be used except when COTFS is using an IPackageStoreWriter
-					TUniquePtr<FPlatformContext>& Context = PlatformContexts.Add(Client.PlatformName, MakeUnique<FPlatformContext>(Client.PlatformName, PackageWriter));
+					TUniquePtr<FPlatformContext>& Context = PlatformContexts.Add(PlatformName, MakeUnique<FPlatformContext>(PlatformName, PackageWriter));
 					
 					PackageWriter->GetEntries([&Context](TArrayView<const FPackageStoreEntryResource> Entries,
 						TArrayView<const IPackageStoreWriter::FOplogCookInfo> CookInfos)
@@ -250,54 +671,54 @@ private:
 		}
 		else
 		{
-			CookOnTheFlyServer.RemovePlatform(Client.PlatformName);
+			CookOnTheFlyServer.RemovePlatform(PlatformName);
 			return true;
 		}
 	}
 
-	bool HandleClientRequest(UE::Cook::FCookOnTheFlyClient Client, const UE::Cook::FCookOnTheFlyRequest& Request, UE::Cook::FCookOnTheFlyResponse& Response)
+	bool HandleClientRequest(const FName& PlatformName, const UE::Cook::FCookOnTheFlyRequest& Request, UE::Cook::FCookOnTheFlyResponse& Response)
 	{
 		bool bRequestOk = false;
 
 		const double StartTime = FPlatformTime::Seconds();
 
-		UE_LOG(LogCookOnTheFly, Verbose, TEXT("New request, Type='%s', Client='%s (%d)'"), LexToString(Request.GetHeader().MessageType), *Client.PlatformName.ToString(), Client.ClientId);
+		UE_LOG(LogCookOnTheFly, Verbose, TEXT("New request, Type='%s', Client='%s'"), LexToString(Request.GetHeader().MessageType), *PlatformName.ToString());
 
 		switch (Request.GetHeader().MessageType)
 		{
 			case UE::Cook::ECookOnTheFlyMessage::CookPackage:
-				bRequestOk = HandleCookPackageRequest(Client, Request, Response);
+				bRequestOk = HandleCookPackageRequest(PlatformName, Request, Response);
 				break;
 			case UE::Cook::ECookOnTheFlyMessage::GetCookedPackages:
-				bRequestOk = HandleGetCookedPackagesRequest(Client, Request, Response);
+				bRequestOk = HandleGetCookedPackagesRequest(PlatformName, Request, Response);
 				break;
 			case UE::Cook::ECookOnTheFlyMessage::RecompileShaders:
-				bRequestOk = HandleRecompileShadersRequest(Client, Request, Response);
+				bRequestOk = HandleRecompileShadersRequest(PlatformName, Request, Response);
 				break;
 			default:
-				UE_LOG(LogCookOnTheFly, Fatal, TEXT("Unknown request, Type='%s', Client='%s (%d)'"), LexToString(Request.GetHeader().MessageType), *Client.PlatformName.ToString(), Client.ClientId);
+				UE_LOG(LogCookOnTheFly, Fatal, TEXT("Unknown request, Type='%s', Client='%s'"), LexToString(Request.GetHeader().MessageType), *PlatformName.ToString());
 				break;
 		}
 
 		const double Duration = FPlatformTime::Seconds() - StartTime;
 
-		UE_LOG(LogCookOnTheFly, Verbose, TEXT("Request handled, Type='%s', Client='%s (%u)', Status='%s', Duration='%.6lfs'"),
+		UE_LOG(LogCookOnTheFly, Verbose, TEXT("Request handled, Type='%s', Client='%s', Status='%s', Duration='%.6lfs'"),
 			LexToString(Request.GetHeader().MessageType),
-			*Client.PlatformName.ToString(), Client.ClientId,
+			*PlatformName.ToString(),
 			bRequestOk ? TEXT("Ok") : TEXT("Failed"),
 			Duration);
 
 		return bRequestOk;
 	}
 
-	bool HandleGetCookedPackagesRequest(UE::Cook::FCookOnTheFlyClient Client, const UE::Cook::FCookOnTheFlyRequest& Request, UE::Cook::FCookOnTheFlyResponse& Response)
+	bool HandleGetCookedPackagesRequest(const FName& PlatformName, const UE::Cook::FCookOnTheFlyRequest& Request, UE::Cook::FCookOnTheFlyResponse& Response)
 	{
 		using namespace UE::Cook;
 		using namespace UE::ZenCookOnTheFly::Messaging;
 
 		FCompletedPackages CompletedPackages;
 		{
-			FPlatformContext& Context = GetContext(Client.PlatformName);
+			FPlatformContext& Context = GetContext(PlatformName);
 			FScopeLock _(&Context.GetLock());
 			Context.GetCompletedPackages(CompletedPackages);
 		}
@@ -308,7 +729,7 @@ private:
 		return true;
 	}
 
-	bool HandleCookPackageRequest(UE::Cook::FCookOnTheFlyClient Client, const UE::Cook::FCookOnTheFlyRequest& Request, UE::Cook::FCookOnTheFlyResponse& Response)
+	bool HandleCookPackageRequest(const FName& PlatformName, const UE::Cook::FCookOnTheFlyRequest& Request, UE::Cook::FCookOnTheFlyResponse& Response)
 	{
 		using namespace UE::ZenCookOnTheFly::Messaging;
 
@@ -316,7 +737,7 @@ private:
 
 		FCookPackageRequest CookRequest = Request.GetBodyAs<FCookPackageRequest>();
 		UE_LOG(LogCookOnTheFly, Verbose, TEXT("Received cook request 0x%llX"), CookRequest.PackageId.ValueForDebugging());
-		FPlatformContext& Context = GetContext(Client.PlatformName);
+		FPlatformContext& Context = GetContext(PlatformName);
 		{
 			auto GetPackageNameFunc = [this, &CookRequest]()
 			{
@@ -421,7 +842,7 @@ private:
 		}
 	}
 
-	bool HandleRecompileShadersRequest(UE::Cook::FCookOnTheFlyClient Client, const UE::Cook::FCookOnTheFlyRequest& Request, UE::Cook::FCookOnTheFlyResponse& Response)
+	bool HandleRecompileShadersRequest(const FName& PlatformName, const UE::Cook::FCookOnTheFlyRequest& Request, UE::Cook::FCookOnTheFlyResponse& Response)
 	{
 		TRACE_CPUPROFILER_EVENT_SCOPE(CookOnTheFly::HandleRecompileShadersRequest);
 
@@ -429,7 +850,7 @@ private:
 		TArray<uint8> MeshMaterialMaps;
 		TArray<uint8> GlobalShaderMap;
 
-		FShaderRecompileData RecompileData(Client.PlatformName.ToString(), &RecompileModifiedFiles, &MeshMaterialMaps, &GlobalShaderMap);
+		FShaderRecompileData RecompileData(PlatformName.ToString(), &RecompileModifiedFiles, &MeshMaterialMaps, &GlobalShaderMap);
 		{
 			int32 iShaderPlatform = static_cast<int32>(RecompileData.ShaderPlatform);
 			TUniquePtr<FArchive> Ar = Request.ReadBody();
@@ -481,7 +902,7 @@ private:
 
 	UE::Cook::ICookOnTheFlyServer& CookOnTheFlyServer;
 	UE::Cook::FIoStoreCookOnTheFlyServerOptions Options;
-	TUniquePtr<UE::Cook::ICookOnTheFlyConnectionServer> ConnectionServer;
+	TUniquePtr<FIoStoreCookOnTheFlyNetworkServer> ConnectionServer;
 	FCriticalSection ContextsCriticalSection;
 	TMap<FName, TUniquePtr<FPlatformContext>> PlatformContexts;
 	FCriticalSection AllKnownPackagesCriticalSection;
@@ -498,4 +919,3 @@ TUniquePtr<ICookOnTheFlyRequestManager> MakeIoStoreCookOnTheFlyRequestManager(IC
 
 }}
 
-#endif
diff --git a/Engine/Source/Runtime/CookOnTheFly/Internal/CookOnTheFly.h b/Engine/Source/Runtime/CookOnTheFly/Internal/CookOnTheFly.h
--- a/Engine/Source/Runtime/CookOnTheFly/Internal/CookOnTheFly.h
+++ b/Engine/Source/Runtime/CookOnTheFly/Internal/CookOnTheFly.h
@@ -246,49 +246,6 @@ using FCookOnTheFlyRequest = FCookOnTheFlyMessage;
 
 using FCookOnTheFlyResponse = FCookOnTheFlyMessage;
 
-/**
- * Connection status
- */
-enum class ECookOnTheFlyConnectionStatus
-{
-	Disconnected,
-	Connected
-};
-
-/**
- * A connected cook-on-the-fly client.
- */
-struct FCookOnTheFlyClient
-{
-	/** A client ID set by the server. */
-	uint32 ClientId = 0;
-	/** The platform. */
-	FName PlatformName;
-};
-
-using FCookOnTheFlyRequestHandler = TFunction<bool(FCookOnTheFlyClient, const FCookOnTheFlyRequest&, FCookOnTheFlyResponse&)>;
-
-using FCookOnTheFlyClientConnectionHandler = TFunction<bool(FCookOnTheFlyClient, ECookOnTheFlyConnectionStatus)>;
-
-using FFillRequest = TFunction<void(FArchive&)>;
-
-using FProcessResponse = TFunction<bool(FArchive&)>;
-
-/**
- * Cook-on-the-fly connection server options.
- */
-struct FCookOnTheFlyServerOptions
-{
-	/** The port to listen for new connections. */
-	int32 Port = DefaultCookOnTheFlyServingPort;
-
-	/** Callback invoked when a client has connected or disconnected. */
-	FCookOnTheFlyClientConnectionHandler HandleClientConnection;
-
-	/** Callback invoked when the server receives a new request. */
-	FCookOnTheFlyRequestHandler HandleRequest;
-};
-
 /**
  * Cook-on-the-fly host address.
  */
@@ -302,24 +259,6 @@ struct FCookOnTheFlyHostOptions
 	FTimespan ServerStartupWaitTime;
 };
 
-/**
- * A connection server used to communicate with cook-on-the-fly clients.
- */
-class ICookOnTheFlyConnectionServer
-{
-public:
-	virtual ~ICookOnTheFlyConnectionServer() { }
-
-	/** Start the cook-on-the-fly server. */
-	virtual bool StartServer() = 0;
-
-	/** Stop the cook-on-the-fly server. */
-	virtual void StopServer() = 0;
-
-	/** Broadcast message to all connected clients for the specified platform. */
-	virtual bool BroadcastMessage(const FCookOnTheFlyMessage& Message, const FName& PlatformName = NAME_None) = 0;
-};
-
 /**
  * A connection used to communicate with the cook-on-the-fly server.
  */
@@ -364,13 +303,6 @@ class ICookOnTheFlyModule
 public:
 	virtual ~ICookOnTheFlyModule() { }
 
-	/**
-	 * Creates a new instance of a cook-on-the-fly connection server.
-	 *
-	 * @param Options The cook-on-the-fly connection server options.
-	 */
-	virtual TUniquePtr<ICookOnTheFlyConnectionServer> CreateConnectionServer(FCookOnTheFlyServerOptions Options) = 0;
-
 	/**
 	 * Connect to the cook-on-the-fly server.
 	 *
diff --git a/Engine/Source/Runtime/CookOnTheFly/Private/CookOnTheFlyConnectionServer.cpp b/Engine/Source/Runtime/CookOnTheFly/Private/CookOnTheFlyConnectionServer.cpp
deleted file mode 100644
diff --git a/Engine/Source/Runtime/CookOnTheFly/Private/CookOnTheFlyModule.cpp b/Engine/Source/Runtime/CookOnTheFly/Private/CookOnTheFlyModule.cpp
--- a/Engine/Source/Runtime/CookOnTheFly/Private/CookOnTheFlyModule.cpp
+++ b/Engine/Source/Runtime/CookOnTheFly/Private/CookOnTheFlyModule.cpp
@@ -8,8 +8,6 @@ DEFINE_LOG_CATEGORY(LogCookOnTheFly);
 namespace UE { namespace Cook
 {
 
-TUniquePtr<ICookOnTheFlyConnectionServer> MakeCookOnTheFlyConnectionServer(FCookOnTheFlyServerOptions Options);
-	
 TUniquePtr<ICookOnTheFlyServerConnection> MakeServerConnection(const UE::Cook::FCookOnTheFlyHostOptions& HostOptions);
 
 }}
@@ -33,11 +31,6 @@ public:
 	// ICookOnTheFlyModule interface
 	//
 
-	virtual TUniquePtr<UE::Cook::ICookOnTheFlyConnectionServer> CreateConnectionServer(UE::Cook::FCookOnTheFlyServerOptions Options) override
-	{
-		return UE::Cook::MakeCookOnTheFlyConnectionServer(Options);
-	}
-
 	virtual TUniquePtr<UE::Cook::ICookOnTheFlyServerConnection> ConnectToServer(const UE::Cook::FCookOnTheFlyHostOptions& HostOptions) override
 	{
 		return UE::Cook::MakeServerConnection(HostOptions);
