ZenCOTF: Allow connections from the editor (no target platform)
#rnx

diff --git a/Engine/Source/Editor/UnrealEd/Private/Cooker/IoStoreCookOnTheFlyRequestManager.cpp b/Engine/Source/Editor/UnrealEd/Private/Cooker/IoStoreCookOnTheFlyRequestManager.cpp
--- a/Engine/Source/Editor/UnrealEd/Private/Cooker/IoStoreCookOnTheFlyRequestManager.cpp
+++ b/Engine/Source/Editor/UnrealEd/Private/Cooker/IoStoreCookOnTheFlyRequestManager.cpp
@@ -402,14 +402,9 @@ private:
 		if (PlatformName.Len())
 		{
 			Client.PlatformName = FName(*PlatformName);
-
-			Response.SetBodyTo(Client.ClientId);
-			Response.SetStatus(UE::Cook::ECookOnTheFlyMessageStatus::Ok);
-		}
-		else
-		{
-			Response.SetStatus(UE::Cook::ECookOnTheFlyMessageStatus::Error);
 		}
+		Response.SetBodyTo(Client.ClientId);
+		Response.SetStatus(UE::Cook::ECookOnTheFlyMessageStatus::Ok);
 	}
 
 	FClient* GetClientById(uint32 ClientId)
@@ -643,6 +638,11 @@ private:
 	{
 		FScopeLock _(&ContextsCriticalSection);
 
+		if (PlatformName.IsNone())
+		{
+			return true;
+		}
+
 		if (ConnectionStatus == FIoStoreCookOnTheFlyNetworkServer::EConnectionStatus::Connected)
 		{
 			const ITargetPlatform* TargetPlatform = CookOnTheFlyServer.AddPlatform(PlatformName);
@@ -716,6 +716,13 @@ private:
 		using namespace UE::Cook;
 		using namespace UE::ZenCookOnTheFly::Messaging;
 
+		if (PlatformName.IsNone())
+		{
+			UE_LOG(LogCookOnTheFly, Warning, TEXT("GetCookedPackagesRequest from editor client"));
+			Response.SetStatus(ECookOnTheFlyMessageStatus::Error);
+			return true;
+		}
+		
 		FCompletedPackages CompletedPackages;
 		{
 			FPlatformContext& Context = GetContext(PlatformName);
@@ -733,6 +740,13 @@ private:
 	{
 		using namespace UE::ZenCookOnTheFly::Messaging;
 
+		if (PlatformName.IsNone())
+		{
+			UE_LOG(LogCookOnTheFly, Warning, TEXT("CookPackageRequest from editor client"));
+			Response.SetStatus(UE::Cook::ECookOnTheFlyMessageStatus::Error);
+			return true;
+		}
+
 		TRACE_CPUPROFILER_EVENT_SCOPE(CookOnTheFly::HandleCookPackageRequest);
 
 		FCookPackageRequest CookRequest = Request.GetBodyAs<FCookPackageRequest>();
@@ -846,6 +860,13 @@ private:
 	{
 		TRACE_CPUPROFILER_EVENT_SCOPE(CookOnTheFly::HandleRecompileShadersRequest);
 
+		if (PlatformName.IsNone())
+		{
+			UE_LOG(LogCookOnTheFly, Warning, TEXT("RecompileShadersRequest from editor client"));
+			Response.SetStatus(UE::Cook::ECookOnTheFlyMessageStatus::Error);
+			return true;
+		}
+
 		TArray<FString> RecompileModifiedFiles;
 		TArray<uint8> MeshMaterialMaps;
 		TArray<uint8> GlobalShaderMap;
diff --git a/Engine/Source/Runtime/CookOnTheFly/Private/CookOnTheFlyServerConnection.cpp b/Engine/Source/Runtime/CookOnTheFly/Private/CookOnTheFlyServerConnection.cpp
--- a/Engine/Source/Runtime/CookOnTheFly/Private/CookOnTheFlyServerConnection.cpp
+++ b/Engine/Source/Runtime/CookOnTheFly/Private/CookOnTheFlyServerConnection.cpp
@@ -212,10 +212,14 @@ public:
 
 		FCookOnTheFlyMessage HandshakeRequest(ECookOnTheFlyMessage::Handshake | ECookOnTheFlyMessage::Request);
 		{
-			TArray<FString> TargetPlatformNames;
-			FPlatformMisc::GetValidTargetPlatforms(TargetPlatformNames);
-			check(TargetPlatformNames.Num() > 0);
-			FString PlatformName(MoveTemp(TargetPlatformNames[0]));
+			FString PlatformName;
+			if (FPlatformProperties::RequiresCookedData())
+			{
+				TArray<FString> TargetPlatformNames;
+				FPlatformMisc::GetValidTargetPlatforms(TargetPlatformNames);
+				check(TargetPlatformNames.Num() > 0);
+				PlatformName = MoveTemp(TargetPlatformNames[0]);
+			}
 			FString ProjectName(FApp::GetProjectName());
 
 			TUniquePtr<FArchive> Ar = HandshakeRequest.WriteBody();
