Cooker: Turn off COOK_CHECKSLOW_PACKAGEDATA since it shows up in profiles and has not broken for > 6 months.
#rnx

diff --git a/Engine/Source/Editor/UnrealEd/Private/Cooker/CookTypes.h b/Engine/Source/Editor/UnrealEd/Private/Cooker/CookTypes.h
--- a/Engine/Source/Editor/UnrealEd/Private/Cooker/CookTypes.h
+++ b/Engine/Source/Editor/UnrealEd/Private/Cooker/CookTypes.h
@@ -17,7 +17,7 @@
 class ITargetPlatform;
 namespace UE::DerivedData { class FBuildDefinition; }
 
-#define COOK_CHECKSLOW_PACKAGEDATA 1
+#define COOK_CHECKSLOW_PACKAGEDATA 0
 #define DEBUG_COOKONTHEFLY 0
 
 /** A BaseKeyFuncs for Maps and Sets with a quicker hash function for pointers than TDefaultMapKeyFuncs */
