Interchange
- Interchange now adds its supported extensions to the import and import into level file pickers.
- Added a GetSupportedFormats on the Interchange Translators.
- UInterchangeTranslatorBase::CanImportSource now defaults to checking the file extension and matching it against the supported formats.
- Added a translator type on the translators to indicate what type of data they can translate (assets or actors for now). This is used by the file pickers to indicate which extensions they want to gather.

diff --git a/Engine/Plugins/Experimental/Interchange/Runtime/Source/Import/Private/Fbx/InterchangeFbxTranslator.cpp b/Engine/Plugins/Experimental/Interchange/Runtime/Source/Import/Private/Fbx/InterchangeFbxTranslator.cpp
--- a/Engine/Plugins/Experimental/Interchange/Runtime/Source/Import/Private/Fbx/InterchangeFbxTranslator.cpp
+++ b/Engine/Plugins/Experimental/Interchange/Runtime/Source/Import/Private/Fbx/InterchangeFbxTranslator.cpp
@@ -23,17 +23,20 @@
 #include "UObject/GCObjectScopeGuard.h"
 
 
-UInterchangeFbxTranslator::UInterchangeFbxTranslator(const class FObjectInitializer& ObjectInitializer)
+UInterchangeFbxTranslator::UInterchangeFbxTranslator()
 {
 	Dispatcher = nullptr;
 }
 
-bool UInterchangeFbxTranslator::CanImportSourceData(const UInterchangeSourceData* InSourceData) const
+EInterchangeTranslatorType UInterchangeFbxTranslator::GetTranslatorType() const
 {
-	const bool bIncludeDot = false;
-	FString Extension = FPaths::GetExtension(InSourceData->GetFilename(), bIncludeDot);
-	FString FbxExtension = (TEXT("fbx"));
-	return FbxExtension.StartsWith(Extension,ESearchCase::IgnoreCase);
+	return EInterchangeTranslatorType::Scenes;
+}
+
+TArray<FString> UInterchangeFbxTranslator::GetSupportedFormats() const
+{
+	TArray<FString> Formats {TEXT("fbx;Filmbox")};
+	return Formats;
 }
 
 bool UInterchangeFbxTranslator::Translate(UInterchangeBaseNodeContainer& BaseNodeContainer) const
diff --git a/Engine/Plugins/Experimental/Interchange/Runtime/Source/Import/Private/Gltf/InterchangeGltfTranslator.cpp b/Engine/Plugins/Experimental/Interchange/Runtime/Source/Import/Private/Gltf/InterchangeGltfTranslator.cpp
--- a/Engine/Plugins/Experimental/Interchange/Runtime/Source/Import/Private/Gltf/InterchangeGltfTranslator.cpp
+++ b/Engine/Plugins/Experimental/Interchange/Runtime/Source/Import/Private/Gltf/InterchangeGltfTranslator.cpp
@@ -291,23 +291,19 @@ void UInterchangeGltfTranslator::HandleGltfMaterial( UInterchangeBaseNodeContain
 	}
 }
 
-bool UInterchangeGltfTranslator::CanImportSourceData( const UInterchangeSourceData* InSourceData ) const
+EInterchangeTranslatorType UInterchangeGltfTranslator::GetTranslatorType() const
 {
-	const bool bIncludeDot = false;
-	const FString Extension = FPaths::GetExtension( InSourceData->GetFilename(), bIncludeDot );
-
-	TStaticArray<FString, 2> GltfExtensions;
-	GltfExtensions[0] = TEXT("gltf;GL Transmission Format");
-	GltfExtensions[1] = TEXT("glb;GL Transmission Format (Binary)");
+	return EInterchangeTranslatorType::Scenes;
+}
 
-	const bool bExtensionMatches =
-		Algo::FindByPredicate( GltfExtensions,
-		[ &Extension ]( const FString& GltfExtension )
-		{
-			return GltfExtension.StartsWith( Extension );
-		}) != nullptr;
+TArray<FString> UInterchangeGltfTranslator::GetSupportedFormats() const
+{
+	TArray<FString> GltfExtensions;
+	GltfExtensions.Reserve(2);
+	GltfExtensions.Add(TEXT("gltf;GL Transmission Format"));
+	GltfExtensions.Add(TEXT("glb;GL Transmission Format (Binary)"));
 
-	return bExtensionMatches;
+	return GltfExtensions;
 }
 
 bool UInterchangeGltfTranslator::Translate( UInterchangeBaseNodeContainer& NodeContainer ) const
diff --git a/Engine/Plugins/Experimental/Interchange/Runtime/Source/Import/Private/Mesh/InterchangeOBJTranslator.cpp b/Engine/Plugins/Experimental/Interchange/Runtime/Source/Import/Private/Mesh/InterchangeOBJTranslator.cpp
--- a/Engine/Plugins/Experimental/Interchange/Runtime/Source/Import/Private/Mesh/InterchangeOBJTranslator.cpp
+++ b/Engine/Plugins/Experimental/Interchange/Runtime/Source/Import/Private/Mesh/InterchangeOBJTranslator.cpp
@@ -810,15 +810,12 @@ UInterchangeOBJTranslator::~UInterchangeOBJTranslator()
 {
 }
 
-
-bool UInterchangeOBJTranslator::CanImportSourceData(const UInterchangeSourceData* InSourceData) const
+TArray<FString> UInterchangeOBJTranslator::GetSupportedFormats() const
 {
-	const bool bIncludeDot = false;
-	FString Extension = FPaths::GetExtension(InSourceData->GetFilename(), bIncludeDot);
-	return Extension == FString(TEXT("obj"));
+	TArray<FString> Formats {TEXT("obj;OBJ File Format")};
+	return Formats;
 }
 
-
 bool UInterchangeOBJTranslator::Translate(UInterchangeBaseNodeContainer& BaseNodeContainer) const
 {
 	// Get filename from SourceData
diff --git a/Engine/Plugins/Experimental/Interchange/Runtime/Source/Import/Private/Texture/InterchangeBMPTranslator.cpp b/Engine/Plugins/Experimental/Interchange/Runtime/Source/Import/Private/Texture/InterchangeBMPTranslator.cpp
--- a/Engine/Plugins/Experimental/Interchange/Runtime/Source/Import/Private/Texture/InterchangeBMPTranslator.cpp
+++ b/Engine/Plugins/Experimental/Interchange/Runtime/Source/Import/Private/Texture/InterchangeBMPTranslator.cpp
@@ -16,12 +16,10 @@
 #include "Nodes/InterchangeBaseNodeContainer.h"
 #include "Texture/TextureTranslatorUtilities.h"
 
-
-bool UInterchangeBMPTranslator::CanImportSourceData(const UInterchangeSourceData* InSourceData) const
+TArray<FString> UInterchangeBMPTranslator::GetSupportedFormats() const
 {
-	FString Extension = FPaths::GetExtension(InSourceData->GetFilename());
-	FString BMPExtension = (TEXT("bmp;Texture"));
-	return BMPExtension.StartsWith(Extension);
+	TArray<FString> Formats {TEXT("bmp;Bitmap image")};
+	return Formats;
 }
 
 bool UInterchangeBMPTranslator::Translate(UInterchangeBaseNodeContainer& BaseNodeContainer) const
diff --git a/Engine/Plugins/Experimental/Interchange/Runtime/Source/Import/Private/Texture/InterchangeDDSTranslator.cpp b/Engine/Plugins/Experimental/Interchange/Runtime/Source/Import/Private/Texture/InterchangeDDSTranslator.cpp
--- a/Engine/Plugins/Experimental/Interchange/Runtime/Source/Import/Private/Texture/InterchangeDDSTranslator.cpp
+++ b/Engine/Plugins/Experimental/Interchange/Runtime/Source/Import/Private/Texture/InterchangeDDSTranslator.cpp
@@ -50,11 +50,15 @@ namespace UE::Interchange::Private::InterchangeDDSTranslator
 	}
 }
 
+TArray<FString> UInterchangeDDSTranslator::GetSupportedFormats() const
+{
+	TArray<FString> Formats {TEXT("dds;DirectDraw Surface")};
+	return Formats;
+}
+
 bool UInterchangeDDSTranslator::CanImportSourceData(const UInterchangeSourceData* InSourceData) const
 {
-	FString Extension = FPaths::GetExtension(InSourceData->GetFilename());
-	FString DDSExtension = (TEXT("dds;Texture"));
-	if (!DDSExtension.StartsWith(Extension))
+	if (!UInterchangeTranslatorBase::CanImportSourceData(InSourceData))
 	{
 		return false;
 	}
diff --git a/Engine/Plugins/Experimental/Interchange/Runtime/Source/Import/Private/Texture/InterchangeEXRTranslator.cpp b/Engine/Plugins/Experimental/Interchange/Runtime/Source/Import/Private/Texture/InterchangeEXRTranslator.cpp
--- a/Engine/Plugins/Experimental/Interchange/Runtime/Source/Import/Private/Texture/InterchangeEXRTranslator.cpp
+++ b/Engine/Plugins/Experimental/Interchange/Runtime/Source/Import/Private/Texture/InterchangeEXRTranslator.cpp
@@ -15,11 +15,10 @@
 #include "Nodes/InterchangeBaseNodeContainer.h"
 #include "Texture/TextureTranslatorUtilities.h"
 
-bool UInterchangeEXRTranslator::CanImportSourceData(const UInterchangeSourceData* InSourceData) const
+TArray<FString> UInterchangeEXRTranslator::GetSupportedFormats() const
 {
-	FString Extension = FPaths::GetExtension(InSourceData->GetFilename());
-	FString EXRExtension = (TEXT("exr;Texture"));
-	return EXRExtension.StartsWith(Extension);
+	TArray<FString> Formats {TEXT("exr;OpenEXR image")};
+	return Formats;
 }
 
 bool UInterchangeEXRTranslator::Translate(UInterchangeBaseNodeContainer& BaseNodeContainer) const
diff --git a/Engine/Plugins/Experimental/Interchange/Runtime/Source/Import/Private/Texture/InterchangeHDRTranslator.cpp b/Engine/Plugins/Experimental/Interchange/Runtime/Source/Import/Private/Texture/InterchangeHDRTranslator.cpp
--- a/Engine/Plugins/Experimental/Interchange/Runtime/Source/Import/Private/Texture/InterchangeHDRTranslator.cpp
+++ b/Engine/Plugins/Experimental/Interchange/Runtime/Source/Import/Private/Texture/InterchangeHDRTranslator.cpp
@@ -15,12 +15,10 @@
 #include "Nodes/InterchangeBaseNodeContainer.h"
 #include "Texture/TextureTranslatorUtilities.h"
 
-
-bool UInterchangeHDRTranslator::CanImportSourceData(const UInterchangeSourceData* InSourceData) const
+TArray<FString> UInterchangeHDRTranslator::GetSupportedFormats() const
 {
-	FString Extension = FPaths::GetExtension(InSourceData->GetFilename());
-	FString HDRExtension = (TEXT("hdr"));
-	return HDRExtension.StartsWith(Extension);
+	TArray<FString> Formats {TEXT("hdr;High Dynamic Range image")};
+	return Formats;
 }
 
 bool UInterchangeHDRTranslator::Translate(UInterchangeBaseNodeContainer& BaseNodeContainer) const
diff --git a/Engine/Plugins/Experimental/Interchange/Runtime/Source/Import/Private/Texture/InterchangeIESTranslator.cpp b/Engine/Plugins/Experimental/Interchange/Runtime/Source/Import/Private/Texture/InterchangeIESTranslator.cpp
--- a/Engine/Plugins/Experimental/Interchange/Runtime/Source/Import/Private/Texture/InterchangeIESTranslator.cpp
+++ b/Engine/Plugins/Experimental/Interchange/Runtime/Source/Import/Private/Texture/InterchangeIESTranslator.cpp
@@ -11,11 +11,10 @@
 #include "Misc/Paths.h"
 #include "Texture/TextureTranslatorUtilities.h"
 
-bool UInterchangeIESTranslator::CanImportSourceData(const UInterchangeSourceData* InSourceData) const
+TArray<FString> UInterchangeIESTranslator::GetSupportedFormats() const
 {
-	FString Extension = FPaths::GetExtension(InSourceData->GetFilename());
-	FString IESExtension = (TEXT("ies"));
-	return IESExtension.StartsWith(Extension);
+	TArray<FString> Formats {TEXT("ies;IES light profile")};
+	return Formats;
 }
 
 bool UInterchangeIESTranslator::Translate(UInterchangeBaseNodeContainer& BaseNodeContainer) const
diff --git a/Engine/Plugins/Experimental/Interchange/Runtime/Source/Import/Private/Texture/InterchangeJPGTranslator.cpp b/Engine/Plugins/Experimental/Interchange/Runtime/Source/Import/Private/Texture/InterchangeJPGTranslator.cpp
--- a/Engine/Plugins/Experimental/Interchange/Runtime/Source/Import/Private/Texture/InterchangeJPGTranslator.cpp
+++ b/Engine/Plugins/Experimental/Interchange/Runtime/Source/Import/Private/Texture/InterchangeJPGTranslator.cpp
@@ -17,22 +17,15 @@
 #include "Nodes/InterchangeBaseNodeContainer.h"
 #include "Texture/TextureTranslatorUtilities.h"
 
-bool UInterchangeJPGTranslator::CanImportSourceData(const UInterchangeSourceData* InSourceData) const
+TArray<FString> UInterchangeJPGTranslator::GetSupportedFormats() const
 {
-	FString Extension = FPaths::GetExtension(InSourceData->GetFilename());
+	TArray<FString> Formats;
+	Formats.Reserve(2);
 
-	TStaticArray<FString, 2> JPGExtensions;
-	JPGExtensions[0] = TEXT("jpg;Texture");
-	JPGExtensions[1] = TEXT("jpeg;Texture");
+	Formats.Add(TEXT("jpg;JPEG image"));
+	Formats.Add(TEXT("jpeg;JPEG image"));
 
-	const bool bExtensionMatches =
-		Algo::FindByPredicate(JPGExtensions,
-		[&Extension](const FString& JPGExtension)
-		{
-			return JPGExtension.StartsWith(Extension);
-		}) != nullptr;
-
-	return bExtensionMatches;
+	return Formats;
 }
 
 bool UInterchangeJPGTranslator::Translate(UInterchangeBaseNodeContainer& BaseNodeContainer) const
diff --git a/Engine/Plugins/Experimental/Interchange/Runtime/Source/Import/Private/Texture/InterchangePCXTranslator.cpp b/Engine/Plugins/Experimental/Interchange/Runtime/Source/Import/Private/Texture/InterchangePCXTranslator.cpp
--- a/Engine/Plugins/Experimental/Interchange/Runtime/Source/Import/Private/Texture/InterchangePCXTranslator.cpp
+++ b/Engine/Plugins/Experimental/Interchange/Runtime/Source/Import/Private/Texture/InterchangePCXTranslator.cpp
@@ -66,11 +66,11 @@ namespace UE
 	}//ns Interchange
 }//ns UE
 
-bool UInterchangePCXTranslator::CanImportSourceData(const UInterchangeSourceData* InSourceData) const
+TArray<FString> UInterchangePCXTranslator::GetSupportedFormats() const
 {
-	FString Extension = FPaths::GetExtension(InSourceData->GetFilename());
-	FString PCXExtension = (TEXT("pcx;Texture"));
-	return PCXExtension.StartsWith(Extension);
+	TArray<FString> Formats {TEXT("pcx;Picture Exchange")};
+
+	return Formats;
 }
 
 bool UInterchangePCXTranslator::Translate(UInterchangeBaseNodeContainer& BaseNodeContainer) const
diff --git a/Engine/Plugins/Experimental/Interchange/Runtime/Source/Import/Private/Texture/InterchangePNGTranslator.cpp b/Engine/Plugins/Experimental/Interchange/Runtime/Source/Import/Private/Texture/InterchangePNGTranslator.cpp
--- a/Engine/Plugins/Experimental/Interchange/Runtime/Source/Import/Private/Texture/InterchangePNGTranslator.cpp
+++ b/Engine/Plugins/Experimental/Interchange/Runtime/Source/Import/Private/Texture/InterchangePNGTranslator.cpp
@@ -187,11 +187,11 @@ namespace UE
 	}//ns Interchange
 }//ns UE
 
-bool UInterchangePNGTranslator::CanImportSourceData(const UInterchangeSourceData* InSourceData) const
+TArray<FString> UInterchangePNGTranslator::GetSupportedFormats() const
 {
-	FString Extension = FPaths::GetExtension(InSourceData->GetFilename());
-	FString PNGExtension = (TEXT("png;Texture"));
-	return PNGExtension.StartsWith(Extension);
+	TArray<FString> Formats {TEXT("png;Portable Network Graphic")};
+
+	return Formats;
 }
 
 bool UInterchangePNGTranslator::Translate(UInterchangeBaseNodeContainer& BaseNodeContainer) const
diff --git a/Engine/Plugins/Experimental/Interchange/Runtime/Source/Import/Private/Texture/InterchangePSDTranslator.cpp b/Engine/Plugins/Experimental/Interchange/Runtime/Source/Import/Private/Texture/InterchangePSDTranslator.cpp
--- a/Engine/Plugins/Experimental/Interchange/Runtime/Source/Import/Private/Texture/InterchangePSDTranslator.cpp
+++ b/Engine/Plugins/Experimental/Interchange/Runtime/Source/Import/Private/Texture/InterchangePSDTranslator.cpp
@@ -367,12 +367,11 @@ namespace UE
 	}//ns Interchange
 }//ns UE
 
-
-bool UInterchangePSDTranslator::CanImportSourceData(const UInterchangeSourceData* InSourceData) const
+TArray<FString> UInterchangePSDTranslator::GetSupportedFormats() const
 {
-	FString Extension = FPaths::GetExtension(InSourceData->GetFilename());
-	FString PNGExtension = (TEXT("psd;Texture"));
-	return PNGExtension.StartsWith(Extension);
+	TArray<FString> Formats {TEXT("psd;Photoshop Document")};
+
+	return Formats;
 }
 
 bool UInterchangePSDTranslator::Translate(UInterchangeBaseNodeContainer& BaseNodeContainer) const
diff --git a/Engine/Plugins/Experimental/Interchange/Runtime/Source/Import/Private/Texture/InterchangeTGATranslator.cpp b/Engine/Plugins/Experimental/Interchange/Runtime/Source/Import/Private/Texture/InterchangeTGATranslator.cpp
--- a/Engine/Plugins/Experimental/Interchange/Runtime/Source/Import/Private/Texture/InterchangeTGATranslator.cpp
+++ b/Engine/Plugins/Experimental/Interchange/Runtime/Source/Import/Private/Texture/InterchangeTGATranslator.cpp
@@ -72,11 +72,11 @@ namespace TGATranslatorImpl
 	}
 }
 
-bool UInterchangeTGATranslator::CanImportSourceData(const UInterchangeSourceData* InSourceData) const
+TArray<FString> UInterchangeTGATranslator::GetSupportedFormats() const
 {
-	FString Extension = FPaths::GetExtension(InSourceData->GetFilename());
-	FString TGAExtension = (TEXT("tga;Texture"));
-	return TGAExtension.StartsWith(Extension);
+	TArray<FString> Formats {TEXT("tga;Targa image")};
+
+	return Formats;
 }
 
 bool UInterchangeTGATranslator::Translate(UInterchangeBaseNodeContainer& BaseNodeContainer) const
diff --git a/Engine/Plugins/Experimental/Interchange/Runtime/Source/Import/Private/Texture/InterchangeTIFFTranslator.cpp b/Engine/Plugins/Experimental/Interchange/Runtime/Source/Import/Private/Texture/InterchangeTIFFTranslator.cpp
--- a/Engine/Plugins/Experimental/Interchange/Runtime/Source/Import/Private/Texture/InterchangeTIFFTranslator.cpp
+++ b/Engine/Plugins/Experimental/Interchange/Runtime/Source/Import/Private/Texture/InterchangeTIFFTranslator.cpp
@@ -15,12 +15,15 @@
 #include "Nodes/InterchangeBaseNodeContainer.h"
 #include "Texture/TextureTranslatorUtilities.h"
 
-
-bool UInterchangeTIFFTranslator::CanImportSourceData(const UInterchangeSourceData* InSourceData) const
+TArray<FString> UInterchangeTIFFTranslator::GetSupportedFormats() const
 {
-	FString Extension = FPaths::GetExtension(InSourceData->GetFilename());
-	FString TIFFExtension = (TEXT("tiff"));
-	return TIFFExtension.StartsWith(Extension);
+	TArray<FString> Formats;
+	Formats.Reserve(2);
+
+	Formats.Add(TEXT("tif;Tag Image File Format"));
+	Formats.Add(TEXT("tiff;Tag Image File Format"));
+
+	return Formats;
 }
 
 bool UInterchangeTIFFTranslator::Translate(UInterchangeBaseNodeContainer& BaseNodeContainer) const
diff --git a/Engine/Plugins/Experimental/Interchange/Runtime/Source/Import/Public/Fbx/InterchangeFbxTranslator.h b/Engine/Plugins/Experimental/Interchange/Runtime/Source/Import/Public/Fbx/InterchangeFbxTranslator.h
--- a/Engine/Plugins/Experimental/Interchange/Runtime/Source/Import/Public/Fbx/InterchangeFbxTranslator.h
+++ b/Engine/Plugins/Experimental/Interchange/Runtime/Source/Import/Public/Fbx/InterchangeFbxTranslator.h
@@ -25,10 +25,11 @@ class INTERCHANGEIMPORT_API UInterchangeFbxTranslator : public UInterchangeTrans
 {
 	GENERATED_BODY()
 public:
-	UInterchangeFbxTranslator(const class FObjectInitializer& ObjectInitializer);
+	UInterchangeFbxTranslator();
 
 	/** Begin UInterchangeTranslatorBase API*/
-	virtual bool CanImportSourceData(const UInterchangeSourceData* InSourceData) const override;
+	EInterchangeTranslatorType GetTranslatorType() const override;
+	virtual TArray<FString> GetSupportedFormats() const override;
 	virtual bool Translate(UInterchangeBaseNodeContainer& BaseNodeContainer) const override;
 	virtual void ReleaseSource() override;
 	virtual void ImportFinish() override;
diff --git a/Engine/Plugins/Experimental/Interchange/Runtime/Source/Import/Public/Gltf/InterchangeGltfTranslator.h b/Engine/Plugins/Experimental/Interchange/Runtime/Source/Import/Public/Gltf/InterchangeGltfTranslator.h
--- a/Engine/Plugins/Experimental/Interchange/Runtime/Source/Import/Public/Gltf/InterchangeGltfTranslator.h
+++ b/Engine/Plugins/Experimental/Interchange/Runtime/Source/Import/Public/Gltf/InterchangeGltfTranslator.h
@@ -24,7 +24,8 @@ class UInterchangeGltfTranslator : public UInterchangeTranslatorBase,
 
 public:
 	/** Begin UInterchangeTranslatorBase API*/
-	virtual bool CanImportSourceData( const UInterchangeSourceData* InSourceData ) const override;
+	virtual EInterchangeTranslatorType GetTranslatorType() const override;
+	virtual TArray<FString> GetSupportedFormats() const override;
 	virtual bool Translate( UInterchangeBaseNodeContainer& BaseNodeContainer ) const override;
 	/** End UInterchangeTranslatorBase API*/
 
diff --git a/Engine/Plugins/Experimental/Interchange/Runtime/Source/Import/Public/Mesh/InterchangeOBJTranslator.h b/Engine/Plugins/Experimental/Interchange/Runtime/Source/Import/Public/Mesh/InterchangeOBJTranslator.h
--- a/Engine/Plugins/Experimental/Interchange/Runtime/Source/Import/Public/Mesh/InterchangeOBJTranslator.h
+++ b/Engine/Plugins/Experimental/Interchange/Runtime/Source/Import/Public/Mesh/InterchangeOBJTranslator.h
@@ -30,10 +30,7 @@ public:
 	virtual ~UInterchangeOBJTranslator();
 
 
-	/*
-	 * return true if the translator can translate the specified source data.
-	 */
-	virtual bool CanImportSourceData(const UInterchangeSourceData* InSourceData) const override;
+	virtual TArray<FString> GetSupportedFormats() const override;
 
 	/**
 	 * Translate the associated source data into a node hold by the specified nodes container.
diff --git a/Engine/Plugins/Experimental/Interchange/Runtime/Source/Import/Public/Texture/InterchangeBMPTranslator.h b/Engine/Plugins/Experimental/Interchange/Runtime/Source/Import/Public/Texture/InterchangeBMPTranslator.h
--- a/Engine/Plugins/Experimental/Interchange/Runtime/Source/Import/Public/Texture/InterchangeBMPTranslator.h
+++ b/Engine/Plugins/Experimental/Interchange/Runtime/Source/Import/Public/Texture/InterchangeBMPTranslator.h
@@ -18,10 +18,7 @@ class INTERCHANGEIMPORT_API UInterchangeBMPTranslator : public UInterchangeTrans
 	GENERATED_BODY()
 public:
 
-	/*
-	 * return true if the translator can translate the specified source data.
-	 */
-	virtual bool CanImportSourceData(const UInterchangeSourceData* InSourceData) const override;
+	TArray<FString> GetSupportedFormats() const override;
 
 	/**
 	 * Translate the associated source data into a node hold by the specified nodes container.
diff --git a/Engine/Plugins/Experimental/Interchange/Runtime/Source/Import/Public/Texture/InterchangeDDSTranslator.h b/Engine/Plugins/Experimental/Interchange/Runtime/Source/Import/Public/Texture/InterchangeDDSTranslator.h
--- a/Engine/Plugins/Experimental/Interchange/Runtime/Source/Import/Public/Texture/InterchangeDDSTranslator.h
+++ b/Engine/Plugins/Experimental/Interchange/Runtime/Source/Import/Public/Texture/InterchangeDDSTranslator.h
@@ -19,6 +19,8 @@ class INTERCHANGEIMPORT_API UInterchangeDDSTranslator : public UInterchangeTrans
 {
 	GENERATED_BODY()
 public:
+
+	TArray<FString> GetSupportedFormats() const override;
 	
 	/*
 	 * return true if the translator can translate the specified source data.
diff --git a/Engine/Plugins/Experimental/Interchange/Runtime/Source/Import/Public/Texture/InterchangeEXRTranslator.h b/Engine/Plugins/Experimental/Interchange/Runtime/Source/Import/Public/Texture/InterchangeEXRTranslator.h
--- a/Engine/Plugins/Experimental/Interchange/Runtime/Source/Import/Public/Texture/InterchangeEXRTranslator.h
+++ b/Engine/Plugins/Experimental/Interchange/Runtime/Source/Import/Public/Texture/InterchangeEXRTranslator.h
@@ -18,10 +18,7 @@ class INTERCHANGEIMPORT_API UInterchangeEXRTranslator : public UInterchangeTrans
 	GENERATED_BODY()
 public:
 	
-	/*
-	 * return true if the translator can translate the specified source data.
-	 */
-	virtual bool CanImportSourceData(const UInterchangeSourceData* InSourceData) const override;
+	TArray<FString> GetSupportedFormats() const override;
 
 	/**
 	 * Translate the associated source data into a node hold by the specified nodes container.
diff --git a/Engine/Plugins/Experimental/Interchange/Runtime/Source/Import/Public/Texture/InterchangeHDRTranslator.h b/Engine/Plugins/Experimental/Interchange/Runtime/Source/Import/Public/Texture/InterchangeHDRTranslator.h
--- a/Engine/Plugins/Experimental/Interchange/Runtime/Source/Import/Public/Texture/InterchangeHDRTranslator.h
+++ b/Engine/Plugins/Experimental/Interchange/Runtime/Source/Import/Public/Texture/InterchangeHDRTranslator.h
@@ -18,10 +18,7 @@ class INTERCHANGEIMPORT_API UInterchangeHDRTranslator : public UInterchangeTrans
 	GENERATED_BODY()
 public:
 
-	/*
-	 * return true if the translator can translate the specified file.
-	 */
-	virtual bool CanImportSourceData(const UInterchangeSourceData* InSourceData) const override;
+	TArray<FString> GetSupportedFormats() const override;
 
 	/**
 	 * Translate the associated source data into a node hold by the specified nodes container.
diff --git a/Engine/Plugins/Experimental/Interchange/Runtime/Source/Import/Public/Texture/InterchangeIESTranslator.h b/Engine/Plugins/Experimental/Interchange/Runtime/Source/Import/Public/Texture/InterchangeIESTranslator.h
--- a/Engine/Plugins/Experimental/Interchange/Runtime/Source/Import/Public/Texture/InterchangeIESTranslator.h
+++ b/Engine/Plugins/Experimental/Interchange/Runtime/Source/Import/Public/Texture/InterchangeIESTranslator.h
@@ -18,10 +18,7 @@ class INTERCHANGEIMPORT_API UInterchangeIESTranslator : public UInterchangeTrans
 	GENERATED_BODY()
 public:
 
-	/*
-	 * return true if the translator can translate the specified file.
-	 */
-	virtual bool CanImportSourceData(const UInterchangeSourceData* InSourceData) const override;
+	TArray<FString> GetSupportedFormats() const override;
 
 	/**
 	 * Translate the associated source data into a node hold by the specified nodes container.
diff --git a/Engine/Plugins/Experimental/Interchange/Runtime/Source/Import/Public/Texture/InterchangeJPGTranslator.h b/Engine/Plugins/Experimental/Interchange/Runtime/Source/Import/Public/Texture/InterchangeJPGTranslator.h
--- a/Engine/Plugins/Experimental/Interchange/Runtime/Source/Import/Public/Texture/InterchangeJPGTranslator.h
+++ b/Engine/Plugins/Experimental/Interchange/Runtime/Source/Import/Public/Texture/InterchangeJPGTranslator.h
@@ -18,10 +18,7 @@ class INTERCHANGEIMPORT_API UInterchangeJPGTranslator : public UInterchangeTrans
 	GENERATED_BODY()
 public:
 	
-	/**
-	 * return true if the translator can translate the specified source data.
-	 */
-	virtual bool CanImportSourceData(const UInterchangeSourceData* InSourceData) const override;
+	TArray<FString> GetSupportedFormats() const override;
 
 	/**
 	 * Translate the associated source data into a node hold by the specified nodes container.
diff --git a/Engine/Plugins/Experimental/Interchange/Runtime/Source/Import/Public/Texture/InterchangePCXTranslator.h b/Engine/Plugins/Experimental/Interchange/Runtime/Source/Import/Public/Texture/InterchangePCXTranslator.h
--- a/Engine/Plugins/Experimental/Interchange/Runtime/Source/Import/Public/Texture/InterchangePCXTranslator.h
+++ b/Engine/Plugins/Experimental/Interchange/Runtime/Source/Import/Public/Texture/InterchangePCXTranslator.h
@@ -18,10 +18,7 @@ class INTERCHANGEIMPORT_API UInterchangePCXTranslator : public UInterchangeTrans
 	GENERATED_BODY()
 public:
 	
-	/*
-	 * return true if the translator can translate the specified source data.
-	 */
-	virtual bool CanImportSourceData(const UInterchangeSourceData* InSourceData) const override;
+	TArray<FString> GetSupportedFormats() const override;
 
 	/**
 	 * Translates the associated source data into a node hold by the specified nodes container.
diff --git a/Engine/Plugins/Experimental/Interchange/Runtime/Source/Import/Public/Texture/InterchangePNGTranslator.h b/Engine/Plugins/Experimental/Interchange/Runtime/Source/Import/Public/Texture/InterchangePNGTranslator.h
--- a/Engine/Plugins/Experimental/Interchange/Runtime/Source/Import/Public/Texture/InterchangePNGTranslator.h
+++ b/Engine/Plugins/Experimental/Interchange/Runtime/Source/Import/Public/Texture/InterchangePNGTranslator.h
@@ -17,11 +17,8 @@ class INTERCHANGEIMPORT_API UInterchangePNGTranslator : public UInterchangeTrans
 {
 	GENERATED_BODY()
 public:
-	
-	/*
-	 * return true if the translator can translate the specified source data.
-	 */
-	virtual bool CanImportSourceData(const UInterchangeSourceData* InSourceData) const override;
+
+	TArray<FString> GetSupportedFormats() const override;
 
 	/**
 	 * Translate the associated source data into a node hold by the specified nodes container.
diff --git a/Engine/Plugins/Experimental/Interchange/Runtime/Source/Import/Public/Texture/InterchangePSDTranslator.h b/Engine/Plugins/Experimental/Interchange/Runtime/Source/Import/Public/Texture/InterchangePSDTranslator.h
--- a/Engine/Plugins/Experimental/Interchange/Runtime/Source/Import/Public/Texture/InterchangePSDTranslator.h
+++ b/Engine/Plugins/Experimental/Interchange/Runtime/Source/Import/Public/Texture/InterchangePSDTranslator.h
@@ -18,10 +18,7 @@ class INTERCHANGEIMPORT_API UInterchangePSDTranslator : public UInterchangeTrans
 	GENERATED_BODY()
 public:
 
-	/*
-	 * return true if the translator can translate the specified source data.
-	 */
-	virtual bool CanImportSourceData(const UInterchangeSourceData* InSourceData) const override;
+	TArray<FString> GetSupportedFormats() const override;
 
 	/**
 	 * Translate the associated source data into a node hold by the specified nodes container.
diff --git a/Engine/Plugins/Experimental/Interchange/Runtime/Source/Import/Public/Texture/InterchangeTGATranslator.h b/Engine/Plugins/Experimental/Interchange/Runtime/Source/Import/Public/Texture/InterchangeTGATranslator.h
--- a/Engine/Plugins/Experimental/Interchange/Runtime/Source/Import/Public/Texture/InterchangeTGATranslator.h
+++ b/Engine/Plugins/Experimental/Interchange/Runtime/Source/Import/Public/Texture/InterchangeTGATranslator.h
@@ -17,10 +17,7 @@ class INTERCHANGEIMPORT_API UInterchangeTGATranslator : public UInterchangeTrans
 	GENERATED_BODY()
 public:
 	
-	/**
-	 * return true if the translator can translate the specified source data.
-	 */
-	virtual bool CanImportSourceData(const UInterchangeSourceData* InSourceData) const override;
+	TArray<FString> GetSupportedFormats() const override;
 
 	/**
 	 * Translates the associated source data into a node hold by the specified nodes container.
diff --git a/Engine/Plugins/Experimental/Interchange/Runtime/Source/Import/Public/Texture/InterchangeTIFFTranslator.h b/Engine/Plugins/Experimental/Interchange/Runtime/Source/Import/Public/Texture/InterchangeTIFFTranslator.h
--- a/Engine/Plugins/Experimental/Interchange/Runtime/Source/Import/Public/Texture/InterchangeTIFFTranslator.h
+++ b/Engine/Plugins/Experimental/Interchange/Runtime/Source/Import/Public/Texture/InterchangeTIFFTranslator.h
@@ -18,10 +18,7 @@ class INTERCHANGEIMPORT_API UInterchangeTIFFTranslator : public UInterchangeTran
 	GENERATED_BODY()
 public:
 
-	/*
-	 * return true if the translator can translate the specified file.
-	 */
-	virtual bool CanImportSourceData(const UInterchangeSourceData* InSourceData) const override;
+	TArray<FString> GetSupportedFormats() const override;
 
 	/**
 	 * Translate a source data into a node hold by the specified nodes container.
diff --git a/Engine/Plugins/Experimental/Interchange/Runtime/Source/Import/Public/Texture/InterchangeUDIMTranslator.h b/Engine/Plugins/Experimental/Interchange/Runtime/Source/Import/Public/Texture/InterchangeUDIMTranslator.h
--- a/Engine/Plugins/Experimental/Interchange/Runtime/Source/Import/Public/Texture/InterchangeUDIMTranslator.h
+++ b/Engine/Plugins/Experimental/Interchange/Runtime/Source/Import/Public/Texture/InterchangeUDIMTranslator.h
@@ -24,6 +24,8 @@ public:
 
 	UInterchangeUDIMTranslator();
 
+	virtual TArray<FString> GetSupportedFormats() const override {return TArray<FString>();}
+
 	/*
 	 * return true if the translator can translate the specified source data.
 	 */
diff --git a/Engine/Source/Developer/AssetTools/Private/AssetTools.cpp b/Engine/Source/Developer/AssetTools/Private/AssetTools.cpp
--- a/Engine/Source/Developer/AssetTools/Private/AssetTools.cpp
+++ b/Engine/Source/Developer/AssetTools/Private/AssetTools.cpp
@@ -3836,6 +3836,13 @@ TArray<UObject*> UAssetToolsImpl::ImportAssetsWithDialogImplementation(const FSt
 	// Generate the file types and extensions represented by the selected factories
 	ObjectTools::GenerateFactoryFileExtensions(Factories, FileTypes, AllExtensions, FilterIndexToFactory);
 
+	const UEditorExperimentalSettings* EditorExperimentalSettings = GetDefault<UEditorExperimentalSettings>();
+	if (EditorExperimentalSettings->bEnableInterchangeFramework)
+	{
+		TArray<FString> InterchangeFileExtensions = UInterchangeManager::GetInterchangeManager().GetSupportedFormats(EInterchangeTranslatorType::Assets);
+		ObjectTools::AppendFormatsFileExtensions(InterchangeFileExtensions, FileTypes, AllExtensions);
+	}
+
 	FileTypes = FString::Printf(TEXT("All Files (%s)|%s|%s"), *AllExtensions, *AllExtensions, *FileTypes);
 
 	// Prompt the user for the filenames
diff --git a/Engine/Source/Editor/UnrealEd/Private/FileHelpers.cpp b/Engine/Source/Editor/UnrealEd/Private/FileHelpers.cpp
--- a/Engine/Source/Editor/UnrealEd/Private/FileHelpers.cpp
+++ b/Engine/Source/Editor/UnrealEd/Private/FileHelpers.cpp
@@ -27,6 +27,7 @@
 #include "UncontrolledChangelistsModule.h"
 #include "SourceControlOperations.h"
 #include "Editor/UnrealEdEngine.h"
+#include "Settings/EditorExperimentalSettings.h"
 #include "Settings/EditorLoadingSavingSettings.h"
 #include "Factories/Factory.h"
 #include "Factories/FbxSceneImportFactory.h"
@@ -78,6 +79,7 @@
 #include "WorldPartition/WorldPartitionEditorPerProjectUserSettings.h"
 #include "PackageSourceControlHelper.h"
 #include "ActorFolder.h"
+#include "InterchangeManager.h"
 
 DEFINE_LOG_CATEGORY_STATIC(LogFileHelpers, Log, All);
 
@@ -484,8 +486,14 @@ FString FEditorFileUtils::GetFilterString(EFileInteraction Interaction)
 
 				ObjectTools::GenerateFactoryFileExtensions(Factories, FileTypes, AllExtensions, FilterIndexToFactory);
 
-				FileTypes = FString::Printf(TEXT("All Files (%s)|%s|%s"), *AllExtensions, *AllExtensions, *FileTypes);
+				const UEditorExperimentalSettings* EditorExperimentalSettings = GetDefault<UEditorExperimentalSettings>();
+				if (EditorExperimentalSettings->bEnableInterchangeFramework)
+				{
+					TArray<FString> InterchangeFileExtensions = UInterchangeManager::GetInterchangeManager().GetSupportedFormats(EInterchangeTranslatorType::Scenes);
+					ObjectTools::AppendFormatsFileExtensions(InterchangeFileExtensions, FileTypes, AllExtensions);
+				}
 
+				FileTypes = FString::Printf(TEXT("All Files (%s)|%s|%s"), *AllExtensions, *AllExtensions, *FileTypes);
 
 				Result = FileTypes;
 			}
@@ -1417,28 +1425,38 @@ void FEditorFileUtils::Import(const FString& InFilename)
 {
 	const FScopedBusyCursor BusyCursor;
 
-	USceneImportFactory *SceneFactory = nullptr;
-	for (UClass* Class : TObjectRange<UClass>())
+	const UEditorExperimentalSettings* EditorExperimentalSettings = GetDefault<UEditorExperimentalSettings>();
+
+	UE::Interchange::FScopedSourceData ScopedSourceData(InFilename);
+	const bool bImportThroughInterchange = EditorExperimentalSettings->bEnableInterchangeFramework && UInterchangeManager::GetInterchangeManager().CanTranslateSourceData(ScopedSourceData.GetSourceData());
+
+	USceneImportFactory* SceneFactory = nullptr;
+
+	if (!bImportThroughInterchange)
 	{
-		if (Class->IsChildOf<USceneImportFactory>() && !Class->HasAnyClassFlags(CLASS_Abstract | CLASS_Deprecated | CLASS_NewerVersionExists))
+		for (UClass* Class : TObjectRange<UClass>())
 		{
-			USceneImportFactory* TestFactory = Class->GetDefaultObject<USceneImportFactory>();
-			if (TestFactory->FactoryCanImport(InFilename))
+			if (Class->IsChildOf<USceneImportFactory>() && !Class->HasAnyClassFlags(CLASS_Abstract | CLASS_Deprecated | CLASS_NewerVersionExists))
 			{
-				/// Pick the first one for now 
-				SceneFactory = TestFactory;
-				break;
+				USceneImportFactory* TestFactory = Class->GetDefaultObject<USceneImportFactory>();
+				if (TestFactory->FactoryCanImport(InFilename))
+				{
+					/// Pick the first one for now 
+					SceneFactory = TestFactory;
+					break;
+				}
 			}
 		}
-
 	}
 
-	if (SceneFactory)
+	if (SceneFactory || bImportThroughInterchange)
 	{
 		FString Path = "/Game";
 
+		const bool bImportsAssets = bImportThroughInterchange || SceneFactory->ImportsAssets();
+
 		//Ask the user for the root path where they want to any content to be placed
-		if(SceneFactory->ImportsAssets())
+		if(bImportsAssets)
 		{
 			TSharedRef<SDlgPickPath> PickContentPathDlg =
 				SNew(SDlgPickPath)
@@ -1458,7 +1476,7 @@ void FEditorFileUtils::Import(const FString& InFilename)
 		TArray<FString> Files;
 		Files.Add(InFilename);
 
-		const bool bSyncToBrowser = SceneFactory->ImportsAssets();
+		const bool bSyncToBrowser = bImportsAssets;
 		AssetToolsModule.Get().ImportAssets(Files, Path, SceneFactory, bSyncToBrowser, nullptr, true);
 	}
 	else
diff --git a/Engine/Source/Runtime/Experimental/Interchange/Core/Private/InterchangeTranslatorBase.cpp b/Engine/Source/Runtime/Experimental/Interchange/Core/Private/InterchangeTranslatorBase.cpp
new file mode 100644
--- /dev/null
+++ b/Engine/Source/Runtime/Experimental/Interchange/Core/Private/InterchangeTranslatorBase.cpp
@@ -0,0 +1,20 @@
+// Copyright Epic Games, Inc. All Rights Reserved.
+
+#include "InterchangeTranslatorBase.h"
+
+#include "Algo/Find.h"
+
+bool UInterchangeTranslatorBase::CanImportSourceData(const UInterchangeSourceData* InSourceData) const
+{
+	const bool bIncludeDot = false;
+	const FString Extension = FPaths::GetExtension(InSourceData->GetFilename(), bIncludeDot);
+
+	const bool bExtensionMatches =
+		Algo::FindByPredicate( GetSupportedFormats(),
+		[ &Extension ]( const FString& Format )
+		{
+			return Format.StartsWith( Extension );
+		}) != nullptr;
+
+	return bExtensionMatches;
+}
\ No newline at end of file
diff --git a/Engine/Source/Runtime/Experimental/Interchange/Core/Public/InterchangeTranslatorBase.h b/Engine/Source/Runtime/Experimental/Interchange/Core/Public/InterchangeTranslatorBase.h
--- a/Engine/Source/Runtime/Experimental/Interchange/Core/Public/InterchangeTranslatorBase.h
+++ b/Engine/Source/Runtime/Experimental/Interchange/Core/Public/InterchangeTranslatorBase.h
@@ -11,6 +11,15 @@
 
 #include "InterchangeTranslatorBase.generated.h"
 
+UENUM(BlueprintType)
+enum class EInterchangeTranslatorType : uint8
+{
+	Invalid = 0,
+	Assets = 1 << 1,
+	Actors = 1 << 2,
+	Scenes = Assets | Actors
+};
+
 UCLASS(BlueprintType, Blueprintable, Abstract, Experimental)
 class INTERCHANGECORE_API UInterchangeTranslatorBase : public UObject
 {
@@ -18,10 +27,13 @@ class INTERCHANGECORE_API UInterchangeTranslatorBase : public UObject
 public:
 
 	/** return true if the translator can translate the given source data. */
-	virtual bool CanImportSourceData(const UInterchangeSourceData* InSourceData) const
-	{
-		return false;
-	}
+	virtual bool CanImportSourceData(const UInterchangeSourceData* InSourceData) const;
+
+	/** Specifies the capabilities of a translator. */
+	virtual EInterchangeTranslatorType GetTranslatorType() const { return EInterchangeTranslatorType::Assets; }
+
+	/** List of formats supported by the translator. Each entry is of the form "ext;Description" where ext is the file extension. */
+	virtual TArray<FString> GetSupportedFormats() const PURE_VIRTUAL(UInterchangeTranslatorBase::GetSupportedExtensions, return TArray<FString>(););
 
 	/**
 	 * Translate the associated source data into node(s) that are hold in the specified nodes container.
diff --git a/Engine/Source/Runtime/Experimental/Interchange/Engine/Private/InterchangeManager.cpp b/Engine/Source/Runtime/Experimental/Interchange/Engine/Private/InterchangeManager.cpp
--- a/Engine/Source/Runtime/Experimental/Interchange/Engine/Private/InterchangeManager.cpp
+++ b/Engine/Source/Runtime/Experimental/Interchange/Engine/Private/InterchangeManager.cpp
@@ -457,14 +457,26 @@ bool UInterchangeManager::RegisterWriter(const UClass* WriterClass)
 	return true;
 }
 
-bool UInterchangeManager::CanTranslateSourceData(const UInterchangeSourceData* SourceData) const
+TArray<FString> UInterchangeManager::GetSupportedFormats(const EInterchangeTranslatorType ForTranslatorType) const
 {
-	// Can we find a translator
-	if (GetTranslatorForSourceData(SourceData))
+	TArray<FString> FileExtensions;
+
+	for (const UClass* TranslatorClass : RegisteredTranslatorsClass)
 	{
-		return true;
+		const UInterchangeTranslatorBase* TranslatorBaseCDO = TranslatorClass->GetDefaultObject<UInterchangeTranslatorBase>();
+
+		if (TranslatorBaseCDO->GetTranslatorType() == ForTranslatorType)
+		{
+			FileExtensions.Append(TranslatorBaseCDO->GetSupportedFormats());
+		}
 	}
-	return false;
+
+	return FileExtensions;
+}
+
+bool UInterchangeManager::CanTranslateSourceData(const UInterchangeSourceData* SourceData) const
+{
+	return GetTranslatorForSourceData(SourceData) != nullptr;
 }
 
 void UInterchangeManager::StartQueuedTasks(bool bCancelAllTasks /*= false*/)
diff --git a/Engine/Source/Runtime/Experimental/Interchange/Engine/Public/InterchangeManager.h b/Engine/Source/Runtime/Experimental/Interchange/Engine/Public/InterchangeManager.h
--- a/Engine/Source/Runtime/Experimental/Interchange/Engine/Public/InterchangeManager.h
+++ b/Engine/Source/Runtime/Experimental/Interchange/Engine/Public/InterchangeManager.h
@@ -339,6 +339,11 @@ public:
 	 */
 	bool RegisterWriter(const UClass* Writer);
 
+	/**
+	 * Returns the list of supported formats for a given translator type.
+	 */
+	TArray<FString> GetSupportedFormats(const EInterchangeTranslatorType ForTranslatorType) const;
+
 	/**
 	 * Look if there is a registered translator for this source data.
 	 * This allow us to by pass the original asset tools system to import supported asset.
