Cooker: Fix tagging of assets in collections; AssetData tags now always require a non-empty value.
#rnx

diff --git a/Engine/Source/Editor/UnrealEd/Private/Commandlets/AssetRegistryGenerator.cpp b/Engine/Source/Editor/UnrealEd/Private/Commandlets/AssetRegistryGenerator.cpp
--- a/Engine/Source/Editor/UnrealEd/Private/Commandlets/AssetRegistryGenerator.cpp
+++ b/Engine/Source/Editor/UnrealEd/Private/Commandlets/AssetRegistryGenerator.cpp
@@ -1036,6 +1036,9 @@ void FAssetRegistryGenerator::UpdateCollectionAssetData()
 	}
 
 	// Apply the collection tags to the asset registry state
+	// Collection tags are queried only by the existence of the key, the value is never used. But Tag Values are not allowed
+	// to be empty. Set the value for each tag to an arbitrary field, something short to avoid wasting memory. We use 1 (aka "true") for now.
+	FStringView CollectionValue(TEXTVIEW("1"));
 	for (const auto& AssetPathNameToCollectionTagsPair : AssetPathNamesToCollectionTags)
 	{
 		const FName AssetPathName = AssetPathNameToCollectionTagsPair.Key;
@@ -1047,7 +1050,7 @@ void FAssetRegistryGenerator::UpdateCollectionAssetData()
 			FAssetDataTagMap TagsAndValues = AssetData->TagsAndValues.CopyMap();
 			for (const FName& CollectionTagName : CollectionTagsForAsset)
 			{
-				TagsAndValues.Add(CollectionTagName, FString()); // TODO: Does this need a value to avoid being trimmed?
+				TagsAndValues.Add(CollectionTagName, FString(CollectionValue));
 			}
 			State.UpdateAssetData(FAssetData(AssetData->PackageName, AssetData->PackagePath, AssetData->AssetName, AssetData->AssetClass, TagsAndValues, AssetData->ChunkIDs, AssetData->PackageFlags));
 		}
