QuickLaunch: Add handling of the CookerSettings iterative flags when cooking in commandlet. Previously they were only respected when cooking in editor.
#rnx

diff --git a/Engine/Source/Editor/UnrealEd/Private/PlayLevelViaLauncher.cpp b/Engine/Source/Editor/UnrealEd/Private/PlayLevelViaLauncher.cpp
--- a/Engine/Source/Editor/UnrealEd/Private/PlayLevelViaLauncher.cpp
+++ b/Engine/Source/Editor/UnrealEd/Private/PlayLevelViaLauncher.cpp
@@ -10,11 +10,13 @@
 #include "Misc/CoreMisc.h"
 #include "GameProjectGenerationModule.h"
 #include "CookerSettings.h"
+#include "Settings/EditorExperimentalSettings.h"
 #include "UnrealEdMisc.h"
 #include "Interfaces/ITargetPlatformManagerModule.h"
 #include "Settings/ProjectPackagingSettings.h"
 #include "Framework/Notifications/NotificationManager.h"
 #include "PlayLevel.h"
+#include "Algo/AllOf.h"
 #include "Async/Async.h"
 #include "Logging/MessageLog.h"
 #include "TargetReceipt.h"
@@ -214,38 +216,31 @@ void UEditorEngine::StartPlayUsingLauncherSession(FRequestPlaySessionParams& InR
 		LauncherProfile->SetAdditionalCommandLineParameters(InRequestParams.EditorPlaySettings->AdditionalLaunchParameters);
 	}
 
-	// select the quickest cook mode based on which in editor cook mode is enabled
-	bool bIncrimentalCooking = true;
 	LauncherProfile->AddCookedPlatform(LaunchPlatformName);
+
+	// select the quickest cook mode based on which in editor cook mode is enabled
+	const UCookerSettings& CookerSettings = *GetDefault<UCookerSettings>();
+	const UEditorExperimentalSettings& ExperimentalSettings = *GetDefault<UEditorExperimentalSettings>();
+
+	bool bInEditorCooking = false;
+	bool bCookOnTheFly = false;
 	ELauncherProfileCookModes::Type CurrentLauncherCookMode = ELauncherProfileCookModes::ByTheBook;
-	bool bCanCookByTheBookInEditor = true;
-	bool bCanCookOnTheFlyInEditor = true;
-	for (const FString& PlatformName : LauncherProfile->GetCookedPlatforms())
-	{
-		if (CanCookByTheBookInEditor(PlatformName) == false)
-		{
-			bCanCookByTheBookInEditor = false;
-		}
-		if (CanCookOnTheFlyInEditor(PlatformName) == false)
-		{
-			bCanCookOnTheFlyInEditor = false;
-		}
-	}
-	if (bCanCookByTheBookInEditor)
-	{
-		CurrentLauncherCookMode = ELauncherProfileCookModes::ByTheBookInEditor;
-	}
-	if (bCanCookOnTheFlyInEditor)
+	if (!CookerSettings.bCookOnTheFlyForLaunchOn)
 	{
-		CurrentLauncherCookMode = ELauncherProfileCookModes::OnTheFlyInEditor;
-		bIncrimentalCooking = false;
+		bInEditorCooking = Algo::AllOf(LauncherProfile->GetCookedPlatforms(),
+			[this](const FString& PlatformName) { return CanCookByTheBookInEditor(PlatformName); });
+		CurrentLauncherCookMode = bInEditorCooking ? ELauncherProfileCookModes::ByTheBookInEditor: ELauncherProfileCookModes::ByTheBook;
 	}
-	if (GetDefault<UCookerSettings>()->bCookOnTheFlyForLaunchOn)
+	else
 	{
-		CurrentLauncherCookMode = ELauncherProfileCookModes::OnTheFly;
-		bIncrimentalCooking = false;
+		bCookOnTheFly = true;
+		bInEditorCooking = Algo::AllOf(LauncherProfile->GetCookedPlatforms(),
+			[this](const FString& PlatformName) { return CanCookOnTheFlyInEditor(PlatformName); });
+		CurrentLauncherCookMode = bInEditorCooking ? ELauncherProfileCookModes::OnTheFlyInEditor : ELauncherProfileCookModes::OnTheFly;
 	}
 
+	bool bIncrementalCooking = (CookerSettings.bIterativeCookingForLaunchOn || ExperimentalSettings.bSharedCookedBuilds) && !bCookOnTheFly;
+
 	if (CurrentLauncherCookMode == ELauncherProfileCookModes::OnTheFlyInEditor ||
 		CurrentLauncherCookMode == ELauncherProfileCookModes::ByTheBookInEditor)
 	{
@@ -266,6 +261,19 @@ void UEditorEngine::StartPlayUsingLauncherSession(FRequestPlaySessionParams& InR
 		}
 	}
 
+	TStringBuilder<256> CookOptions;
+	CookOptions << LauncherProfile->GetCookOptions();
+	ensure(CookOptions.Len() == 0);
+	auto SetCookOption = [&CookOptions](FStringView Option, bool bOptionOn)
+	{
+		ensure(CookOptions.ToView().Find(Option) == INDEX_NONE);
+		if (bOptionOn)
+		{
+			CookOptions << (CookOptions.Len() > 0 ? TEXTVIEW(" ") : TEXTVIEW(""));
+			CookOptions << Option;
+		}
+	};
+
 	// content only projects won't have multiple targets to pick from, and pasing -target=UnrealGame will fail if what C++ thinks
 	// is a content only project needs a temporary target.cs file in UBT, 
 	// only set the BuildTarget in code-based projects
@@ -274,12 +282,16 @@ void UEditorEngine::StartPlayUsingLauncherSession(FRequestPlaySessionParams& InR
 		LauncherProfile->SetBuildTarget(GetDefault<UProjectPackagingSettings>()->GetLaunchOnTargetInfo()->Name);
 	}
 	LauncherProfile->SetCookMode(CurrentLauncherCookMode);
-	LauncherProfile->SetUnversionedCooking(!bIncrimentalCooking);
-	LauncherProfile->SetIncrementalCooking(bIncrimentalCooking);
+	LauncherProfile->SetUnversionedCooking(!bIncrementalCooking); // Unversioned cooking is not allowed with incremental cooking
+	LauncherProfile->SetIncrementalCooking(bIncrementalCooking);
+	SetCookOption(TEXTVIEW("-IgnoreIniSettingsOutOfDate"), bIncrementalCooking && CookerSettings.bIgnoreIniSettingsOutOfDateForIteration);
+	SetCookOption(TEXTVIEW("-IgnoreScriptPackagesOutOfDate"), bIncrementalCooking && CookerSettings.bIgnoreScriptPackagesOutOfDateForIteration);
+	SetCookOption(TEXTVIEW("-IterateSharedCookedbuild"), bIncrementalCooking && ExperimentalSettings.bSharedCookedBuilds);
 	LauncherProfile->SetDeployedDeviceGroup(DeviceGroup);
-	LauncherProfile->SetIncrementalDeploying(bIncrimentalCooking);
+	LauncherProfile->SetIncrementalDeploying(bIncrementalCooking);
 	LauncherProfile->SetEditorExe(FUnrealEdMisc::Get().GetExecutableForCommandlets());
 	LauncherProfile->SetShouldUpdateDeviceFlash(InRequestParams.LauncherTargetDevice->bUpdateDeviceFlash);
+	LauncherProfile->SetCookOptions(*CookOptions);
 	
 	if (LauncherProfile->IsBuildingUAT() && !GetDefault<UEditorPerProjectUserSettings>()->bAlwaysBuildUAT && bUATSuccessfullyCompiledOnce)
 	{
