Disable code navigation when C++ is not allowed in the editor

diff --git a/Engine/Source/Editor/BlueprintGraph/Private/K2Node_CallFunction.cpp b/Engine/Source/Editor/BlueprintGraph/Private/K2Node_CallFunction.cpp
--- a/Engine/Source/Editor/BlueprintGraph/Private/K2Node_CallFunction.cpp
+++ b/Engine/Source/Editor/BlueprintGraph/Private/K2Node_CallFunction.cpp
@@ -24,6 +24,9 @@
 #include "Kismet2/BlueprintEditorUtils.h"
 #include "Classes/EditorStyleSettings.h"
 #include "Editor.h"
+#include "Editor/UnrealEdEngine.h"
+#include "Preferences/UnrealEdOptions.h"
+#include "UnrealEdGlobals.h"
 #include "EdGraphUtilities.h"
 
 #include "KismetCompiler.h"
@@ -3283,64 +3286,68 @@ bool UK2Node_CallFunction::CanJumpToDefinition() const
 {
 	const UFunction* TargetFunction = GetTargetFunction();
 	const bool bNativeFunction = (TargetFunction != nullptr) && (TargetFunction->IsNative());
-	return bNativeFunction || (GetJumpTargetForDoubleClick() != nullptr);
+	const bool bCanJumpToNativeFunction = bNativeFunction && ensure(GUnrealEd) && GUnrealEd->GetUnrealEdOptions()->IsCPPAllowed();
+	return bCanJumpToNativeFunction || (GetJumpTargetForDoubleClick() != nullptr);
 }
 
 void UK2Node_CallFunction::JumpToDefinition() const
 {
-	// For native functions, try going to the function definition in C++ if available
-	if (UFunction* TargetFunction = GetTargetFunction())
+	if (ensure(GUnrealEd) && GUnrealEd->GetUnrealEdOptions()->IsCPPAllowed())
 	{
-		if (TargetFunction->IsNative())
+		// For native functions, try going to the function definition in C++ if available
+		if (UFunction* TargetFunction = GetTargetFunction())
 		{
-			// First try the nice way that will get to the right line number
-			bool bSucceeded = false;
-			const bool bNavigateToNativeFunctions = GetDefault<UBlueprintEditorSettings>()->bNavigateToNativeFunctionsFromCallNodes;
-			
-			if(bNavigateToNativeFunctions) 
+			if (TargetFunction->IsNative())
 			{
-				if(FSourceCodeNavigation::CanNavigateToFunction(TargetFunction))
-				{
-					bSucceeded = FSourceCodeNavigation::NavigateToFunction(TargetFunction);
-				}
+				// First try the nice way that will get to the right line number
+				bool bSucceeded = false;
+				const bool bNavigateToNativeFunctions = GetDefault<UBlueprintEditorSettings>()->bNavigateToNativeFunctionsFromCallNodes;
 
-				// Failing that, fall back to the older method which will still get the file open assuming it exists
-				if (!bSucceeded)
+				if (bNavigateToNativeFunctions)
 				{
-					FString NativeParentClassHeaderPath;
-					const bool bFileFound = FSourceCodeNavigation::FindClassHeaderPath(TargetFunction, NativeParentClassHeaderPath) && (IFileManager::Get().FileSize(*NativeParentClassHeaderPath) != INDEX_NONE);
-					if (bFileFound)
+					if (FSourceCodeNavigation::CanNavigateToFunction(TargetFunction))
 					{
-						const FString AbsNativeParentClassHeaderPath = FPaths::ConvertRelativePathToFull(NativeParentClassHeaderPath);
-						bSucceeded = FSourceCodeNavigation::OpenSourceFile(AbsNativeParentClassHeaderPath);
+						bSucceeded = FSourceCodeNavigation::NavigateToFunction(TargetFunction);
 					}
-				}
-			}
-			else
-			{	
-				// Inform user that the function is native, give them opportunity to enable navigation to native
-				// functions:
-				FNotificationInfo Info(LOCTEXT("NavigateToNativeDisabled", "Navigation to Native (c++) Functions Disabled"));
-				Info.ExpireDuration = 10.0f;
-				Info.CheckBoxState = bNavigateToNativeFunctions ? ECheckBoxState::Checked : ECheckBoxState::Unchecked;
-			
-				Info.CheckBoxStateChanged = FOnCheckStateChanged::CreateStatic(
-					[](ECheckBoxState NewState)
+
+					// Failing that, fall back to the older method which will still get the file open assuming it exists
+					if (!bSucceeded)
 					{
-						const FScopedTransaction Transaction(LOCTEXT("ChangeNavigateToNativeFunctionsFromCallNodes", "Change Navigate to Native Functions from Call Nodes Setting"));
-	
-						UBlueprintEditorSettings* MutableEditorSetings = GetMutableDefault<UBlueprintEditorSettings>();
-						MutableEditorSetings->Modify();
-						MutableEditorSetings->bNavigateToNativeFunctionsFromCallNodes = (NewState == ECheckBoxState::Checked) ? true : false;
-						MutableEditorSetings->SaveConfig();
+						FString NativeParentClassHeaderPath;
+						const bool bFileFound = FSourceCodeNavigation::FindClassHeaderPath(TargetFunction, NativeParentClassHeaderPath) && (IFileManager::Get().FileSize(*NativeParentClassHeaderPath) != INDEX_NONE);
+						if (bFileFound)
+						{
+							const FString AbsNativeParentClassHeaderPath = FPaths::ConvertRelativePathToFull(NativeParentClassHeaderPath);
+							bSucceeded = FSourceCodeNavigation::OpenSourceFile(AbsNativeParentClassHeaderPath);
+						}
 					}
-				);
-				Info.CheckBoxText = LOCTEXT("EnableNavigationToNative", "Navigate to Native Functions from Blueprint Call Nodes?");
-			
-				FSlateNotificationManager::Get().AddNotification(Info);
-			}
+				}
+				else
+				{
+					// Inform user that the function is native, give them opportunity to enable navigation to native
+					// functions:
+					FNotificationInfo Info(LOCTEXT("NavigateToNativeDisabled", "Navigation to Native (c++) Functions Disabled"));
+					Info.ExpireDuration = 10.0f;
+					Info.CheckBoxState = bNavigateToNativeFunctions ? ECheckBoxState::Checked : ECheckBoxState::Unchecked;
+
+					Info.CheckBoxStateChanged = FOnCheckStateChanged::CreateStatic(
+						[](ECheckBoxState NewState)
+						{
+							const FScopedTransaction Transaction(LOCTEXT("ChangeNavigateToNativeFunctionsFromCallNodes", "Change Navigate to Native Functions from Call Nodes Setting"));
 
-			return;
+							UBlueprintEditorSettings* MutableEditorSetings = GetMutableDefault<UBlueprintEditorSettings>();
+							MutableEditorSetings->Modify();
+							MutableEditorSetings->bNavigateToNativeFunctionsFromCallNodes = (NewState == ECheckBoxState::Checked) ? true : false;
+							MutableEditorSetings->SaveConfig();
+						}
+					);
+					Info.CheckBoxText = LOCTEXT("EnableNavigationToNative", "Navigate to Native Functions from Blueprint Call Nodes?");
+
+					FSlateNotificationManager::Get().AddNotification(Info);
+				}
+
+				return;
+			}
 		}
 	}
 
diff --git a/Engine/Source/Editor/BlueprintGraph/Private/K2Node_GetSubsystem.cpp b/Engine/Source/Editor/BlueprintGraph/Private/K2Node_GetSubsystem.cpp
--- a/Engine/Source/Editor/BlueprintGraph/Private/K2Node_GetSubsystem.cpp
+++ b/Engine/Source/Editor/BlueprintGraph/Private/K2Node_GetSubsystem.cpp
@@ -22,6 +22,9 @@
 #include "Subsystems/WorldSubsystem.h"
 #include "SourceCodeNavigation.h"
 #include "BlueprintEditorSettings.h"
+#include "Editor/UnrealEdEngine.h"
+#include "Preferences/UnrealEdOptions.h"
+#include "UnrealEdGlobals.h"
 
 // ************************************************************************************
 //    UK2Node_GetSubsystem
@@ -737,20 +740,27 @@ FText UK2Node_GetEditorSubsystem::GetTooltipText() const
 
 bool UK2Node_GetSubsystem::CanJumpToDefinition() const
 {
-	return CustomClass != nullptr;
+	if (CustomClass && ensure(GUnrealEd) && GUnrealEd->GetUnrealEdOptions()->IsCPPAllowed())
+	{
+		return true;
+	}
+	return Super::CanJumpToDefinition();
 }
 
 void UK2Node_GetSubsystem::JumpToDefinition() const
 {
 	bool bSucceeded = false;
 	
-	// Attempt to navigate to the header file where the class is defined if the 
-	// blueprint preferences allow for it
-	if (GetDefault<UBlueprintEditorSettings>()->bNavigateToNativeFunctionsFromCallNodes)
+	if (CustomClass && ensure(GUnrealEd) && GUnrealEd->GetUnrealEdOptions()->IsCPPAllowed())
 	{
-		if (FSourceCodeNavigation::CanNavigateToClass(CustomClass))
+		// Attempt to navigate to the header file where the class is defined if the 
+		// blueprint preferences allow for it
+		if (GetDefault<UBlueprintEditorSettings>()->bNavigateToNativeFunctionsFromCallNodes)
 		{
-			bSucceeded = FSourceCodeNavigation::NavigateToClass(CustomClass);
+			if (FSourceCodeNavigation::CanNavigateToClass(CustomClass))
+			{
+				bSucceeded = FSourceCodeNavigation::NavigateToClass(CustomClass);
+			}
 		}
 	}
 
@@ -784,4 +794,4 @@ void UK2Node_GetEditorSubsystem::ValidateNodeDuringCompilation(class FCompilerRe
 		const FText ErrorText = NSLOCTEXT("K2Node", "GetEditorSubsystem_Error", "Editor Subsystems can only be used in Editor Utilities / Blutilities");
 		MessageLog.Error(*ErrorText.ToString(), this);
 	}
-}
\ No newline at end of file
+}
diff --git a/Engine/Source/Editor/BlueprintGraph/Private/K2Node_Variable.cpp b/Engine/Source/Editor/BlueprintGraph/Private/K2Node_Variable.cpp
--- a/Engine/Source/Editor/BlueprintGraph/Private/K2Node_Variable.cpp
+++ b/Engine/Source/Editor/BlueprintGraph/Private/K2Node_Variable.cpp
@@ -17,6 +17,9 @@
 #include "Logging/MessageLog.h"
 #include "SourceCodeNavigation.h"
 #include "Engine/SimpleConstructionScript.h"
+#include "Editor/UnrealEdEngine.h"
+#include "Preferences/UnrealEdOptions.h"
+#include "UnrealEdGlobals.h"
 
 #define LOCTEXT_NAMESPACE "K2Node"
 
@@ -1103,11 +1106,14 @@ bool UK2Node_Variable::CanJumpToDefinition() const
 {
 	const FProperty* VariableProperty = GetPropertyForVariable();
 	const bool bNativeVariable = (VariableProperty != nullptr) && (VariableProperty->IsNative());
-	return bNativeVariable || (GetJumpTargetForDoubleClick() != nullptr);
+	const bool bCanJumpToNativeVariable = bNativeVariable && ensure(GUnrealEd) && GUnrealEd->GetUnrealEdOptions()->IsCPPAllowed();
+	return bCanJumpToNativeVariable || (GetJumpTargetForDoubleClick() != nullptr);
 }
 
 void UK2Node_Variable::JumpToDefinition() const
 {
+	if (ensure(GUnrealEd) && GUnrealEd->GetUnrealEdOptions()->IsCPPAllowed())
+	{
 	// For native variables, try going to the variable definition in C++ if available
 	if (FProperty* VariableProperty = GetPropertyForVariable())
 	{
@@ -1117,6 +1123,7 @@ void UK2Node_Variable::JumpToDefinition() const
 			return;
 		}
 	}
+	}
 
 	// Otherwise, fall back to the inherited behavior
 	Super::JumpToDefinition();
diff --git a/Engine/Source/Editor/LevelEditor/Private/LevelEditor.cpp b/Engine/Source/Editor/LevelEditor/Private/LevelEditor.cpp
--- a/Engine/Source/Editor/LevelEditor/Private/LevelEditor.cpp
+++ b/Engine/Source/Editor/LevelEditor/Private/LevelEditor.cpp
@@ -754,7 +754,9 @@ void FLevelEditorModule::BindGlobalLevelEditorCommands()
 	ActionList.MapAction(
 		Commands.GoToCodeForActor, 
 		FExecuteAction::CreateStatic( &FLevelEditorActionCallbacks::GoToCodeForActor_Clicked ),
-		FCanExecuteAction::CreateStatic( &FLevelEditorActionCallbacks::GoToCodeForActor_CanExecute )
+		FCanExecuteAction::CreateStatic( &FLevelEditorActionCallbacks::GoToCodeForActor_CanExecute ),
+		FIsActionChecked(),
+		FIsActionButtonVisible::CreateStatic(&FLevelEditorActionCallbacks::GoToCodeForActor_IsVisible)
 		);
 
 	ActionList.MapAction( 
diff --git a/Engine/Source/Editor/LevelEditor/Private/LevelEditorActions.cpp b/Engine/Source/Editor/LevelEditor/Private/LevelEditorActions.cpp
--- a/Engine/Source/Editor/LevelEditor/Private/LevelEditorActions.cpp
+++ b/Engine/Source/Editor/LevelEditor/Private/LevelEditorActions.cpp
@@ -106,6 +106,9 @@
 #include "Misc/ScopeExit.h"
 #include "Dialogs/DlgPickPath.h"
 #include "AssetToolsModule.h"
+#include "Editor/UnrealEdEngine.h"
+#include "Preferences/UnrealEdOptions.h"
+#include "UnrealEdGlobals.h"
 
 #include "Internationalization/Culture.h"
 
@@ -1196,6 +1199,11 @@ bool FLevelEditorActionCallbacks::GoToCodeForActor_CanExecute()
 	return FSourceCodeNavigation::CanNavigateToClass(SelectedActorInfo.SelectionClass);
 }
 
+bool FLevelEditorActionCallbacks::GoToCodeForActor_IsVisible()
+{
+	return ensure(GUnrealEd) && GUnrealEd->GetUnrealEdOptions()->IsCPPAllowed();
+}
+
 void FLevelEditorActionCallbacks::GoToDocsForActor_Clicked()
 {
 	const auto& SelectedActorInfo = AssetSelectionUtils::GetSelectedActorInfo();
diff --git a/Engine/Source/Editor/LevelEditor/Public/LevelEditorActions.h b/Engine/Source/Editor/LevelEditor/Public/LevelEditorActions.h
--- a/Engine/Source/Editor/LevelEditor/Public/LevelEditorActions.h
+++ b/Engine/Source/Editor/LevelEditor/Public/LevelEditorActions.h
@@ -917,6 +917,7 @@ public:
 	/** Called when 'Go to Code for Actor' is clicked */
 	static void GoToCodeForActor_Clicked();
 	static bool GoToCodeForActor_CanExecute();
+	static bool GoToCodeForActor_IsVisible();
 
 	/** Called when 'Go to Documentation for Actor' is clicked */
 	static void GoToDocsForActor_Clicked();
diff --git a/Engine/Source/Editor/SceneOutliner/Private/SceneOutlinerActorInfoColumn.cpp b/Engine/Source/Editor/SceneOutliner/Private/SceneOutlinerActorInfoColumn.cpp
--- a/Engine/Source/Editor/SceneOutliner/Private/SceneOutlinerActorInfoColumn.cpp
+++ b/Engine/Source/Editor/SceneOutliner/Private/SceneOutlinerActorInfoColumn.cpp
@@ -181,7 +181,11 @@ TSharedPtr<SWidget> FTypeInfoColumn::ConstructClassHyperlink( ISceneOutlinerTree
 
 				if (bIsBlueprintClass || bIsGameClass)
 				{
-					return FEditorClassUtils::GetSourceLink(ActorClass, Actor);
+					FEditorClassUtils::FSourceLinkParams SourceLinkParams;
+					SourceLinkParams.Object = Actor;
+					SourceLinkParams.bUseDefaultFormat = true;
+
+					return FEditorClassUtils::GetSourceLink(ActorClass, SourceLinkParams);
 				}
 			}
 		}
diff --git a/Engine/Source/Editor/UMGEditor/Private/Details/SWidgetDetailsView.cpp b/Engine/Source/Editor/UMGEditor/Private/Details/SWidgetDetailsView.cpp
--- a/Engine/Source/Editor/UMGEditor/Private/Details/SWidgetDetailsView.cpp
+++ b/Engine/Source/Editor/UMGEditor/Private/Details/SWidgetDetailsView.cpp
@@ -257,7 +257,10 @@ void SWidgetDetailsView::OnEditorSelectionChanged()
 	// If only 1 valid selected object exists, update the class link to point to the right class.
 	if ( SelectedObjects.Num() == 1 && SelectedObjects[0].IsValid() )
 	{
-		ClassLinkArea->SetContent(FEditorClassUtils::GetSourceLink(SelectedObjects[0]->GetClass(), TWeakObjectPtr<UObject>()));
+		FEditorClassUtils::FSourceLinkParams SourceLinkParams;
+		SourceLinkParams.bUseDefaultFormat = true;
+
+		ClassLinkArea->SetContent(FEditorClassUtils::GetSourceLink(SelectedObjects[0]->GetClass(), SourceLinkParams));
 	}
 	else
 	{
diff --git a/Engine/Source/Editor/UnrealEd/Private/EditorClassUtils.cpp b/Engine/Source/Editor/UnrealEd/Private/EditorClassUtils.cpp
--- a/Engine/Source/Editor/UnrealEd/Private/EditorClassUtils.cpp
+++ b/Engine/Source/Editor/UnrealEd/Private/EditorClassUtils.cpp
@@ -88,24 +88,16 @@ TSharedRef<SWidget> FEditorClassUtils::GetDynamicDocumentationLinkWidget(const T
 	return IDocumentation::Get()->CreateAnchor(TAttribute<FString>::CreateLambda(GetLink));
 }
 
-TSharedRef<SWidget> FEditorClassUtils::GetSourceLink(const UClass* Class, const TWeakObjectPtr<UObject> ObjectWeakPtr)
+TSharedRef<SWidget> FEditorClassUtils::GetSourceLink(const UClass* Class, const FSourceLinkParams& Params)
 {
-	const FText BlueprintFormat = NSLOCTEXT("SourceHyperlink", "EditBlueprint", "Edit {0}");
-	const FText CodeFormat = NSLOCTEXT("SourceHyperlink", "GoToCode", "Open {0}");
-
-	return GetSourceLinkFormatted(Class, ObjectWeakPtr, BlueprintFormat, CodeFormat);
-}
+	TSharedRef<SWidget> Link = SNew(SSpacer);
 
-TSharedRef<SWidget> FEditorClassUtils::GetSourceLinkFormatted(const UClass* Class, const TWeakObjectPtr<UObject> ObjectWeakPtr, const FText& BlueprintFormat, const FText& CodeFormat)
-{
-	TSharedRef<SWidget> SourceHyperlink = SNew( SSpacer );
 	UBlueprint* Blueprint = (Class ? Cast<UBlueprint>(Class->ClassGeneratedBy) : nullptr);
-
 	if (Blueprint)
 	{
 		struct Local
 		{
-			static void OnEditBlueprintClicked( TWeakObjectPtr<UBlueprint> InBlueprint, TWeakObjectPtr<UObject> InAsset )
+			static void OnEditBlueprintClicked(TWeakObjectPtr<UBlueprint> InBlueprint, TWeakObjectPtr<UObject> InAsset)
 			{
 				if (UBlueprint* BlueprintToEdit = InBlueprint.Get())
 				{
@@ -116,37 +108,124 @@ TSharedRef<SWidget> FEditorClassUtils::GetSourceLinkFormatted(const UClass* Clas
 						BlueprintToEdit->SetObjectBeingDebugged(Asset);
 					}
 					// Open the blueprint
-					GEditor->EditObject( BlueprintToEdit );
+					GEditor->EditObject(BlueprintToEdit);
 				}
 			}
 		};
 
 		TWeakObjectPtr<UBlueprint> BlueprintPtr = Blueprint;
 
-		SourceHyperlink = SNew(SHyperlink)
+		FText Text = FText::FromName(Blueprint->GetFName());
+		if (Params.BlueprintFormat)
+		{
+			Text = FText::Format(*Params.BlueprintFormat, Text);
+		}
+		else if (Params.bUseDefaultFormat)
+		{
+			const FText DefaultBlueprintFormat = NSLOCTEXT("SourceHyperlink", "EditBlueprint", "Edit {0}");
+			Text = FText::Format(DefaultBlueprintFormat, Text);
+		}
+
+		Link = SNew(SHyperlink)
 			.Style(FEditorStyle::Get(), "Common.GotoBlueprintHyperlink")
-			.OnNavigate_Static(&Local::OnEditBlueprintClicked, BlueprintPtr, ObjectWeakPtr)
-			.Text(FText::Format(BlueprintFormat, FText::FromString(Blueprint->GetName())))
+			.OnNavigate_Static(&Local::OnEditBlueprintClicked, BlueprintPtr, Params.Object)
+			.Text(Text)
 			.ToolTipText(NSLOCTEXT("SourceHyperlink", "EditBlueprint_ToolTip", "Click to edit the blueprint"));
 	}
-	else if (Class && FSourceCodeNavigation::CanNavigateToClass(Class))
+	else if (Class)
 	{
-		struct Local
+		TWeakObjectPtr<const UClass> WeakClassPtr(Class);
+
+		auto OnNavigateToClassCode = [WeakClassPtr]()
+		{
+			if (const UClass* Class = WeakClassPtr.Get())
+			{
+				FSourceCodeNavigation::NavigateToClass(Class);
+			}
+		};
+
+		auto CanNavigateToClassCode = [WeakClassPtr]()
 		{
-			static void OnEditCodeClicked(const UClass* TargetClass)
+			if (const UClass* Class = WeakClassPtr.Get())
 			{
-				FSourceCodeNavigation::NavigateToClass(TargetClass);
+				return FSourceCodeNavigation::CanNavigateToClass(Class);
 			}
+			return false;
 		};
 
-		SourceHyperlink = SNew(SHyperlink)
-			.Style(FEditorStyle::Get(), "Common.GotoNativeCodeHyperlink")
-			.OnNavigate_Static(&Local::OnEditCodeClicked, Class)
-			.Text(FText::Format(CodeFormat, FText::FromString((Class->GetName()))))
-			.ToolTipText(FText::Format(NSLOCTEXT("SourceHyperlink", "GoToCode_ToolTip", "Click to open this source file in {0}"), FSourceCodeNavigation::GetSelectedSourceCodeIDE()));
+		FText ClassNameText = FText::FromName(Class->GetFName());
+
+		FText FormattedText;
+		if (Params.CodeFormat)
+		{
+			FormattedText = FText::Format(*Params.CodeFormat, ClassNameText);
+		}
+		else if (Params.bUseDefaultFormat)
+		{
+			const FText DefaultCodeFormat = NSLOCTEXT("SourceHyperlink", "GoToCode", "Open {0}");
+			FormattedText = FText::Format(DefaultCodeFormat, ClassNameText);
+		}
+		else
+		{
+			FormattedText = ClassNameText;
+		}
+
+		TSharedRef<SWidget> NoLinkWidget = SNullWidget::NullWidget;
+		if (Params.bEmptyIfNoLink)
+		{
+			NoLinkWidget = SNew(SSpacer)
+				.Visibility_Lambda([CanNavigateToClassCode]() { return CanNavigateToClassCode() ? EVisibility::Collapsed : EVisibility::Visible; });
+		}
+		else
+		{
+			NoLinkWidget = SNew(STextBlock)
+				.Text(Params.bUseFormatIfNoLink ? FormattedText : ClassNameText)
+				.ColorAndOpacity(FSlateColor::UseSubduedForeground())
+				.Visibility_Lambda([CanNavigateToClassCode]() { return CanNavigateToClassCode() ? EVisibility::Collapsed : EVisibility::Visible; });
+		}
+
+		Link = SNew(SHorizontalBox)
+			+ SHorizontalBox::Slot()
+			.AutoWidth()
+			.VAlign(VAlign_Center)
+			[
+				SNew(SHyperlink)
+				.Style(FEditorStyle::Get(), "Common.GotoNativeCodeHyperlink")
+				.OnNavigate_Lambda(OnNavigateToClassCode)
+				.Text(FormattedText)
+				.ToolTipText(FText::Format(NSLOCTEXT("SourceHyperlink", "GoToCode_ToolTip", "Click to open this source file in {0}"), FSourceCodeNavigation::GetSelectedSourceCodeIDE()))
+				.Visibility_Lambda([CanNavigateToClassCode]() { return CanNavigateToClassCode() ? EVisibility::Visible : EVisibility::Collapsed; })
+			]
+			+ SHorizontalBox::Slot()
+			.AutoWidth()
+			.VAlign(VAlign_Center)
+			[
+				NoLinkWidget
+			];
 	}
 
-	return SourceHyperlink;
+	return Link;
+}
+
+TSharedRef<SWidget> FEditorClassUtils::GetSourceLink(const UClass* Class, const TWeakObjectPtr<UObject> ObjectWeakPtr)
+{
+	FSourceLinkParams Params;
+	Params.Object = ObjectWeakPtr;
+	Params.bUseDefaultFormat = true;
+	Params.bEmptyIfNoLink = true;
+
+	return GetSourceLink(Class, Params);
+}
+
+TSharedRef<SWidget> FEditorClassUtils::GetSourceLinkFormatted(const UClass* Class, const TWeakObjectPtr<UObject> ObjectWeakPtr, const FText& BlueprintFormat, const FText& CodeFormat)
+{
+	FSourceLinkParams Params;
+	Params.Object = ObjectWeakPtr;
+	Params.BlueprintFormat = &BlueprintFormat;
+	Params.CodeFormat = &CodeFormat;
+	Params.bEmptyIfNoLink = true;
+
+	return GetSourceLink(Class, Params);
 }
 
 UClass* FEditorClassUtils::GetClassFromString(const FString& ClassName)
diff --git a/Engine/Source/Editor/UnrealEd/Private/SourceCodeNavigation.cpp b/Engine/Source/Editor/UnrealEd/Private/SourceCodeNavigation.cpp
--- a/Engine/Source/Editor/UnrealEd/Private/SourceCodeNavigation.cpp
+++ b/Engine/Source/Editor/UnrealEd/Private/SourceCodeNavigation.cpp
@@ -27,6 +27,9 @@
 #include "Interfaces/IHttpResponse.h"
 #include "Interfaces/IHttpRequest.h"
 #include "HttpModule.h"
+#include "Editor/UnrealEdEngine.h"
+#include "Preferences/UnrealEdOptions.h"
+#include "UnrealEdGlobals.h"
 
 #if PLATFORM_WINDOWS
 #include "Windows/WindowsHWrapper.h"
@@ -1487,6 +1490,11 @@ bool FSourceCodeNavigation::CanNavigateToClass(const UClass* InClass)
 		return false;
 	}
 
+	if (ensure(GUnrealEd) && !GUnrealEd->GetUnrealEdOptions()->IsCPPAllowed())
+	{
+		return false;
+	}
+
 	for (int32 i = 0; i < SourceCodeNavigationHandlers.Num(); ++i)
 	{
 		ISourceCodeNavigationHandler* Handler = SourceCodeNavigationHandlers[i];
@@ -1532,6 +1540,11 @@ bool FSourceCodeNavigation::CanNavigateToStruct(const UScriptStruct* InStruct)
 		return false;
 	}
 
+	if (ensure(GUnrealEd) && !GUnrealEd->GetUnrealEdOptions()->IsCPPAllowed())
+	{
+		return false;
+	}
+
 	for (int32 i = 0; i < SourceCodeNavigationHandlers.Num(); ++i)
 	{
 		ISourceCodeNavigationHandler* Handler = SourceCodeNavigationHandlers[i];
@@ -1577,6 +1590,11 @@ bool FSourceCodeNavigation::CanNavigateToFunction(const UFunction* InFunction)
 		return false;
 	}
 
+	if (ensure(GUnrealEd) && !GUnrealEd->GetUnrealEdOptions()->IsCPPAllowed())
+	{
+		return false;
+	}
+
 	for (int32 i = 0; i < SourceCodeNavigationHandlers.Num(); ++i)
 	{
 		ISourceCodeNavigationHandler* Handler = SourceCodeNavigationHandlers[i];
@@ -1631,6 +1649,11 @@ bool FSourceCodeNavigation::CanNavigateToProperty(const FProperty* InProperty)
 		return false;
 	}
 
+	if (ensure(GUnrealEd) && !GUnrealEd->GetUnrealEdOptions()->IsCPPAllowed())
+	{
+		return false;
+	}
+
 	for (int32 i = 0; i < SourceCodeNavigationHandlers.Num(); ++i)
 	{
 		ISourceCodeNavigationHandler* Handler = SourceCodeNavigationHandlers[i];
@@ -1681,6 +1704,11 @@ bool FSourceCodeNavigation::CanNavigateToStruct(const UStruct* InStruct)
 		return false;
 	}
 
+	if (ensure(GUnrealEd) && !GUnrealEd->GetUnrealEdOptions()->IsCPPAllowed())
+	{
+		return false;
+	}
+
 	for (int32 i = 0; i < SourceCodeNavigationHandlers.Num(); ++i)
 	{
 		ISourceCodeNavigationHandler* handler = SourceCodeNavigationHandlers[i];
diff --git a/Engine/Source/Editor/UnrealEd/Private/Toolkits/SimpleAssetEditor.cpp b/Engine/Source/Editor/UnrealEd/Private/Toolkits/SimpleAssetEditor.cpp
--- a/Engine/Source/Editor/UnrealEd/Private/Toolkits/SimpleAssetEditor.cpp
+++ b/Engine/Source/Editor/UnrealEd/Private/Toolkits/SimpleAssetEditor.cpp
@@ -6,9 +6,8 @@
 #include "Widgets/Docking/SDockTab.h"
 #include "PropertyEditorModule.h"
 #include "IDetailsView.h"
+#include "EditorClassUtils.h"
 #include "Editor.h"
-#include "Widgets/Input/SHyperlink.h"
-#include "SourceCodeNavigation.h"
 
 
 #define LOCTEXT_NAMESPACE "GenericEditor"
@@ -314,18 +313,6 @@ void FSimpleAssetEditor::PostRegenerateMenusAndToolbars()
 	// Provide a hyperlink to view that native class
 	if (CommonDenominatorClass)
 	{
-		TWeakObjectPtr<UClass> WeakClassPtr(CommonDenominatorClass);
-		auto OnNavigateToClassCode = [WeakClassPtr]()
-		{
-			if (UClass* StrongClassPtr = WeakClassPtr.Get())
-			{
-				if (FSourceCodeNavigation::CanNavigateToClass(StrongClassPtr))
-				{
-					FSourceCodeNavigation::NavigateToClass(StrongClassPtr);
-				}
-			}
-		};
-
 		// build and attach the menu overlay
 		TSharedRef<SHorizontalBox> MenuOverlayBox = SNew(SHorizontalBox)
 			+ SHorizontalBox::Slot()
@@ -342,11 +329,7 @@ void FSimpleAssetEditor::PostRegenerateMenusAndToolbars()
 			.VAlign(VAlign_Center)
 			.Padding(0.0f, 0.0f, 8.0f, 0.0f)
 			[
-				SNew(SHyperlink)
-				.Style(FEditorStyle::Get(), "Common.GotoNativeCodeHyperlink")
-				.OnNavigate_Lambda(OnNavigateToClassCode)
-				.Text(FText::FromName(CommonDenominatorClass->GetFName()))
-				.ToolTipText(FText::Format(LOCTEXT("GoToCode_ToolTip", "Click to open this source file in {0}"), FSourceCodeNavigation::GetSelectedSourceCodeIDE()))
+				FEditorClassUtils::GetSourceLink(CommonDenominatorClass)
 			];
 	
 		SetMenuOverlay(MenuOverlayBox);
diff --git a/Engine/Source/Editor/UnrealEd/Public/EditorClassUtils.h b/Engine/Source/Editor/UnrealEd/Public/EditorClassUtils.h
--- a/Engine/Source/Editor/UnrealEd/Public/EditorClassUtils.h
+++ b/Engine/Source/Editor/UnrealEd/Public/EditorClassUtils.h
@@ -73,23 +73,60 @@ namespace FEditorClassUtils
 	 */
 	UNREALED_API TSharedRef<SWidget> GetDynamicDocumentationLinkWidget(const TAttribute<const UClass*>& ClassAttribute);
 
+	/** Optional GetSourceLink parameters */
+	struct FSourceLinkParams
+	{
+		/* Object to set blueprint debugging to in the case we have a blueprint generated class */
+		TWeakObjectPtr<UObject> Object;
+
+		/* Text format for blueprint links */
+		const FText* BlueprintFormat = nullptr;
+		
+		/* Text format for C++ code file links */
+		const FText* CodeFormat = nullptr;
+
+		/** 
+		 * If true, use default values for BlueprintFormat (Edit ...) and CodeFormat (Open ...) when unspecified
+		 * If false, use only the class name when BlueprintFormat and CodeFormat are unspecified 
+		 */
+		bool bUseDefaultFormat = false;
+
+		/** If true, use specified text format if the link is unavailable for some reason, otherwise use only the class name */
+		bool bUseFormatIfNoLink = false;
+
+		/** Whether a spacer widget is used if the link is unavailable for some reason */
+		bool bEmptyIfNoLink = false;
+	};
+
+	/**
+	 * Creates an hyperlink to the source code or blueprint for a given class
+	 * (or a text block/spacer if the link is unavailable for some reason)
+	 *
+	 * @param	Class	Class we want to build a link for
+	 * @param	Params	See FSourceLinkOptionalParams
+	 * @return			Shared pointer to the constructed widget
+	 */
+	UNREALED_API TSharedRef<SWidget> GetSourceLink(const UClass* Class, const FSourceLinkParams& Params = FSourceLinkParams());
+
 	/**
-	 * Creates a link to the source code or blueprint for a given class
+	 * Creates an hyperlink to the source code or blueprint for a given class
+	 * (or a spacer if the link is unavailable for some reason)
 	 *
 	 * @param	Class			Class we want to build a link for
 	 * @param	ObjectWeakPtr	Optional object to set blueprint debugging to in the case we are choosing a blueprint
-	 * @return					Shared pointer to the constructed tooltip
+	 * @return					Shared pointer to the constructed widget
 	 */
 	UNREALED_API TSharedRef<SWidget> GetSourceLink(const UClass* Class, const TWeakObjectPtr<UObject> ObjectWeakPtr);
 
 	/**
-	 * Creates a link to the source code or blueprint for a given class formatted however you need. Example "Edit {0}"
+	 * Creates an hyperlink to the source code or blueprint for a given class formatted however you need. Example "Edit {0}"
+	 * (or a spacer if the link is unavailable for some reason)
 	 *
 	 * @param	Class			Class we want to build a link for
 	 * @param	ObjectWeakPtr	Optional object to set blueprint debugging to in the case we are choosing a blueprint
-	 * @param	BlueprintFormat	The text format for blueprint links.
-	 * @param	CodeFormat		The text format for C++ code file links.
-	 * @return					Shared pointer to the constructed tooltip
+	 * @param	BlueprintFormat	The text format for blueprint links
+	 * @param	CodeFormat		The text format for C++ code file links
+	 * @return					Shared pointer to the constructed widget
 	 */
 	UNREALED_API TSharedRef<SWidget> GetSourceLinkFormatted(const UClass* Class, const TWeakObjectPtr<UObject> ObjectWeakPtr, const FText& BlueprintFormat, const FText& CodeFormat);
 
