LowLevelMemTracker - add scopes for some large allocations during cooking.
#rnx

diff --git a/Engine/Source/Developer/CollectionManager/Private/CollectionManager.cpp b/Engine/Source/Developer/CollectionManager/Private/CollectionManager.cpp
--- a/Engine/Source/Developer/CollectionManager/Private/CollectionManager.cpp
+++ b/Engine/Source/Developer/CollectionManager/Private/CollectionManager.cpp
@@ -1830,6 +1830,7 @@ void FCollectionManager::LoadCollections()
 
 	const double LoadStartTime = FPlatformTime::Seconds();
 	const int32 PrevNumCollections = AvailableCollections.Num();
+	LLM_SCOPE_BYNAME(TEXT("CollectionManager"));
 
 	FRWLock CollectionLock;
 	ParallelFor(
diff --git a/Engine/Source/Developer/DerivedDataCache/Private/DerivedDataBackendInterface.h b/Engine/Source/Developer/DerivedDataCache/Private/DerivedDataBackendInterface.h
--- a/Engine/Source/Developer/DerivedDataCache/Private/DerivedDataBackendInterface.h
+++ b/Engine/Source/Developer/DerivedDataCache/Private/DerivedDataBackendInterface.h
@@ -7,6 +7,7 @@
 #include "Containers/BitArray.h"
 #include "DerivedDataCache.h"
 #include "DerivedDataLegacyCacheStore.h"
+#include "HAL/LowLevelMemTracker.h"
 #include "Stats/Stats.h"
 
 class FDerivedDataCacheUsageStats;
@@ -24,6 +25,12 @@ DECLARE_FLOAT_ACCUMULATOR_STAT_EXTERN(TEXT("Sync Put Time"),STAT_DDC_PutTime,STA
 DECLARE_FLOAT_ACCUMULATOR_STAT_EXTERN(TEXT("Sync Build Time"),STAT_DDC_SyncBuildTime,STATGROUP_DDC, );
 DECLARE_FLOAT_ACCUMULATOR_STAT_EXTERN(TEXT("Exists Time"),STAT_DDC_ExistTime,STATGROUP_DDC, );
 
+/** Memory allocation tag for internal allocations of DDC backends */
+LLM_DECLARE_TAG_API(DDCBackend, DERIVEDDATACACHE_API);
+
+/** Memory allocation tag for allocations made by DDC but that are passed to other systems */
+LLM_DECLARE_TAG_API(UntaggedDDCResult, DERIVEDDATACACHE_API);
+
 namespace UE::DerivedData
 {
 
diff --git a/Engine/Source/Developer/DerivedDataCache/Private/DerivedDataBackends.cpp b/Engine/Source/Developer/DerivedDataCache/Private/DerivedDataBackends.cpp
--- a/Engine/Source/Developer/DerivedDataCache/Private/DerivedDataBackends.cpp
+++ b/Engine/Source/Developer/DerivedDataCache/Private/DerivedDataBackends.cpp
@@ -27,6 +27,8 @@
 #include <atomic>
 
 DEFINE_LOG_CATEGORY(LogDerivedDataCache);
+LLM_DEFINE_TAG(UntaggedDDCResult);
+LLM_DEFINE_TAG(DDCBackend);
 
 #define MAX_BACKEND_KEY_LENGTH (120)
 #define LOCTEXT_NAMESPACE "DerivedDataBackendGraph"
diff --git a/Engine/Source/Developer/DerivedDataCache/Private/ZenDerivedDataBackend.cpp b/Engine/Source/Developer/DerivedDataCache/Private/ZenDerivedDataBackend.cpp
--- a/Engine/Source/Developer/DerivedDataCache/Private/ZenDerivedDataBackend.cpp
+++ b/Engine/Source/Developer/DerivedDataCache/Private/ZenDerivedDataBackend.cpp
@@ -858,6 +858,7 @@ void FZenDerivedDataBackend::Get(
 		Zen::FZenHttpRequest::Result HttpResult = Zen::FZenHttpRequest::Result::Failed;
 
 		{
+			LLM_SCOPE_BYTAG(UntaggedDDCResult);
 			Zen::FZenScopedRequestPtr Request(RequestPool.Get());
 			HttpResult = Request->PerformRpc(TEXT("/z$/$rpc"_SV), BatchRequest.Save().AsObject(), BatchResponse);
 		}
@@ -1100,6 +1101,7 @@ void FZenDerivedDataBackend::GetValue(
 			Zen::FZenHttpRequest::Result HttpResult = Zen::FZenHttpRequest::Result::Failed;
 
 			{
+				LLM_SCOPE_BYTAG(UntaggedDDCResult);
 				Zen::FZenScopedRequestPtr Request(RequestPool.Get());
 				HttpResult = Request->PerformRpc(TEXT("/z$/$rpc"_SV), BatchRequest.Save().AsObject(), BatchResponse);
 			}
@@ -1249,6 +1251,7 @@ void FZenDerivedDataBackend::GetChunks(
 		Zen::FZenHttpRequest::Result HttpResult = Zen::FZenHttpRequest::Result::Failed;
 
 		{
+			LLM_SCOPE_BYTAG(UntaggedDDCResult);
 			Zen::FZenScopedRequestPtr Request(RequestPool.Get());
 			HttpResult = Request->PerformRpc(TEXT("/z$/$rpc"_SV), BatchRequest.Save().AsObject(), BatchResponse);
 		}
diff --git a/Engine/Source/Developer/Zen/Internal/ZenServerHttp.cpp b/Engine/Source/Developer/Zen/Internal/ZenServerHttp.cpp
--- a/Engine/Source/Developer/Zen/Internal/ZenServerHttp.cpp
+++ b/Engine/Source/Developer/Zen/Internal/ZenServerHttp.cpp
@@ -28,6 +28,7 @@
 #include "Compression/CompressedBuffer.h"
 #include "Compression/OodleDataCompression.h"
 #include "Containers/StringFwd.h"
+#include "HAL/LowLevelMemTracker.h"
 #include "Memory/CompositeBuffer.h"
 #include "Serialization/CompactBinary.h"
 #include "Serialization/CompactBinaryPackage.h"
@@ -38,6 +39,8 @@
 #include "HAL/PlatformProcess.h"
 #include "ProfilingDebugging/CpuProfilerTrace.h"
 
+LLM_DEFINE_TAG(ZenDDC, NAME_None, TEXT("DDCBackend"));
+
 DEFINE_LOG_CATEGORY_STATIC(LogZenHttp, Log, All);
 
 namespace UE::Zen {
@@ -450,6 +453,7 @@ namespace UE::Zen {
 
 	FZenHttpRequest::Result FZenHttpRequest::PerformBlocking(FStringView Uri, RequestVerb Verb, uint32 ContentLength)
 	{
+		LLM_SCOPE_BYTAG(ZenDDC);
 		// Strip any leading slashes because we compose the prefix and the suffix with a separating slash below
 		while (Uri.StartsWith('/'))
 		{
diff --git a/Engine/Source/Editor/UnrealEd/Private/CookOnTheFlyServer.cpp b/Engine/Source/Editor/UnrealEd/Private/CookOnTheFlyServer.cpp
--- a/Engine/Source/Editor/UnrealEd/Private/CookOnTheFlyServer.cpp
+++ b/Engine/Source/Editor/UnrealEd/Private/CookOnTheFlyServer.cpp
@@ -123,6 +123,8 @@
 #define LOCTEXT_NAMESPACE "Cooker"
 
 DEFINE_LOG_CATEGORY(LogCook);
+LLM_DEFINE_TAG(Cooker);
+
 
 int32 GCookProgressDisplay = (int32)ECookProgressDisplayMode::RemainingPackages;
 static FAutoConsoleVariableRef CVarCookDisplayMode(
@@ -405,6 +407,7 @@ UCookOnTheFlyServer::~UCookOnTheFlyServer()
 void UCookOnTheFlyServer::Tick(float DeltaTime)
 {
 	TRACE_CPUPROFILER_EVENT_SCOPE(UCookOnTheFlyServer::Tick);
+	LLM_SCOPE_BYTAG(Cooker);
 
 	check(IsCookingInEditor());
 
@@ -447,6 +450,7 @@ TStatId UCookOnTheFlyServer::GetStatId() const
 
 bool UCookOnTheFlyServer::StartCookOnTheFly(FCookOnTheFlyOptions InCookOnTheFlyOptions)
 {
+	LLM_SCOPE_BYTAG(Cooker);
 #if WITH_COTF
 	check(IsCookOnTheFlyMode());
 	//GetDerivedDataCacheRef().WaitForQuiescence(false);
@@ -1129,6 +1133,7 @@ bool UCookOnTheFlyServer::RequestPackage(const FName& StandardPackageFName, cons
 
 uint32 UCookOnTheFlyServer::TickCookOnTheSide(const float TimeSlice, uint32 &CookedPackageCount, ECookTickFlags TickFlags)
 {
+	LLM_SCOPE_BYTAG(Cooker);
 	TickCancels();
 	TickNetwork();
 	if (!IsInSession())
@@ -2753,7 +2758,11 @@ bool UCookOnTheFlyServer::LoadPackageForCooking(UE::Cook::FPackageData& PackageD
 			}
 		}
 
-		UPackage* LoadedPackage = LoadPackage(LoadIntoPackage, *FileName, LOAD_None);
+		UPackage* LoadedPackage;
+		{
+			LLM_SCOPE(ELLMTag::Untagged); // Reset the scope so that untagged memory in the package shows up as Untagged rather than Cooker
+			LoadedPackage = LoadPackage(LoadIntoPackage, *FileName, LOAD_None);
+		}
 		if (IsValid(LoadedPackage) && LoadedPackage->IsFullyLoaded())
 		{
 			OutPackage = LoadedPackage;
@@ -7400,6 +7409,7 @@ void UCookOnTheFlyServer::TermSandbox()
 void UCookOnTheFlyServer::StartCookByTheBook( const FCookByTheBookStartupOptions& CookByTheBookStartupOptions )
 {
 	UE_SCOPED_COOKTIMER(StartCookByTheBook);
+	LLM_SCOPE_BYTAG(Cooker);
 
 	const TArray<FString>& CookMaps = CookByTheBookStartupOptions.CookMaps;
 	const TArray<FString>& CookDirectories = CookByTheBookStartupOptions.CookDirectories;
@@ -8209,7 +8219,10 @@ uint32 UCookOnTheFlyServer::FullLoadAndSave(uint32& CookedPackageCount)
 					const FString BuildFilename = BuildFilenameFName.ToString();
 					GIsCookerLoadingPackage = true;
 					UE_SCOPED_HIERARCHICAL_COOKTIMER(LoadPackage);
-					LoadPackage(nullptr, *BuildFilename, LOAD_None);
+					{
+						LLM_SCOPE(ELLMTag::Untagged); // Reset the scope so that untagged memory in the package shows up as Untagged rather than Cooker
+						LoadPackage(nullptr, *BuildFilename, LOAD_None);
+					}
 					if (GShaderCompilingManager)
 					{
 						GShaderCompilingManager->ProcessAsyncResults(true, false);
@@ -8440,7 +8453,10 @@ uint32 UCookOnTheFlyServer::FullLoadAndSave(uint32& CookedPackageCount)
 				FName BuildFilenameFName = PackageDatas->GetFileNameByPackageName(FName(*ToLoad));
 				if (!PackageTracker->NeverCookPackageList.Contains(BuildFilenameFName))
 				{
-					LoadPackage(nullptr, *ToLoad, LOAD_None);
+					{
+						LLM_SCOPE(ELLMTag::Untagged); // Reset the scope so that untagged memory in the package shows up as Untagged rather than Cooker
+						LoadPackage(nullptr, *ToLoad, LOAD_None);
+					}
 					if (GShaderCompilingManager)
 					{
 						GShaderCompilingManager->ProcessAsyncResults(true, false);
diff --git a/Engine/Source/Editor/UnrealEd/Private/Cooker/CookTypes.h b/Engine/Source/Editor/UnrealEd/Private/Cooker/CookTypes.h
--- a/Engine/Source/Editor/UnrealEd/Private/Cooker/CookTypes.h
+++ b/Engine/Source/Editor/UnrealEd/Private/Cooker/CookTypes.h
@@ -7,6 +7,7 @@
 #include "Containers/UnrealString.h"
 #include "CookOnTheSide/CookOnTheFlyServer.h" // ECookTickFlags
 #include "DerivedDataRequestOwner.h"
+#include "HAL/LowLevelMemTracker.h"
 #include "HAL/Platform.h"
 #include "Logging/TokenizedMessage.h"
 #include "Serialization/PackageWriter.h"
@@ -260,5 +261,7 @@ public:
 };
 
 void LogCookerMessage(const FString& MessageText, EMessageSeverity::Type Severity);
+LLM_DECLARE_TAG(Cooker);
+
 constexpr uint32 ExpectedMaxNumPlatforms = 32;
 #define REMAPPED_PLUGINS TEXT("RemappedPlugins")
diff --git a/Engine/Source/Editor/UnrealEd/Private/Cooker/PackageBuildDependencyTracker.cpp b/Engine/Source/Editor/UnrealEd/Private/Cooker/PackageBuildDependencyTracker.cpp
--- a/Engine/Source/Editor/UnrealEd/Private/Cooker/PackageBuildDependencyTracker.cpp
+++ b/Engine/Source/Editor/UnrealEd/Private/Cooker/PackageBuildDependencyTracker.cpp
@@ -76,6 +76,8 @@ void FPackageBuildDependencyTracker::StaticOnObjectHandleRead(UObject* ReadObjec
 							*Referencer.ToString(), *Referenced.ToString());
 					}
 
+					LLM_SCOPE_BYNAME(TEXTVIEW("TObjectPtr"));
+
 					FBuildDependencyAccessData AccessData{ Referenced, AccumulatedScopeData->TargetPlatform };
 					FScopeLock RecordsScopeLock(&Singleton.RecordsLock);
 					if (Referencer == Singleton.LastReferencer)
diff --git a/Engine/Source/Editor/UnrealEd/Private/Cooker/PackageTracker.cpp b/Engine/Source/Editor/UnrealEd/Private/Cooker/PackageTracker.cpp
--- a/Engine/Source/Editor/UnrealEd/Private/Cooker/PackageTracker.cpp
+++ b/Engine/Source/Editor/UnrealEd/Private/Cooker/PackageTracker.cpp
@@ -100,6 +100,7 @@ namespace UE::Cook
 
 			if (Package->GetOuter() == nullptr)
 			{
+				LLM_SCOPE_BYTAG(Cooker);
 				if (LoadingPackageData && Package->GetFName() != LoadingPackageData->GetPackageName())
 				{
 					COOK_STAT(++Stats::NumInlineLoads);
diff --git a/Engine/Source/Runtime/AssetRegistry/Private/AssetDataGatherer.cpp b/Engine/Source/Runtime/AssetRegistry/Private/AssetDataGatherer.cpp
--- a/Engine/Source/Runtime/AssetRegistry/Private/AssetDataGatherer.cpp
+++ b/Engine/Source/Runtime/AssetRegistry/Private/AssetDataGatherer.cpp
@@ -1421,6 +1421,8 @@ void FAssetDataDiscovery::EnsureCompletion()
 
 void FAssetDataDiscovery::TickInternal()
 {
+	LLM_SCOPE(ELLMTag::AssetRegistry);
+
 	TStringBuilder<256> DirLocalAbsPath;
 	TStringBuilder<128> DirLongPackageName;
 	TStringBuilder<128> DirMountRelPath;
@@ -2725,6 +2727,8 @@ void FAssetDataGatherer::EnsureCompletion()
 
 void FAssetDataGatherer::TickInternal(bool& bOutIsTickInterrupt)
 {
+	LLM_SCOPE(ELLMTag::AssetRegistry);
+
 	using namespace UE::AssetDataGather::Private;
 
 	const int32 BatchSize = FMath::Max(1, FTaskGraphInterface::Get().GetNumWorkerThreads()) * AssetDataGathererConstants::SingleThreadFilesPerBatch;
diff --git a/Engine/Source/Runtime/AssetRegistry/Private/AssetRegistry.cpp b/Engine/Source/Runtime/AssetRegistry/Private/AssetRegistry.cpp
--- a/Engine/Source/Runtime/AssetRegistry/Private/AssetRegistry.cpp
+++ b/Engine/Source/Runtime/AssetRegistry/Private/AssetRegistry.cpp
@@ -1112,6 +1112,8 @@ namespace UE::AssetRegistry
 
 void FAssetRegistryImpl::SearchAllAssets(Impl::FEventContext& EventContext, bool bSynchronousSearch)
 {
+	LLM_SCOPE(ELLMTag::AssetRegistry);
+
 	ConstructGatherer();
 	FAssetDataGatherer& Gatherer = *GlobalGatherer;
 	if (Gatherer.IsSynchronous())
@@ -1159,6 +1161,7 @@ void FAssetRegistryImpl::SearchAllAssets(Impl::FEventContext& EventContext, bool
 void UAssetRegistryImpl::WaitForCompletion()
 {
 	using namespace UE::AssetRegistry::Impl;
+	LLM_SCOPE(ELLMTag::AssetRegistry);
 
 	for (;;)
 	{
@@ -2880,6 +2883,7 @@ bool FAssetRegistryImpl::IsLoadingAssets() const
 void UAssetRegistryImpl::Tick(float DeltaTime)
 {
 	checkf(IsInGameThread(), TEXT("The tick function executes deferred loads and events and must be on the game thread to do so."));
+	LLM_SCOPE(ELLMTag::AssetRegistry);
 
 	UE::AssetRegistry::Impl::EGatherStatus Status = UE::AssetRegistry::Impl::EGatherStatus::Active;
 	double TickStartTime = -1; // Force a full flush if DeltaTime < 0
diff --git a/Engine/Source/Runtime/CoreUObject/Private/UObject/LinkerLoad.cpp b/Engine/Source/Runtime/CoreUObject/Private/UObject/LinkerLoad.cpp
--- a/Engine/Source/Runtime/CoreUObject/Private/UObject/LinkerLoad.cpp
+++ b/Engine/Source/Runtime/CoreUObject/Private/UObject/LinkerLoad.cpp
@@ -1364,6 +1364,7 @@ FLinkerLoad::ELinkerStatus FLinkerLoad::SerializePackageFileSummaryInternal()
 FLinkerLoad::ELinkerStatus FLinkerLoad::SerializePackageFileSummary()
 {
 	DECLARE_SCOPE_CYCLE_COUNTER( TEXT( "FLinkerLoad::SerializePackageFileSummary" ), STAT_LinkerLoad_SerializePackageFileSummary, STATGROUP_LinkerLoad );
+	LLM_SCOPE(ELLMTag::UObject);
 
 	if (bHasSerializedPackageFileSummary == false)
 	{
diff --git a/Engine/Source/Runtime/CoreUObject/Private/UObject/SavePackage/SavePackageUtilities.cpp b/Engine/Source/Runtime/CoreUObject/Private/UObject/SavePackage/SavePackageUtilities.cpp
--- a/Engine/Source/Runtime/CoreUObject/Private/UObject/SavePackage/SavePackageUtilities.cpp
+++ b/Engine/Source/Runtime/CoreUObject/Private/UObject/SavePackage/SavePackageUtilities.cpp
@@ -1377,12 +1377,15 @@ void FEDLCookChecker::Reset()
 	bIsActive = false;
 }
 
+LLM_DEFINE_TAG(EDLCookChecker);
+
 void FEDLCookChecker::AddImport(UObject* Import, UPackage* ImportingPackage)
 {
 	if (bIsActive)
 	{
 		if (!Import->GetOutermost()->HasAnyPackageFlags(PKG_CompiledIn))
 		{
+			LLM_SCOPE_BYTAG(EDLCookChecker);
 			FEDLNodeID NodeId = FindOrAddNode(FEDLNodeHash(Import, EObjectEvent::Serialize));
 			FEDLNodeData& NodeData = Nodes[NodeId];
 			FName ImportingPackageName = ImportingPackage->GetFName();
@@ -1399,6 +1402,7 @@ void FEDLCookChecker::AddExport(UObject* Export)
 {
 	if (bIsActive)
 	{
+		LLM_SCOPE_BYTAG(EDLCookChecker);
 		FEDLNodeID SerializeID = FindOrAddNode(FEDLNodeHash(Export, EObjectEvent::Serialize));
 		Nodes[SerializeID].bIsExport = true;
 		FEDLNodeID CreateID = FindOrAddNode(FEDLNodeHash(Export, EObjectEvent::Create));
@@ -1411,6 +1415,7 @@ void FEDLCookChecker::AddArc(UObject* DepObject, bool bDepIsSerialize, UObject*
 {
 	if (bIsActive)
 	{
+		LLM_SCOPE_BYTAG(EDLCookChecker);
 		FEDLNodeID ExportID = FindOrAddNode(FEDLNodeHash(Export, bExportIsSerialize ? EObjectEvent::Serialize : EObjectEvent::Create));
 		FEDLNodeID DepID = FindOrAddNode(FEDLNodeHash(DepObject, bDepIsSerialize ? EObjectEvent::Serialize : EObjectEvent::Create));
 		AddDependency(ExportID, DepID);
@@ -1421,6 +1426,7 @@ void FEDLCookChecker::AddPackageWithUnknownExports(FName LongPackageName)
 {
 	if (bIsActive)
 	{
+		LLM_SCOPE_BYTAG(EDLCookChecker);
 		PackagesWithUnknownExports.Add(LongPackageName);
 	}
 }
@@ -1433,6 +1439,7 @@ void FEDLCookChecker::AddDependency(FEDLNodeID SourceID, FEDLNodeID TargetID)
 
 void FEDLCookChecker::StartSavingEDLCookInfoForVerification()
 {
+	LLM_SCOPE_BYTAG(EDLCookChecker);
 	FScopeLock CookCheckerInstanceLock(&CookCheckerInstanceCritical);
 	for (FEDLCookChecker* Checker : CookCheckerInstances)
 	{
diff --git a/Engine/Source/Runtime/Engine/Private/AssetManager.cpp b/Engine/Source/Runtime/Engine/Private/AssetManager.cpp
--- a/Engine/Source/Runtime/Engine/Private/AssetManager.cpp
+++ b/Engine/Source/Runtime/Engine/Private/AssetManager.cpp
@@ -38,6 +38,7 @@
 #endif
 
 #define LOCTEXT_NAMESPACE "AssetManager"
+LLM_DEFINE_TAG(AssetManager);
 
 DEFINE_LOG_CATEGORY(LogAssetManager);
 
@@ -806,6 +807,7 @@ void UAssetManager::ScanPathsSynchronous(const TArray<FString>& PathsToScan) con
 
 int32 UAssetManager::ScanPathsForPrimaryAssets(FPrimaryAssetType PrimaryAssetType, const TArray<FString>& Paths, UClass* BaseClass, bool bHasBlueprintClasses, bool bIsEditorOnly, bool bForceSynchronousScan)
 {
+	LLM_SCOPE_BYTAG(AssetManager);
 	TRACE_CPUPROFILER_EVENT_SCOPE(UAssetManager::ScanPathsForPrimaryAssets)
 	TSharedRef<FPrimaryAssetTypeData>* FoundType = AssetTypeMap.Find(PrimaryAssetType);
 
@@ -3643,6 +3645,7 @@ void UAssetManager::UpdateManagementDatabase(bool bForceRefresh)
 	{
 		return;
 	}
+	LLM_SCOPE_BYTAG(AssetManager);
 
 	ManagementParentMap.Reset();
 
diff --git a/Engine/Source/Runtime/Engine/Private/FontFace.cpp b/Engine/Source/Runtime/Engine/Private/FontFace.cpp
--- a/Engine/Source/Runtime/Engine/Private/FontFace.cpp
+++ b/Engine/Source/Runtime/Engine/Private/FontFace.cpp
@@ -22,6 +22,7 @@ UFontFace::UFontFace()
 
 void UFontFace::Serialize(FArchive& Ar)
 {
+	LLM_SCOPE_BYNAME(TEXT("FontFaceData"));
 	Ar.UsingCustomVersion(FEditorObjectVersion::GUID);
 	Ar.UsingCustomVersion(FUE5MainStreamObjectVersion::GUID);
 
diff --git a/Engine/Source/Runtime/Foliage/Private/InstancedFoliage.cpp b/Engine/Source/Runtime/Foliage/Private/InstancedFoliage.cpp
--- a/Engine/Source/Runtime/Foliage/Private/InstancedFoliage.cpp
+++ b/Engine/Source/Runtime/Foliage/Private/InstancedFoliage.cpp
@@ -4129,6 +4129,7 @@ FArchive& operator<<(FArchive& Ar, FFoliageMeshInfo_Old& MeshInfo)
 
 void AInstancedFoliageActor::Serialize(FArchive& Ar)
 {
+	LLM_SCOPE_BYNAME(TEXT("Foliage"));
 	Super::Serialize(Ar);
 
 	Ar.UsingCustomVersion(FFoliageCustomVersion::GUID);
