Allow macro graphs to be renamed in a Blueprint macro library without breaking existing package imports for dependent macro instance nodes that reference the renamed graph.

diff --git a/Engine/Source/Editor/UnrealEd/Private/Kismet2/BlueprintEditorUtils.cpp b/Engine/Source/Editor/UnrealEd/Private/Kismet2/BlueprintEditorUtils.cpp
--- a/Engine/Source/Editor/UnrealEd/Private/Kismet2/BlueprintEditorUtils.cpp
+++ b/Engine/Source/Editor/UnrealEd/Private/Kismet2/BlueprintEditorUtils.cpp
@@ -2721,12 +2721,28 @@ void FBlueprintEditorUtils::RemoveGraph(UBlueprint* Blueprint, class UEdGraph* G
 /** Rename a graph and mark objects for modified */
 void FBlueprintEditorUtils::RenameGraph(UEdGraph* Graph, const FString& NewNameStr)
 {
-	if (Graph && Graph->Rename(*NewNameStr, Graph->GetOuter(), REN_Test))
+	if (Graph)
 	{
 		// Cache old name
 		const FName OldGraphName = Graph->GetFName();
 		UBlueprint* Blueprint = FBlueprintEditorUtils::FindBlueprintForGraphChecked(Graph);
 
+		// When renaming a graph inside a macro library, we may have dropped a redirector after a previous
+		// rename (see below). If we're now trying to reuse it, move the redirector aside to free up the name.
+		if (Blueprint->BlueprintType == BPTYPE_MacroLibrary)
+		{
+			if (UObjectRedirector* Redirector = FindObjectFast<UObjectRedirector>(Graph->GetOuter(), *NewNameStr))
+			{
+				Redirector->Rename(nullptr, GetTransientPackage(), REN_DontCreateRedirectors);
+			}
+		}
+
+		// Ensure that there are no collisions; leave the name as-is if this fails for some reason.
+		if (!Graph->Rename(*NewNameStr, Graph->GetOuter(), REN_Test))
+		{
+			return;
+		}
+
 		auto RenameGraphLambda = [](UEdGraph* GraphToRename, const FName LocalOldGraphName, const FName LocalNewGraphName, ERenameFlags RenameFlags)
 		{
 			// Ensure we have undo records
@@ -2760,9 +2776,23 @@ void FBlueprintEditorUtils::RenameGraph(UEdGraph* Graph, const FString& NewNameS
 			}
 		};
 
+		ERenameFlags RenameFlagsToApply = REN_None;
+		if (Blueprint->bIsRegeneratingOnLoad)
+		{
+			RenameFlagsToApply |= REN_ForceNoResetLoaders;
+		}
+
+		// Macro library graphs are referenced indirectly and resolved at edit/compile time via GUID (see FGraphReference).
+		// However, they will be exported by name at save time, so renaming a macro library graph implies we should also
+		// export a redirector with the old name so that the linker will still be able to resolve existing imports on load.
+		if (Blueprint->BlueprintType != BPTYPE_MacroLibrary)
+		{
+			RenameFlagsToApply |= REN_DontCreateRedirectors;
+		}
+
 		// Apply new name
 		const FName NewGraphName(*NewNameStr);
-		RenameGraphLambda(Graph, OldGraphName, NewGraphName, (Blueprint->bIsRegeneratingOnLoad ? REN_ForceNoResetLoaders : 0) | REN_DontCreateRedirectors);
+		RenameGraphLambda(Graph, OldGraphName, NewGraphName, RenameFlagsToApply);
 
 		TArray<UBlueprint*> ModifiedBlueprints;
 		ModifiedBlueprints.Add(Blueprint);
diff --git a/Engine/Source/Editor/UnrealEd/Private/Kismet2/Kismet2NameValidators.cpp b/Engine/Source/Editor/UnrealEd/Private/Kismet2/Kismet2NameValidators.cpp
--- a/Engine/Source/Editor/UnrealEd/Private/Kismet2/Kismet2NameValidators.cpp
+++ b/Engine/Source/Editor/UnrealEd/Private/Kismet2/Kismet2NameValidators.cpp
@@ -152,13 +152,29 @@ EValidatorResult FKismetNameValidator::IsValid(const FName& Name, bool /* bOrigi
 		// If it is in the names list then it is already in use.
 		if(!Names.Contains(Name))
 		{
-			UObject* ExistingObject = StaticFindObject(/*Class=*/ NULL, const_cast<UBlueprint*>(BlueprintObject), *Name.ToString(), true);
-			ValidatorResult = (ExistingObject == NULL)? EValidatorResult::Ok : EValidatorResult::AlreadyInUse;
+			ValidatorResult = EValidatorResult::Ok;
+
+			// Check for collision with an existing object.
+			if (UObject* ExistingObject = StaticFindObject(/*Class=*/ nullptr, const_cast<UBlueprint*>(BlueprintObject), *Name.ToString(), true))
+			{
+				// To allow the linker to resolve imports when dependent Blueprints are loaded, macro libraries
+				// will leave behind a redirector whenever one of its macro graph objects are renamed. These get
+				// moved aside to allow their names to be recycled, so we consider a collision with an existing
+				// redirector object in that case to be a false positive.
+				// 
+				// Note that while the linker will resolve references to either a graph or redirector by name on load,
+				// macro instance nodes will instead resolve references to any external macro library's graph by GUID. 
+				const bool bIgnoreRedirectors = BlueprintObject && BlueprintObject->BlueprintType == BPTYPE_MacroLibrary;
+				if (!bIgnoreRedirectors || !ExistingObject->IsA<UObjectRedirector>())
+				{
+					ValidatorResult = EValidatorResult::AlreadyInUse;
+				}
+			}
 		}
 
 		if(ValidatorResult == EValidatorResult::Ok)
 		{
-			if(Scope == NULL)
+			if(Scope == nullptr)
 			{
 				// Search through all functions for their local variables and prevent duplicate names
 				TArray<UK2Node_FunctionEntry*> FunctionEntryNodes;
@@ -182,7 +198,7 @@ EValidatorResult FKismetNameValidator::IsValid(const FName& Name, bool /* bOrigi
 			}
 			else
 			{
-				if(FindFProperty<const FProperty>(Scope, *Name.ToString()) != NULL)
+				if(FindFProperty<const FProperty>(Scope, *Name.ToString()) != nullptr)
 				{
 					ValidatorResult = EValidatorResult::LocallyInUse;
 				}
