ZenCOTF: Added RecookPackages request
#rnx

diff --git a/Engine/Source/Editor/UnrealEd/Classes/CookOnTheSide/CookOnTheFlyServer.h b/Engine/Source/Editor/UnrealEd/Classes/CookOnTheSide/CookOnTheFlyServer.h
--- a/Engine/Source/Editor/UnrealEd/Classes/CookOnTheSide/CookOnTheFlyServer.h
+++ b/Engine/Source/Editor/UnrealEd/Classes/CookOnTheSide/CookOnTheFlyServer.h
@@ -544,6 +544,8 @@ public:
 	*/
 	void StopAndClearCookedData();
 
+	void TickRequestManager();
+
 	/** 
 	* Process any shader recompile requests
 	*/
diff --git a/Engine/Source/Editor/UnrealEd/Private/Commandlets/CookCommandlet.cpp b/Engine/Source/Editor/UnrealEd/Private/Commandlets/CookCommandlet.cpp
--- a/Engine/Source/Editor/UnrealEd/Private/Commandlets/CookCommandlet.cpp
+++ b/Engine/Source/Editor/UnrealEd/Private/Commandlets/CookCommandlet.cpp
@@ -556,6 +556,8 @@ bool UCookCommandlet::CookOnTheFly( FGuid InstanceId, int32 Timeout, bool bForce
 
 		while ( (CookOnTheFlyServer->HasRemainingWork() == false) && !IsEngineExitRequested())
 		{
+			CookOnTheFlyServer->TickRequestManager();
+
 			CookOnTheFlyServer->TickRecompileShaderRequests();
 
 			// Shaders need to be updated
diff --git a/Engine/Source/Editor/UnrealEd/Private/CookOnTheFlyServer.cpp b/Engine/Source/Editor/UnrealEd/Private/CookOnTheFlyServer.cpp
--- a/Engine/Source/Editor/UnrealEd/Private/CookOnTheFlyServer.cpp
+++ b/Engine/Source/Editor/UnrealEd/Private/CookOnTheFlyServer.cpp
@@ -355,6 +355,27 @@ public:
 		}
 	};
 
+	virtual void MarkPackageDirty(const FName& PackageName) override
+	{
+		Cooker.ExternalRequests->AddCallback([this, PackageName]()
+			{
+				UE::Cook::FPackageData* PackageData = Cooker.PackageDatas->FindPackageDataByPackageName(PackageName);
+				if (!PackageData)
+				{
+					return;
+				}
+				if (PackageData->IsInProgress())
+				{
+					return;
+				}
+				if (!PackageData->HasAnyCookedPlatform())
+				{
+					return;
+				}
+				PackageData->ClearCookProgress();
+			});
+	}
+
 	virtual bool EnqueueRecompileShaderRequest(UE::Cook::FRecompileShaderRequest RecompileShaderRequest) override
 	{
 		UE_LOG(LogCook, Verbose, TEXT("Enqueing recompile shader(s) request, Platform='%s', ShaderPlatform='%d'"),
@@ -435,6 +456,7 @@ void UCookOnTheFlyServer::Tick(float DeltaTime)
 	uint32 CookedPackagesCount = 0;
 	const static float CookOnTheSideTimeSlice = 0.1f; // seconds
 	TickCookOnTheSide( CookOnTheSideTimeSlice, CookedPackagesCount);
+	TickRequestManager();
 	TickRecompileShaderRequests();
 }
 
@@ -1251,6 +1273,11 @@ uint32 UCookOnTheFlyServer::TickCookOnTheSide(const float TimeSlice, uint32 &Coo
 		CookByTheBookFinished();
 	}
 
+	if (IsCookOnTheFlyMode() && bCookComplete)
+	{
+		CollectGarbage(RF_NoFlags);
+	}
+
 	CookedPackageCount += StackData.CookedPackageCount;
 	return StackData.ResultFlags;
 }
@@ -3820,6 +3847,14 @@ void UCookOnTheFlyServer::BeginDestroy()
 	Super::BeginDestroy();
 }
 
+void UCookOnTheFlyServer::TickRequestManager()
+{
+	if (CookOnTheFlyRequestManager)
+	{
+		CookOnTheFlyRequestManager->Tick();
+	}
+}
+
 void UCookOnTheFlyServer::TickRecompileShaderRequests()
 {
 	// try to pull off a request
diff --git a/Engine/Source/Editor/UnrealEd/Private/Cooker/CookOnTheFlyServerInterface.h b/Engine/Source/Editor/UnrealEd/Private/Cooker/CookOnTheFlyServerInterface.h
--- a/Engine/Source/Editor/UnrealEd/Private/Cooker/CookOnTheFlyServerInterface.h
+++ b/Engine/Source/Editor/UnrealEd/Private/Cooker/CookOnTheFlyServerInterface.h
@@ -87,6 +87,8 @@ public:
 	 */
 	virtual bool EnqueueCookRequest(FCookPackageRequest CookPackageRequest) = 0;
 
+	virtual void MarkPackageDirty(const FName& PackageName) = 0;
+
 	/**
 	 * Enqueue a new shader compile request.
 	 *
@@ -113,6 +115,8 @@ public:
 	/** Initialze the request manager. */
 	virtual bool Initialize() = 0;
 
+	virtual void Tick() = 0;
+
 	/** Shutdown the reques tmanager. */
 	virtual void Shutdown() = 0;
 
diff --git a/Engine/Source/Editor/UnrealEd/Private/Cooker/IoStoreCookOnTheFlyRequestManager.cpp b/Engine/Source/Editor/UnrealEd/Private/Cooker/IoStoreCookOnTheFlyRequestManager.cpp
--- a/Engine/Source/Editor/UnrealEd/Private/Cooker/IoStoreCookOnTheFlyRequestManager.cpp
+++ b/Engine/Source/Editor/UnrealEd/Private/Cooker/IoStoreCookOnTheFlyRequestManager.cpp
@@ -19,6 +19,7 @@
 #include "Async/Async.h"
 #include "Serialization/BufferArchive.h"
 #include "NetworkMessage.h"
+#include "Engine/Engine.h"
 
 class FIoStoreCookOnTheFlyNetworkServer
 {
@@ -537,6 +538,29 @@ private:
 			}
 		}
 
+		void RequestRecook(UE::Cook::ICookOnTheFlyServer& InCookOnTheFlyServer, const FPackageId& PackageId, const FName& PackageName)
+		{
+			FPackage& Package = GetPackage(PackageId);
+			if (Package.Status != EPackageStatus::Cooked && Package.Status != EPackageStatus::Failed)
+			{
+				UE_LOG(LogCookOnTheFly, Verbose, TEXT("Skipping recook of package 0x%llX '%s' that was not cooked"), PackageId.ValueForDebugging(), *PackageName.ToString());
+				return;
+			}
+			FString Filename;
+			if (FPackageName::TryConvertLongPackageNameToFilename(PackageName.ToString(), Filename))
+			{
+				UE_LOG(LogCookOnTheFly, Verbose, TEXT("Recooking package 0x%llX '%s'"), PackageId.ValueForDebugging(), *PackageName.ToString());
+				Package.Status = EPackageStatus::Cooking;
+				const bool bEnqueued = InCookOnTheFlyServer.EnqueueCookRequest(UE::Cook::FCookPackageRequest{ PlatformName, Filename });
+				check(bEnqueued);
+			}
+			else
+			{
+				UE_LOG(LogCookOnTheFly, Warning, TEXT("Failed to recook package 0x%llX '%s' (File not found)"), PackageId.ValueForDebugging(), *PackageName.ToString());
+				Package.Status = EPackageStatus::Failed;
+			}
+		}
+
 		void MarkAsFailed(FPackageId PackageId, UE::ZenCookOnTheFly::Messaging::FCompletedPackages& OutCompletedPackages)
 		{
 			UE_LOG(LogCookOnTheFly, Warning, TEXT("0x%llX failed"), PackageId.ValueForDebugging());
@@ -633,6 +657,70 @@ private:
 		AllKnownPackagesMap.Add(PackageId, PackageName);
 	}
 
+	void TickRecookPackages()
+	{
+		TArray<FPackageId, TInlineAllocator<128>> PackageIds;
+		{
+			FScopeLock _(&PackagesToRecookCritical);
+			if (PackagesToRecook.IsEmpty())
+			{
+				return;
+			}
+			PackageIds = PackagesToRecook.Array();
+			PackagesToRecook.Empty();
+		}
+
+		TArray<FName, TInlineAllocator<128>> PackageNames;
+		PackageNames.Reserve(PackageIds.Num());
+
+		CollectGarbage(RF_NoFlags);
+
+		for (FPackageId PackageId : PackageIds)
+		{
+			FName PackageName = AllKnownPackagesMap.FindRef(PackageId);
+			if (!PackageName.IsNone())
+			{
+				PackageNames.Add(PackageName);
+
+				UPackage* Package = FindObjectFast<UPackage>(nullptr, PackageName);
+				if (Package)
+				{
+					UE_LOG(LogCookOnTheFly, Warning, TEXT("Can't recook package '%s'"), *PackageName.ToString());
+					UEngine::FindAndPrintStaleReferencesToObject(Package, ELogVerbosity::Display);
+				}
+				else
+				{
+					UE_LOG(LogCookOnTheFly, Verbose, TEXT("Recooking package '%s'"), *PackageName.ToString());
+				}
+			}
+		}
+
+		for (const FName& PackageName : PackageNames)
+		{
+			CookOnTheFlyServer.MarkPackageDirty(PackageName);
+		}
+
+		ForEachContext([this, &PackageIds, &PackageNames](FPlatformContext& Context)
+			{
+				FScopeLock _(&Context.GetLock());
+				for (int32 PackageIndex = 0; PackageIndex < PackageNames.Num(); ++PackageIndex)
+				{
+					const FPackageId& PackageId = PackageIds[PackageIndex];
+					const FName& PackageName = PackageNames[PackageIndex];
+					if (!PackageName.IsNone())
+					{
+						Context.RequestRecook(CookOnTheFlyServer, PackageId, PackageName);
+					}
+				}
+				return true;
+			});
+	}
+
+	virtual void Tick() override
+	{
+		TickRecookPackages();
+	}
+
 private:
 	bool HandleClientConnection(const FName& PlatformName, FIoStoreCookOnTheFlyNetworkServer::EConnectionStatus ConnectionStatus)
 	{
@@ -695,6 +783,9 @@ private:
 			case UE::Cook::ECookOnTheFlyMessage::RecompileShaders:
 				bRequestOk = HandleRecompileShadersRequest(PlatformName, Request, Response);
 				break;
+			case UE::Cook::ECookOnTheFlyMessage::RecookPackages:
+				bRequestOk = HandleRecookPackagesRequest(PlatformName, Request, Response);
+				break;
 			default:
 				UE_LOG(LogCookOnTheFly, Fatal, TEXT("Unknown request, Type='%s', Client='%s'"), LexToString(Request.GetHeader().MessageType), *PlatformName.ToString());
 				break;
@@ -769,6 +860,29 @@ private:
 		return true;
 	}
 
+	bool HandleRecookPackagesRequest(const FName& PlatformName, const UE::Cook::FCookOnTheFlyRequest& Request, UE::Cook::FCookOnTheFlyResponse& Response)
+	{
+		using namespace UE::ZenCookOnTheFly::Messaging;
+
+		TRACE_CPUPROFILER_EVENT_SCOPE(CookOnTheFly::HandleRecookPackagesRequest);
+
+		FRecookPackagesRequest RecookRequest = Request.GetBodyAs<FRecookPackagesRequest>();
+
+		UE_LOG(LogCookOnTheFly, Display, TEXT("Received recook request for %d packages"), RecookRequest.PackageIds.Num());
+
+		{
+			FScopeLock _(&PackagesToRecookCritical);
+			for (FPackageId PackageId : RecookRequest.PackageIds)
+			{
+				PackagesToRecook.Add(PackageId);
+			}
+		}
+
+		Response.SetStatus(UE::Cook::ECookOnTheFlyMessageStatus::Ok);
+
+		return true;
+	}
+
 	void OnPackageCooked(const IPackageStoreWriter::FCommitEventArgs& EventArgs)
 	{
 		using namespace UE::Cook;
@@ -921,6 +1035,20 @@ private:
 		return *Ctx;
 	}
 
+	void ForEachContext(TFunctionRef<bool(FPlatformContext&)> Callback)
+	{
+		FScopeLock _(&ContextsCriticalSection);
+		for (auto& KV : PlatformContexts)
+		{
+			TUniquePtr<FPlatformContext>& Ctx = KV.Value;
+			check(Ctx.IsValid());
+			if (!Callback(*Ctx))
+			{
+				return;
+			}
+		}
+	}
+
 	UE::Cook::ICookOnTheFlyServer& CookOnTheFlyServer;
 	UE::Cook::FIoStoreCookOnTheFlyServerOptions Options;
 	TUniquePtr<FIoStoreCookOnTheFlyNetworkServer> ConnectionServer;
@@ -928,6 +1056,8 @@ private:
 	TMap<FName, TUniquePtr<FPlatformContext>> PlatformContexts;
 	FCriticalSection AllKnownPackagesCriticalSection;
 	TMap<FPackageId, FName> AllKnownPackagesMap;
+	FCriticalSection PackagesToRecookCritical;
+	TSet<FPackageId> PackagesToRecook;
 };
 
 namespace UE { namespace Cook
diff --git a/Engine/Source/Editor/UnrealEd/Private/Cooker/NetworkFileCookOnTheFlyRequestManager.cpp b/Engine/Source/Editor/UnrealEd/Private/Cooker/NetworkFileCookOnTheFlyRequestManager.cpp
--- a/Engine/Source/Editor/UnrealEd/Private/Cooker/NetworkFileCookOnTheFlyRequestManager.cpp
+++ b/Engine/Source/Editor/UnrealEd/Private/Cooker/NetworkFileCookOnTheFlyRequestManager.cpp
@@ -34,6 +34,10 @@ public:
 		return NetworkFileServer.IsValid();
 	}
 
+	virtual void Tick() override
+	{
+	}
+
 	virtual void Shutdown() override
 	{
 		if (NetworkFileServer.IsValid())
diff --git a/Engine/Source/Runtime/CookOnTheFly/Internal/CookOnTheFly.h b/Engine/Source/Runtime/CookOnTheFly/Internal/CookOnTheFly.h
--- a/Engine/Source/Runtime/CookOnTheFly/Internal/CookOnTheFly.h
+++ b/Engine/Source/Runtime/CookOnTheFly/Internal/CookOnTheFly.h
@@ -52,6 +52,8 @@ enum class ECookOnTheFlyMessage : uint32
 	FilesAdded			= 0x60,
 	/* Heartbeat message. */
 	Heartbeat			= 0x70,
+	/* Request to recook packages. */
+	RecookPackages		= 0x80,
 };
 ENUM_CLASS_FLAGS(ECookOnTheFlyMessage);
 
@@ -80,6 +82,8 @@ inline const TCHAR* LexToString(ECookOnTheFlyMessage Message)
 			return TEXT("FilesAdded");
 		case ECookOnTheFlyMessage::Heartbeat:
 			return TEXT("Heartbeat");
+		case ECookOnTheFlyMessage::RecookPackages:
+			return TEXT("RecookPackages");
 		default:
 			return TEXT("Unknown");
 	};
diff --git a/Engine/Source/Runtime/CookOnTheFly/Internal/CookOnTheFlyMessages.h b/Engine/Source/Runtime/CookOnTheFly/Internal/CookOnTheFlyMessages.h
--- a/Engine/Source/Runtime/CookOnTheFly/Internal/CookOnTheFlyMessages.h
+++ b/Engine/Source/Runtime/CookOnTheFly/Internal/CookOnTheFlyMessages.h
@@ -34,6 +34,18 @@ struct FCookPackageResponse
 	COOKONTHEFLY_API friend FArchive& operator<<(FArchive& Ar, FCookPackageResponse& Response);
 };
 
+struct FRecookPackagesRequest
+{
+	TArray<FPackageId> PackageIds;
+
+	COOKONTHEFLY_API friend FArchive& operator<<(FArchive& Ar, FRecookPackagesRequest& Request);
+};
+
+struct FRecookPackagesResponse
+{
+	COOKONTHEFLY_API friend FArchive& operator<<(FArchive& Ar, FRecookPackagesResponse& Response);
+};
+
 }}} // namesapce UE::ZenCookOnTheFly::Messaging
 
 #endif // WITH_COTF
diff --git a/Engine/Source/Runtime/CookOnTheFly/Private/CookOnTheFlyMessages.cpp b/Engine/Source/Runtime/CookOnTheFly/Private/CookOnTheFlyMessages.cpp
--- a/Engine/Source/Runtime/CookOnTheFly/Private/CookOnTheFlyMessages.cpp
+++ b/Engine/Source/Runtime/CookOnTheFly/Private/CookOnTheFlyMessages.cpp
@@ -33,6 +33,18 @@ FArchive& operator<<(FArchive& Ar, FCookPackageResponse& Response)
 	return Ar;
 }
 
+FArchive& operator<<(FArchive& Ar, FRecookPackagesRequest& Request)
+{
+	Ar << Request.PackageIds;
+
+	return Ar;
+}
+
+FArchive& operator<<(FArchive& Ar, FRecookPackagesResponse& Response)
+{
+	return Ar;
+}
+
 }}} // namesapce UE::ZenCookOnTheFly::Messaging
 
 #endif // WITH_COTF
