ZenCOTF:
Created ExternalCookOnTheFlyServer editor object for communicating with a running Zen COTF server
Recook packages when they are saved in the editor
#rnx

diff --git a/Engine/Source/Editor/UnrealEd/Classes/Editor/UnrealEdEngine.h b/Engine/Source/Editor/UnrealEd/Classes/Editor/UnrealEdEngine.h
--- a/Engine/Source/Editor/UnrealEd/Classes/Editor/UnrealEdEngine.h
+++ b/Engine/Source/Editor/UnrealEd/Classes/Editor/UnrealEdEngine.h
@@ -949,6 +949,9 @@ private:
 	/** Transient unsaved version of template map infos used by the editor. */
 	TArray<FTemplateMapInfo> TemplateMapInfoCache;
 
+	/** Proxy for a cotf server running in a separate process */
+	class FExternalCookOnTheFlyServer* ExternalCookOnTheFlyServer = nullptr;
+
 	/**
 	* Internal helper function to count how many dirty packages require checkout.
 	*
diff --git a/Engine/Source/Editor/UnrealEd/Private/Cooker/ExternalCookOnTheFlyServer.cpp b/Engine/Source/Editor/UnrealEd/Private/Cooker/ExternalCookOnTheFlyServer.cpp
new file mode 100644
--- /dev/null
+++ b/Engine/Source/Editor/UnrealEd/Private/Cooker/ExternalCookOnTheFlyServer.cpp
@@ -0,0 +1,126 @@
+// Copyright Epic Games, Inc. All Rights Reserved.
+
+#include "ExternalCookOnTheFlyServer.h"
+#include "CookOnTheFly.h"
+#include "CookOnTheFlyMessages.h"
+#include "AssetRegistry/AssetRegistryModule.h"
+#include "MessageEndpointBuilder.h"
+#include "Misc/Paths.h"
+#include "Misc/PathViews.h"
+
+FExternalCookOnTheFlyServer::FExternalCookOnTheFlyServer()
+	: CookOnTheFlyModule(FModuleManager::LoadModuleChecked<UE::Cook::ICookOnTheFlyModule>(TEXT("CookOnTheFly")))
+	, AssetRegistry(FModuleManager::LoadModuleChecked<FAssetRegistryModule>(AssetRegistryConstants::ModuleName).Get())
+	, ServiceId(GenerateServiceId())
+{
+	MessageEndpoint = FMessageEndpoint::Builder("FCookOnTheFly")
+		.Handling<FZenCookOnTheFlyRegisterServiceMessage>(this, &FExternalCookOnTheFlyServer::HandleRegisterServiceMessage)
+		.ReceivingOnThread(ENamedThreads::GameThread);
+
+	if (MessageEndpoint.IsValid())
+	{
+		MessageEndpoint->Subscribe<FZenCookOnTheFlyRegisterServiceMessage>();
+	}
+}
+
+FExternalCookOnTheFlyServer::~FExternalCookOnTheFlyServer()
+{
+	AssetRegistry.OnAssetUpdatedOnDisk().RemoveAll(this);
+}
+
+void FExternalCookOnTheFlyServer::HandleRegisterServiceMessage(const FZenCookOnTheFlyRegisterServiceMessage& Message, const TSharedRef<IMessageContext, ESPMode::ThreadSafe>& Context)
+{
+	if (Message.ServiceId != ServiceId)
+	{
+		UE_LOG(LogCookOnTheFly, Verbose, TEXT("Rejecting service from %s"), *Context->GetSender().ToString());
+		return;
+	}
+	UE_LOG(LogCookOnTheFly, Verbose, TEXT("Accepting service from %s"), *Context->GetSender().ToString());
+	UE::Cook::FCookOnTheFlyHostOptions HostOptions;
+	HostOptions.Hosts.Add(TEXT("127.0.0.1"));
+	HostOptions.Port = Message.Port;
+	ServerConnection = CookOnTheFlyModule.ConnectToServer(HostOptions);
+	if (ServerConnection)
+	{
+		UE_LOG(LogCookOnTheFly, Display, TEXT("Connected to server"));
+		AssetRegistry.OnAssetUpdatedOnDisk().AddRaw(this, &FExternalCookOnTheFlyServer::AssetUpdatedOnDisk);
+	}
+	else
+	{
+		UE_LOG(LogCookOnTheFly, Warning, TEXT("Failed connecting to server"));
+	}
+}
+
+void FExternalCookOnTheFlyServer::Tick(float DeltaSeconds)
+{
+	{
+		FScopeLock _(&AllPackagesToRecookCritical);
+		if (!AllPackagesToRecook.IsEmpty())
+		{
+			TArray<FPackageId> PackageIdsToRecook;
+			PackageIdsToRecook.Reserve(AllPackagesToRecook.Num());
+			for (FName PackageName : AllPackagesToRecook)
+			{
+				PackageIdsToRecook.Add(FPackageId::FromName(PackageName));
+			}
+			AllPackagesToRecook.Empty();
+
+			UE::Cook::FCookOnTheFlyRequest Request(UE::Cook::ECookOnTheFlyMessage::RecookPackages);
+			Request.SetBodyTo(UE::ZenCookOnTheFly::Messaging::FRecookPackagesRequest{ MoveTemp(PackageIdsToRecook) });
+			UE::Cook::FCookOnTheFlyResponse Response = ServerConnection->SendRequest(Request).Get();
+			if (!Response.IsOk())
+			{
+				UE_LOG(LogCookOnTheFly, Warning, TEXT("Failed to send RecookPackages request"));
+			}
+		}
+	}
+	if (ServerConnection && !ServerConnection->IsConnected())
+	{
+		AssetRegistry.OnAssetUpdatedOnDisk().RemoveAll(this);
+		ServerConnection.Reset();
+		UE_LOG(LogCookOnTheFly, Display, TEXT("Disconnected from server"));
+	}
+}
+
+bool FExternalCookOnTheFlyServer::IsTickable() const
+{
+	return ServerConnection.IsValid();
+}
+
+void FExternalCookOnTheFlyServer::AssetUpdatedOnDisk(const FAssetData& AssetData)
+{
+	FScopeLock _(&AllPackagesToRecookCritical);
+	if (AllPackagesToRecook.Contains(AssetData.PackageName))
+	{
+		return;
+	}
+
+	TArray<FName> ModifiedPackagesToRecurse;
+	ModifiedPackagesToRecurse.Push(AssetData.PackageName);
+	AllPackagesToRecook.Add(AssetData.PackageName);
+	while (!ModifiedPackagesToRecurse.IsEmpty())
+	{
+		FName ModifiedPackage = ModifiedPackagesToRecurse.Pop(false);
+		TArray<FName> Referencers;
+		AssetRegistry.GetReferencers(ModifiedPackage, Referencers, UE::AssetRegistry::EDependencyCategory::Package, UE::AssetRegistry::EDependencyQuery::Hard);
+		for (FName Referencer : Referencers)
+		{
+			if (!AllPackagesToRecook.Contains(Referencer))
+			{
+				ModifiedPackagesToRecurse.Push(Referencer);
+				AllPackagesToRecook.Add(Referencer);
+			}
+		}
+	}
+}
+
+FString FExternalCookOnTheFlyServer::GenerateServiceId()
+{
+	TStringBuilder<256> ServiceIdString;
+	ServiceIdString.Append(FPlatformProcess::UserName());
+	ServiceIdString.Append("@");
+	ServiceIdString.Append(FPlatformProcess::ComputerName());
+	ServiceIdString.Append(":");
+	FPathViews::ToAbsolutePath(FPaths::GetProjectFilePath(), ServiceIdString);
+	return FMD5::HashBytes((unsigned char*)*ServiceIdString, ServiceIdString.Len() * sizeof(TCHAR));
+}
diff --git a/Engine/Source/Editor/UnrealEd/Private/Cooker/ExternalCookOnTheFlyServer.h b/Engine/Source/Editor/UnrealEd/Private/Cooker/ExternalCookOnTheFlyServer.h
new file mode 100644
--- /dev/null
+++ b/Engine/Source/Editor/UnrealEd/Private/Cooker/ExternalCookOnTheFlyServer.h
@@ -0,0 +1,49 @@
+// Copyright Epic Games, Inc. All Rights Reserved.
+
+#pragma once
+
+#include "CookOnTheFly.h"
+#include "TickableEditorObject.h"
+#include "IMessageContext.h"
+#include "MessageEndpoint.h"
+#include "ExternalCookOnTheFlyServer.generated.h"
+
+class IAssetRegistry;
+
+class FExternalCookOnTheFlyServer
+	: public FTickableEditorObject
+{
+public:
+	FExternalCookOnTheFlyServer();
+	~FExternalCookOnTheFlyServer();
+	void HandleRegisterServiceMessage(const FZenCookOnTheFlyRegisterServiceMessage& Message, const TSharedRef<IMessageContext, ESPMode::ThreadSafe>& Context);
+
+	static FString GenerateServiceId();
+
+private:
+	virtual void Tick(float DeltaSeconds) override;
+	virtual bool IsTickable() const override;
+	virtual TStatId GetStatId() const override { return TStatId(); }
+
+	void AssetUpdatedOnDisk(const FAssetData& AssetData);
+
+	UE::Cook::ICookOnTheFlyModule& CookOnTheFlyModule;
+	IAssetRegistry& AssetRegistry;
+	TSharedPtr<FMessageEndpoint, ESPMode::ThreadSafe> MessageEndpoint;
+	const FString ServiceId;
+	TUniquePtr<UE::Cook::ICookOnTheFlyServerConnection> ServerConnection;
+	FCriticalSection AllPackagesToRecookCritical;
+	TSet<FName> AllPackagesToRecook;
+};
+
+USTRUCT()
+struct FZenCookOnTheFlyRegisterServiceMessage
+{
+	GENERATED_USTRUCT_BODY()
+
+	UPROPERTY(EditAnywhere, Category = "Message")
+	FString ServiceId;
+
+	UPROPERTY(EditAnywhere, Category = "Message")
+	int32 Port = -1;
+};
\ No newline at end of file
diff --git a/Engine/Source/Editor/UnrealEd/Private/Cooker/IoStoreCookOnTheFlyRequestManager.cpp b/Engine/Source/Editor/UnrealEd/Private/Cooker/IoStoreCookOnTheFlyRequestManager.cpp
--- a/Engine/Source/Editor/UnrealEd/Private/Cooker/IoStoreCookOnTheFlyRequestManager.cpp
+++ b/Engine/Source/Editor/UnrealEd/Private/Cooker/IoStoreCookOnTheFlyRequestManager.cpp
@@ -20,6 +20,9 @@
 #include "Serialization/BufferArchive.h"
 #include "NetworkMessage.h"
 #include "Engine/Engine.h"
+#include "MessageEndpoint.h"
+#include "MessageEndpointBuilder.h"
+#include "Cooker/ExternalCookOnTheFlyServer.h"
 
 class FIoStoreCookOnTheFlyNetworkServer
 {
@@ -61,7 +64,9 @@ public:
 
 	FIoStoreCookOnTheFlyNetworkServer(FServerOptions InOptions)
 		: Options(MoveTemp(InOptions))
+		, ServiceId(FExternalCookOnTheFlyServer::GenerateServiceId())
 	{
+		MessageEndpoint = FMessageEndpoint::Builder("FCookOnTheFly");
 	}
 
 	~FIoStoreCookOnTheFlyNetworkServer()
@@ -112,6 +117,14 @@ public:
 
 		UE_LOG(LogCookOnTheFly, Display, TEXT("COTF server is ready for client(s) on '%s'!"), *ListenAddr->ToString(true));
 
+		if (MessageEndpoint.IsValid())
+		{
+			FZenCookOnTheFlyRegisterServiceMessage* RegisterServiceMessage = FMessageEndpoint::MakeMessage<FZenCookOnTheFlyRegisterServiceMessage>();
+			RegisterServiceMessage->ServiceId = ServiceId;
+			RegisterServiceMessage->Port = ListenAddr->GetPort();
+			MessageEndpoint->Publish(RegisterServiceMessage);
+		}
+
 		return true;
 	}
 
@@ -433,6 +446,8 @@ private:
 	TAtomic<bool> bStopRequested{ false };
 	uint32 NextClientId = 1;
 	TAtomic<uint32> NextCorrelationId{ 1 };
+	TSharedPtr<FMessageEndpoint, ESPMode::ThreadSafe> MessageEndpoint;
+	const FString ServiceId;
 };
 
 class FIoStoreCookOnTheFlyRequestManager final
diff --git a/Engine/Source/Editor/UnrealEd/Private/UnrealEdEngine.cpp b/Engine/Source/Editor/UnrealEd/Private/UnrealEdEngine.cpp
--- a/Engine/Source/Editor/UnrealEd/Private/UnrealEdEngine.cpp
+++ b/Engine/Source/Editor/UnrealEd/Private/UnrealEdEngine.cpp
@@ -71,7 +71,7 @@
 #include "ProfilingDebugging/StallDetector.h"
 #include "GameMapsSettings.h"
 #include "HAL/PlatformApplicationMisc.h"
-
+#include "Cooker/ExternalCookOnTheFlyServer.h"
 
 DEFINE_LOG_CATEGORY_STATIC(LogUnrealEdEngine, Log, All);
 
@@ -217,6 +217,8 @@ void UUnrealEdEngine::Init(IEngineLoop* InEngineLoop)
 			CookServer = NewObject<UCookOnTheFlyServer>();
 			CookServer->Initialize(ECookMode::CookByTheBookFromTheEditor, BaseCookingFlags);
 		}
+
+		ExternalCookOnTheFlyServer = new FExternalCookOnTheFlyServer();
 	}
 
 	if (FParse::Param(FCommandLine::Get(), TEXT("nomcp")))
@@ -431,6 +433,12 @@ void UUnrealEdEngine::FinishDestroy()
 		FCoreUObjectDelegates::OnObjectModified.RemoveAll(CookServer);
 	}
 
+	if (ExternalCookOnTheFlyServer)
+	{
+		delete ExternalCookOnTheFlyServer;
+		ExternalCookOnTheFlyServer = nullptr;
+	}
+
 	if(PackageAutoSaver.Get())
 	{
 		// We've finished shutting down, so disable the auto-save restore
diff --git a/Engine/Source/Programs/AutomationTool/Scripts/CookCommand.Automation.cs b/Engine/Source/Programs/AutomationTool/Scripts/CookCommand.Automation.cs
--- a/Engine/Source/Programs/AutomationTool/Scripts/CookCommand.Automation.cs
+++ b/Engine/Source/Programs/AutomationTool/Scripts/CookCommand.Automation.cs
@@ -253,7 +253,7 @@ namespace AutomationScripts
 
 					if (Params.ZenStore)
 					{
-						COTFCommandLine += " -zenstore";
+						COTFCommandLine += " -zenstore -messaging";
 					}
 
 					var ServerLogFile = CombinePaths(LogFolderOutsideOfSandbox, "Server.log");
