Fix TextureFactory::AlphaCoverageThresholds should have had default value 0
AlphaCoverage scaling was accidentally turned on for all imported textures after 2016
InterchangeTextureFactory replicated the broken behavior of TextureFactory
this fixes the import but broken textures are widespread

diff --git a/Engine/Plugins/Experimental/Interchange/Runtime/Source/Import/Private/Texture/InterchangeTextureFactory.cpp b/Engine/Plugins/Experimental/Interchange/Runtime/Source/Import/Private/Texture/InterchangeTextureFactory.cpp
--- a/Engine/Plugins/Experimental/Interchange/Runtime/Source/Import/Private/Texture/InterchangeTextureFactory.cpp
+++ b/Engine/Plugins/Experimental/Interchange/Runtime/Source/Import/Private/Texture/InterchangeTextureFactory.cpp
@@ -807,7 +807,6 @@ UObject* UInterchangeTextureFactory::CreateEmptyAsset(const FCreateAssetParams&
 	if (!ExistingAsset)
 	{
 		UTexture* NewTexture = NewObject<UTexture>(Arguments.Parent, TextureClass, *Arguments.AssetName, RF_Public | RF_Standalone);
-		NewTexture->AlphaCoverageThresholds = FVector4(0.0, 0.0, 0.0, 1.0);
 		Texture = NewTexture;
 
 		if (UTextureLightProfile* LightProfile = Cast<UTextureLightProfile>(NewTexture))
diff --git a/Engine/Source/Editor/UnrealEd/Classes/Factories/TextureFactory.h b/Engine/Source/Editor/UnrealEd/Classes/Factories/TextureFactory.h
--- a/Engine/Source/Editor/UnrealEd/Classes/Factories/TextureFactory.h
+++ b/Engine/Source/Editor/UnrealEd/Classes/Factories/TextureFactory.h
@@ -118,7 +118,7 @@ class UNREALED_API UTextureFactory : public UFactory, public IImportSettingsPars
 
 	/** Channel values to compare to when preserving alpha coverage from a mask. */
 	UPROPERTY(EditAnywhere, Category=PreserveAlphaCoverage, meta=(ToolTip="Channel values to compare to when preserving alpha coverage from a mask for mips"))
-	FVector4 AlphaCoverageThresholds;
+	FVector4 AlphaCoverageThresholds = FVector4(0,0,0,0);
 
 	/** If enabled, preserve the value of border pixels when creating mip-maps */
 	UPROPERTY(EditAnywhere, Category=PreserveBorder, meta=(ToolTip="If enabled, preserve the value of border pixels when creating mip-maps"))
