WorldPartition: Fix drag & drop move WP World item in ContentBrowser

#rnx

diff --git a/Engine/Source/Editor/UnrealEd/Private/DragAndDrop/AssetDragDropOp.cpp b/Engine/Source/Editor/UnrealEd/Private/DragAndDrop/AssetDragDropOp.cpp
--- a/Engine/Source/Editor/UnrealEd/Private/DragAndDrop/AssetDragDropOp.cpp
+++ b/Engine/Source/Editor/UnrealEd/Private/DragAndDrop/AssetDragDropOp.cpp
@@ -210,6 +210,7 @@ void FAssetDragDropOp::Init(TArray<FAssetData> InAssetData, TArray<FString> InAs
 	// Can cause unsafe frame reentry 
 	for (FAssetData& Data : AssetData)
 	{
+		FScopedLoadAllExternalObjects Scope(Data.PackageName);
 		Data.GetAsset();
 	}
 
diff --git a/Engine/Source/Editor/UnrealEd/Private/FileHelpers.cpp b/Engine/Source/Editor/UnrealEd/Private/FileHelpers.cpp
--- a/Engine/Source/Editor/UnrealEd/Private/FileHelpers.cpp
+++ b/Engine/Source/Editor/UnrealEd/Private/FileHelpers.cpp
@@ -689,6 +689,7 @@ static bool SaveWorld(UWorld* World,
 	const FString OriginalWorldName = World->GetName();
 	const FString OriginalPackageName = Package->GetName();
 	const FString NewWorldAssetName = FPackageName::GetLongPackageAssetName(NewPackageName);
+	const bool bNewPackageExists = FPackageName::DoesPackageExist(NewPackageName);
 	bool bValidWorldName = true;
 	bool bPackageNeedsRename = false;
 	bool bWorldNeedsRename = false;
@@ -876,7 +877,7 @@ static bool SaveWorld(UWorld* World,
 			// Make sure when we exit SaveWorld AssetRegistry is up to date with saved map
 			AssetRegistry.ScanModifiedAssetFiles({ FinalFilename });
 			
-			if (bPackageNeedsRename || bNewlyCreated)
+			if (bPackageNeedsRename || bNewlyCreated || !bNewPackageExists)
 			{
 				// Force rescan to make sure assets are found on map open or world partition initialize
 				AssetRegistry.ScanPathsSynchronous( ULevel::GetExternalObjectsPaths(NewPackageName) , true);
