SavePackage: Improve feedback and avoid crashing when the temporary file fails to write.
#rnx

diff --git a/Engine/Source/Editor/UnrealEd/Private/Cooker/LooseCookedPackageWriter.cpp b/Engine/Source/Editor/UnrealEd/Private/Cooker/LooseCookedPackageWriter.cpp
--- a/Engine/Source/Editor/UnrealEd/Private/Cooker/LooseCookedPackageWriter.cpp
+++ b/Engine/Source/Editor/UnrealEd/Private/Cooker/LooseCookedPackageWriter.cpp
@@ -17,6 +17,7 @@
 #include "Misc/CString.h"
 #include "Misc/ConfigCacheIni.h"
 #include "Misc/FileHelper.h"
+#include "Misc/Optional.h"
 #include "Misc/PackageName.h"
 #include "Misc/Paths.h"
 #include "Misc/PathViews.h"
@@ -279,30 +280,60 @@ static void WriteToFile(const FString& Filename, const FCompositeBuffer& Buffer)
 {
 	IFileManager& FileManager = IFileManager::Get();
 
+	struct FFailureReason
+	{
+		uint32 LastErrorCode = 0;
+		bool bSizeMatchFailed = false;
+	};
+	TOptional<FFailureReason> FailureReason;
+
 	for (int32 Tries = 0; Tries < 3; ++Tries)
 	{
-		if (FArchive* Ar = FileManager.CreateFileWriter(*Filename))
+		FArchive* Ar = FileManager.CreateFileWriter(*Filename);
+		if (!Ar)
 		{
-			int64 DataSize = 0;
-			for (const FSharedBuffer& Segment : Buffer.GetSegments())
+			if (!FailureReason)
 			{
-				int64 SegmentSize = static_cast<int64>(Segment.GetSize());
-				Ar->Serialize(const_cast<void*>(Segment.GetData()), SegmentSize);
-				DataSize += SegmentSize;
+				FailureReason = FFailureReason{ FPlatformMisc::GetLastError(), false };
 			}
-			delete Ar;
+			continue;
+		}
 
-			if (FileManager.FileSize(*Filename) != DataSize)
-			{
-				FileManager.Delete(*Filename);
+		int64 DataSize = 0;
+		for (const FSharedBuffer& Segment : Buffer.GetSegments())
+		{
+			int64 SegmentSize = static_cast<int64>(Segment.GetSize());
+			Ar->Serialize(const_cast<void*>(Segment.GetData()), SegmentSize);
+			DataSize += SegmentSize;
+		}
+		delete Ar;
 
-				UE_LOG(LogSavePackage, Fatal, TEXT("Could not save to %s!"), *Filename);
+		if (FileManager.FileSize(*Filename) != DataSize)
+		{
+			if (!FailureReason)
+			{
+				FailureReason = FFailureReason{ 0, true };
 			}
-			return;
+			FileManager.Delete(*Filename);
+			continue;
 		}
+		return;
 	}
 
-	UE_LOG(LogSavePackage, Fatal, TEXT("Could not write to %s!"), *Filename);
+	TCHAR LastErrorText[1024];
+	if (FailureReason && FailureReason->bSizeMatchFailed)
+	{
+		FCString::Strcpy(LastErrorText, TEXT("Unexpected file size. Another operation is modifying the file, or the write operation failed to write completely."));
+	}
+	else if (FailureReason && FailureReason->LastErrorCode != 0)
+	{
+		FPlatformMisc::GetSystemErrorMessage(LastErrorText, UE_ARRAY_COUNT(LastErrorText), FailureReason->LastErrorCode);
+	}
+	else
+	{
+		FCString::Strcpy(LastErrorText, TEXT("Unknown failure reason."));
+	}
+	UE_LOG(LogSavePackage, Fatal, TEXT("SavePackage Async write %s failed: %s"), *Filename, LastErrorText);
 }
 
 void FLooseCookedPackageWriter::FWriteFileData::Write(FMD5& AccumulatedHash, EWriteOptions WriteOptions) const
@@ -700,4 +731,4 @@ EPackageExtension FLooseCookedPackageWriter::BulkDataTypeToExtension(FBulkDataIn
 		checkNoEntry();
 		return EPackageExtension::Unspecified;
 	}
-}
+}
\ No newline at end of file
diff --git a/Engine/Source/Runtime/CoreUObject/Private/UObject/LinkerSave.cpp b/Engine/Source/Runtime/CoreUObject/Private/UObject/LinkerSave.cpp
--- a/Engine/Source/Runtime/CoreUObject/Private/UObject/LinkerSave.cpp
+++ b/Engine/Source/Runtime/CoreUObject/Private/UObject/LinkerSave.cpp
@@ -34,7 +34,18 @@ FLinkerSave::FLinkerSave(UPackage* InParent, const TCHAR* InFilename, bool bForc
 		Saver = IFileManager::Get().CreateFileWriter( InFilename, 0 );
 		if( !Saver )
 		{
-			UE_LOG(LogLinker, Fatal, TEXT("%s"), *FString::Printf( TEXT("Error opening file '%s'."), InFilename ) );
+			TCHAR LastErrorText[1024];
+			uint32 LastError = FPlatformMisc::GetLastError();
+			if (LastError != 0)
+			{
+				FPlatformMisc::GetSystemErrorMessage(LastErrorText, UE_ARRAY_COUNT(LastErrorText), LastError);
+			}
+			else
+			{
+				FCString::Strcpy(LastErrorText, TEXT("Unknown failure reason."));
+			}
+			UE_LOG(LogLinker, Error, TEXT("Error opening file '%s': %s"), InFilename, LastErrorText);
+			return; // Caller must test this->Saver to detect the failure
 		}
 
 		UPackage* Package = dynamic_cast<UPackage*>(LinkerRoot);
diff --git a/Engine/Source/Runtime/CoreUObject/Private/UObject/SavePackage.cpp b/Engine/Source/Runtime/CoreUObject/Private/UObject/SavePackage.cpp
--- a/Engine/Source/Runtime/CoreUObject/Private/UObject/SavePackage.cpp
+++ b/Engine/Source/Runtime/CoreUObject/Private/UObject/SavePackage.cpp
@@ -2153,6 +2153,18 @@ FSavePackageResultStruct UPackage::Save(UPackage* InOuter, UObject* InAsset, con
 						// Allocate the linker with a tempfile, forcing byte swapping if wanted.
 						TempFilename = FPaths::CreateTempFilename(*FPaths::ProjectSavedDir(), *BaseFilename.Left(32));
 						Linker = MakePimpl<FLinkerSave>(InOuter, *TempFilename.GetValue(), bForceByteSwapping, bSaveUnversionedNative);
+						if (!Linker->Saver)
+						{
+							FFormatNamedArguments Arguments;
+							Arguments.Add(TEXT("Name"), FText::FromString(*TempFilename));
+							FText ErrorText = FText::Format(NSLOCTEXT("SavePackage", "CouldNotCreateSaveFile", "Could not create temporary save filename {Name}."), Arguments);
+							UE_LOG(LogSavePackage, Error, TEXT("%s"), *ErrorText.ToString());
+							if (!(SaveFlags & SAVE_NoError))
+							{
+								Error->Logf(ELogVerbosity::Error, TEXT("%s"), *ErrorText.ToString());
+							}
+							return ESavePackageResult::Error;
+						}
 					}
 
 					Linker->bProceduralSave = ObjectSaveContext.bProceduralSave;
diff --git a/Engine/Source/Runtime/CoreUObject/Private/UObject/SavePackage2.cpp b/Engine/Source/Runtime/CoreUObject/Private/UObject/SavePackage2.cpp
--- a/Engine/Source/Runtime/CoreUObject/Private/UObject/SavePackage2.cpp
+++ b/Engine/Source/Runtime/CoreUObject/Private/UObject/SavePackage2.cpp
@@ -808,6 +808,18 @@ ESavePackageResult CreateLinker(FSaveContext& SaveContext)
 			// Allocate the linker with a tempfile, forcing byte swapping if wanted.
 			SaveContext.SetTempFilename(FPaths::CreateTempFilename(*FPaths::ProjectSavedDir(), *BaseFilename.Left(32)));
 			SaveContext.SetLinker(MakePimpl<FLinkerSave>(SaveContext.GetPackage(), *SaveContext.GetTempFilename().GetValue(), SaveContext.IsForceByteSwapping(), SaveContext.IsSaveUnversionedNative()));
+			if (!SaveContext.GetLinker()->Saver)
+			{
+				FFormatNamedArguments Arguments;
+				Arguments.Add(TEXT("Name"), FText::FromString(*SaveContext.GetTempFilename()));
+				FText ErrorText = FText::Format(NSLOCTEXT("SavePackage", "CouldNotCreateSaveFile", "Could not create temporary save filename {Name}."), Arguments);
+				UE_LOG(LogSavePackage, Error, TEXT("%s"), *ErrorText.ToString());
+				if (SaveContext.IsGenerateSaveError())
+				{
+					SaveContext.GetError()->Logf(ELogVerbosity::Error, TEXT("%s"), *ErrorText.ToString());
+				}
+				return ESavePackageResult::Error;
+			}
 		}
 
 		if (SaveContext.IsGenerateSaveError())
@@ -2006,12 +2018,12 @@ ESavePackageResult FinalizeFile(FStructuredArchive::FRecord& StructuredArchiveRo
 		{
 			if (SaveContext.IsGenerateSaveError())
 			{
-				UE_LOG(LogSavePackage, Warning, TEXT("%s"), *FString::Printf(TEXT("Error saving '%s'"), SaveContext.GetFilename()));
+				UE_LOG(LogSavePackage, Error, TEXT("%s"), *FString::Printf(TEXT("Error saving '%s'"), SaveContext.GetFilename()));
+				SaveContext.GetError()->Logf(ELogVerbosity::Warning, TEXT("%s"), *FText::Format(NSLOCTEXT("Core", "SaveWarning", "Error saving '{0}'"), FText::FromString(FString(SaveContext.GetFilename()))).ToString());
 			}
 			else
 			{
-				UE_LOG(LogSavePackage, Error, TEXT("%s"), *FString::Printf(TEXT("Error saving '%s'"), SaveContext.GetFilename()));
-				SaveContext.GetError()->Logf(ELogVerbosity::Warning, TEXT("%s"), *FText::Format(NSLOCTEXT("Core", "SaveWarning", "Error saving '{0}'"), FText::FromString(FString(SaveContext.GetFilename()))).ToString());
+				UE_LOG(LogSavePackage, Warning, TEXT("%s"), *FString::Printf(TEXT("Error saving '%s'"), SaveContext.GetFilename()));
 			}
 			return FinalizeResult;
 		}
