Prevent package reload calling PostEditChangeProperty on objects that have an outdated class during PostPackageFixup. In this case the class has already been regenerated and attempting to register actor components with an outdated class will trigger a check().

This could occur when an actor on a sublevel has its class modified when ReloadPackages() is called.

This was likely caused by change 17997995 where FBlueprintCompileReinstancer::ReplaceInstancesOfClass_Inner sets CLASS_NewerVersionExists. This causes the flag to be present during PostPackageFixup which didn't occur in previous versions.

diff --git a/Engine/Source/Editor/UnrealEd/Private/PackageTools.cpp b/Engine/Source/Editor/UnrealEd/Private/PackageTools.cpp
--- a/Engine/Source/Editor/UnrealEd/Private/PackageTools.cpp
+++ b/Engine/Source/Editor/UnrealEd/Private/PackageTools.cpp
@@ -1009,9 +1009,13 @@ UPackageTools::UPackageTools(const FObjectInitializer& ObjectInitializer)
 					continue;
 				}
 
-				FPropertyChangedEvent PropertyEvent(nullptr, EPropertyChangeType::Redirected);
-				ObjectReferencerPtr->PostEditChangeProperty(PropertyEvent);
-
+				if (!ObjectReferencerPtr->GetClass()->HasAnyClassFlags(CLASS_NewerVersionExists))
+				{
+					// Calling PostEditChangeProperty on an actor with an outdated class will trigger a check() during construction scripts.
+					FPropertyChangedEvent PropertyEvent(nullptr, EPropertyChangeType::Redirected);
+					ObjectReferencerPtr->PostEditChangeProperty(PropertyEvent);
+				}
+				
 				// We need to recompile any Blueprints that had properties changed to make sure their generated class is up-to-date and has no lingering references to the old objects
 				UBlueprint* BlueprintToRecompile = nullptr;
 				if (UBlueprint* BlueprintReferencer = Cast<UBlueprint>(ObjectReferencerPtr))
