Fix r.VT.EnableAutoImport which got broken in some UE5 merge.
#rnx

diff --git a/Engine/Source/Editor/UnrealEd/Private/Factories/EditorFactories.cpp b/Engine/Source/Editor/UnrealEd/Private/Factories/EditorFactories.cpp
--- a/Engine/Source/Editor/UnrealEd/Private/Factories/EditorFactories.cpp
+++ b/Engine/Source/Editor/UnrealEd/Private/Factories/EditorFactories.cpp
@@ -4357,6 +4357,35 @@ UObject* UTextureFactory::FactoryCreateBinary
 
 	UTexture2D* Texture2D = Cast<UTexture2D>(Texture);
 
+	if (Texture2D)
+	{
+		// If the texture is larger than a certain threshold make it VT.
+		// Note that previously for re-imports we still checked size and potentially changed the VT status.
+		// But that was unintuitive for many users so now for re-imports we will end up ignoring this and respecting the existing setting below.
+		static const auto CVarVirtualTexturesEnabled = IConsoleManager::Get().FindTConsoleVariableDataInt(TEXT("r.VirtualTextures"));
+		check(CVarVirtualTexturesEnabled);
+		static const auto CVarVirtualTexturesAutoImportEnabled = IConsoleManager::Get().FindTConsoleVariableDataInt(TEXT("r.VT.EnableAutoImport"));
+		check(CVarVirtualTexturesAutoImportEnabled);
+
+		if (CVarVirtualTexturesEnabled->GetValueOnAnyThread() && CVarVirtualTexturesAutoImportEnabled->GetValueOnAnyThread())
+		{
+			const int32 VirtualTextureAutoEnableThreshold = GetDefault<UTextureImportSettings>()->AutoVTSize;
+			const int32 VirtualTextureAutoEnableThresholdPixels = VirtualTextureAutoEnableThreshold * VirtualTextureAutoEnableThreshold;
+
+			// We do this in pixels so a 8192 x 128 texture won't get VT enabled 
+			// We use the Source size instead of simple Texture2D->GetSizeX() as this uses the size of the platform data
+			// however for a new texture platform data may not be generated yet, and for an reimport of a texture this is the size of the
+			// old texture. 
+			// Using source size gives one small caveat. It looks at the size before mipmap power of two padding adjustment.
+			// Textures with more than 1 block (UDIM textures) must be imported as VT
+			if (Texture->Source.GetNumBlocks() > 1 ||
+				Texture2D->Source.GetSizeX() * Texture2D->Source.GetSizeY() >= VirtualTextureAutoEnableThresholdPixels)
+			{
+				Texture2D->VirtualTextureStreaming = true;
+			}
+		}
+	}
+
 	// Restore user set options
 	if (ExistingTexture && bUsingExistingSettings)
 	{
@@ -4398,35 +4427,6 @@ UObject* UTextureFactory::FactoryCreateBinary
 		GConfig->SetBool( TEXT("/Script/UnrealEd.EditorEngine"), TEXT("FlipNormalMapGreenChannel"), bFlipNormalMapGreenChannel, GEngineIni );
 	}
 
-	if(Texture2D)
-	{
-		// The texture has been imported and has no editor specific changes applied so we clear the painted flag.
-		Texture2D->bHasBeenPaintedInEditor = false;
-
-		// If the texture is larger than a certain threshold make it VT. This is explicitly done after the
-		// application of the existing settings above, so if a texture gets reimported at a larger size it will
-		// still be properly flagged as a VT (note: What about reimporting at a lower resolution?)
-		static const auto CVarVirtualTexturesEnabled = IConsoleManager::Get().FindTConsoleVariableDataInt(TEXT("r.VirtualTextures")); check(CVarVirtualTexturesEnabled);
-		
-		if (CVarVirtualTexturesEnabled->GetValueOnAnyThread())
-		{
-			int virtualTextureAutoEnableThreshold = GetDefault<UTextureImportSettings>()->AutoVTSize;
-			int virtualTextureAutoEnableThresholdPixels = virtualTextureAutoEnableThreshold * virtualTextureAutoEnableThreshold;
-
-			// We do this in pixels so a 8192 x 128 texture won't get VT enabled 
-			// We use the Source size instead of simple Texture2D->GetSizeX() as this uses the size of the platform data
-			// however for a new texture platform data may not be generated yet, and for an reimport of a texture this is the size of the
-			// old texture. 
-			// Using source size gives one small caveat. It looks at the size before mipmap power of two padding adjustment.
-			// Textures with more than 1 block (UDIM textures) must be imported as VT
-			if (Texture->Source.GetNumBlocks() > 1 ||
-				Texture2D->Source.GetSizeX()*Texture2D->Source.GetSizeY() >= virtualTextureAutoEnableThresholdPixels)
-			{
-				Texture2D->VirtualTextureStreaming = true;
-			}
-		}
-	}
-
 	// Automatically detect if the texture is a normal map and configure its properties accordingly
 
 	if (!bUsingExistingSettings)
