Adds FSourceControlFilesDeletedDelegate. Should be broadcasted when revert/markfordelete operation could require cleanups (eg package/assets/ref deleted).
Fixes Reverting mark for add asset preventing to recreate an asset with the same name in the same editor session (Asset was still present in AssetRegistry).

#tb Tim.Gautier

diff --git a/Engine/Plugins/Developer/PerforceSourceControl/Source/PerforceSourceControl/Private/PerforceSourceControlOperations.cpp b/Engine/Plugins/Developer/PerforceSourceControl/Source/PerforceSourceControl/Private/PerforceSourceControlOperations.cpp
--- a/Engine/Plugins/Developer/PerforceSourceControl/Source/PerforceSourceControl/Private/PerforceSourceControlOperations.cpp
+++ b/Engine/Plugins/Developer/PerforceSourceControl/Source/PerforceSourceControl/Private/PerforceSourceControlOperations.cpp
@@ -3,6 +3,7 @@
 #include "PerforceSourceControlOperations.h"
 #include "PerforceSourceControlPrivate.h"
 #include "HAL/FileManager.h"
+#include "ISourceControlModule.h"
 #include "Misc/Paths.h"
 #include "Misc/EngineVersion.h"
 #include "Modules/ModuleManager.h"
@@ -76,7 +77,7 @@ static bool IsAddAction(const FString& Action)
 /** Checks if the name of an action corresponds to EPerforceState::MarkedForDelete */
 static bool IsDeleteAction(const FString& Action)
 {
-	return Action == TEXT("delete") || Action == TEXT("move/delete");
+	return Action == TEXT("delete") || Action == TEXT("move/delete") || Action == TEXT("deleted");
 }
 
 /**
@@ -734,6 +735,15 @@ bool FPerforceRevertWorker::Execute(FPerforceSourceControlCommand& InCommand)
 		FP4RecordSet Records;
 		InCommand.bCommandSuccessful = Connection.RunCommand(TEXT("revert"), Parameters, Records, InCommand.ResultInfo.ErrorMessages, FOnIsCancelled::CreateRaw(&InCommand, &FPerforceSourceControlCommand::IsCanceled), InCommand.bConnectionDropped);
 		ParseRecordSetForState(Records, OutResults);
+
+		for (TMap<FString, EPerforceState::Type>::TConstIterator It(OutResults); It; ++It)
+		{
+			if (It->Value == EPerforceState::MarkedForDelete)
+			{
+				RevertOperation->AddDeletedFile(It->Key);
+			}
+		}
+
 		ChangelistToUpdate = InCommand.Changelist;
 	}
 	return InCommand.bCommandSuccessful;
@@ -754,7 +764,7 @@ bool FPerforceRevertWorker::UpdateStates() const
 
 		bUpdatedChangelists |= State->Changelist.IsInitialized() && RemoveFilesFromChangelist(OutResults, State->Changelist);
 	}
-	
+
 	return bUpdatedCachedStates || bUpdatedChangelists;
 }
 
diff --git a/Engine/Source/Developer/SourceControl/Private/SourceControlModule.cpp b/Engine/Source/Developer/SourceControl/Private/SourceControlModule.cpp
--- a/Engine/Source/Developer/SourceControl/Private/SourceControlModule.cpp
+++ b/Engine/Source/Developer/SourceControl/Private/SourceControlModule.cpp
@@ -491,6 +491,21 @@ const FSourceControlPreSubmitFinalizeDelegate& FSourceControlModule::GetOnPreSub
 	return OnPresubmitFinalize;
 }
 
+FDelegateHandle FSourceControlModule::RegisterFilesDeleted(const FSourceControlFilesDeletedDelegate::FDelegate& InDelegate)
+{
+	return OnFilesDeleted.Add(InDelegate);
+}
+
+void FSourceControlModule::UnregisterFilesDeleted(FDelegateHandle InHandle)
+{
+	OnFilesDeleted.Remove(InHandle);
+}
+
+const FSourceControlFilesDeletedDelegate& FSourceControlModule::GetOnFilesDeleted() const
+{
+	return OnFilesDeleted;
+}
+
 IMPLEMENT_MODULE( FSourceControlModule, SourceControl );
 
 #undef LOCTEXT_NAMESPACE
diff --git a/Engine/Source/Developer/SourceControl/Private/SourceControlModule.h b/Engine/Source/Developer/SourceControl/Private/SourceControlModule.h
--- a/Engine/Source/Developer/SourceControl/Private/SourceControlModule.h
+++ b/Engine/Source/Developer/SourceControl/Private/SourceControlModule.h
@@ -49,6 +49,10 @@ public:
 	virtual void UnregisterPreSubmitFinalize(FDelegateHandle Handle) override;
 	virtual const FSourceControlPreSubmitFinalizeDelegate& GetOnPreSubmitFinalize() const override;
 
+	virtual FDelegateHandle RegisterFilesDeleted(const FSourceControlFilesDeletedDelegate::FDelegate& InDelegate) override;
+	virtual void UnregisterFilesDeleted(FDelegateHandle InHandle) override;
+	virtual const FSourceControlFilesDeletedDelegate& GetOnFilesDeleted() const override;
+
 	/** Save the settings to the ini file */
 	void SaveSettings();
 
@@ -135,6 +139,9 @@ private:
 	/** To be called right before files are submitted, allowing for additional last minute validation. @see FSourceControlPreSubmitFinalizeDelegate */
 	FSourceControlPreSubmitFinalizeDelegate OnPresubmitFinalize;
 
+	/** To be called after a source control operations deleted files */
+	FSourceControlFilesDeletedDelegate OnFilesDeleted;
+
 	/** Used to cache source controlled AssetData information */
 	FSourceControlAssetDataCache AssetDataCache;
 };
diff --git a/Engine/Source/Developer/SourceControl/Public/ISourceControlModule.h b/Engine/Source/Developer/SourceControl/Public/ISourceControlModule.h
--- a/Engine/Source/Developer/SourceControl/Public/ISourceControlModule.h
+++ b/Engine/Source/Developer/SourceControl/Public/ISourceControlModule.h
@@ -34,6 +34,13 @@ DECLARE_DELEGATE_FourParams(FSourceControlPreSubmitDataValidationDelegate, FSour
  */
 DECLARE_MULTICAST_DELEGATE_ThreeParams(FSourceControlPreSubmitFinalizeDelegate, const TArray<FString>& /*FilesToSubmit*/, TArray<FText>& /*OutDescriptionTags*/, TArray<FText>& /*OutErrors*/);
 
+/**
+ * Delegate called after source control operations deleted files.
+ *
+ * @param DeletedFiles			The absolute file paths of the deleted files
+ */
+DECLARE_MULTICAST_DELEGATE_OneParam(FSourceControlFilesDeletedDelegate, const TArray<FString>& /*InDeletedFiles*/);
+
 /**
  * The modality of the login window.
  */
@@ -184,6 +191,21 @@ public:
 	 */
 	virtual const FSourceControlPreSubmitFinalizeDelegate& GetOnPreSubmitFinalize() const = 0;
 
+	/**
+	 * Registers a delegate that is invoked after source control operation deleted files.  @see FSourceControlFilesDeletedDelegate
+	 */
+	virtual FDelegateHandle RegisterFilesDeleted(const FSourceControlFilesDeletedDelegate::FDelegate& InDelegate) = 0;
+	
+	/**
+	 * Unregister a previously registered delegate. @see FSourceControlFilesDeletedDelegate
+	 */
+	virtual void UnregisterFilesDeleted(FDelegateHandle InHandle) = 0;
+
+	/**
+	 * Returns access to the delegate so that it can be broadcast as needed. @see FSourceControlFilesDeletedDelegate
+	 */
+	virtual const FSourceControlFilesDeletedDelegate& GetOnFilesDeleted() const = 0;
+
 	/**
 	 * Gets a reference to the source control module instance.
 	 *
diff --git a/Engine/Source/Developer/SourceControl/Public/SourceControlOperations.h b/Engine/Source/Developer/SourceControl/Public/SourceControlOperations.h
--- a/Engine/Source/Developer/SourceControl/Public/SourceControlOperations.h
+++ b/Engine/Source/Developer/SourceControl/Public/SourceControlOperations.h
@@ -188,8 +188,16 @@ public:
 		return USourceControlPreferences::ShouldDeleteNewFilesOnRevert();
 	}
 
+	const TArray<FString>& GetDeletedFiles() const { return DeletedFiles; }
+
+	void AddDeletedFile(const FString& InDeletedFile)
+	{
+		DeletedFiles.Add(InDeletedFile);
+	}
+
 protected:
-	bool	bIsSoftRevert = false;
+	bool				bIsSoftRevert = false;
+	TArray<FString>		DeletedFiles;
 };
 
 /**
diff --git a/Engine/Source/Editor/SourceControlWindows/Private/SSourceControlChangelists.cpp b/Engine/Source/Editor/SourceControlWindows/Private/SSourceControlChangelists.cpp
--- a/Engine/Source/Editor/SourceControlWindows/Private/SSourceControlChangelists.cpp
+++ b/Engine/Source/Editor/SourceControlWindows/Private/SSourceControlChangelists.cpp
@@ -754,7 +754,14 @@ void SSourceControlChangelistsWidget::OnRevert()
 	if (SelectedChangelist.IsValid() || (!SelectedControlledFiles.IsEmpty()))
 	{
 		auto RevertOperation = ISourceControlOperation::Create<FRevert>();
-		SourceControlProvider.Execute(RevertOperation, SelectedChangelist, SelectedControlledFiles);
+		SourceControlProvider.Execute(RevertOperation, SelectedChangelist, SelectedControlledFiles, EConcurrency::Synchronous, FSourceControlOperationComplete::CreateLambda([](const FSourceControlOperationRef& Operation, ECommandResult::Type InResult)
+		{
+			if (Operation->GetName() == TEXT("Revert"))
+			{
+				TSharedRef<FRevert> RevertOperation = StaticCastSharedRef<FRevert>(Operation);
+				ISourceControlModule::Get().GetOnFilesDeleted().Broadcast(RevertOperation->GetDeletedFiles());
+			}
+		}));
 	}
 
 	FUncontrolledChangelistStatePtr SelectedUncontrolledChangelist = GetCurrentUncontrolledChangelistState();
diff --git a/Engine/Source/Editor/UnrealEd/Classes/Editor/UnrealEdEngine.h b/Engine/Source/Editor/UnrealEd/Classes/Editor/UnrealEdEngine.h
--- a/Engine/Source/Editor/UnrealEd/Classes/Editor/UnrealEdEngine.h
+++ b/Engine/Source/Editor/UnrealEd/Classes/Editor/UnrealEdEngine.h
@@ -891,6 +891,10 @@ public:
 
 	/** Gets the project default map templates without any runtime overrides */
 	const TArray<FTemplateMapInfo>& GetProjectDefaultMapTemplates() const;
+
+	/** Called after files are deleted to perform necessary cleanups. */
+	void OnSourceControlFilesDeleted(const TArray<FString>& InDeletedFiles);
+
 protected:
 
 	/** Called when global editor selection changes */
@@ -969,4 +973,7 @@ private:
 
 	/** Internal function to filter and add visualizers to a specific list */
 	void AddVisualizers(AActor* Actor, TArray<FCachedComponentVisualizer>& Visualizers, TFunctionRef<bool(const TSharedPtr<FComponentVisualizer>&)> Condition);
+
+	/** Delegate Called after files have been deleted to perform necessary cleanups. */
+	FDelegateHandle SourceControlFilesDeletedHandle;
 };
diff --git a/Engine/Source/Editor/UnrealEd/Private/UnrealEdEngine.cpp b/Engine/Source/Editor/UnrealEd/Private/UnrealEdEngine.cpp
--- a/Engine/Source/Editor/UnrealEd/Private/UnrealEdEngine.cpp
+++ b/Engine/Source/Editor/UnrealEd/Private/UnrealEdEngine.cpp
@@ -71,6 +71,8 @@
 #include "ProfilingDebugging/StallDetector.h"
 #include "GameMapsSettings.h"
 #include "HAL/PlatformApplicationMisc.h"
+#include "AssetRegistryModule.h"
+#include "ObjectTools.h"
 #include "Cooker/ExternalCookOnTheFlyServer.h"
 
 DEFINE_LOG_CATEGORY_STATIC(LogUnrealEdEngine, Log, All);
@@ -246,6 +248,12 @@ void UUnrealEdEngine::Init(IEngineLoop* InEngineLoop)
 			}
 		});
 #endif
+
+	// Delay this until after the source control module has loaded
+	FDelayedAutoRegisterHelper(EDelayedRegisterRunPhase::EarliestPossiblePluginsLoaded, [this]()
+	{
+		SourceControlFilesDeletedHandle = ISourceControlModule::Get().RegisterFilesDeleted(FSourceControlFilesDeletedDelegate::FDelegate::CreateUObject(this, &UUnrealEdEngine::OnSourceControlFilesDeleted));
+	});
 }
 
 bool CanCookForPlatformInThisProcess( const FString& PlatformName )
@@ -451,6 +459,11 @@ void UUnrealEdEngine::FinishDestroy()
 		delete PerformanceMonitor;
 	}
 
+	if (ISourceControlModule* SourceControlModule = FModuleManager::GetModulePtr<ISourceControlModule>(FName("SourceControl")))
+	{
+		SourceControlModule->UnregisterFilesDeleted(SourceControlFilesDeletedHandle);
+	}
+
 	FPackageName::OnContentPathMounted().RemoveAll(this);
 	FPackageName::OnContentPathDismounted().RemoveAll(this);
 	UPackage::PackageDirtyStateChangedEvent.RemoveAll(this);
@@ -916,6 +929,43 @@ const TArray<FTemplateMapInfo>& UUnrealEdEngine::GetProjectDefaultMapTemplates()
 	return TemplateMapInfoCache;
 }
 
+void UUnrealEdEngine::OnSourceControlFilesDeleted(const TArray<FString>& InDeletedFiles)
+{
+	TArray<UObject*> ObjectsToDelete;
+
+	if (InDeletedFiles.IsEmpty())
+	{
+		return;
+	}
+
+	FAssetRegistryModule& AssetRegistryModule = FModuleManager::LoadModuleChecked<FAssetRegistryModule>(TEXT("AssetRegistry"));
+
+	for (const FString& File : InDeletedFiles)
+	{
+		TArray<FAssetData> Assets;
+		FString PackageName;
+
+		if (!FPackageName::TryConvertFilenameToLongPackageName(File, PackageName))
+		{
+			continue;
+		}
+
+		AssetRegistryModule.Get().GetAssetsByPackageName(*PackageName, Assets);
+
+		for (const FAssetData& Asset : Assets)
+		{
+			UObject* ObjectToDelete = Asset.GetAsset();
+
+			if (ObjectToDelete != nullptr)
+			{
+				ObjectsToDelete.Add(ObjectToDelete);
+			}
+		}
+	}
+
+	ObjectTools::DeleteObjectsUnchecked(ObjectsToDelete);
+}
+
 void UUnrealEdEngine::OnHISMTreeBuilt(UHierarchicalInstancedStaticMeshComponent* Component, bool bWasAsyncBuild)
 {
 	if (bWasAsyncBuild)
