Add filtering system for editor modes

diff --git a/Engine/Source/Editor/UnrealEd/Private/Subsystems/AssetEditorSubsystem.cpp b/Engine/Source/Editor/UnrealEd/Private/Subsystems/AssetEditorSubsystem.cpp
--- a/Engine/Source/Editor/UnrealEd/Private/Subsystems/AssetEditorSubsystem.cpp
+++ b/Engine/Source/Editor/UnrealEd/Private/Subsystems/AssetEditorSubsystem.cpp
@@ -638,6 +638,11 @@ UEdMode* UAssetEditorSubsystem::CreateEditorModeWithToolsOwner(FEditorModeID Mod
 
 bool UAssetEditorSubsystem::FindEditorModeInfo(const FEditorModeID& InModeID, FEditorModeInfo& OutModeInfo) const
 {
+	if (!IsEditorModeAllowed(InModeID))
+	{
+		return false;
+	}
+	
 	const TSharedRef<IEditorModeFactory>* ModeFactory = FEditorModeRegistry::Get().GetFactoryMap().Find(InModeID);
 	if (ModeFactory)
 	{
@@ -660,11 +665,19 @@ TArray<FEditorModeInfo> UAssetEditorSubsystem::GetEditorModeInfoOrderedByPriorit
 
 	for (const auto& Pair : FEditorModeRegistry::Get().GetFactoryMap())
 	{
-		ModeInfoArray.Add(Pair.Value->GetModeInfo());
+		FEditorModeInfo ModeInfo = Pair.Value->GetModeInfo();
+		if (IsEditorModeAllowed(ModeInfo.ID))
+		{
+			ModeInfoArray.Add(MoveTemp(ModeInfo));
+		}
 	}
 	for (const auto& EditorMode : EditorModes)
 	{
-		ModeInfoArray.Add(EditorMode.Value.ModeInfo);
+		const FEditorModeInfo& ModeInfo = EditorMode.Value.ModeInfo;
+		if (IsEditorModeAllowed(ModeInfo.ID))
+		{
+			ModeInfoArray.Add(ModeInfo);
+		}
 	}
 
 	ModeInfoArray.Sort([](const FEditorModeInfo& A, const FEditorModeInfo& B) {
@@ -1018,6 +1031,16 @@ void UAssetEditorSubsystem::UnregisterEditorModes()
 	EditorModes.Empty();
 }
 
+FNamePermissionList& UAssetEditorSubsystem::GetAllowedEditorModes()
+{
+	return AllowedEditorModes;
+}
+
+bool UAssetEditorSubsystem::IsEditorModeAllowed(const FName ModeId) const
+{
+	return AllowedEditorModes.PassesFilter(ModeId);
+}
+
 void UAssetEditorSubsystem::OnSMInstanceElementsEnabled()
 {
 	// Let the modes know that SM instance elements may have been enabled or disabled and update state accordingly
diff --git a/Engine/Source/Editor/UnrealEd/Public/Subsystems/AssetEditorSubsystem.h b/Engine/Source/Editor/UnrealEd/Public/Subsystems/AssetEditorSubsystem.h
--- a/Engine/Source/Editor/UnrealEd/Public/Subsystems/AssetEditorSubsystem.h
+++ b/Engine/Source/Editor/UnrealEd/Public/Subsystems/AssetEditorSubsystem.h
@@ -8,6 +8,7 @@
 
 #include "Containers/Ticker.h"
 #include "Tools/Modes.h"
+#include "Misc/NamePermissionList.h"
 #include "AssetEditorSubsystem.generated.h"
 
 class UAssetEditor;
@@ -198,6 +199,9 @@ public:
 	 */
 	FOnModeUnregistered& OnEditorModeUnregistered();
 
+	/** Get the permission list that controls which editor modes are exposed */
+	FNamePermissionList& GetAllowedEditorModes();
+
 private:
 
 	/** Handles FAssetEditorRequestOpenAsset messages. */
@@ -250,6 +254,8 @@ private:
 	void UnregisterEditorModes();
 	void OnSMInstanceElementsEnabled();
 
+	bool IsEditorModeAllowed(const FName ModeId) const;
+
 private:
 
 	/** struct used by OpenedEditorTimes map to store editor names and times */
@@ -334,4 +340,10 @@ private:
 
 	/** Event that is triggered whenever a mode is unregistered */
 	FOnModeUnregistered OnEditorModeUnregisteredEvent;
+	
+	/**
+	 * Which FEditorModeInfo data should be returned when queried, filtered by the mode's ID.
+	 * Note that this does not disable or unregister disallowed modes, it simply removes them from the query results.
+	 */
+	FNamePermissionList AllowedEditorModes;
 };
