Fix Preview Rendering Level change not being applied to either editor-viewport or separate-window in-process PIE

diff --git a/Engine/Source/Editor/LevelEditor/Private/SLevelViewport.cpp b/Engine/Source/Editor/LevelEditor/Private/SLevelViewport.cpp
--- a/Engine/Source/Editor/LevelEditor/Private/SLevelViewport.cpp
+++ b/Engine/Source/Editor/LevelEditor/Private/SLevelViewport.cpp
@@ -4130,7 +4130,6 @@ void SLevelViewport::StartPlayInEditorSession(UGameViewportClient* PlayClient, c
 
 	bPIEHasFocus = ActiveViewport->HasMouseCapture();
 
-
 	if(EditorPlayInSettings->ShowMouseControlLabel && !GEngine->IsStereoscopic3D( ActiveViewport.Get() ) )
 	{
 		ELabelAnchorMode AnchorMode = EditorPlayInSettings->MouseControlLabelPosition.GetValue();
@@ -4139,6 +4138,15 @@ void SLevelViewport::StartPlayInEditorSession(UGameViewportClient* PlayClient, c
 	}
 
 	GEngine->BroadcastLevelActorListChanged();
+
+	// register for preview feature level change
+	UEditorEngine* Editor = CastChecked<UEditorEngine>(GEngine);
+	UWorld* PIEWorld = PlayClient->GetWorld();
+	PIEPreviewFeatureLevelChangedHandle = Editor->OnPreviewFeatureLevelChanged().AddLambda([PIEWorld](ERHIFeatureLevel::Type NewFeatureLevel)
+		{
+			PIEWorld->GetWorld()->ChangeFeatureLevel(NewFeatureLevel);
+		});
+
 }
 
 EVisibility SLevelViewport::GetMouseCaptureLabelVisibility() const
@@ -4353,6 +4361,13 @@ void SLevelViewport::EndPlayInEditorSession()
 	}
 
 	GEngine->BroadcastLevelActorListChanged();
+
+	// Remove preview feature level delegate if set
+	if (PIEPreviewFeatureLevelChangedHandle.IsValid())
+	{
+		CastChecked<UEditorEngine>(GEngine)->OnPreviewFeatureLevelChanged().Remove(PIEPreviewFeatureLevelChangedHandle);
+		PIEPreviewFeatureLevelChangedHandle.Reset();
+	}
 }
 
 void SLevelViewport::SwapViewportsForSimulateInEditor()
diff --git a/Engine/Source/Editor/LevelEditor/Public/SLevelViewport.h b/Engine/Source/Editor/LevelEditor/Public/SLevelViewport.h
--- a/Engine/Source/Editor/LevelEditor/Public/SLevelViewport.h
+++ b/Engine/Source/Editor/LevelEditor/Public/SLevelViewport.h
@@ -783,6 +783,11 @@ private:
 	 */
 	TSharedPtr<SWidget> InactiveViewportWidgetEditorContent;
 
+	/** 
+	 *  When PIE is active, the handle for the change feature level delegate
+	 */
+	FDelegateHandle PIEPreviewFeatureLevelChangedHandle;
+
 	/** Level viewport client */
 	TSharedPtr<FLevelEditorViewportClient> LevelViewportClient;
 
diff --git a/Engine/Source/Editor/UnrealEd/Private/PlayLevel.cpp b/Engine/Source/Editor/UnrealEd/Private/PlayLevel.cpp
--- a/Engine/Source/Editor/UnrealEd/Private/PlayLevel.cpp
+++ b/Engine/Source/Editor/UnrealEd/Private/PlayLevel.cpp
@@ -70,6 +70,8 @@
 #include "Slate/SceneViewport.h"
 #include "Kismet2/KismetEditorUtilities.h"
 #include "Kismet2/BlueprintEditorUtils.h"
+#include "RHIShaderPlatformDefinitions.inl"
+
 
 #include "LevelEditor.h"
 #include "IAssetViewport.h"
@@ -3083,8 +3085,8 @@ FText GeneratePIEViewportWindowTitle(const EPlayNetMode InNetMode, const ERHIFea
 	FFormatNamedArguments Args;
 	Args.Add(TEXT("GameName"), FText::FromString(FString(WindowTitleOverride.IsEmpty() ? FApp::GetProjectName() : WindowTitleOverride.ToString())));
 	Args.Add(TEXT("PlatformBits"), FText::FromString(PlatformBitsString));
-	Args.Add(TEXT("RHIName"), FText::FromName(ShaderPlatformToPlatformName(GetFeatureLevelShaderPlatform(InFeatureLevel))));
-
+	Args.Add(TEXT("RHIName"), GetFriendlyShaderPlatformName(GetFeatureLevelShaderPlatform(InFeatureLevel)));
+	
 	if (InNetMode == PIE_Client)
 	{
 		Args.Add(TEXT("NetMode"), FText::FromString(FString::Printf(TEXT("Client %d"), ClientIndex)));
@@ -3284,7 +3286,7 @@ TSharedRef<SPIEViewport> UEditorEngine::GeneratePIEViewportWindow(const FRequest
 				}
 			}
 
-			static void OnPIEWindowClosed(const TSharedRef<SWindow>& WindowBeingClosed, TWeakPtr<SViewport> PIEViewportWidget, TWeakObjectPtr<UEditorEngine> OwningEditorEngine, int32 ViewportIndex, bool bRestoreRootWindow)
+			static void OnPIEWindowClosed(const TSharedRef<SWindow>& WindowBeingClosed, TWeakPtr<SViewport> PIEViewportWidget, TWeakObjectPtr<UEditorEngine> OwningEditorEngine, int32 ViewportIndex, bool bRestoreRootWindow, FDelegateHandle PreviewFeatureLevelChangedHandle)
 			{
 				// Save off the window position
 				const FVector2D PIEWindowPos = WindowBeingClosed->GetLocalToScreenTransform().GetTranslation();
@@ -3295,6 +3297,11 @@ TSharedRef<SPIEViewport> UEditorEngine::GeneratePIEViewportWindow(const FRequest
 				if (OwningEditorEngine.IsValid())
 				{
 					OwningEditorEngine->StoreWindowSizeAndPositionForInstanceIndex(ViewportIndex, WindowSize, WindowPosition);
+
+					if (PreviewFeatureLevelChangedHandle.IsValid())
+					{
+						OwningEditorEngine->OnPreviewFeatureLevelChanged().Remove(PreviewFeatureLevelChangedHandle);
+					}
 				}
 
 				// Route the callback
@@ -3312,9 +3319,19 @@ TSharedRef<SPIEViewport> UEditorEngine::GeneratePIEViewportWindow(const FRequest
 			}
 		};
 
+		FDelegateHandle PreviewFeatureLevelChangedHandle = OnPreviewFeatureLevelChanged().AddLambda([PieWindow, bHasCustomWindow, InViewportClient, InNetMode, InSessionParams, InWorldContext](ERHIFeatureLevel::Type NewFeatureLevel)
+			{
+				if (!bHasCustomWindow)
+				{
+					FText ViewportName = GeneratePIEViewportWindowTitle(InNetMode, NewFeatureLevel, InSessionParams, InWorldContext.PIEInstance, InWorldContext.PIEFixedTickSeconds);
+					PieWindow->SetTitle(ViewportName);
+				}
+				InViewportClient->GetWorld()->ChangeFeatureLevel(NewFeatureLevel);
+			});
+
 		PieWindow->SetRequestDestroyWindowOverride(FRequestDestroyWindowOverride::CreateStatic(&FLocal::RequestDestroyPIEWindowOverride, TWeakObjectPtr<UEditorEngine>(this)));
 		PieWindow->SetOnWindowClosed(FOnWindowClosed::CreateStatic(&FLocal::OnPIEWindowClosed, TWeakPtr<SViewport>(PieViewportWidget), TWeakObjectPtr<UEditorEngine>(this),
-			InViewportIndex, bShouldMinimizeRootWindow));
+			InViewportIndex, bShouldMinimizeRootWindow, PreviewFeatureLevelChangedHandle));
 	}
 
 	// Create a new viewport that the viewport widget will use to render the game
diff --git a/Engine/Source/Editor/UnrealEd/Private/SEditorViewport.cpp b/Engine/Source/Editor/UnrealEd/Private/SEditorViewport.cpp
--- a/Engine/Source/Editor/UnrealEd/Private/SEditorViewport.cpp
+++ b/Engine/Source/Editor/UnrealEd/Private/SEditorViewport.cpp
@@ -818,9 +818,9 @@ TSharedRef<SWidget> SEditorViewport::BuildFeatureLevelWidget() const
 
 EVisibility SEditorViewport::GetCurrentFeatureLevelPreviewTextVisibility() const
 {
-	if (GetWorld())
+	if (Client->GetWorld())
 	{
-		return (GetWorld()->FeatureLevel != GMaxRHIFeatureLevel) ? EVisibility::SelfHitTestInvisible : EVisibility::Collapsed;
+		return (Client->GetWorld()->FeatureLevel != GMaxRHIFeatureLevel) ? EVisibility::SelfHitTestInvisible : EVisibility::Collapsed;
 	}
 	else
 	{
@@ -838,7 +838,7 @@ FText SEditorViewport::GetCurrentFeatureLevelPreviewText(bool bDrawOnlyLabel) co
 	}
 	else
 	{
-		UWorld* World = GetWorld();
+		UWorld* World = Client->GetWorld();
 		if (World != nullptr)
 		{
 			ERHIFeatureLevel::Type TargetFeatureLevel = World->FeatureLevel;
