Post-code review adjustments

diff --git a/Engine/Source/Editor/UnrealEd/Classes/ThumbnailRendering/LevelThumbnailRenderer.h b/Engine/Source/Editor/UnrealEd/Classes/ThumbnailRendering/LevelThumbnailRenderer.h
--- a/Engine/Source/Editor/UnrealEd/Classes/ThumbnailRendering/LevelThumbnailRenderer.h
+++ b/Engine/Source/Editor/UnrealEd/Classes/ThumbnailRendering/LevelThumbnailRenderer.h
@@ -29,6 +29,7 @@ class ULevelThumbnailRenderer : public UDefaultSizedThumbnailRenderer
 	// End UThumbnailRenderer Object
 
 private:
-	FSceneView* CreateView(ULevel* Level, FSceneViewFamily* ViewFamily, int32 X, int32 Y, uint32 SizeX, uint32 SizeY) const; 
+	/** Allocates then adds an FSceneView to the ViewFamily. */
+	[[nodiscard]] FSceneView* CreateView(ULevel* Level, FSceneViewFamily* ViewFamily, int32 X, int32 Y, uint32 SizeX, uint32 SizeY) const; 
 };
 
diff --git a/Engine/Source/Editor/UnrealEd/Classes/ThumbnailRendering/WorldThumbnailRenderer.h b/Engine/Source/Editor/UnrealEd/Classes/ThumbnailRendering/WorldThumbnailRenderer.h
--- a/Engine/Source/Editor/UnrealEd/Classes/ThumbnailRendering/WorldThumbnailRenderer.h
+++ b/Engine/Source/Editor/UnrealEd/Classes/ThumbnailRendering/WorldThumbnailRenderer.h
@@ -28,7 +28,7 @@ class UWorldThumbnailRenderer : public UDefaultSizedThumbnailRenderer
 
 private:
 	/** Allocates then adds an FSceneView to the ViewFamily. */
-	FSceneView* CreateView(UWorld* World, FSceneViewFamily* ViewFamily, int32 X, int32 Y, uint32 SizeX, uint32 SizeY) const;
+	[[nodiscard]] FSceneView* CreateView(UWorld* World, FSceneViewFamily* ViewFamily, int32 X, int32 Y, uint32 SizeX, uint32 SizeY) const;
 
 private:
 	/** Offset used to orient all worlds to show a more vertical camera, if necessary. Individual thumbnail infos can provide additional offset. */
diff --git a/Engine/Source/Editor/UnrealEd/Private/ThumbnailHelpers.cpp b/Engine/Source/Editor/UnrealEd/Private/ThumbnailHelpers.cpp
--- a/Engine/Source/Editor/UnrealEd/Private/ThumbnailHelpers.cpp
+++ b/Engine/Source/Editor/UnrealEd/Private/ThumbnailHelpers.cpp
@@ -70,7 +70,8 @@ FThumbnailPreviewScene::FThumbnailPreviewScene()
 PRAGMA_DISABLE_DEPRECATION_WARNINGS
 void FThumbnailPreviewScene::GetView(FSceneViewFamily* ViewFamily, int32 X, int32 Y, uint32 SizeX, uint32 SizeY) const
 {
-	CreateView(ViewFamily, X, Y, SizeX, SizeY);
+	// CreateView allocates a FSceneView, which is only accessible as a const pointer in FSceneViewFamily afterwards so CreateView is the new way and is marked as [[nodiscard]], hence the static_cast<void>, to avoid a compiler warning :
+	static_cast<void>(CreateView(ViewFamily, X, Y, SizeX, SizeY)); 
 }
 PRAGMA_ENABLE_DEPRECATION_WARNINGS
 
diff --git a/Engine/Source/Editor/UnrealEd/Public/ThumbnailHelpers.h b/Engine/Source/Editor/UnrealEd/Public/ThumbnailHelpers.h
--- a/Engine/Source/Editor/UnrealEd/Public/ThumbnailHelpers.h
+++ b/Engine/Source/Editor/UnrealEd/Public/ThumbnailHelpers.h
@@ -25,7 +25,7 @@ public:
 	void GetView(FSceneViewFamily* ViewFamily, int32 X, int32 Y, uint32 SizeX, uint32 SizeY) const;
 
 	/** Allocates then adds an FSceneView to the ViewFamily. */
-	FSceneView* CreateView(FSceneViewFamily* ViewFamily, int32 X, int32 Y, uint32 SizeX, uint32 SizeY) const;
+	[[nodiscard]] FSceneView* CreateView(FSceneViewFamily* ViewFamily, int32 X, int32 Y, uint32 SizeX, uint32 SizeY) const;
 
 	/* Begin FTickableEditorObject */
 	virtual void Tick(float DeltaTime) override;
