Fix editor throttling while VR headset attached.

diff --git a/Engine/Source/Editor/UnrealEd/Private/EditorEngine.cpp b/Engine/Source/Editor/UnrealEd/Private/EditorEngine.cpp
--- a/Engine/Source/Editor/UnrealEd/Private/EditorEngine.cpp
+++ b/Engine/Source/Editor/UnrealEd/Private/EditorEngine.cpp
@@ -6154,8 +6154,11 @@ bool UEditorEngine::ShouldThrottleCPUUsage() const
 	
 
 	bool bShouldThrottle = false;
-
-	if (!FApp::HasFocus() && !IsRunningCommandlet() && !GIsAutomationTesting && !FApp::IsBenchmarking())
+	
+	// Always check if in foreground, since VR apps will only have focus when running (PIE, etc).
+	const bool bIsForeground = FPlatformApplicationMisc::IsThisApplicationForeground();
+	
+	if (!bIsForeground &&!FApp::HasFocus() && !IsRunningCommandlet() && !GIsAutomationTesting && !FApp::IsBenchmarking())
 	{
 		const UEditorPerformanceSettings* Settings = GetDefault<UEditorPerformanceSettings>();
 		bShouldThrottle = Settings->bThrottleCPUWhenNotForeground;
