Parallelize expensive components of UUnrealEdEngine::DeleteActors

This cuts several seconds off of deleting an actor in a big map.

FBlueprintEditorUtils::GetActorReferenceMap: ~750ms -> ~40ms
FAssetRenameManager::PopulateAssetReferencers: ~1800ms -> ~400ms

#rnx

diff --git a/Engine/Source/Editor/UnrealEd/Private/Kismet2/BlueprintEditorUtils.cpp b/Engine/Source/Editor/UnrealEd/Private/Kismet2/BlueprintEditorUtils.cpp
--- a/Engine/Source/Editor/UnrealEd/Private/Kismet2/BlueprintEditorUtils.cpp
+++ b/Engine/Source/Editor/UnrealEd/Private/Kismet2/BlueprintEditorUtils.cpp
@@ -8168,6 +8168,9 @@ void FBlueprintEditorUtils::FindActorsThatReferenceActor( AActor* InActor, TArra
 void FBlueprintEditorUtils::GetActorReferenceMap(UWorld* InWorld, TArray<UClass*>& InClassesToIgnore, TMap<AActor*, TArray<AActor*> >& OutReferencingActors)
 {
 	TRACE_CPUPROFILER_EVENT_SCOPE(FBlueprintEditorUtils::GetActorReferenceMap);
+
+	TArray<AActor*> Actors;
+
 	// Iterate all actors in the same world as InActor
 	for (FActorIterator ActorIt(InWorld); ActorIt; ++ActorIt)
 	{
@@ -8187,18 +8190,43 @@ void FBlueprintEditorUtils::GetActorReferenceMap(UWorld* InWorld, TArray<UClass*
 
 			if (!bShouldIgnore)
 			{
-				// Get all references from CurrentActor and see if any Actors
-				TArray<UObject*> References;
-				FReferenceFinder Finder(References);
-				Finder.FindReferences(CurrentActor);
+				Actors.Add(CurrentActor);
+			}
+		}
+	}
 
-				for (int32 RefIdx = 0; RefIdx < References.Num(); RefIdx++)
-				{
-					if (References[RefIdx] && References[RefIdx]->IsA(AActor::StaticClass()))
-					{
-						OutReferencingActors.FindOrAdd(Cast<AActor>(References[RefIdx])).Add(CurrentActor);
-					}
-				}
+	TArray<TArray<UObject*>> ActorsReferences;
+	ActorsReferences.AddDefaulted(Actors.Num());
+
+	ParallelFor(Actors.Num(), [&Actors = std::as_const(Actors), &ActorsReferences](int32 Index) {
+		AActor* Actor(Actors[Index]);
+		TArray<UObject*>& References(ActorsReferences[Index]);
+
+		// Find all references
+		FReferenceFinder Finder(References);
+		Finder.FindReferences(Actor);
+
+		// Null out references that aren't actors
+		for (int32 ReferenceIndex = 0; ReferenceIndex < References.Num(); ReferenceIndex++)
+		{
+			UObject* Reference = References[ReferenceIndex];
+			if (Reference && !Reference->IsA(AActor::StaticClass()))
+			{
+				References[ReferenceIndex] = nullptr;
+			}
+		}
+	});
+
+	for (int ActorIndex=0; ActorIndex <ActorsReferences.Num(); ActorIndex++)
+	{
+		TArray<UObject*>& References(ActorsReferences[ActorIndex]);
+		for (int32 ReferenceIndex=0; ReferenceIndex<References.Num(); ReferenceIndex++)
+		{
+			// This is a sparse array as we may have nulled out non-actors
+			UObject* Reference = References[ReferenceIndex];
+			if (Reference)
+			{
+				OutReferencingActors.FindOrAdd(Cast<AActor>(Reference)).Add(Actors[ActorIndex]);
 			}
 		}
 	}
diff --git a/Engine/Source/Runtime/Engine/Private/Level.cpp b/Engine/Source/Runtime/Engine/Private/Level.cpp
--- a/Engine/Source/Runtime/Engine/Private/Level.cpp
+++ b/Engine/Source/Runtime/Engine/Private/Level.cpp
@@ -3219,17 +3219,35 @@ TArray<UPackage*> ULevel::GetLoadedExternalObjectPackages() const
 
 	if (!ExternalObjectsPaths.IsEmpty())
 	{
-		for (TObjectIterator<UPackage> It; It; ++It)
-		{
-			TStringBuilder<256> PackageName;
-			It->GetLoadedPath().AppendPackageName(PackageName);
-			FStringView PackageNameStringView(PackageName);
-			for (const FString& ExternalObjectsPath : ExternalObjectsPaths)
+			TArray<UObject*> Packages;
+			GetObjectsOfClass(UPackage::StaticClass(), Packages, /*bIncludeDerivedClasses =*/ true, /*EObjectFlags ExcludeFlags =*/ RF_ClassDefaultObject,/*EInternalObjectFlags ExclusionInternalFlags =*/ EInternalObjectFlags::None);
+
+			TArray<bool> PackageIsInExternalObjectsPath;
+			PackageIsInExternalObjectsPath.InsertUninitialized(0, Packages.Num());
+
+			ParallelFor(Packages.Num(), [&Packages = std::as_const(Packages), &ExternalObjectsPath = std::as_const(ExternalObjectsPath), &PackageIsInExternalObjectsPath](int32 Index) {
+				UPackage* Package = static_cast<UPackage*>(Packages[Index]);
+
+				TStringBuilder<256> PackageName;
+				Package->GetLoadedPath().AppendPackageName(PackageName);
+				FStringView PackageNameStringView(PackageName);
+				bool bIsInExternalObjectsPath = false;
+				for (const FString& ExternalObjectsPath : ExternalObjectsPaths)
+				{
+					bIsInExternalObjectsPath = PackageNameStringView.Contains(ExternalObjectsPath);
+					if (bIsInExternalObjectsPath)
+					{
+						break;
+					}
+				}
+				PackageIsInExternalObjectsPath[Index] = bIsInExternalObjectsPath;
+			}
+
+			for (int Index=0; Index<PackageIsInExternalObjectsPath.Num(); Index++)
 			{
-				if (PackageNameStringView.Contains(ExternalObjectsPath))
+				if (PackageIsInExternalObjectsPath[Index])
 				{
-					ExternalObjectPackages.Add(*It);
-					break;
+					ActorPackages.Add(static_cast<UPackage*>(Packages[Index]));
 				}
 			}
 		}
