Fix EditorDomain to refuse to save packages into EditorDomain that add new classes during PostLoad/PreSave, as the schema for these classes was not part of the key and will cause serialization errors if the class properties are modified after the save into the editor domain.
#rnx

diff --git a/Engine/Config/BaseEditor.ini b/Engine/Config/BaseEditor.ini
--- a/Engine/Config/BaseEditor.ini
+++ b/Engine/Config/BaseEditor.ini
@@ -368,6 +368,7 @@ SaveUnversioned=true
 ; +PackageBlockList=/Long/PackageName OR ../../../Standard/Filename/NoExtension OR D:\Absolute\Filename\NoExtension
 
 +ClassBlockList=/Script/Engine.StaticMesh
++ClassBlockList=/Script/Engine.Blueprint
 
 ; Known bad files that are problematic if all of them come from EditorDomain. Blocking any one of these is sufficient to prevent the bug.
 +PackageBlockList=/Niagara/Functions/RandomRangeFloat
@@ -546,4 +547,4 @@ bCanModifyPluginsFromBrowser=true
 
 [/Script/GameFeaturesEditor.GameFeaturesEditorSettings]
 +PluginTemplates=(Path=(Path="../../Plugins/Experimental/GameFeatures/Templates/GameFeaturePluginContentOnly"),Label=NSLOCTEXT("[/Script/GameFeaturesEditor]", "E17DB7E94A2C0F5968A52E89959864C0", "Game Feature (Content Only)"),Description=NSLOCTEXT("[/Script/GameFeaturesEditor]", "CB5136B64AFAFBA5EDEA078FCAC7BCC9", "Create a new Game Feature Plugin."),DefaultGameFeatureDataClass=None,DefaultGameFeatureDataName="")
-+PluginTemplates=(Path=(Path="../../Plugins/Experimental/GameFeatures/Templates/GameFeaturePluginWithCode"),Label=NSLOCTEXT("[/Script/GameFeaturesEditor]", "90D6FD2A455176CEC42DD79A23AEC3B0", "Game Feature (with C++)"),Description=NSLOCTEXT("[/Script/GameFeaturesEditor]", "2B00A0E7415C8CEAA6DA609A141237F0", "Create a new Game Feature Plugin with a minimal amount of code."),DefaultGameFeatureDataClass=None,DefaultGameFeatureDataName="")
++PluginTemplates=(Path=(Path="../../Plugins/Experimental/GameFeatures/Templates/GameFeaturePluginWithCode"),Label=NSLOCTEXT("[/Script/GameFeaturesEditor]", "90D6FD2A455176CEC42DD79A23AEC3B0", "Game Feature (with C++)"),Description=NSLOCTEXT("[/Script/GameFeaturesEditor]", "2B00A0E7415C8CEAA6DA609A141237F0", "Create a new Game Feature Plugin with a minimal amount of code."),DefaultGameFeatureDataClass=None,DefaultGameFeatureDataName="")
\ No newline at end of file
diff --git a/Engine/Source/Editor/UnrealEd/Private/EditorDomain/EditorDomain.h b/Engine/Source/Editor/UnrealEd/Private/EditorDomain/EditorDomain.h
--- a/Engine/Source/Editor/UnrealEd/Private/EditorDomain/EditorDomain.h
+++ b/Engine/Source/Editor/UnrealEd/Private/EditorDomain/EditorDomain.h
@@ -61,14 +61,16 @@ struct FPackageDigest
 	 * created from properties that change if the saved version would change (package bytes, serialization versions).
 	 */
 	FIoHash Hash;
+	/** Classes used by the resaved package: imports or instances that can be created by PostLoad/PreSave */
+	TArray<FName> ImportedClasses;
+	/** List of CustomVersions used to save the package. */
+	UE::AssetRegistry::FPackageCustomVersionsHandle CustomVersions;
 	/** Allow flags for whether the package can be saved/loaded from EditorDomain. */
 	EDomainUse DomainUse = EDomainUse::None;
 	/** Status for creation of this digest. Either Success or an error code for why it couldn't be created. */
 	EStatus Status = EStatus::NotYetRequested;
 	/** Extended information for the status description (e.g. missing class name). */
 	FName StatusArg;
-	/** List of CustomVersions used to save the package. */
-	UE::AssetRegistry::FPackageCustomVersionsHandle CustomVersions;
 };
 
 }
diff --git a/Engine/Source/Editor/UnrealEd/Private/EditorDomain/EditorDomainUtils.cpp b/Engine/Source/Editor/UnrealEd/Private/EditorDomain/EditorDomainUtils.cpp
--- a/Engine/Source/Editor/UnrealEd/Private/EditorDomain/EditorDomainUtils.cpp
+++ b/Engine/Source/Editor/UnrealEd/Private/EditorDomain/EditorDomainUtils.cpp
@@ -59,7 +59,7 @@ namespace UE::EditorDomain
 
 TArray<FGuid> GetCustomVersions(UClass& Class);
 TMap<FGuid, UObject*> FindCustomVersionCulprits(TConstArrayView<FGuid> UnknownGuids, UPackage* Package);
-void InitializeGlobalAddedCustomVersions();
+void InitializeGlobalConstructClasses();
 
 FClassDigestMap GClassDigests;
 FClassDigestMap& GetClassDigests()
@@ -70,16 +70,15 @@ FClassDigestMap& GetClassDigests()
 TMap<FName, EDomainUse> GClassBlockedUses;
 TMap<FName, EDomainUse> GPackageBlockedUses;
 TMultiMap<FName, FName> GConstructClasses;
-TArray<FName> GGlobalConstructClasses;
 TSet<FName> GTargetDomainClassBlockList;
 bool GTargetDomainClassUseAllowList = true;
 bool GTargetDomainClassEmptyAllowList = false;
-TArray<int32> GGlobalAddedCustomVersions;
-bool bGGlobalAddedCustomVersionsInitialized = false;
+TArray<FName> GGlobalConstructClasses;
+bool bGGlobalConstructClassesInitialized = false;
 int64 GMaxBulkDataSize = -1;
 
 // Change to a new guid when EditorDomain needs to be invalidated
-const TCHAR* EditorDomainVersion = TEXT("54D98F079D2A46A9802B2B990560285D");
+const TCHAR* EditorDomainVersion = TEXT("B0ED553FD35F4185B095EF52ED119C7F");
 
 // Identifier of the CacheBuckets for EditorDomain tables
 const TCHAR* EditorDomainPackageBucketName = TEXT("EditorDomainPackage");
@@ -218,42 +217,93 @@ FPackageDigest CalculatePackageDigest(IAssetRegistry& AssetRegistry, FName Packa
 	FKnownCustomVersions::FindOrAddHandles(CustomVersionHandles,
 		PackageVersions.Num(), [PackageVersions](int32 Index) -> const FGuid& { return PackageVersions[Index].Key;});
 
-	int32 NextClass = 0;
+	// Loop over all ImportedClasses used by the package on disk, fetch the ClassDigestData for each one, and
+	// extend the list with the transitive closure of classes that can be created by any of those classes.
 	FClassDigestMap& ClassDigests = GetClassDigests();
-	for (int32 Attempt = 0; NextClass < PackageData.ImportedClasses.Num(); ++Attempt)
+	TArray<FName> ImportedClasses = MoveTemp(PackageData.ImportedClasses);
+	TSet<FName> DiscoveredClasses(ImportedClasses);
+	bool bNeedsResort = false;
+	bool bHasTriedPrecacheClassDigests = false;
+	auto AddClass = [&DiscoveredClasses, &ImportedClasses, &bHasTriedPrecacheClassDigests, &bNeedsResort](FName ClassName)
+	{
+		bool bAlreadyDiscovered;
+		DiscoveredClasses.Add(ClassName, &bAlreadyDiscovered);
+		if (!bAlreadyDiscovered)
+		{
+			ImportedClasses.Add(ClassName);
+			bHasTriedPrecacheClassDigests = false;
+			bNeedsResort = true;
+		}
+	};
+
+	InitializeGlobalConstructClasses();
+	for (FName GlobalAddedClass : GGlobalConstructClasses)
+	{
+		AddClass(GlobalAddedClass);
+	}
+	int32 NextClass = 0;
+	while (NextClass < ImportedClasses.Num())
 	{
-		if (Attempt > 0)
+		{
+			FReadScopeLock ClassDigestsScopeLock(ClassDigests.Lock);
+			while (NextClass < ImportedClasses.Num())
+			{
+				FName ClassName = ImportedClasses[NextClass];
+				FClassDigestData* ExistingData = ClassDigests.Map.Find(ClassName);
+				if (!ExistingData)
+				{
+					break;
+				}
+				if (ExistingData->ResolvedClosestNative == ClassName)
+				{
+					++NextClass;
+				}
+				else
+				{
+					ImportedClasses.RemoveAt(NextClass);
+					AddClass(ExistingData->ResolvedClosestNative);
+				}
+				for (FName ConstructClass : ExistingData->ConstructClasses)
+				{
+					AddClass(ConstructClass);
+				}
+			}
+		}
+
+		if (NextClass < ImportedClasses.Num())
 		{
 			// EDITORDOMAIN_TODO: Remove this !IsInGameThread check once FindObject no longer asserts if GIsSavingPackage
-			if (Attempt > 1 || !IsInGameThread())
+			if (bHasTriedPrecacheClassDigests || !IsInGameThread())
 			{
-				return FPackageDigest(FPackageDigest::EStatus::MissingClass, PackageData.ImportedClasses[NextClass]);
+				return FPackageDigest(FPackageDigest::EStatus::MissingClass, ImportedClasses[NextClass]);
 			}
-			TConstArrayView<FName> RemainingClasses(PackageData.ImportedClasses);
-			RemainingClasses = RemainingClasses.Slice(NextClass, PackageData.ImportedClasses.Num() - NextClass);
+			TConstArrayView<FName> RemainingClasses(ImportedClasses);
+			RemainingClasses = RemainingClasses.Slice(NextClass, ImportedClasses.Num() - NextClass);
 			PrecacheClassDigests(RemainingClasses);
+			bHasTriedPrecacheClassDigests = true;
 		}
+	}
+
+	// Loop over the sorted list of classes a second time, now that we have all of them,
+	// and add their schemas
+	if (bNeedsResort)
+	{
+		ImportedClasses.Sort(FNameLexicalLess());
+	}
+	{
 		FReadScopeLock ClassDigestsScopeLock(ClassDigests.Lock);
-		for (; NextClass < PackageData.ImportedClasses.Num(); ++NextClass)
+		for (FName ClassName : ImportedClasses)
 		{
-			FName ClassName = PackageData.ImportedClasses[NextClass];
 			FClassDigestData* ExistingData = ClassDigests.Map.Find(ClassName);
-			if (!ExistingData)
-			{
-				break;
-			}
-			if (ExistingData->bNative)
-			{
-				Writer.Update(&ExistingData->SchemaHash, sizeof(ExistingData->SchemaHash));
-			}
+			check(ExistingData); // This should have been added above if not already existing
+			check(ExistingData->bNative); // Non-native classes were removed by the loop above
+			Writer.Update(&ExistingData->SchemaHash, sizeof(ExistingData->SchemaHash));
 			CustomVersionHandles.Append(ExistingData->CustomVersionHandles);
 			EnumSetFlagsAnd(Result.DomainUse, EDomainUse::LoadEnabled | EDomainUse::SaveEnabled,
 				Result.DomainUse, ExistingData->EditorDomainUse);
 		}
 	}
 
-	InitializeGlobalAddedCustomVersions();
-	CustomVersionHandles.Append(GGlobalAddedCustomVersions);
 	CustomVersionHandles.Sort();
 	CustomVersionHandles.SetNum(Algo::Unique(CustomVersionHandles));
 
@@ -278,6 +328,7 @@ FPackageDigest CalculatePackageDigest(IAssetRegistry& AssetRegistry, FName Packa
 	Result.Hash = Writer.Finalize();
 	Result.Status = FPackageDigest::EStatus::Successful;
 	Result.CustomVersions = UE::AssetRegistry::FPackageCustomVersionsHandle::FindOrAdd(CustomVersions);
+	Result.ImportedClasses = MoveTemp(ImportedClasses);
 	return Result;
 }
 
@@ -329,6 +380,7 @@ private:
 
 FClassDigestData* FPrecacheClassDigest::GetRecursive(FName ClassName, bool bAllowRedirects)
 {
+	check(!ClassName.IsNone());
 	// Called within ClassDigestsMap.Lock.WriteLock()
 	FClassDigestData* DigestData = &ClassDigests.FindOrAdd(ClassName);
 	if (DigestData->bConstructed)
@@ -337,7 +389,7 @@ FClassDigestData* FPrecacheClassDigest::GetRecursive(FName ClassName, bool bAllo
 	}
 	DigestData->bConstructed = true;
 
-	FName LookupName = ClassName;
+	FName RedirectsResolvedName = ClassName;
 	ClassName.ToString(NameStringBuffer);
 	if (bAllowRedirects)
 	{
@@ -346,7 +398,7 @@ FClassDigestData* FPrecacheClassDigest::GetRecursive(FName ClassName, bool bAllo
 		if (ClassNameRedirect != RedirectedClassNameRedirect)
 		{
 			NameStringBuffer = RedirectedClassNameRedirect.ToString();
-			LookupName = FName(NameStringBuffer);
+			RedirectsResolvedName = FName(NameStringBuffer);
 		}
 	}
 
@@ -365,16 +417,16 @@ FClassDigestData* FPrecacheClassDigest::GetRecursive(FName ClassName, bool bAllo
 	// Fill in digest data config-driven flags
 	DigestData->EditorDomainUse = EDomainUse::LoadEnabled | EDomainUse::SaveEnabled;
 	DigestData->EditorDomainUse &= ~MapFindRef(GClassBlockedUses, ClassName, EDomainUse::None);
-	if (LookupName != ClassName)
+	if (RedirectsResolvedName != ClassName)
 	{
-		DigestData->EditorDomainUse &= ~MapFindRef(GClassBlockedUses, LookupName, EDomainUse::None);
+		DigestData->EditorDomainUse &= ~MapFindRef(GClassBlockedUses, RedirectsResolvedName, EDomainUse::None);
 	}
 	if (!GTargetDomainClassUseAllowList)
 	{
 		DigestData->bTargetIterativeEnabled = !GTargetDomainClassBlockList.Contains(ClassName);
-		if (LookupName != ClassName)
+		if (RedirectsResolvedName != ClassName)
 		{
-			DigestData->bTargetIterativeEnabled &= !GTargetDomainClassBlockList.Contains(LookupName);
+			DigestData->bTargetIterativeEnabled &= !GTargetDomainClassBlockList.Contains(RedirectsResolvedName);
 		}
 	}
 
@@ -382,6 +434,7 @@ FClassDigestData* FPrecacheClassDigest::GetRecursive(FName ClassName, bool bAllo
 	FName ParentName;
 	if (Struct)
 	{
+		DigestData->ResolvedClosestNative = RedirectsResolvedName;
 		DigestData->bNative = true;
 		DigestData->SchemaHash = Struct->GetSchemaHash(false /* bSkipEditorOnly */);
 		UStruct* ParentStruct = Struct->GetSuperStruct();
@@ -423,7 +476,13 @@ FClassDigestData* FPrecacheClassDigest::GetRecursive(FName ClassName, bool bAllo
 				}
 			}
 		}
+		if (ParentName.IsNone())
+		{
+			ParentName = FName(*UObject::StaticClass()->GetPathName());
+		}
+		DigestData->ResolvedClosestNative = ParentName;
 	}
+	check(!DigestData->ResolvedClosestNative.IsNone());
 
 	// Get the CustomVersions used by the native class; GetCustomVersions already returns all custom versions used
 	// by the parent class so we do not need to copy data from the parent
@@ -479,18 +538,18 @@ FClassDigestData* FPrecacheClassDigest::GetRecursive(FName ClassName, bool bAllo
 		}
 	}
 
-	// Propagate values from the ConstructClasses
+	// Get and validate the list of ConstructClasses by recursively initializing them
 	TArray<FName> ConstructClasses; // Not a class scratch variable because we use it across recursive calls
 	GConstructClasses.MultiFind(ClassName, ConstructClasses);
-	if (LookupName != ClassName)
+	if (RedirectsResolvedName != ClassName)
 	{
-		GConstructClasses.MultiFind(LookupName, ConstructClasses);
+		GConstructClasses.MultiFind(RedirectsResolvedName, ConstructClasses);
 	}
 	if (!ConstructClasses.IsEmpty())
 	{
-		TArray<int32> ConstructCustomVersions; // Not a class scratch variable because we use it across recursive calls
-		for (FName ConstructClass : ConstructClasses)
+		for (int32 Index = 0; Index < ConstructClasses.Num(); )
 		{
+			FName ConstructClass = ConstructClasses[Index];
 			FClassDigestData* ConstructClassDigest = GetRecursive(ConstructClass, true /* bAllowRedirects */);
 			if (!ConstructClassDigest)
 			{
@@ -499,6 +558,7 @@ FClassDigestData* FPrecacheClassDigest::GetRecursive(FName ClassName, bool bAllo
 					TEXT("This is a class that can be constructed by postload upgrades of class %s. ")
 					TEXT("Old packages with class %s will load more slowly."),
 					*ConstructClass.ToString(), *ClassName.ToString(), *ClassName.ToString(), *ClassName.ToString());
+				ConstructClasses.RemoveAt(Index);
 			}
 			else
 			{
@@ -508,14 +568,15 @@ FClassDigestData* FPrecacheClassDigest::GetRecursive(FName ClassName, bool bAllo
 						TEXT("Cycle detected in Editor.ini:[EditorDomain]:PostLoadConstructClasses of class %s. This is unexpected, but not a problem."),
 						*ClassName.ToString());
 				}
-				ConstructCustomVersions.Append(ConstructClassDigest->CustomVersionHandles);
+				++Index;
 			}
 		}
+		Algo::Sort(ConstructClasses, FNameLexicalLess());
+		ConstructClasses.SetNum(Algo::Unique(ConstructClasses));
+
 		// The map has possibly been modified so we need to recalculate the address of ClassName's DigestData
 		DigestData = &ClassDigests.FindChecked(ClassName);
-		DigestData->CustomVersionHandles.Append(MoveTemp(ConstructCustomVersions));
-		Algo::Sort(DigestData->CustomVersionHandles);
-		DigestData->CustomVersionHandles.SetNum(Algo::Unique(DigestData->CustomVersionHandles));
+		DigestData->ConstructClasses = MoveTemp(ConstructClasses);
 	}
 
 	DigestData->bConstructionComplete = true;
@@ -532,30 +593,30 @@ void PrecacheClassDigests(TConstArrayView<FName> ClassNames)
 	}
 }
 
-/** Construct GGlobalAddedCustomVersions from the classes specified by config */
-void InitializeGlobalAddedCustomVersions()
+/** Construct GGlobalConstructClasses from the classes specified by config */
+void InitializeGlobalConstructClasses()
 {
-	if (bGGlobalAddedCustomVersionsInitialized)
+	if (bGGlobalConstructClassesInitialized)
 	{
 		return;
 	}
-	bGGlobalAddedCustomVersionsInitialized = true;
-	TArray<FName> GlobalAddedClassNames;
+	bGGlobalConstructClassesInitialized = true;
 	{
 		TArray<FString> Lines;
 		GConfig->GetArray(TEXT("EditorDomain"), TEXT("GlobalCanConstructClasses"), Lines, GEditorIni);
-		GlobalAddedClassNames.Reserve(Lines.Num());
+		GGlobalConstructClasses.Reserve(Lines.Num());
 		for (const FString& Line : Lines)
 		{
-			GlobalAddedClassNames.Add(FName(FStringView(Line).TrimStartAndEnd()));
+			GGlobalConstructClasses.Add(FName(FStringView(Line).TrimStartAndEnd()));
 		}
 	}
 
-	PrecacheClassDigests(GlobalAddedClassNames);
+	PrecacheClassDigests(GGlobalConstructClasses);
 	FClassDigestMap& ClassDigests = GetClassDigests();
 	FReadScopeLock ClassDigestsScopeLock(ClassDigests.Lock);
-	for (FName ClassName : GlobalAddedClassNames)
+	for (int32 Index = 0; Index < GGlobalConstructClasses.Num();)
 	{
+		FName ClassName = GGlobalConstructClasses[Index];
 		FClassDigestData* ExistingData = ClassDigests.Map.Find(ClassName);
 		if (!ExistingData)
 		{
@@ -563,14 +624,15 @@ void InitializeGlobalAddedCustomVersions()
 				TEXT("This is a class that can be constructed automatically by SavePackage when saving old packages. ")
 				TEXT("Old packages that do not yet have this class will load more slowly."),
 				*ClassName.ToString());
+			GGlobalConstructClasses.RemoveAt(Index);
 		}
 		else
 		{
-			GGlobalAddedCustomVersions.Append(ExistingData->CustomVersionHandles);
+			++Index;
 		}
 	}
-	Algo::Sort(GGlobalAddedCustomVersions);
-	GGlobalAddedCustomVersions.SetNum(Algo::Unique(GGlobalAddedCustomVersions));
+	Algo::Sort(GGlobalConstructClasses, FNameLexicalLess());
+	GGlobalConstructClasses.SetNum(Algo::Unique(GGlobalConstructClasses));
 }
 
 /** An archive that just collects custom versions. */
@@ -939,6 +1001,295 @@ void RequestEditorDomainPackage(const FPackagePath& PackagePath,
 	Cache.Get({{{PackagePath.GetDebugName()}, GetEditorDomainPackageKey(EditorDomainHash), CachePolicy}}, Owner, MoveTemp(Callback));
 }
 
+/** TODO: Delete this duplicate of AssetRegistry/PackageReader once we have approved making that class Public */
+class FPackageReader : public FArchiveUObject
+{
+public:
+	FPackageReader()
+		: Loader(nullptr)
+		, PackageFileSize(0)
+	{
+		SetIsLoading(true);
+		SetIsPersistent(true);
+	}
+
+	bool OpenPackageFile(FArchive* InLoader)
+	{
+		Loader = InLoader;
+		return OpenPackageFile();
+	}
+	bool OpenPackageFile()
+	{
+		// Read package file summary from the file
+		*this << PackageFileSummary;
+		PackageFileSize = TotalSize();
+		return true;
+	}
+	bool ReadImportedClasses(TArray<FName>& OutClassNames)
+	{
+		if (!SerializeNameMap())
+		{
+			return false;
+		}
+
+		TArray<FObjectImport> ImportMap;
+		if (!SerializeImportMap(ImportMap))
+		{
+			return false;
+		}
+		if (!SerializeImportedClasses(ImportMap, OutClassNames))
+		{
+			return false;
+		}
+		return true;
+	}
+	const FPackageFileSummary& GetPackageFileSummary() const { return PackageFileSummary; }
+
+
+	// Farchive implementation to redirect requests to the Loader
+	virtual void Serialize(void* V, int64 Length) override
+	{
+		check(Loader);
+		Loader->Serialize(V, Length);
+		if (Loader->IsError())
+		{
+			SetError();
+		}
+	}
+
+	virtual bool Precache(int64 PrecacheOffset, int64 PrecacheSize) override
+	{
+		check(Loader);
+		return Loader->Precache(PrecacheOffset, PrecacheSize);
+	}
+
+	virtual void Seek(int64 InPos) override
+	{
+		check(Loader);
+		Loader->Seek(InPos);
+		if (Loader->IsError())
+		{
+			SetError();
+		}
+	}
+
+	virtual int64 Tell() override
+	{
+		check(Loader);
+		return Loader->Tell();
+	}
+
+	virtual int64 TotalSize() override
+	{
+		check(Loader);
+		return Loader->TotalSize();
+	}
+
+	virtual FArchive& operator<<(FName& Name) override
+	{
+		int32 NameIndex;
+		FArchive& Ar = *this;
+		Ar << NameIndex;
+
+		if (!NameMap.IsValidIndex(NameIndex))
+		{
+			SetError();
+			return *this;
+		}
+
+		// if the name wasn't loaded (because it wasn't valid in this context)
+		if (NameMap[NameIndex] == NAME_None)
+		{
+			int32 TempNumber;
+			Ar << TempNumber;
+			Name = NAME_None;
+		}
+		else
+		{
+			int32 Number;
+			Ar << Number;
+			// simply create the name from the NameMap's name and the serialized instance number
+			Name = FName(NameMap[NameIndex], Number);
+		}
+
+		return *this;
+	}
+
+	virtual FString GetArchiveName() const override
+	{
+		return TEXT("EditorDomainPackageReader");
+	}
+private:
+	// Serializers for different package maps
+	bool SerializeNameMap()
+	{
+		if (NameMap.Num() == 0 && PackageFileSummary.NameCount > 0)
+		{
+			if (!StartSerializeSection(PackageFileSummary.NameOffset))
+			{
+				return false;
+			}
+
+			const int MinSizePerNameEntry = 1;
+			if (PackageFileSize < Tell() + PackageFileSummary.NameCount * MinSizePerNameEntry)
+			{
+				return false;
+			}
+
+			for (int32 NameMapIdx = 0; NameMapIdx < PackageFileSummary.NameCount; ++NameMapIdx)
+			{
+				// Read the name entry from the file.
+				FNameEntrySerialized NameEntry(ENAME_LinkerConstructor);
+				*this << NameEntry;
+				if (IsError())
+				{
+					return false;
+				}
+				NameMap.Add(FName(NameEntry));
+			}
+		}
+
+		return true;
+	}
+	bool SerializeImportMap(TArray<FObjectImport>& OutImportMap)
+	{
+		if (PackageFileSummary.ImportCount > 0)
+		{
+			if (!StartSerializeSection(PackageFileSummary.ImportOffset))
+			{
+				return false;
+			}
+
+			const int MinSizePerImport = 1;
+			if (PackageFileSize < Tell() + PackageFileSummary.ImportCount * MinSizePerImport)
+			{
+				return false;
+			}
+			for (int32 ImportMapIdx = 0; ImportMapIdx < PackageFileSummary.ImportCount; ++ImportMapIdx)
+			{
+				FObjectImport& Import = OutImportMap.Emplace_GetRef();
+				*this << Import;
+				if (IsError())
+				{
+					return false;
+				}
+			}
+		}
+
+		return true;
+	}
+	bool SerializeImportedClasses(const TArray<FObjectImport>& ImportMap, TArray<FName>& OutClassNames)
+	{
+		OutClassNames.Reset();
+
+		TSet<int32> ClassImportIndices;
+		if (PackageFileSummary.ExportCount > 0)
+		{
+			if (!StartSerializeSection(PackageFileSummary.ExportOffset))
+			{
+				return false;
+			}
+
+			const int MinSizePerExport = 1;
+			if (PackageFileSize < Tell() + PackageFileSummary.ExportCount * MinSizePerExport)
+			{
+				return false;
+			}
+			FObjectExport ExportBuffer;
+			for (int32 ExportMapIdx = 0; ExportMapIdx < PackageFileSummary.ExportCount; ++ExportMapIdx)
+			{
+				*this << ExportBuffer;
+				if (IsError())
+				{
+					return false;
+				}
+				if (ExportBuffer.ClassIndex.IsImport())
+				{
+					ClassImportIndices.Add(ExportBuffer.ClassIndex.ToImport());
+				}
+			}
+		}
+
+		TArray<FName, TInlineAllocator<5>>  ParentChain;
+		FNameBuilder ClassObjectPath;
+		for (int32 ClassImportIndex : ClassImportIndices)
+		{
+			ParentChain.Reset();
+			ClassObjectPath.Reset();
+			if (!ImportMap.IsValidIndex(ClassImportIndex))
+			{
+				return false;
+			}
+			bool bParentChainComplete = false;
+			int32 CurrentParentIndex = ClassImportIndex;
+			for (;;)
+			{
+				const FObjectImport& ObjectImport = ImportMap[CurrentParentIndex];
+				ParentChain.Add(ObjectImport.ObjectName);
+				if (ObjectImport.OuterIndex.IsImport())
+				{
+					CurrentParentIndex = ObjectImport.OuterIndex.ToImport();
+					if (!ImportMap.IsValidIndex(CurrentParentIndex))
+					{
+						return false;
+					}
+				}
+				else if (ObjectImport.OuterIndex.IsNull())
+				{
+					bParentChainComplete = true;
+					break;
+				}
+				else
+				{
+					check(ObjectImport.OuterIndex.IsExport());
+					// Ignore classes in an external package but with an object in this package as one of their outers;
+					// We do not need to handle that case yet for Import Classes, and we would have to make this
+					// loop more complex (searching in both ExportMap and ImportMap) to do so
+					break;
+				}
+			}
+
+			if (bParentChainComplete)
+			{
+				int32 NumTokens = ParentChain.Num();
+				check(NumTokens >= 1);
+				const TCHAR Delimiters[] = { '.', SUBOBJECT_DELIMITER_CHAR, '.' };
+				int32 DelimiterIndex = 0;
+				ParentChain[NumTokens - 1].AppendString(ClassObjectPath);
+				for (int32 TokenIndex = NumTokens - 2; TokenIndex >= 0; --TokenIndex)
+				{
+					ClassObjectPath << Delimiters[DelimiterIndex];
+					DelimiterIndex = FMath::Min(DelimiterIndex + 1, static_cast<int32>(UE_ARRAY_COUNT(Delimiters)) - 1);
+					ParentChain[TokenIndex].AppendString(ClassObjectPath);
+				}
+				OutClassNames.Emplace(ClassObjectPath);
+			}
+		}
+
+		OutClassNames.Sort(FNameLexicalLess());
+		return true;
+	}
+	bool StartSerializeSection(int64 Offset)
+	{
+		check(Loader);
+		if (Offset <= 0 || Offset > PackageFileSize)
+		{
+			return false;
+		}
+		ClearError();
+		Loader->ClearError();
+		Seek(Offset);
+		return !IsError();
+	}
+
+	FArchive* Loader;
+	FPackageFileSummary PackageFileSummary;
+	TArray<FName> NameMap;
+	int64 PackageFileSize;
+	int64 AssetRegistryDependencyDataOffset;
+	bool bLoaderOwner;
+};
+
 /** Stores data from SavePackage in accessible fields */
 class FEditorDomainPackageWriter final : public TPackageWriterToSharedBuffer<IPackageWriter>
 {
@@ -952,16 +1303,20 @@ public:
 	}
 
 	/** Deserialize the CustomVersions out of the PackageFileSummary that was serialized into the header */
-	bool TryGetCustomVersions(FCustomVersionContainer& OutVersions)
+	bool TryGetClassesAndVersions(TArray<FName>& OutImportedClasses, FCustomVersionContainer& OutVersions)
 	{
 		FMemoryReaderView HeaderArchive(WritePackageRecord.Buffer.GetView());
-		FPackageFileSummary Summary;
-		HeaderArchive << Summary;
-		if (HeaderArchive.IsError())
+		FPackageReader PackageReader;
+		if (!PackageReader.OpenPackageFile(&HeaderArchive))
+		{
+			return false;
+		}
+
+		OutVersions = PackageReader.GetPackageFileSummary().GetCustomVersionContainer();
+		if (!PackageReader.ReadImportedClasses(OutImportedClasses))
 		{
 			return false;
 		}
-		OutVersions = Summary.GetCustomVersionContainer();
 		return true;
 	}
 
@@ -1196,13 +1551,31 @@ bool TrySavePackage(UPackage* Package)
 	Info.WriteOptions = IPackageWriter::EWriteOptions::Write;
 	PackageWriter->CommitPackage(MoveTemp(Info));
 
+	TArray<FName> SavedImportedClasses;
 	FCustomVersionContainer SavedCustomVersions;
-	if (!PackageWriter->TryGetCustomVersions(SavedCustomVersions))
+	if (!PackageWriter->TryGetClassesAndVersions(SavedImportedClasses, SavedCustomVersions))
 	{
 		UE_LOG(LogEditorDomain, Warning, TEXT("Could not save %s to EditorDomain: Could not read the PackageFileSummary from the saved bytes."),
 			*Package->GetName());
 		return false;
 	}
+	TSet<FName> KnownImportedClasses(PackageDigest.ImportedClasses);
+	for (FName ImportedClass : SavedImportedClasses)
+	{
+		if (!KnownImportedClasses.Contains(ImportedClass))
+		{
+			// Suggested debugging technique for this message: Add a conditional breakpoint on the packagename
+			// at the start of LoadPackageInternal. After it gets hit, add a breakpoint in the constructor
+			// of the ImportedClass.
+			UE_LOG(LogEditorDomain, Display, TEXT("Could not save %s to EditorDomain: It uses an unexpected class. ")
+				TEXT("Optimized loading and iterative cooking will be disabled for this package.\n\t")
+				TEXT("Find the class that added %s and add ")
+				TEXT("Editor.ini:[EditorDomain]:+PostLoadCanConstructClasses=<ConstructingClass>,%s"),
+				*Package->GetName(), *ImportedClass.ToString(), *ImportedClass.ToString());
+			return false;
+		}
+	}
+
 	TConstArrayView<UE::AssetRegistry::FPackageCustomVersion> KnownCustomVersions = PackageDigest.CustomVersions.Get();
 	TSet<FGuid> KnownGuids;
 	KnownGuids.Reserve(KnownCustomVersions.Num());
@@ -1221,74 +1594,23 @@ bool TrySavePackage(UPackage* Package)
 	if (!UnknownGuids.IsEmpty())
 	{
 		TMap<FGuid, UObject*> Culprits = FindCustomVersionCulprits(UnknownGuids, Package);
-
-		// First check whether the culprit for (one of) the missing CustomVersion is an instance
-		// that was added during PostLoad. If so, advise adding an entry to PostLoadCanConstructClasses.
-		UObject* ConstructedCulprit = nullptr;
-		TOptional<FAssetPackageData> PackageData = IAssetRegistry::Get()->GetAssetPackageDataCopy(Package->GetFName());
+		TStringBuilder<256> VersionsText;
 		for (const FGuid& CustomVersionGuid : UnknownGuids)
 		{
+			TOptional<FCustomVersion> CustomVersion = FCurrentCustomVersions::Get(CustomVersionGuid);
 			UObject* Culprit = Culprits.FindOrAdd(CustomVersionGuid);
-			FName CulpritClassName = Culprit ? FName(*Culprit->GetClass()->GetPathName()) : NAME_None;
-			if (CulpritClassName.IsNone() || PackageData->ImportedClasses.Contains(CulpritClassName))
-			{
-				continue;
-			}
-			// If the culprit class does not declare the version either, then we still need to give the message
-			// advising adding an entry in DeclareCustomVersions
-			bool bConstructedClassDeclaresTheVersion = true;
-			{
-				PrecacheClassDigests({ CulpritClassName });
-				FClassDigestMap& ClassDigests = GetClassDigests();
-				FReadScopeLock ClassDigestsScopeLock(ClassDigests.Lock);
-				FClassDigestData* ExistingData = ClassDigests.Map.Find(CulpritClassName);
-				if (!ExistingData)
-				{
-					bConstructedClassDeclaresTheVersion = false;
-				}
-				else
-				{
-					TArray<FGuid> ClassCustomVersionGuids;
-					FKnownCustomVersions::FindGuidsChecked(ClassCustomVersionGuids, ExistingData->CustomVersionHandles);
-					bConstructedClassDeclaresTheVersion = ClassCustomVersionGuids.Contains(CustomVersionGuid);
-				}
-			}
-			if (bConstructedClassDeclaresTheVersion)
-			{
-				ConstructedCulprit = Culprit;
-				break;
-			}
+			VersionsText << TEXT("\n\tCustomVersion(Guid=") << CustomVersionGuid << TEXT(", Name=")
+				<< (CustomVersion ? *CustomVersion->GetFriendlyName().ToString() : TEXT("<Unknown>"))
+				<< TEXT("): Used by ")
+				<< (Culprit ? *Culprit->GetClass()->GetPathName() : TEXT("<CulpritUnknown>"));
 		}
-		TStringBuilder<128> FixupSuggestion;
-		if (ConstructedCulprit)
-		{
-			// Suggested debugging technique for this message: Add a conditional breakpoint on the packagename
-			// at the start of LoadPackageInternal. After it gets hit, add a breakpoint in the constructor
-			// of the ConstructedCulprit class.
-			FixupSuggestion << TEXT("The custom version is used by a class which was created after load of the package. ")
-				<< TEXT("Find the class that added ") << ConstructedCulprit->GetFullName() << TEXT(" and add ")
-				<< TEXT("Editor.ini:[EditorDomain]:+PostLoadCanConstructClasses=<ConstructingClass>,")
-				<< ConstructedCulprit->GetClass()->GetPathName();
-		}
-		else
-		{
-			// Suggested debugging technique for this message: SetNextStatement back to beginning of the function,
-			// add a conditional breakpoint in FArchive::UsingCustomVersion with Key.A == 0x<FirstHexWordFromGuid>
-			FixupSuggestion << TEXT("Modify the classes or structs used in the package to call Ar.UsingCustomVersion(Guid) in Serialize or DeclareCustomVersions.");
-			for (const FGuid& CustomVersionGuid : UnknownGuids)
-			{
-				TOptional<FCustomVersion> CustomVersion = FCurrentCustomVersions::Get(CustomVersionGuid);
-				UObject* Culprit = Culprits.FindOrAdd(CustomVersionGuid);
-				FixupSuggestion << TEXT("\n\tCustomVersion(Guid=") << CustomVersionGuid << TEXT(", Name=")
-					<< (CustomVersion ? *CustomVersion->GetFriendlyName().ToString() : TEXT("<Unknown>"))
-					<< TEXT("): Used by ")
-					<< (Culprit ? *Culprit->GetClass()->GetPathName() : TEXT("<CulpritUnknown>"));
-			}
 
-		}
+		// Suggested debugging technique for this message: SetNextStatement back to beginning of the function,
+		// add a conditional breakpoint in FArchive::UsingCustomVersion with Key.A == 0x<FirstHexWordFromGuid>
 		UE_LOG(LogEditorDomain, Display, TEXT("Could not save %s to EditorDomain: It uses an unexpected custom version. ")
-			TEXT("Optimized loading and iterative cooking will be disabled for this package.\n\t%s"),
-			*Package->GetName(), FixupSuggestion.ToString());
+			TEXT("Optimized loading and iterative cooking will be disabled for this package. ")
+			TEXT("Modify the classes/structs to call Ar.UsingCustomVersion(Guid) in Serialize or DeclareCustomVersions.%s"),
+			*Package->GetName(), VersionsText.ToString());
 		return false;
 	}
 
diff --git a/Engine/Source/Editor/UnrealEd/Private/EditorDomain/EditorDomainUtils.h b/Engine/Source/Editor/UnrealEd/Private/EditorDomain/EditorDomainUtils.h
--- a/Engine/Source/Editor/UnrealEd/Private/EditorDomain/EditorDomainUtils.h
+++ b/Engine/Source/Editor/UnrealEd/Private/EditorDomain/EditorDomainUtils.h
@@ -35,11 +35,16 @@ struct FClassDigestData
 {
 public:
 	FBlake3Hash SchemaHash;
+	/** CustomVersions used by the class */
 	TArray<int32> CustomVersionHandles;
-	bool bNative = false;
+	/** Classes that can be created by the class during PostLoad/PreSave (parsed from ini) */
+	TArray<FName> ConstructClasses;
+	/** The closest native parent of (the possibly CoreRedirected target of) the class */
+	FName ResolvedClosestNative;
 	/** EditorDomainEnabled allows everything and uses only a blocklist, so DomainUse by default is enabled. */
 	EDomainUse EditorDomainUse = EDomainUse::LoadEnabled | EDomainUse::SaveEnabled;
 
+	bool bNative = false;
 	/** bTargetIterativeEnabled uses an allowlist (with a blocklist override), so defaults to false. */
 	bool bTargetIterativeEnabled = false;
 
