[WP Navmesh] Optimize Perforce access when building WP navmesh
Change FPackageSourceControlHelper::Delete(const TArray<UPackage*>& Packages) to be batched
Also update dialog being displayed when an error occurred during the build process
#rnx

diff --git a/Engine/Source/Editor/UnrealEd/Private/EditorBuildUtils.cpp b/Engine/Source/Editor/UnrealEd/Private/EditorBuildUtils.cpp
--- a/Engine/Source/Editor/UnrealEd/Private/EditorBuildUtils.cpp
+++ b/Engine/Source/Editor/UnrealEd/Private/EditorBuildUtils.cpp
@@ -963,29 +963,26 @@ bool FEditorBuildUtils::WorldPartitionBuildNavigation(const FString& InLongPacka
 		int32 Result = 0;
 		if (!bCancelled && FPlatformProcess::GetProcReturnCode(ProcessHandle, &Result))
 		{	
-			if (Result == 0)
-			{
-				// Force a directory watcher tick for the asset registry to get notified of the changes
-				FDirectoryWatcherModule& DirectoryWatcherModule = FModuleManager::Get().LoadModuleChecked<FDirectoryWatcherModule>(TEXT("DirectoryWatcher"));
-				DirectoryWatcherModule.Get()->Tick(-1.0f);
-				
-				// Unload any loaded map
-				if (!UEditorLoadingAndSavingUtils::NewBlankMap(/*bSaveExistingMap*/false))
-				{
-					return false;
-				}
-			
-				// Force registry update before loading converted map
-				const FAssetRegistryModule& AssetRegistryModule = FModuleManager::LoadModuleChecked<FAssetRegistryModule>("AssetRegistry");
-				IAssetRegistry& AssetRegistry = AssetRegistryModule.Get();
-											
-				FString MapToLoad = InLongPackageName;
-			
-				AssetRegistry.ScanModifiedAssetFiles({ MapToLoad });
-				AssetRegistry.ScanPathsSynchronous(ULevel::GetExternalObjectsPaths(MapToLoad), true);
+			// Force a directory watcher tick for the asset registry to get notified of the changes
+			FDirectoryWatcherModule& DirectoryWatcherModule = FModuleManager::Get().LoadModuleChecked<FDirectoryWatcherModule>(TEXT("DirectoryWatcher"));
+			DirectoryWatcherModule.Get()->Tick(-1.0f);
 			
-				FEditorFileUtils::LoadMap(MapToLoad);
+			// Unload any loaded map
+			if (!UEditorLoadingAndSavingUtils::NewBlankMap(/*bSaveExistingMap*/false))
+			{
+				return false;
 			}
+		
+			// Force registry update before loading converted map
+			const FAssetRegistryModule& AssetRegistryModule = FModuleManager::LoadModuleChecked<FAssetRegistryModule>("AssetRegistry");
+			IAssetRegistry& AssetRegistry = AssetRegistryModule.Get();
+										
+			FString MapToLoad = InLongPackageName;
+		
+			AssetRegistry.ScanModifiedAssetFiles({ MapToLoad });
+			AssetRegistry.ScanPathsSynchronous(ULevel::GetExternalObjectsPaths(MapToLoad), true);
+		
+			FEditorFileUtils::LoadMap(MapToLoad);
 		}
 		else if (bCancelled)
 		{
@@ -994,7 +991,7 @@ bool FEditorBuildUtils::WorldPartitionBuildNavigation(const FString& InLongPacka
 	
 		if (Result != 0)
 		{
-			FMessageDialog::Open(EAppMsgType::Ok, LOCTEXT("WorldPartitionBuildNavigationFailed", "Building navigation failed!"));
+			FMessageDialog::Open(EAppMsgType::Ok, LOCTEXT("WorldPartitionBuildNavigationFailed", "Errors occured during the build process, please refer to the logs ('WPNavigationBuilderLog.txt')."));
 		}
 	}
 	
diff --git a/Engine/Source/Editor/UnrealEd/Private/PackageSourceControlHelper.cpp b/Engine/Source/Editor/UnrealEd/Private/PackageSourceControlHelper.cpp
--- a/Engine/Source/Editor/UnrealEd/Private/PackageSourceControlHelper.cpp
+++ b/Engine/Source/Editor/UnrealEd/Private/PackageSourceControlHelper.cpp
@@ -184,14 +184,16 @@ bool FPackageSourceControlHelper::Delete(const TArray<UPackage*>& Packages) cons
 		return true;
 	}
 
-	bool bAllPackagesDeleted = true;
+	TArray<FString> PackageNames;
+	PackageNames.Reserve(Packages.Num());
+	
 	for (UPackage* Package : Packages)
 	{
+		PackageNames.Add(Package->GetName());
 		ResetLoaders(Package);
-		bAllPackagesDeleted &= Delete(Package->GetName());
 	}
 
-	return bAllPackagesDeleted;
+	return Delete(PackageNames);
 }
 
 bool FPackageSourceControlHelper::AddToSourceControl(UPackage* Package) const
diff --git a/Engine/Source/Editor/UnrealEd/Private/WorldPartition/WorldPartitionNavigationDataBuilder.cpp b/Engine/Source/Editor/UnrealEd/Private/WorldPartition/WorldPartitionNavigationDataBuilder.cpp
--- a/Engine/Source/Editor/UnrealEd/Private/WorldPartition/WorldPartitionNavigationDataBuilder.cpp
+++ b/Engine/Source/Editor/UnrealEd/Private/WorldPartition/WorldPartitionNavigationDataBuilder.cpp
@@ -116,7 +116,7 @@ bool UWorldPartitionNavigationDataBuilder::RunInternal(UWorld* World, const FCel
 		UE_LOG(LogWorldPartitionNavigationDataBuilder, Verbose, TEXT("   Number of packages to clear: %i"), PackagesToClean.Num());
 		
 		// Just delete all ANavigationDataChunkActor packages
-		if (!PackageHelper.Delete(PackagesToClean.Array()))
+		if (!DeletePackages(PackageHelper, PackagesToClean.Array()))
 		{
 			UE_LOG(LogWorldPartitionNavigationDataBuilder, Error, TEXT("Error deleting packages."));
 		}
@@ -183,19 +183,10 @@ bool UWorldPartitionNavigationDataBuilder::RunInternal(UWorld* World, const FCel
 	}
 
 	// Delete packages
-	if (!PackagesToDelete.IsEmpty())
+	if (!DeletePackages(PackageHelper, PackagesToDelete))
 	{
-		UE_LOG(LogWorldPartitionNavigationDataBuilder, Log, TEXT("Deleting %d packages."), PackagesToDelete.Num());
-		for (const UPackage* Package : PackagesToDelete)
-		{
-			UE_LOG(LogWorldPartitionNavigationDataBuilder, Verbose, TEXT("   Deleting package  %s."), *Package->GetName());	
-		}
-		
-		if (!PackageHelper.Delete(PackagesToDelete))
-		{
-			UE_LOG(LogWorldPartitionNavigationDataBuilder, Error, TEXT("Error deleting packages."));
-			return true;
-		}
+		UE_LOG(LogWorldPartitionNavigationDataBuilder, Error, TEXT("Error deleting packages."));
+		return true;
 	}
 
 	// Save packages
@@ -245,13 +236,17 @@ bool UWorldPartitionNavigationDataBuilder::RunInternal(UWorld* World, const FCel
 			TRACE_CPUPROFILER_EVENT_SCOPE(AddingToSourceControl);
 			UE_LOG(LogWorldPartitionNavigationDataBuilder, Log, TEXT("Adding packages to source control."));
 
-			for (UPackage* Package : PackagesToAdd)
+			TArray<FString> PackageNamesToAdd;
+			PackageNamesToAdd.Reserve(PackagesToAdd.Num());
+			for (const UPackage* Package : PackagesToAdd)
 			{
-				if (!PackageHelper.AddToSourceControl(Package))
-				{
-					UE_LOG(LogWorldPartitionNavigationDataBuilder, Error, TEXT("Error adding package %s to source control."), *Package->GetName());
-					return true;
-				}
+				PackageNamesToAdd.Add(Package->GetName());
+			}
+			
+			if (!PackageHelper.AddToSourceControl(PackageNamesToAdd))
+			{
+				UE_LOG(LogWorldPartitionNavigationDataBuilder, Error, TEXT("Error adding packages."));
+				return true;
 			}
 		}
 
@@ -398,3 +393,23 @@ bool UWorldPartitionNavigationDataBuilder::SavePackages(const TArray<UPackage*>&
 	return true;
 }
 
+bool UWorldPartitionNavigationDataBuilder::DeletePackages(FPackageSourceControlHelper& PackageHelper, const TArray<UPackage*>& PackagesToDelete) const
+{
+	if (!PackagesToDelete.IsEmpty())
+	{
+		TRACE_CPUPROFILER_EVENT_SCOPE(DeletePackages);
+		
+		UE_LOG(LogWorldPartitionNavigationDataBuilder, Log, TEXT("Deleting %d packages."), PackagesToDelete.Num());
+		for (const UPackage* Package : PackagesToDelete)
+		{
+			UE_LOG(LogWorldPartitionNavigationDataBuilder, Verbose, TEXT("   Deleting package  %s."), *Package->GetName());
+		}
+		
+		if (!PackageHelper.Delete(PackagesToDelete))
+		{
+			return false;
+		}
+	}
+
+	return true;
+}
\ No newline at end of file
diff --git a/Engine/Source/Editor/UnrealEd/Public/WorldPartition/WorldPartitionNavigationDataBuilder.h b/Engine/Source/Editor/UnrealEd/Public/WorldPartition/WorldPartitionNavigationDataBuilder.h
--- a/Engine/Source/Editor/UnrealEd/Public/WorldPartition/WorldPartitionNavigationDataBuilder.h
+++ b/Engine/Source/Editor/UnrealEd/Public/WorldPartition/WorldPartitionNavigationDataBuilder.h
@@ -25,6 +25,7 @@ protected:
 	bool GenerateNavigationData(UWorldPartition* WorldPartition, const FBox& LoadedBounds, const FBox& GeneratingBounds) const;
 
 	bool SavePackages(const TArray<UPackage*>& PackagesToSave) const;
+	bool DeletePackages(FPackageSourceControlHelper& PackageHelper, const TArray<UPackage*>& PackagesToDelete) const;
 
 	bool bCleanBuilderPackages = false;
 };
