Fixes issue replacing skeleton on AnimBP that is missing skeleton asset.

diff --git a/Engine/Source/Developer/AssetTools/Private/AssetTypeActions/AssetTypeActions_AnimBlueprint.cpp b/Engine/Source/Developer/AssetTools/Private/AssetTypeActions/AssetTypeActions_AnimBlueprint.cpp
--- a/Engine/Source/Developer/AssetTools/Private/AssetTypeActions/AssetTypeActions_AnimBlueprint.cpp
+++ b/Engine/Source/Developer/AssetTools/Private/AssetTypeActions/AssetTypeActions_AnimBlueprint.cpp
@@ -184,6 +184,28 @@ void FAssetTypeActions_AnimBlueprint::OpenAssetEditor( const TArray<UObject*>& I
 		auto AnimBlueprint = Cast<UAnimBlueprint>(*ObjIt);
 		if (AnimBlueprint != NULL && AnimBlueprint->SkeletonGeneratedClass && AnimBlueprint->GeneratedClass)
 		{
+			if(AnimBlueprint->BlueprintType != BPTYPE_Interface && !AnimBlueprint->TargetSkeleton && !AnimBlueprint->bIsTemplate)
+			{
+				FText ShouldRetargetMessage = LOCTEXT("ShouldRetarget_Message", "Could not find the skeleton for Anim Blueprint '{BlueprintName}' Would you like to choose a new one?");
+				
+				FFormatNamedArguments Arguments;
+				Arguments.Add( TEXT("BlueprintName"), FText::FromString(AnimBlueprint->GetName()));
+
+				if (FMessageDialog::Open(EAppMsgType::YesNo, FText::Format(ShouldRetargetMessage, Arguments)) == EAppReturnType::Yes)
+				{
+					TArray<TObjectPtr<UObject>> AssetsToRetarget;
+					AssetsToRetarget.Add(AnimBlueprint);
+					const bool bSkeletonReplaced = ReplaceMissingSkeleton(AssetsToRetarget);
+					if (!bSkeletonReplaced)
+					{
+						return; // Persona will crash if trying to load asset without a skeleton
+					}
+				}
+				else
+				{
+					return;
+				}
+			}
 			const bool bBringToFrontIfOpen = true;
 #if WITH_EDITOR
 			if (IAssetEditorInstance* EditorInstance = GEditor->GetEditorSubsystem<UAssetEditorSubsystem>()->FindEditorForAsset(AnimBlueprint, bBringToFrontIfOpen))
@@ -264,6 +286,16 @@ void FAssetTypeActions_AnimBlueprint::ExecuteFindSkeleton(TArray<TWeakObjectPtr<
 	}
 }
 
+bool FAssetTypeActions_AnimBlueprint::ReplaceMissingSkeleton(TArray<UObject*> InAnimBlueprints) const
+{
+	// record anim assets that need skeleton replaced
+	const TArray<TWeakObjectPtr<UObject>> ABPsToFix = GetTypedWeakObjectPtrs<UObject>(InAnimBlueprints);
+	// get a skeleton from the user and replace it
+	const TSharedPtr<SReplaceMissingSkeletonDialog> PickSkeletonWindow = SNew(SReplaceMissingSkeletonDialog).AnimAssets(ABPsToFix);
+	const bool bWasSkeletonReplaced = PickSkeletonWindow.Get()->ShowModal();
+	return bWasSkeletonReplaced;
+}
+
 TSharedPtr<SWidget> FAssetTypeActions_AnimBlueprint::GetThumbnailOverlay(const FAssetData& AssetData) const
 {
 	const FSlateBrush* Icon = FSlateIconFinder::FindIconBrushForClass(UAnimBlueprint::StaticClass());
diff --git a/Engine/Source/Developer/AssetTools/Private/AssetTypeActions/AssetTypeActions_AnimBlueprint.h b/Engine/Source/Developer/AssetTools/Private/AssetTypeActions/AssetTypeActions_AnimBlueprint.h
--- a/Engine/Source/Developer/AssetTools/Private/AssetTypeActions/AssetTypeActions_AnimBlueprint.h
+++ b/Engine/Source/Developer/AssetTools/Private/AssetTypeActions/AssetTypeActions_AnimBlueprint.h
@@ -34,4 +34,7 @@ private:
 	
 	/** Certain options are not available for template anim BPs */
 	bool AreAnyNonTemplateAnimBlueprintsSelected(TArray<TWeakObjectPtr<UAnimBlueprint>> Objects) const;
+
+	/** When skeleton asset is missing, allow replacing skeleton asset */ 
+	bool ReplaceMissingSkeleton(TArray<UObject*> InAnimBlueprints) const;
 };
diff --git a/Engine/Source/Developer/AssetTools/Private/AssetTypeActions/AssetTypeActions_AnimationAsset.cpp b/Engine/Source/Developer/AssetTools/Private/AssetTypeActions/AssetTypeActions_AnimationAsset.cpp
--- a/Engine/Source/Developer/AssetTools/Private/AssetTypeActions/AssetTypeActions_AnimationAsset.cpp
+++ b/Engine/Source/Developer/AssetTools/Private/AssetTypeActions/AssetTypeActions_AnimationAsset.cpp
@@ -90,6 +90,30 @@ void FAssetTypeActions_AnimationAsset::OpenAnimAssetEditor(const TArray<UObject*
 	// For each one..
 	for (UAnimationAsset* AnimAsset : AnimAssets)
 	{
+		USkeleton* AnimSkeleton = AnimAsset->GetSkeleton();
+		if (!AnimSkeleton)
+		{
+			FText ShouldRetargetMessage = LOCTEXT("ShouldRetargetAnimAsset_Message", "Could not find the skeleton for Anim '{AnimName}' Would you like to choose a new one?");
+
+			FFormatNamedArguments Arguments;
+			Arguments.Add(TEXT("AnimName"), FText::FromString(AnimAsset->GetName()));
+
+			if (FMessageDialog::Open(EAppMsgType::YesNo, FText::Format(ShouldRetargetMessage, Arguments)) == EAppReturnType::Yes)
+			{
+				TArray<TObjectPtr<UObject>> AssetsToRetarget;
+				AssetsToRetarget.Add(AnimAsset);
+				const bool bSkeletonReplaced = ReplaceMissingSkeleton(AssetsToRetarget);
+				if (!bSkeletonReplaced)
+				{
+					return; // Persona will crash if trying to load asset without a skeleton
+				}
+			}
+			else
+			{
+				return;
+			}
+		}
+		
 		// First see if we already have it open
 		const bool bBringToFrontIfOpen = true;
 #if WITH_EDITOR
@@ -157,7 +181,6 @@ void FAssetTypeActions_AnimationAsset::OpenAnimAssetEditor(const TArray<UObject*
 	}
 }
 
-
 void FAssetTypeActions_AnimationAsset::ExecuteOpenInNewWindow(TArray<TWeakObjectPtr<UAnimationAsset>> Objects)
 {
 	TArray<UObject*> ObjectsToSync;
@@ -194,4 +217,14 @@ void FAssetTypeActions_AnimationAsset::ExecuteFindSkeleton(TArray<TWeakObjectPtr
 	}
 }
 
+bool FAssetTypeActions_AnimationAsset::ReplaceMissingSkeleton(TArray<TObjectPtr<UObject>> InAnimationAssets) const
+{
+	// record anim assets that need skeleton replaced
+	const TArray<TWeakObjectPtr<UObject>> AnimsToFix = GetTypedWeakObjectPtrs<UObject>(InAnimationAssets);
+	// get a skeleton from the user and replace it
+	const TSharedPtr<SReplaceMissingSkeletonDialog> PickSkeletonWindow = SNew(SReplaceMissingSkeletonDialog).AnimAssets(AnimsToFix);
+	const bool bWasSkeletonReplaced = PickSkeletonWindow.Get()->ShowModal();
+	return bWasSkeletonReplaced;
+}
+
 #undef LOCTEXT_NAMESPACE
diff --git a/Engine/Source/Developer/AssetTools/Public/AssetTypeActions/AssetTypeActions_AnimationAsset.h b/Engine/Source/Developer/AssetTools/Public/AssetTypeActions/AssetTypeActions_AnimationAsset.h
--- a/Engine/Source/Developer/AssetTools/Public/AssetTypeActions/AssetTypeActions_AnimationAsset.h
+++ b/Engine/Source/Developer/AssetTools/Public/AssetTypeActions/AssetTypeActions_AnimationAsset.h
@@ -7,6 +7,8 @@
 #include "AssetTypeActions_Base.h"
 #include "Animation/AnimationAsset.h"
 #include "EditorAnimUtils.h"
+#include "SSkeletonWidget.h"
+
 
 class ASSETTOOLS_API FAssetTypeActions_AnimationAsset : public FAssetTypeActions_Base
 {
@@ -32,4 +34,7 @@ private:
 
 	/** Open animation asset, will find existing editor if desired. */
 	void OpenAnimAssetEditor(const TArray<UObject*>& InObjects, bool bForceNewEditor, TSharedPtr<IToolkitHost> EditWithinLevelEditor);
+	
+	/** Replace skeleton when USkeleton is missing. Returns true only if Skeleton was replaced. */
+	bool ReplaceMissingSkeleton(TArray<TObjectPtr<UObject>> InAnimationAssets) const;
 };
diff --git a/Engine/Source/Editor/UnrealEd/Private/SSkeletonWidget.cpp b/Engine/Source/Editor/UnrealEd/Private/SSkeletonWidget.cpp
--- a/Engine/Source/Editor/UnrealEd/Private/SSkeletonWidget.cpp
+++ b/Engine/Source/Editor/UnrealEd/Private/SSkeletonWidget.cpp
@@ -767,7 +767,113 @@ FDisplayedAssetEntryInfo::FDisplayedAssetEntryInfo(UObject* InAsset, USkeleton*
 	, AnimAsset(InAsset)
 	, RemapAsset(nullptr)
 {
+}
+
+void SReplaceMissingSkeletonDialog::Construct(const FArguments& InArgs)
+{
+	AssetsToReplaceSkeletonOn = InArgs._AnimAssets;
+	
+	// Load the content browser module to display an asset picker
+	const FContentBrowserModule& ContentBrowserModule = FModuleManager::LoadModuleChecked<FContentBrowserModule>("ContentBrowser");
+
+	FAssetPickerConfig AssetPickerConfig;
+
+	/** The asset picker will only show skeleton assets */
+	AssetPickerConfig.Filter.ClassNames.Add(USkeleton::StaticClass()->GetFName());
+	AssetPickerConfig.Filter.bRecursiveClasses = true;
+
+	/** The delegate that fires when an asset was selected */
+	AssetPickerConfig.OnAssetSelected = FOnAssetSelected::CreateSP(this, & SReplaceMissingSkeletonDialog::OnSkeletonSelected);
+
+	/** The default view mode should be a list view */
+	AssetPickerConfig.InitialAssetViewType = EAssetViewType::List;
 
+	SWindow::Construct(SWindow::FArguments()
+		.Title(LOCTEXT("ReplaceSkeletonOptions", "Pick Replacement Skeleton"))
+		.ClientSize(FVector2D(500, 600))
+		.SupportsMinimize(false)
+		.SupportsMaximize(false)
+		.ClientSize(FVector2D(450, 450))
+		[
+			SNew(SVerticalBox)
+
+			+ SVerticalBox::Slot() // Add user input block
+			[
+				ContentBrowserModule.Get().CreateAssetPicker(AssetPickerConfig)
+			]
+
+			+SVerticalBox::Slot()
+			.AutoHeight()
+			.HAlign(HAlign_Right)
+			.Padding(5)
+			[
+				SNew(SUniformGridPanel)
+				.SlotPadding(FEditorStyle::GetMargin("StandardDialog.SlotPadding"))
+				.MinDesiredSlotWidth(FEditorStyle::GetFloat("StandardDialog.MinDesiredSlotWidth"))
+				.MinDesiredSlotHeight(FEditorStyle::GetFloat("StandardDialog.MinDesiredSlotHeight"))
+				+SUniformGridPanel::Slot(0, 0)
+				[
+					SNew(SButton)
+					.HAlign(HAlign_Center)
+					.ContentPadding(FEditorStyle::GetMargin("StandardDialog.ContentPadding"))
+					.Text(LOCTEXT("OK", "OK"))
+					.OnClicked(this, &SReplaceMissingSkeletonDialog::OnButtonClick, EAppReturnType::Ok)
+				]
+				+SUniformGridPanel::Slot(1, 0)
+				[
+					SNew(SButton)
+					.HAlign(HAlign_Center)
+					.ContentPadding(FEditorStyle::GetMargin("StandardDialog.ContentPadding"))
+					.Text(LOCTEXT("Cancel", "Cancel"))
+					.OnClicked(this, &SReplaceMissingSkeletonDialog::OnButtonClick, EAppReturnType::Cancel)
+				]
+			]
+		]);
+}
+
+void SReplaceMissingSkeletonDialog::OnSkeletonSelected(const FAssetData& Replacement)
+{
+	SelectedAsset = Replacement;
+}
+
+FReply SReplaceMissingSkeletonDialog::OnButtonClick(EAppReturnType::Type ButtonID)
+{
+	UserResponse = ButtonID;
+	RequestDestroyWindow();
+
+	if (ButtonID != EAppReturnType::Ok)
+	{
+		return FReply::Handled();
+	}
+
+	const TObjectPtr<UObject> Asset = SelectedAsset.GetAsset();
+	if (Asset.IsNull())
+	{
+		return FReply::Handled();
+	}
+	
+	if (const TObjectPtr<USkeleton> ReplacementSkeleton = CastChecked<USkeleton>(Asset))
+	{
+		constexpr bool bRetargetReferredAssets = true;
+		constexpr bool bConvertSpaces = false;
+		FAnimationRetargetContext RetargetContext(AssetsToReplaceSkeletonOn, bRetargetReferredAssets, bConvertSpaces);
+		// since we are replacing a missing skeleton, we don't want to duplicate the asset
+		// setting this to null prevents assets from being duplicated
+		const FNameDuplicationRule* NameRule = nullptr;
+
+		EditorAnimUtils::RetargetAnimations(nullptr, ReplacementSkeleton, RetargetContext, bRetargetReferredAssets, NameRule);
+
+		bWasSkeletonReplaced = true;
+	}
+
+	return FReply::Handled();
+}
+
+bool SReplaceMissingSkeletonDialog::ShowModal()
+{
+	bWasSkeletonReplaced = false;
+	GEditor->EditorAddModalWindow(SharedThis(this));
+	return bWasSkeletonReplaced;
 }
 
 #undef LOCTEXT_NAMESPACE 
diff --git a/Engine/Source/Editor/UnrealEd/Public/SSkeletonWidget.h b/Engine/Source/Editor/UnrealEd/Public/SSkeletonWidget.h
--- a/Engine/Source/Editor/UnrealEd/Public/SSkeletonWidget.h
+++ b/Engine/Source/Editor/UnrealEd/Public/SSkeletonWidget.h
@@ -418,4 +418,30 @@ protected:
 	FText AssetPath;
 };
 
+class SReplaceMissingSkeletonDialog : public SWindow
+{
+	
+public:
+	
+	SLATE_BEGIN_ARGS(SReplaceMissingSkeletonDialog)
+	{
+	}
+
+	SLATE_ARGUMENT(TArray<TWeakObjectPtr<UObject>>, AnimAssets)
+	SLATE_END_ARGS()
+
+	void UNREALED_API Construct(const FArguments& InArgs);
+	SReplaceMissingSkeletonDialog() : UserResponse(EAppReturnType::Cancel){}
+
+	bool UNREALED_API ShowModal();
+
+protected:
+	void OnSkeletonSelected(const FAssetData& Replacement);
+	FReply OnButtonClick(EAppReturnType::Type ButtonID);
+	EAppReturnType::Type UserResponse;
+	FAssetData SelectedAsset;
+	TArray<TWeakObjectPtr<UObject>> AssetsToReplaceSkeletonOn;
+	bool bWasSkeletonReplaced;
+};
+
 #undef LOCTEXT_NAMESPACE
