EditorDomain KnownClasses support:
Add DeclareCustomVersions and PostLoadCanConstructClasses  for classes that need them.
Fix PostLoadCanConstructClasses to inherit from parent.
Fix TrySavePackage to replace BlueprintClasses with their native base class.
#rnx

diff --git a/Engine/Config/BaseEditor.ini b/Engine/Config/BaseEditor.ini
--- a/Engine/Config/BaseEditor.ini
+++ b/Engine/Config/BaseEditor.ini
@@ -369,6 +369,7 @@ SaveUnversioned=true
 
 +ClassBlockList=/Script/Engine.StaticMesh
 +ClassBlockList=/Script/Engine.Blueprint
++ClassBlockList=/Script/Engine.UserDefinedStruct
 
 ; Known bad files that are problematic if all of them come from EditorDomain. Blocking any one of these is sufficient to prevent the bug.
 +PackageBlockList=/Niagara/Functions/RandomRangeFloat
@@ -385,10 +386,25 @@ SaveUnversioned=true
 +GlobalCanConstructClasses=/Script/CoreUObject.MetaData
 +GlobalCanConstructClasses=/Script/Engine.AssetImportData
 ; +PostLoadCanConstructClasses=<LongPackageNameOfExistingClass>,<LongPackageNameOfClassItCanCreateInPostLoad>
-+PostLoadCanConstructClasses=/Script/Engine.StaticMesh,/Script/Engine.StaticMeshDescriptionBulkData
-+PostLoadCanConstructClasses=/Script/Engine.MapBuildDataRegistry,/Script/Engine.TextureCube
 +PostLoadCanConstructClasses=/Script/ClothingSystemRuntimeCommon.ClothingAssetCommon,/Script/ChaosCloth.ChaosClothSharedSimConfig
 +PostLoadCanConstructClasses=/Script/ClothingSystemRuntimeCommon.ClothingAssetCommon,/Script/ChaosCloth.ChaosClothConfig
++PostLoadCanConstructClasses=/Script/Engine.AnimSequenceBase,/Script/Engine.AnimDataModel
++PostLoadCanConstructClasses=/Script/Engine.MapBuildDataRegistry,/Script/Engine.TextureCube
++PostLoadCanConstructClasses=/Script/Engine.PhysicsAsset,/Script/Engine.SkeletalBodySetup
++PostLoadCanConstructClasses=/Script/Engine.SkeletalMesh,/Script/Engine.SkeletalMeshEditorData
++PostLoadCanConstructClasses=/Script/Engine.StaticMesh,/Script/Engine.StaticMeshDescriptionBulkData
++PostLoadCanConstructClasses=/Script/LevelSequence.LevelSequenceActor,/Script/LevelSequence.DefaultLevelSequenceInstanceData
++PostLoadCanConstructClasses=/Script/MovieScene.MovieScene,/Script/MovieScene.MovieSceneNodeGroupCollection
++PostLoadCanConstructClasses=/Script/MovieScene.MovieSceneSection,/Script/MovieScene.MovieSceneBuiltInEasingFunction
++PostLoadCanConstructClasses=/Script/Niagara.NiagaraEffectType,/Script/Niagara.NiagaraSignificanceHandlerDistance
++PostLoadCanConstructClasses=/Script/Niagara.NiagaraEmitter,/Script/NiagaraEditor.NiagaraEditorParametersAdapter
++PostLoadCanConstructClasses=/Script/Niagara.NiagaraSystem,/Script/NiagaraEditor.NiagaraEditorParametersAdapter
++PostLoadCanConstructClasses=/Script/NiagaraEditor.NiagaraEmitterEditorData,/Script/Engine.EdGraph
++PostLoadCanConstructClasses=/Script/NiagaraEditor.NiagaraGraph,/Script/NiagaraEditor.NiagaraScriptVariable
++PostLoadCanConstructClasses=/Script/NiagaraEditor.NiagaraSystemEditorData,/Script/NiagaraEditor.NiagaraOverviewNode
+
+
+
 
 [TargetDomain]
 ; IterativeClassAllowListEnabled/IterativeClassAllowList/IterativeClassBlockList:
diff --git a/Engine/Source/Editor/UnrealEd/Private/EditorDomain/EditorDomainUtils.cpp b/Engine/Source/Editor/UnrealEd/Private/EditorDomain/EditorDomainUtils.cpp
--- a/Engine/Source/Editor/UnrealEd/Private/EditorDomain/EditorDomainUtils.cpp
+++ b/Engine/Source/Editor/UnrealEd/Private/EditorDomain/EditorDomainUtils.cpp
@@ -500,6 +500,7 @@ FClassDigestData* FPrecacheClassDigest::GetRecursive(FName ClassName, bool bAllo
 	}
 
 	// Propagate values from the parent
+	TArray<FName> ConstructClasses; // Not a class scratch variable because we use it across recursive calls
 	if (!ParentName.IsNone())
 	{
 		// CoreRedirects are expected to act only on import classes from packages; they are not expected to act on the parent class pointer
@@ -535,21 +536,22 @@ FClassDigestData* FPrecacheClassDigest::GetRecursive(FName ClassName, bool bAllo
 			{
 				DigestData->bTargetIterativeEnabled &= ParentDigest->bTargetIterativeEnabled;
 			}
+			ConstructClasses.Append(ParentDigest->ConstructClasses);
 		}
 	}
 
 	// Get and validate the list of ConstructClasses by recursively initializing them
-	TArray<FName> ConstructClasses; // Not a class scratch variable because we use it across recursive calls
-	GConstructClasses.MultiFind(ClassName, ConstructClasses);
+	TArray<FName> AddedConstructClasses; // Not a class scratch variable because we use it across recursive calls
+	GConstructClasses.MultiFind(ClassName, AddedConstructClasses);
 	if (RedirectsResolvedName != ClassName)
 	{
-		GConstructClasses.MultiFind(RedirectsResolvedName, ConstructClasses);
+		GConstructClasses.MultiFind(RedirectsResolvedName, AddedConstructClasses);
 	}
-	if (!ConstructClasses.IsEmpty())
+	if (!AddedConstructClasses.IsEmpty())
 	{
-		for (int32 Index = 0; Index < ConstructClasses.Num(); )
+		for (int32 Index = 0; Index < AddedConstructClasses.Num(); )
 		{
-			FName ConstructClass = ConstructClasses[Index];
+			FName ConstructClass = AddedConstructClasses[Index];
 			FClassDigestData* ConstructClassDigest = GetRecursive(ConstructClass, true /* bAllowRedirects */);
 			if (!ConstructClassDigest)
 			{
@@ -558,7 +560,7 @@ FClassDigestData* FPrecacheClassDigest::GetRecursive(FName ClassName, bool bAllo
 					TEXT("This is a class that can be constructed by postload upgrades of class %s. ")
 					TEXT("Old packages with class %s will load more slowly."),
 					*ConstructClass.ToString(), *ClassName.ToString(), *ClassName.ToString(), *ClassName.ToString());
-				ConstructClasses.RemoveAt(Index);
+				AddedConstructClasses.RemoveAt(Index);
 			}
 			else
 			{
@@ -571,11 +573,15 @@ FClassDigestData* FPrecacheClassDigest::GetRecursive(FName ClassName, bool bAllo
 				++Index;
 			}
 		}
-		Algo::Sort(ConstructClasses, FNameLexicalLess());
-		ConstructClasses.SetNum(Algo::Unique(ConstructClasses));
+		ConstructClasses.Append(AddedConstructClasses);
 
 		// The map has possibly been modified so we need to recalculate the address of ClassName's DigestData
 		DigestData = &ClassDigests.FindChecked(ClassName);
+	}
+	if (ConstructClasses.Num())
+	{
+		Algo::Sort(ConstructClasses, FNameLexicalLess());
+		ConstructClasses.SetNum(Algo::Unique(ConstructClasses));
 		DigestData->ConstructClasses = MoveTemp(ConstructClasses);
 	}
 
@@ -1559,6 +1565,53 @@ bool TrySavePackage(UPackage* Package)
 			*Package->GetName());
 		return false;
 	}
+
+	// Replace any Blueprint classes in SavedImportedClasses with their native base class, since only native classes
+	// have class schemas and customversions that could change the custom versions. EditorDomain packages with blueprint
+	// instances do not use unversioned properties and so do not need to be keyed to changes in the BP class they depend on.
+	// ImportedClasses used that same replacement, so we need to match it here when identifying which classes are new.
+	int32 NextClass = 0;
+	bool bHasTriedPrecacheClassDigests = false;
+	FClassDigestMap& ClassDigests = GetClassDigests();
+	while (NextClass < SavedImportedClasses.Num())
+	{
+		{
+			FReadScopeLock ClassDigestsScopeLock(ClassDigests.Lock);
+			while (NextClass < SavedImportedClasses.Num())
+			{
+				FName ClassName = SavedImportedClasses[NextClass];
+				FClassDigestData* ExistingData = ClassDigests.Map.Find(ClassName);
+				if (!ExistingData)
+				{
+					break;
+				}
+				SavedImportedClasses[NextClass] = ExistingData->ResolvedClosestNative;
+				++NextClass;
+			}
+		}
+
+		if (NextClass < SavedImportedClasses.Num())
+		{
+			if (bHasTriedPrecacheClassDigests)
+			{
+				UE_LOG(LogEditorDomain, Warning, TEXT("Could not save %s to EditorDomain: It constructed an instance of class %s which which does not exist."),
+					*Package->GetName(), *SavedImportedClasses[NextClass].ToString());
+				return false;
+			}
+			// EDITORDOMAIN_TODO: Remove this !IsInGameThread check once FindObject no longer asserts if GIsSavingPackage
+			if (!IsInGameThread())
+			{
+				UE_LOG(LogEditorDomain, Warning, TEXT("Could not save %s to EditorDomain: It constructed an instance of a class %s which is not yet saved in the EditorDomain, ")
+					TEXT("and we are not on the gamethread so we cannot add it."),
+					*Package->GetName(), *SavedImportedClasses[NextClass].ToString());
+				return false;
+			}
+			TConstArrayView<FName> RemainingClasses = TConstArrayView<FName>(SavedImportedClasses).RightChop(NextClass);
+			PrecacheClassDigests(RemainingClasses);
+			bHasTriedPrecacheClassDigests = true;
+		}
+	}
+
 	TSet<FName> KnownImportedClasses(PackageDigest.ImportedClasses);
 	for (FName ImportedClass : SavedImportedClasses)
 	{
diff --git a/Engine/Source/Runtime/Engine/Classes/Animation/MorphTarget.h b/Engine/Source/Runtime/Engine/Classes/Animation/MorphTarget.h
--- a/Engine/Source/Runtime/Engine/Classes/Animation/MorphTarget.h
+++ b/Engine/Source/Runtime/Engine/Classes/Animation/MorphTarget.h
@@ -175,6 +175,7 @@ public:
 	//~ UObject interface
 
 	ENGINE_API virtual void Serialize(FArchive& Ar) override;
+	ENGINE_API virtual void DeclareCustomVersions(FArchive& Ar) override;
 	ENGINE_API virtual void PostLoad() override;
 
 	/** UObject does not support serialization via FMemoryArchive, so manually handle separately */
diff --git a/Engine/Source/Runtime/Engine/Private/EdGraph/EdGraphPin.cpp b/Engine/Source/Runtime/Engine/Private/EdGraph/EdGraphPin.cpp
--- a/Engine/Source/Runtime/Engine/Private/EdGraph/EdGraphPin.cpp
+++ b/Engine/Source/Runtime/Engine/Private/EdGraph/EdGraphPin.cpp
@@ -1765,6 +1765,7 @@ void UEdGraphPin::DeclareCustomVersions(FArchive& Ar)
 {
 	Ar.UsingCustomVersion(FFrameworkObjectVersion::GUID);
 	Ar.UsingCustomVersion(FUE5MainStreamObjectVersion::GUID);
+	Ar.UsingCustomVersion(FUE5ReleaseStreamObjectVersion::GUID);
 	FEdGraphPinType::DeclareCustomVersions(Ar);
 }
 
diff --git a/Engine/Source/Runtime/Engine/Private/MorphMesh.cpp b/Engine/Source/Runtime/Engine/Private/MorphMesh.cpp
--- a/Engine/Source/Runtime/Engine/Private/MorphMesh.cpp
+++ b/Engine/Source/Runtime/Engine/Private/MorphMesh.cpp
@@ -36,6 +36,13 @@ void UMorphTarget::Serialize( FArchive& Ar )
 	}
 }
 
+void UMorphTarget::DeclareCustomVersions(FArchive& Ar)
+{
+	Super::DeclareCustomVersions(Ar);
+	FMorphTargetLODModel MorphLODModel;
+	Ar << MorphLODModel;
+}
+
 namespace
 {
 	void SerializeMorphLODModels(FMemoryArchive& Ar, TArray<FMorphTargetLODModel>& MorphLODModels)
