Allow feature levels that do not work with Ray tracing to still be active in editor when Ray Tracing is turned on for the project.

diff --git a/Engine/Source/Editor/UnrealEd/Private/EditorEngine.cpp b/Engine/Source/Editor/UnrealEd/Private/EditorEngine.cpp
--- a/Engine/Source/Editor/UnrealEd/Private/EditorEngine.cpp
+++ b/Engine/Source/Editor/UnrealEd/Private/EditorEngine.cpp
@@ -7482,18 +7482,6 @@ void UEditorEngine::OnSceneMaterialsModified()
 
 void UEditorEngine::SetPreviewPlatform(const FPreviewPlatformInfo& NewPreviewPlatform, bool bSaveSettings)
 {
-#if RHI_RAYTRACING
-	if (IsRayTracingEnabled())
-	{
-		if (PreviewPlatform.PreviewFeatureLevel < ERHIFeatureLevel::SM5)
-		{
-			UE_LOG(LogEditor, Warning, TEXT("Preview feature level is incompatible with ray tracing, defaulting to Shader Model 5"));
-			PreviewPlatform.PreviewFeatureLevel = ERHIFeatureLevel::SM5;
-		}
-		return;
-	}
-#endif
-
 	// If we have specified a MaterialQualityPlatform ensure its feature level matches the requested feature level.
 	EShaderPlatform ShaderPlatform = ShaderFormatToLegacyShaderPlatform(NewPreviewPlatform.PreviewShaderFormatName);
 	ERHIFeatureLevel::Type MaxFeatureLevel = NewPreviewPlatform.PreviewShaderFormatName != NAME_None ? (ERHIFeatureLevel::Type)GetMaxSupportedFeatureLevel(ShaderPlatform) : ERHIFeatureLevel::SM5;
