Add an edit condition to global namespace config settings to avoid confusion when Blueprint namespace editor features are disabled.

diff --git a/Engine/Source/Editor/BlueprintGraph/Public/BlueprintEditorSettings.h b/Engine/Source/Editor/BlueprintGraph/Public/BlueprintEditorSettings.h
--- a/Engine/Source/Editor/BlueprintGraph/Public/BlueprintEditorSettings.h
+++ b/Engine/Source/Editor/BlueprintGraph/Public/BlueprintEditorSettings.h
@@ -176,7 +176,7 @@ public:
 	bool bEnableNamespaceImportingFeatures;
 
 	// The list of namespaces to always expose in any Blueprint (local per-user)
-	UPROPERTY(EditAnywhere, config, Category = Experimental)
+	UPROPERTY(EditAnywhere, config, Category = Experimental, meta = (EditCondition = "bEnableNamespaceEditorFeatures"))
 	TArray<FString> NamespacesToAlwaysInclude;
 
 	/** If enabled, then placed cast nodes will default to their "pure" form (meaning: without execution pins). */
diff --git a/Engine/Source/Editor/UnrealEd/Private/Settings/BlueprintEditorProjectSettingsCustomization.cpp b/Engine/Source/Editor/UnrealEd/Private/Settings/BlueprintEditorProjectSettingsCustomization.cpp
new file mode 100644
--- /dev/null
+++ b/Engine/Source/Editor/UnrealEd/Private/Settings/BlueprintEditorProjectSettingsCustomization.cpp
@@ -0,0 +1,38 @@
+// Copyright Epic Games, Inc. All Rights Reserved.
+
+#include "BlueprintEditorProjectSettingsCustomization.h"
+#include "BlueprintEditorSettings.h"
+#include "Settings/BlueprintEditorProjectSettings.h"
+#include "PropertyHandle.h"
+#include "DetailLayoutBuilder.h"
+#include "IDetailPropertyRow.h"
+
+#define LOCTEXT_NAMESPACE "FBlueprintEditorProjectSettingsCustomization"
+
+TSharedRef<IDetailCustomization> FBlueprintEditorProjectSettingsCustomization::MakeInstance()
+{
+	return MakeShared<FBlueprintEditorProjectSettingsCustomization>();
+}
+
+FBlueprintEditorProjectSettingsCustomization::~FBlueprintEditorProjectSettingsCustomization()
+{
+}
+
+void FBlueprintEditorProjectSettingsCustomization::CustomizeDetails(IDetailLayoutBuilder& LayoutBuilder)
+{
+	// Add a custom edit condition on the project-specific 'NamespacesToAlwaysInclude' member to link it to the editor's namespace feature toggle flag (for consistency w/ the editor-specific set).
+	static FName PropertyName_NamespacesToAlwaysInclude = GET_MEMBER_NAME_CHECKED(UBlueprintEditorProjectSettings, NamespacesToAlwaysInclude);
+	TSharedRef<IPropertyHandle> PropertyHandle_NamespacesToAlwaysInclude = LayoutBuilder.GetProperty(PropertyName_NamespacesToAlwaysInclude);
+	if (IDetailPropertyRow* PropertyRow_NamespacesToAlwaysInclude = LayoutBuilder.EditDefaultProperty(PropertyHandle_NamespacesToAlwaysInclude))
+	{
+		PropertyRow_NamespacesToAlwaysInclude->EditCondition(
+			TAttribute<bool>::CreateLambda([]()
+			{
+				return GetDefault<UBlueprintEditorSettings>()->bEnableNamespaceEditorFeatures;
+			}),
+			FOnBooleanValueChanged()
+		);
+	}
+}
+
+#undef LOCTEXT_NAMESPACE
diff --git a/Engine/Source/Editor/UnrealEd/Private/Settings/BlueprintEditorProjectSettingsCustomization.h b/Engine/Source/Editor/UnrealEd/Private/Settings/BlueprintEditorProjectSettingsCustomization.h
new file mode 100644
--- /dev/null
+++ b/Engine/Source/Editor/UnrealEd/Private/Settings/BlueprintEditorProjectSettingsCustomization.h
@@ -0,0 +1,23 @@
+// Copyright Epic Games, Inc. All Rights Reserved.
+
+#pragma once
+
+#include "CoreMinimal.h"
+#include "IDetailCustomization.h"
+
+/**
+ * Implements a details customization for UBlueprintEditorProjectSettings.
+ */
+class FBlueprintEditorProjectSettingsCustomization : public IDetailCustomization
+{
+public:
+	// Factory creation.
+	static TSharedRef<IDetailCustomization> MakeInstance();
+
+	// Virtual destructor.
+	virtual ~FBlueprintEditorProjectSettingsCustomization();
+
+	// IDetailCustomization interface
+	virtual void CustomizeDetails(IDetailLayoutBuilder& LayoutBuilder) override;
+	// End IDetailCustomization interface
+};
\ No newline at end of file
diff --git a/Engine/Source/Editor/UnrealEd/Private/UnrealEdEngine.cpp b/Engine/Source/Editor/UnrealEd/Private/UnrealEdEngine.cpp
--- a/Engine/Source/Editor/UnrealEd/Private/UnrealEdEngine.cpp
+++ b/Engine/Source/Editor/UnrealEd/Private/UnrealEdEngine.cpp
@@ -45,6 +45,7 @@
 #include "LevelEditor.h"
 #include "InstancedStaticMeshDelegates.h"
 #include "Interfaces/IMainFrameModule.h"
+#include "Settings/BlueprintEditorProjectSettingsCustomization.h"
 #include "Settings/EditorLoadingSavingSettingsCustomization.h"
 #include "Settings/GameMapsSettingsCustomization.h"
 #include "Settings/LevelEditorPlaySettingsCustomization.h"
@@ -176,6 +177,7 @@ void UUnrealEdEngine::Init(IEngineLoop* InEngineLoop)
 	{
 		FPropertyEditorModule& PropertyModule = FModuleManager::LoadModuleChecked<FPropertyEditorModule>("PropertyEditor");
 
+		PropertyModule.RegisterCustomClassLayout("BlueprintEditorProjectSettings", FOnGetDetailCustomizationInstance::CreateStatic(&FBlueprintEditorProjectSettingsCustomization::MakeInstance));
 		PropertyModule.RegisterCustomClassLayout("EditorLoadingSavingSettings", FOnGetDetailCustomizationInstance::CreateStatic(&FEditorLoadingSavingSettingsCustomization::MakeInstance));
 		PropertyModule.RegisterCustomClassLayout("GameMapsSettings", FOnGetDetailCustomizationInstance::CreateStatic(&FGameMapsSettingsCustomization::MakeInstance));
 		PropertyModule.RegisterCustomClassLayout("LevelEditorPlaySettings", FOnGetDetailCustomizationInstance::CreateStatic(&FLevelEditorPlaySettingsCustomization::MakeInstance));
diff --git a/Engine/Source/Editor/UnrealEd/Public/Settings/BlueprintEditorProjectSettings.h b/Engine/Source/Editor/UnrealEd/Public/Settings/BlueprintEditorProjectSettings.h
--- a/Engine/Source/Editor/UnrealEd/Public/Settings/BlueprintEditorProjectSettings.h
+++ b/Engine/Source/Editor/UnrealEd/Public/Settings/BlueprintEditorProjectSettings.h
@@ -34,6 +34,16 @@ public:
 	UPROPERTY(EditAnywhere, config, Category = Experimental)
 	uint8 bEnableChildActorExpansionInTreeView : 1;
 
+	/**
+	 * Default view mode to use for child actor components in a Blueprint actor's component tree hierarchy (experimental).
+	 */
+	UPROPERTY(EditAnywhere, config, Category = Experimental, meta = (EditCondition = "bEnableChildActorExpansionInTreeView"))
+	EChildActorComponentTreeViewVisualizationMode DefaultChildActorTreeViewMode;
+
+	// The list of namespaces to always expose in any Blueprint (for all users of the game/project). Requires Blueprint namespace features to be enabled in editor preferences.
+	UPROPERTY(EditAnywhere, config, Category = Experimental)
+	TArray<FString> NamespacesToAlwaysInclude;
+
 	/** 
 	 * List of compiler messages that have been suppressed outside of full, interactive editor sessions for 
 	 * the current project - useful for silencing warnings that were added to the engine after 
@@ -50,16 +60,6 @@ public:
 	UPROPERTY(EditAnywhere, config, Category= Blueprints, DisplayName = "Compiler Messages Disabled Entirely")
 	TArray<FName> DisabledCompilerMessages;
 
-	// The list of namespaces to always expose in any Blueprint (for all users of the game/project)
-	UPROPERTY(EditAnywhere, config, Category=Experimental)
-	TArray<FString> NamespacesToAlwaysInclude;
-	
-	/**
-	 * Default view mode to use for child actor components in a Blueprint actor's component tree hierarchy (experimental).
-	 */
-	UPROPERTY(EditAnywhere, config, Category=Experimental, meta=(EditCondition="bEnableChildActorExpansionInTreeView"))
-	EChildActorComponentTreeViewVisualizationMode DefaultChildActorTreeViewMode;
-
 	/**
 	 * Any blueprint deriving from one of these base classes will be allowed to recompile during Play-in-Editor
 	 * (This setting exists both as an editor preference and project setting, and will be allowed if listed in either place) 
