TraceScopes for Cooking: Add scopes for PreSave, SaveSerialize, BeginCacheForCookedPlatformData, and IsCachedCookedPlatformDataLoaded.
Rename the LoadChannel version of the Serialize scope to LoadSerialize.
Lump PreSaveRoot and PostSave into PreSave since they don't deserve a separate scope.
Change existing classname-parameterized scopes to use WriteToString instead of FStrings to remove an allocation.
#rnx

diff --git a/Engine/Source/Editor/UnrealEd/Private/CookOnTheFlyServer.cpp b/Engine/Source/Editor/UnrealEd/Private/CookOnTheFlyServer.cpp
--- a/Engine/Source/Editor/UnrealEd/Private/CookOnTheFlyServer.cpp
+++ b/Engine/Source/Editor/UnrealEd/Private/CookOnTheFlyServer.cpp
@@ -2534,9 +2534,9 @@ bool UCookOnTheFlyServer::BeginPrepareSave(UE::Cook::FPackageData& PackageData,
 			*CurrentAsyncCache -= 1;
 		}
 
-		Obj->BeginCacheForCookedPlatformData(TargetPlatform);
+		RouteBeginCacheForCookedPlatformData(Obj, TargetPlatform);
 		++CookedPlatformDataNextIndex;
-		if (Obj->IsCachedCookedPlatformDataLoaded(TargetPlatform))
+		if (RouteIsCachedCookedPlatformDataLoaded(Obj, TargetPlatform))
 		{
 			if (CurrentAsyncCache)
 			{
@@ -3258,9 +3258,9 @@ void UCookOnTheFlyServer::TickPrecacheObjectsForPlatforms(const float TimeSlice,
 			{
 				continue;
 			}
-			if (!Material->IsCachedCookedPlatformDataLoaded(TargetPlatform))
+			if (!RouteIsCachedCookedPlatformDataLoaded(Material, TargetPlatform))
 			{
-				Material->BeginCacheForCookedPlatformData(TargetPlatform);
+				RouteBeginCacheForCookedPlatformData(Material, TargetPlatform);
 				AllMaterialsCompiled = false;
 			}
 		}
@@ -3299,9 +3299,9 @@ void UCookOnTheFlyServer::TickPrecacheObjectsForPlatforms(const float TimeSlice,
 			{
 				continue;
 			}
-			if (!Texture->IsCachedCookedPlatformDataLoaded(TargetPlatform))
+			if (!RouteIsCachedCookedPlatformDataLoaded(Texture, TargetPlatform))
 			{
-				Texture->BeginCacheForCookedPlatformData(TargetPlatform);
+				RouteBeginCacheForCookedPlatformData(Texture, TargetPlatform);
 			}
 		}
 		if (Timer.IsTimeUp())
@@ -8422,8 +8422,8 @@ uint32 UCookOnTheFlyServer::FullLoadAndSave(uint32& CookedPackageCount)
 								{
 									UE_SCOPED_HIERARCHICAL_COOKTIMER(FullLoadAndSave_BeginCache);
 									UE_TRACK_REFERENCING_PACKAGE_SCOPED(Package, PackageAccessTrackingOps::NAME_CookerBuildObject);
-									Obj->BeginCacheForCookedPlatformData(TargetPlatform);
-									if (!Obj->IsCachedCookedPlatformDataLoaded(TargetPlatform))
+									RouteBeginCacheForCookedPlatformData(Obj, TargetPlatform);
+									if (!RouteIsCachedCookedPlatformDataLoaded(Obj, TargetPlatform))
 									{
 										bAllPlatformDataLoaded = false;
 									}
@@ -8543,7 +8543,7 @@ uint32 UCookOnTheFlyServer::FullLoadAndSave(uint32& CookedPackageCount)
 				UE_TRACK_REFERENCING_PACKAGE_SCOPED(Obj->GetPackage(), PackageAccessTrackingOps::NAME_CookerBuildObject);
 				for (const ITargetPlatform* TargetPlatform : TargetPlatforms)
 				{
-					if (!Obj->IsCachedCookedPlatformDataLoaded(TargetPlatform))
+					if (!RouteIsCachedCookedPlatformDataLoaded(Obj, TargetPlatform))
 					{
 						bAllPlatformDataLoaded = false;
 						break;
@@ -8612,7 +8612,7 @@ uint32 UCookOnTheFlyServer::FullLoadAndSave(uint32& CookedPackageCount)
 
 						for (const ITargetPlatform* TargetPlatform : TargetPlatforms)
 						{
-							Obj->BeginCacheForCookedPlatformData(TargetPlatform);
+							RouteBeginCacheForCookedPlatformData(Obj, TargetPlatform);
 						}
 					}
 				}
diff --git a/Engine/Source/Editor/UnrealEd/Private/Cooker/CookPackageData.cpp b/Engine/Source/Editor/UnrealEd/Private/Cooker/CookPackageData.cpp
--- a/Engine/Source/Editor/UnrealEd/Private/Cooker/CookPackageData.cpp
+++ b/Engine/Source/Editor/UnrealEd/Private/Cooker/CookPackageData.cpp
@@ -1375,7 +1375,7 @@ bool FPendingCookedPlatformData::PollIsComplete()
 		return true;
 	}
 	UE_TRACK_REFERENCING_PACKAGE_SCOPED(LocalObject->GetPackage(), PackageAccessTrackingOps::NAME_CookerBuildObject);
-	if (LocalObject->IsCachedCookedPlatformDataLoaded(TargetPlatform))
+	if (RouteIsCachedCookedPlatformDataLoaded(LocalObject, TargetPlatform))
 	{
 		Release();
 		return true;
diff --git a/Engine/Source/Editor/UnrealEd/Private/Cooker/CookProfiling.h b/Engine/Source/Editor/UnrealEd/Private/Cooker/CookProfiling.h
--- a/Engine/Source/Editor/UnrealEd/Private/Cooker/CookProfiling.h
+++ b/Engine/Source/Editor/UnrealEd/Private/Cooker/CookProfiling.h
@@ -40,36 +40,70 @@ private:
 
 UE_TRACE_CHANNEL_EXTERN(CookChannel)
 
-#define UE_CREATE_HIERARCHICAL_COOKTIMER(name, incrementScope)	FScopeTimer ScopeTimer##name(__COUNTER__, #name, incrementScope); 
+#define UE_CREATE_HIERARCHICAL_COOKTIMER(name, incrementScope) \
+			FScopeTimer PREPROCESSOR_JOIN(__HierarchicalCookTimerScope, __LINE__)(__COUNTER__, #name, incrementScope); 
+#define UE_CREATE_TEXT_HIERARCHICAL_COOKTIMER(name, incrementScope) \
+			FScopeTimer PREPROCESSOR_JOIN(__HierarchicalCookTimerScope, __LINE__)(__COUNTER__, name, incrementScope);
 
 // Emits trace events denoting scope/lifetime of an activity on the cooking channel.
 #define UE_SCOPED_COOKTIMER(name)						TRACE_CPUPROFILER_EVENT_SCOPE_ON_CHANNEL(name, CookChannel)
+// More expensive version that takes a dynamic TCHAR*
+#define UE_SCOPED_TEXT_COOKTIMER(name)					TRACE_CPUPROFILER_EVENT_SCOPE_TEXT_ON_CHANNEL(name, CookChannel)
+
 #define UE_SCOPED_COOKTIMER_AND_DURATION(name, durationStorage) \
-			FScopedDurationTimer name##Timer(durationStorage); UE_SCOPED_COOKTIMER(name)
+			FScopedDurationTimer PREPROCESSOR_JOIN(__CookTimerScope, __LINE__)(durationStorage); \
+			UE_SCOPED_COOKTIMER(name)
+// More expensive version that takes a dynamic TCHAR*
+#define UE_SCOPED_TEXT_COOKTIMER_AND_DURATION(name, durationStorage) \
+			FScopedDurationTimer PREPROCESSOR_JOIN(__CookTimerScope, __LINE__)(durationStorage); \
+			UE_SCOPED_TEXT_COOKTIMER(name)
 
 // Emits trace events denoting scope/lifetime of an activity on the cooking channel.  Also creates a named hierarchical timer that can be aggregated and reported at cook completion.
-#define UE_SCOPED_HIERARCHICAL_COOKTIMER(name)			UE_CREATE_HIERARCHICAL_COOKTIMER(name, true); ScopeTimer##name.Start(); UE_SCOPED_COOKTIMER(name)
+#define UE_SCOPED_HIERARCHICAL_COOKTIMER(name) \
+			UE_CREATE_HIERARCHICAL_COOKTIMER(name, true); \
+			PREPROCESSOR_JOIN(__HierarchicalCookTimerScope, __LINE__).Start(); \
+			UE_SCOPED_COOKTIMER(name)
+// More expensive version that takes a dynamic TCHAR*
+#define UE_SCOPED_TEXT_HIERARCHICAL_COOKTIMER(name) \
+			UE_CREATE_TEXT_HIERARCHICAL_COOKTIMER(name, true); \
+			PREPROCESSOR_JOIN(__HierarchicalCookTimerScope, __LINE__).Start(); \
+			UE_SCOPED_TEXT_COOKTIMER(name)
+// Version that takes a separate variable for where the duration should be stored
 #define UE_SCOPED_HIERARCHICAL_COOKTIMER_AND_DURATION(name, durationStorage) \
-			FScopedDurationTimer name##Timer(durationStorage); UE_SCOPED_HIERARCHICAL_COOKTIMER(name)
+			FScopedDurationTimer PREPROCESSOR_JOIN(__CookTimerScope, __LINE__)(durationStorage); \
+			UE_SCOPED_HIERARCHICAL_COOKTIMER(name)
+// More expensive version that takes a separate variable for where the duration should be stored and a dynamic TCHAR*
+#define UE_SCOPED_TEXT_HIERARCHICAL_COOKTIMER_AND_DURATION(name, durationStorage) \
+			FScopedDurationTimer PREPROCESSOR_JOIN(__CookTimerScope, __LINE__)(durationStorage); \
+			UE_SCOPED_TEXT_HIERARCHICAL_COOKTIMER(name)
 
 #define UE_CUSTOM_COOKTIMER_LOG Cpu
 
 // Emits trace events denoting scope/lifetime of an activity on the cooking channel.
 #define UE_SCOPED_CUSTOM_COOKTIMER(name)				UE_TRACE_LOG_SCOPED_T(UE_CUSTOM_COOKTIMER_LOG, name, CookChannel)
 #define UE_SCOPED_CUSTOM_COOKTIMER_AND_DURATION(name, durationStorage) \
-			FScopedDurationTimer name##Timer(durationStorage); UE_SCOPED_CUSTOM_COOKTIMER(name)
+			FScopedDurationTimer PREPROCESSOR_JOIN(__CookTimerScope, __LINE__)(durationStorage); \
+			UE_SCOPED_CUSTOM_COOKTIMER(name)
 
 // Emits trace events denoting scope/lifetime of an activity on the cooking channel.  Also creates a named hierarchical timer that can be aggregated and reported at cook completion.
-#define UE_SCOPED_HIERARCHICAL_CUSTOM_COOKTIMER(name)	UE_CREATE_HIERARCHICAL_COOKTIMER(name, true); ScopeTimer##name.Start(); UE_SCOPED_CUSTOM_COOKTIMER(name)
+#define UE_SCOPED_HIERARCHICAL_CUSTOM_COOKTIMER(name) \
+			UE_CREATE_HIERARCHICAL_COOKTIMER(name, true); \
+			PREPROCESSOR_JOIN(__HierarchicalCookTimerScope, __LINE__).Start(); \
+			UE_SCOPED_CUSTOM_COOKTIMER(name)
 #define UE_SCOPED_HIERARCHICAL_CUSTOM_COOKTIMER_AND_DURATION(name, durationStorage) \
-			FScopedDurationTimer name##Timer(durationStorage); UE_SCOPED_HIERARCHICAL_CUSTOM_COOKTIMER(name)
+			FScopedDurationTimer PREPROCESSOR_JOIN(__CookTimerScope, __LINE__)(durationStorage); \
+			UE_SCOPED_HIERARCHICAL_CUSTOM_COOKTIMER(name)
 #define UE_ADD_CUSTOM_COOKTIMER_META(name, key, value) << name.key(value)
 
 #else
 #define UE_SCOPED_COOKTIMER(name)
+#define UE_SCOPED_TEXT_COOKTIMER(name)
 #define UE_SCOPED_COOKTIMER_AND_DURATION(name, durationStorage)
+#define UE_SCOPED_TEXT_COOKTIMER_AND_DURATION(name, durationStorage)
 #define UE_SCOPED_HIERARCHICAL_COOKTIMER(name)
+#define UE_SCOPED_TEXT_HIERARCHICAL_COOKTIMER(name)
 #define UE_SCOPED_HIERARCHICAL_COOKTIMER_AND_DURATION(name, durationStorage)
+#define UE_SCOPED_TEXT_HIERARCHICAL_COOKTIMER_AND_DURATION(name, durationStorage)
 
 #define UE_CUSTOM_COOKTIMER_LOG Cpu
 #define UE_SCOPED_CUSTOM_COOKTIMER(name)
diff --git a/Engine/Source/Editor/UnrealEd/Private/Cooker/CookTypes.h b/Engine/Source/Editor/UnrealEd/Private/Cooker/CookTypes.h
--- a/Engine/Source/Editor/UnrealEd/Private/Cooker/CookTypes.h
+++ b/Engine/Source/Editor/UnrealEd/Private/Cooker/CookTypes.h
@@ -5,6 +5,7 @@
 #include "Containers/Map.h"
 #include "Containers/Set.h"
 #include "Containers/UnrealString.h"
+#include "Cooker/CookProfiling.h"
 #include "CookOnTheSide/CookOnTheFlyServer.h" // ECookTickFlags
 #include "DerivedDataRequestOwner.h"
 #include "HAL/LowLevelMemTracker.h"
@@ -212,6 +213,18 @@ namespace UE::Cook
 	};
 }
 
+inline void RouteBeginCacheForCookedPlatformData(UObject* Obj, const ITargetPlatform* TargetPlatform)
+{
+	UE_SCOPED_TEXT_COOKTIMER(*WriteToString<128>(Obj->GetClass()->GetFName(), TEXT("_BeginCacheForCookedPlatformData")));
+	Obj->BeginCacheForCookedPlatformData(TargetPlatform);
+}
+
+inline bool RouteIsCachedCookedPlatformDataLoaded(UObject* Obj, const ITargetPlatform* TargetPlatform)
+{
+	UE_SCOPED_TEXT_COOKTIMER(*WriteToString<128>(Obj->GetClass()->GetFName(), TEXT("_IsCachedCookedPlatformDataLoaded")));
+	return Obj->IsCachedCookedPlatformDataLoaded(TargetPlatform);
+}
+
 //////////////////////////////////////////////////////////////////////////
 // Cook by the book options
 
diff --git a/Engine/Source/Runtime/CoreUObject/Private/UObject/LinkerLoad.cpp b/Engine/Source/Runtime/CoreUObject/Private/UObject/LinkerLoad.cpp
--- a/Engine/Source/Runtime/CoreUObject/Private/UObject/LinkerLoad.cpp
+++ b/Engine/Source/Runtime/CoreUObject/Private/UObject/LinkerLoad.cpp
@@ -4217,7 +4217,9 @@ void FLinkerLoad::Preload( UObject* Object )
 #endif
 						check(CurrentLoadContext);
 #if WITH_EDITOR
-						SCOPED_LOADTIMER_TEXT(*(Object->GetClass()->GetName() + TEXT("_Serialize")));
+						SCOPED_LOADTIMER_TEXT(*WriteToString<128>(Object->GetClass()->GetFName(), TEXTVIEW("_LoadSerialize")));
+						TRACE_CPUPROFILER_EVENT_SCOPE_TEXT(*WriteToString<128>(Object->GetClass()->GetFName(), TEXTVIEW("_Serialize")));
+						TRACE_CPUPROFILER_EVENT_SCOPE_TEXT(*Object->GetFullName());
 #endif
 
 						// Maintain the current SerializedObjects.
@@ -4243,10 +4245,6 @@ void FLinkerLoad::Preload( UObject* Object )
 						else
 #endif
 						{
-#if WITH_EDITOR
-							TRACE_CPUPROFILER_EVENT_SCOPE_TEXT(*(Object->GetClass()->GetName() + TEXT("_Serialize")));
-							TRACE_CPUPROFILER_EVENT_SCOPE_TEXT(*Object->GetFullName());
-#endif
 							Object->Serialize(*this);
 						}
 
diff --git a/Engine/Source/Runtime/CoreUObject/Private/UObject/Obj.cpp b/Engine/Source/Runtime/CoreUObject/Private/UObject/Obj.cpp
--- a/Engine/Source/Runtime/CoreUObject/Private/UObject/Obj.cpp
+++ b/Engine/Source/Runtime/CoreUObject/Private/UObject/Obj.cpp
@@ -1095,7 +1095,7 @@ void UObject::ConditionalPostLoad()
 			else
 			{
 #if WITH_EDITOR
-				SCOPED_LOADTIMER_TEXT(*(GetClass()->GetName() + TEXT("_PostLoad")));
+				SCOPED_LOADTIMER_TEXT(*WriteToString<128>(GetClass()->GetFName(), TEXTVIEW("_PostLoad")));
 #endif
 				LLM_SCOPED_TAG_WITH_OBJECT_IN_SET(GetOutermost(), ELLMTagSet::Assets);
 				LLM_SCOPED_TAG_WITH_OBJECT_IN_SET(GetClass(), ELLMTagSet::AssetClasses);
diff --git a/Engine/Source/Runtime/CoreUObject/Private/UObject/SavePackage/PackageHarvester.cpp b/Engine/Source/Runtime/CoreUObject/Private/UObject/SavePackage/PackageHarvester.cpp
--- a/Engine/Source/Runtime/CoreUObject/Private/UObject/SavePackage/PackageHarvester.cpp
+++ b/Engine/Source/Runtime/CoreUObject/Private/UObject/SavePackage/PackageHarvester.cpp
@@ -248,7 +248,10 @@ void FPackageHarvester::ProcessExport(const FExportWithContext& InProcessContext
 
 	// In the CDO case the above would serialize most of the references, including transient properties
 	// but we still want to serialize the object using the normal path to collect all custom versions it might be using.
-	Export->Serialize(*this);
+	{
+		SCOPED_SAVETIMER_TEXT(*WriteToString<128>(Export->GetClass()->GetFName(), TEXT("_SaveSerialize")));
+		Export->Serialize(*this);
+	}
 
 	// Gather object preload dependencies
 	if (SaveContext.IsCooking())
diff --git a/Engine/Source/Runtime/CoreUObject/Private/UObject/SavePackage/SavePackageUtilities.cpp b/Engine/Source/Runtime/CoreUObject/Private/UObject/SavePackage/SavePackageUtilities.cpp
--- a/Engine/Source/Runtime/CoreUObject/Private/UObject/SavePackage/SavePackageUtilities.cpp
+++ b/Engine/Source/Runtime/CoreUObject/Private/UObject/SavePackage/SavePackageUtilities.cpp
@@ -852,6 +852,7 @@ bool IsProceduralSave(bool bIsCooking, const FPackagePath& TargetPackagePath, ui
 
 void CallPreSave(UObject* Object, FObjectSaveContextData& ObjectSaveContext)
 {
+	SCOPED_SAVETIMER_TEXT(*WriteToString<256>(Object->GetClass()->GetFName(), TEXTVIEW("_PreSave")));
 	PRAGMA_DISABLE_DEPRECATION_WARNINGS;
 	Object->PreSave(ObjectSaveContext.TargetPlatform);
 	PRAGMA_ENABLE_DEPRECATION_WARNINGS;
@@ -874,6 +875,7 @@ void CallPreSave(UObject* Object, FObjectSaveContextData& ObjectSaveContext)
 
 void CallPreSaveRoot(UObject* Object, FObjectSaveContextData& ObjectSaveContext)
 {
+	SCOPED_SAVETIMER_TEXT(*WriteToString<256>(Object->GetClass()->GetFName(), TEXTVIEW("_PreSave")));
 	PRAGMA_DISABLE_DEPRECATION_WARNINGS;
 	bool bLegacyNeedsCleanup = Object->PreSaveRoot(*ObjectSaveContext.TargetFilename);
 	PRAGMA_ENABLE_DEPRECATION_WARNINGS;
@@ -885,6 +887,7 @@ void CallPreSaveRoot(UObject* Object, FObjectSaveContextData& ObjectSaveContext)
 
 void CallPostSaveRoot(UObject* Object, FObjectSaveContextData& ObjectSaveContext, bool bNeedsCleanup)
 {
+	SCOPED_SAVETIMER_TEXT(*WriteToString<256>(Object->GetClass()->GetFName(), TEXTVIEW("_PreSave")));
 	PRAGMA_DISABLE_DEPRECATION_WARNINGS;
 	Object->PostSaveRoot(bNeedsCleanup);
 	PRAGMA_ENABLE_DEPRECATION_WARNINGS;
diff --git a/Engine/Source/Runtime/CoreUObject/Private/UObject/SavePackage/SavePackageUtilities.h b/Engine/Source/Runtime/CoreUObject/Private/UObject/SavePackage/SavePackageUtilities.h
--- a/Engine/Source/Runtime/CoreUObject/Private/UObject/SavePackage/SavePackageUtilities.h
+++ b/Engine/Source/Runtime/CoreUObject/Private/UObject/SavePackage/SavePackageUtilities.h
@@ -30,8 +30,10 @@ enum class ESavePackageResult;
 #if UE_TRACE_ENABLED && !UE_BUILD_SHIPPING
 UE_TRACE_CHANNEL_EXTERN(SaveTimeChannel)
 #define SCOPED_SAVETIMER(TimerName) TRACE_CPUPROFILER_EVENT_SCOPE_ON_CHANNEL(TimerName, SaveTimeChannel)
+#define SCOPED_SAVETIMER_TEXT(TimerName) TRACE_CPUPROFILER_EVENT_SCOPE_TEXT_ON_CHANNEL(TimerName, SaveTimeChannel)
 #else
 #define SCOPED_SAVETIMER(TimerName)
+#define SCOPED_SAVETIMER_TEXT(TimerName)
 #endif
 
 struct FLargeMemoryDelete
diff --git a/Engine/Source/Runtime/CoreUObject/Private/UObject/SavePackage2.cpp b/Engine/Source/Runtime/CoreUObject/Private/UObject/SavePackage2.cpp
--- a/Engine/Source/Runtime/CoreUObject/Private/UObject/SavePackage2.cpp
+++ b/Engine/Source/Runtime/CoreUObject/Private/UObject/SavePackage2.cpp
@@ -1720,6 +1720,7 @@ ESavePackageResult WriteExports(FStructuredArchive::FRecord& StructuredArchiveRo
 			bool bSupportsText = false;
 #endif
 
+			SCOPED_SAVETIMER_TEXT(*WriteToString<128>(Export.Object->GetClass()->GetFName(), TEXT("_SaveSerialize")));
 			if (Export.Object->HasAnyFlags(RF_ClassDefaultObject))
 			{
 				if (bSupportsText)
