Add trace scope and metadata for TrySavePackage in the EditorDomain code.

diff --git a/Engine/Source/Editor/UnrealEd/Private/EditorDomain/EditorDomainUtils.cpp b/Engine/Source/Editor/UnrealEd/Private/EditorDomain/EditorDomainUtils.cpp
--- a/Engine/Source/Editor/UnrealEd/Private/EditorDomain/EditorDomainUtils.cpp
+++ b/Engine/Source/Editor/UnrealEd/Private/EditorDomain/EditorDomainUtils.cpp
@@ -31,6 +31,8 @@
 #include "String/ParseTokens.h"
 #include "TargetDomain/TargetDomainUtils.h"
 #include "Templates/Function.h"
+#include "Trace/Trace.h"
+#include "Trace/Trace.inl"
 #include "UObject/CoreRedirects.h"
 #include "UObject/ObjectVersion.h"
 #include "UObject/Package.h"
@@ -39,6 +41,35 @@
 #include "UObject/UObjectHash.h"
 #include "UObject/UE5MainStreamObjectVersion.h"
 
+
+#if !defined(EDITORDOMAINTIMEPROFILERTRACE_ENABLED)
+#if UE_TRACE_ENABLED && !UE_BUILD_SHIPPING
+#define EDITORDOMAINTIMEPROFILERTRACE_ENABLED 1
+#else
+#define EDITORDOMAINTIMEPROFILERTRACE_ENABLED 0
+#endif
+#endif
+
+#define CUSTOM_EDITORDOMAINTIMER_LOG Cpu
+
+#if EDITORDOMAINTIMEPROFILERTRACE_ENABLED
+#define SCOPED_EDITORDOMAINTIMER_TEXT(TimerName) TRACE_CPUPROFILER_EVENT_SCOPE_TEXT_ON_CHANNEL(TimerName, CpuChannel)
+#define SCOPED_EDITORDOMAINTIMER(TimerName) TRACE_CPUPROFILER_EVENT_SCOPE_ON_CHANNEL(TimerName, CpuChannel)
+#define SCOPED_CUSTOM_EDITORDOMAINTIMER(TimerName) UE_TRACE_LOG_SCOPED_T(CUSTOM_EDITORDOMAINTIMER_LOG, TimerName, CpuChannel)
+#define ADD_CUSTOM_EDITORDOMAINTIMER_META(TimerName, Key, Value) << TimerName.Key(Value)
+#define SCOPED_EDITORDOMAINTIMER_CNT(TimerName)
+#else
+#define SCOPED_EDITORDOMAINTIMER_TEXT(TimerName)
+#define SCOPED_EDITORDOMAINTIMER(TimerName)
+#define SCOPED_CUSTOM_EDITORDOMAINTIMER(TimerName)
+#define ADD_CUSTOM_EDITORDOMAINTIMER_META(TimerName, Key, Value)
+#define SCOPED_EDITORDOMAINTIMER_CNT(TimerName)
+#endif
+
+UE_TRACE_EVENT_BEGIN(CUSTOM_EDITORDOMAINTIMER_LOG, EditorDomain_TrySavePackage, NoSync)
+UE_TRACE_EVENT_FIELD(UE::Trace::WideString, PackageName)
+UE_TRACE_EVENT_END()
+
 /** Modify the masked bits in the output: set them to A & B. */
 template<typename Enum>
 static void EnumSetFlagsAnd(Enum& Output, Enum Mask, Enum A, Enum B)
@@ -1497,6 +1528,9 @@ bool TrySavePackage(UPackage* Package)
 	}
 	UE_LOG(LogEditorDomain, Verbose, TEXT("Saving to EditorDomain: %s."), *WriteToString<256>(PackageName));
 
+	SCOPED_CUSTOM_EDITORDOMAINTIMER(EditorDomain_TrySavePackage)
+		ADD_CUSTOM_EDITORDOMAINTIMER_META(EditorDomain_TrySavePackage, PackageName, *WriteToString<256>(PackageName));
+
 	uint32 SaveFlags = SAVE_NoError // Do not crash the SaveServer on an error
 		| SAVE_BulkDataByReference	// EditorDomain saves reference bulkdata from the WorkspaceDomain rather than duplicating it
 		| SAVE_Async				// SavePackage support for PackageWriter is only implemented with SAVE_Async
