Fix NoSaveDevAR option also keeping AllChunksInfo.csv from being written.

diff --git a/Engine/Source/Editor/UnrealEd/Private/Commandlets/AssetRegistryGenerator.cpp b/Engine/Source/Editor/UnrealEd/Private/Commandlets/AssetRegistryGenerator.cpp
--- a/Engine/Source/Editor/UnrealEd/Private/Commandlets/AssetRegistryGenerator.cpp
+++ b/Engine/Source/Editor/UnrealEd/Private/Commandlets/AssetRegistryGenerator.cpp
@@ -1402,20 +1402,24 @@ bool FAssetRegistryGenerator::SaveAssetRegistry(const FString& SandboxPath, bool
 	UpdateKeptPackages();
 	UpdateCollectionAssetData();
 
-	if (DevelopmentSaveOptions.bSerializeAssetRegistry && bSerializeDevelopmentAssetRegistry)
+	if (DevelopmentSaveOptions.bSerializeAssetRegistry)
 	{
-		// Create development registry data, used for DLC cooks, iterative cooks, and editor viewing
-		FArrayWriter SerializedAssetRegistry;
-
-		State.Save(SerializedAssetRegistry, DevelopmentSaveOptions);
-
-		// Save the generated registry
 		FString PlatformSandboxPath = SandboxPath.Replace(TEXT("[Platform]"), *TargetPlatform->PlatformName());
 		const TCHAR* DevelopmentAssetRegistryFilename = GetDevelopmentAssetRegistryFilename();
 		PlatformSandboxPath.ReplaceInline(TEXT("AssetRegistry.bin"), *FString::Printf(TEXT("Metadata/%s"), DevelopmentAssetRegistryFilename));
-		FFileHelper::SaveArrayToFile(SerializedAssetRegistry, *PlatformSandboxPath);
 
-		UE_LOG(LogAssetRegistryGenerator, Display, TEXT("Generated development asset registry %s num assets %d, size is %5.2fkb"), *PlatformSandboxPath, State.GetNumAssets(), (float)SerializedAssetRegistry.Num() / 1024.f);
+		if (bSerializeDevelopmentAssetRegistry)
+		{
+			// Create development registry data, used for DLC cooks, iterative cooks, and editor viewing
+			FArrayWriter SerializedAssetRegistry;
+
+			State.Save(SerializedAssetRegistry, DevelopmentSaveOptions);
+
+			// Save the generated registry
+			FFileHelper::SaveArrayToFile(SerializedAssetRegistry, *PlatformSandboxPath);
+
+			UE_LOG(LogAssetRegistryGenerator, Display, TEXT("Generated development asset registry %s num assets %d, size is %5.2fkb"), *PlatformSandboxPath, State.GetNumAssets(), (float)SerializedAssetRegistry.Num() / 1024.f);
+		}
 
 		if (bGenerateChunks && bUseAssetManager)
 		{
