Back out changelist 19000069 to address CIS error.

diff --git a/Engine/Source/Editor/UnrealEd/Private/CookOnTheFlyServer.cpp b/Engine/Source/Editor/UnrealEd/Private/CookOnTheFlyServer.cpp
--- a/Engine/Source/Editor/UnrealEd/Private/CookOnTheFlyServer.cpp
+++ b/Engine/Source/Editor/UnrealEd/Private/CookOnTheFlyServer.cpp
@@ -7710,43 +7710,6 @@ void UCookOnTheFlyServer::StartCookByTheBook( const FCookByTheBookStartupOptions
 			}
 		}
 
-		FString ExtraReleaseVersionAssetsFile;
-		const bool bUsingExtraReleaseVersionAssets = FParse::Value(FCommandLine::Get(), TEXT("ExtraReleaseVersionAssets="), ExtraReleaseVersionAssetsFile);
-		if (bUsingExtraReleaseVersionAssets)
-		{
-			// read the file 
-			TArray<FString> OutAssetPaths;
-			FFileHelper::LoadFileToStringArray(OutAssetPaths, *ExtraReleaseVersionAssetsFile);
-
-			TArray<TTuple<FName, FPackageNameCache::FCachedPackageFilename>> PackageToStandardFileNames;
-			for (const FString& AssetPath : OutAssetPaths)
-			{
-				// register with package cache
-				FName AssetPathFName(*AssetPath);
-				if (UE::Cook::FPackageData* PackageData = PackageDatas->TryAddPackageDataByFileName(AssetPathFName))
-				{
-					// build standard filename
-					FString StandardPackageFilename = AssetPath;
-					FPaths::MakeStandardFilename(StandardPackageFilename);
-
-					// add to override list
-					FName StandardPackageFileFName(*StandardPackageFilename);
-					OverridePackageList.Add(StandardPackageFileFName);
-
-					PackageToStandardFileNames.Add(
-						TTuple<FName, FPackageNameCache::FCachedPackageFilename>(
-							PackageData->GetPackageName(),
-							FPackageNameCache::FCachedPackageFilename(MoveTemp(StandardPackageFilename), StandardPackageFileFName)));
-				}
-				else
-				{
-					UE_LOG(LogCook, Error, TEXT("Failed to resolve package data for package [%s]"), *AssetPathFName.ToString());
-				}
-			}
-			// Add these packages to the cache such that they're resolvable later
-			GetPackageNameCache().AppendCacheResults(MoveTemp(PackageToStandardFileNames));
-		}
-
 		for ( const ITargetPlatform* TargetPlatform: TargetPlatforms )
 		{
 			SCOPED_BOOT_TIMING("AddCookedPlatforms");
