COTF: Fixed issue with GC running too often
#rnx

diff --git a/Engine/Source/Editor/UnrealEd/Private/CookOnTheFlyServer.cpp b/Engine/Source/Editor/UnrealEd/Private/CookOnTheFlyServer.cpp
--- a/Engine/Source/Editor/UnrealEd/Private/CookOnTheFlyServer.cpp
+++ b/Engine/Source/Editor/UnrealEd/Private/CookOnTheFlyServer.cpp
@@ -1273,11 +1273,6 @@ uint32 UCookOnTheFlyServer::TickCookOnTheSide(const float TimeSlice, uint32 &Coo
 		CookByTheBookFinished();
 	}
 
-	if (IsCookOnTheFlyMode() && bCookComplete)
-	{
-		CollectGarbage(RF_NoFlags);
-	}
-
 	CookedPackageCount += StackData.CookedPackageCount;
 	return StackData.ResultFlags;
 }
@@ -3571,7 +3566,15 @@ uint32 UCookOnTheFlyServer::GetPackagesPerPartialGC() const
 
 double UCookOnTheFlyServer::GetIdleTimeToGC() const
 {
-	return IdleTimeToGC;
+	if (CurrentCookMode == ECookMode::CookOnTheFly)
+	{
+		// For COTF outside of the editor we want to release open linker file handles promptly but still give some time for new requests to come in
+		return 0.5;
+	}
+	else
+	{
+		return IdleTimeToGC;
+	}
 }
 
 uint64 UCookOnTheFlyServer::GetMaxMemoryAllowance() const
