Allow script execution to bypass development-only nodes in a Blueprint function graph when a project is configured to compile them out at cook time.

This is a safer fix that moves disabled graph node pruning into the CloneGraph() utility API, which covers all graph types during compilation, unlike my first attempt, which changed the underlying structure of how function graphs are compiled to more closely resemble event graph compilation, but unfortunately also broke compilation of consolidated subgraphs as a result.

diff --git a/Engine/Source/Editor/UnrealEd/Private/EdGraphUtilities.cpp b/Engine/Source/Editor/UnrealEd/Private/EdGraphUtilities.cpp
--- a/Engine/Source/Editor/UnrealEd/Private/EdGraphUtilities.cpp
+++ b/Engine/Source/Editor/UnrealEd/Private/EdGraphUtilities.cpp
@@ -16,7 +16,7 @@
 #include "K2Node_Composite.h"
 
 DECLARE_CYCLE_STAT(TEXT("Clone Graph"), EKismetCompilerStats_CloneGraph, STATGROUP_KismetCompiler);
-DECLARE_CYCLE_STAT(TEXT("Clone Graph - Build Backtrack Map"), EKismetCompilerStats_CloneGraph_BuildBackTrackMap, STATGROUP_KismetCompiler);
+DECLARE_CYCLE_STAT(TEXT("Clone Graph - Post Process"), EKismetCompilerStats_CloneGraph_PostProcess, STATGROUP_KismetCompiler);
 
 /////////////////////////////////////////////////////
 // Local namespace
@@ -246,11 +246,11 @@ UEdGraph* FEdGraphUtilities::CloneGraph(UEdGraph* InSource, UObject* NewOuter, F
 	// Duplicate the graph, keeping track of what was duplicated
 	TMap<UObject*, UObject*> DuplicatedObjectList;
 
-	UObject* UseOuter = (NewOuter != NULL) ? NewOuter : GetTransientPackage();
+	UObject* UseOuter = (NewOuter != nullptr) ? NewOuter : GetTransientPackage();
 	FObjectDuplicationParameters Parameters(InSource, UseOuter);
 	Parameters.CreatedObjects = &DuplicatedObjectList;
 
-	if (bCloningForCompile || (NewOuter == NULL))
+	if (bCloningForCompile || (NewOuter == nullptr))
 	{
 		Parameters.ApplyFlags |= RF_Transient;
 		Parameters.FlagMask &= ~RF_Transactional;
@@ -258,38 +258,82 @@ UEdGraph* FEdGraphUtilities::CloneGraph(UEdGraph* InSource, UObject* NewOuter, F
 
 	UEdGraph* ClonedGraph = CastChecked<UEdGraph>(StaticDuplicateObjectEx(Parameters));
 
-	// Store backtrack links from each duplicated object to the original source object
-	if (MessageLog != NULL)
+	// During compilation, do not clone disabled (e.g. ghost) nodes from non-transient source graphs.
+	const bool bExcludeDisabledNodes = bCloningForCompile && InSource && !InSource->HasAnyFlags(RF_Transient);
+
+	// Exclude disabled nodes and/or store backtrack links from each duplicated object to the original source object.
+	if (bExcludeDisabledNodes || MessageLog != nullptr)
 	{
-		BP_SCOPED_COMPILER_EVENT_STAT(EKismetCompilerStats_CloneGraph_BuildBackTrackMap);
+		BP_SCOPED_COMPILER_EVENT_STAT(EKismetCompilerStats_CloneGraph_PostProcess);
 		
 		for (TMap<UObject*, UObject*>::TIterator It(DuplicatedObjectList); It; ++It)
 		{
+			bool bIsExcludedNode = false;
+
 			UObject* const Source = It.Key();
 			UObject* const Dest = It.Value();
 
-			MessageLog->NotifyIntermediateObjectCreation(Dest, Source);
-
 			UEdGraphNode* SrcNode = Cast<UEdGraphNode>(Source);
 			UEdGraphNode* DstNode = Cast<UEdGraphNode>(Dest);
 			if (SrcNode && DstNode)
 			{
+				// Determine whether the source is a disabled node that should be excluded from the cloned graph.
+				bIsExcludedNode = bExcludeDisabledNodes && !SrcNode->IsNodeEnabled();
+
 				// associate pins, no known case of StaticDuplicateObjectEx resulting in a different number of pins, but
 				// if that does happen we just associate as many pins as we can:
 				ensure(SrcNode->Pins.Num() == DstNode->Pins.Num());
-				for (int32 I = 0; I < SrcNode->Pins.Num() && I < DstNode->Pins.Num(); ++I)
+				for (int32 PinIdx = 0; PinIdx < SrcNode->Pins.Num() && PinIdx < DstNode->Pins.Num(); ++PinIdx)
 				{
-					if (ensure(DstNode->Pins[I] && SrcNode->Pins[I]))
+					UEdGraphPin* SrcPin = SrcNode->Pins[PinIdx];
+					UEdGraphPin* DstPin = DstNode->Pins[PinIdx];
+					if (ensure(SrcPin && DstPin))
 					{
-						MessageLog->NotifyIntermediatePinCreation(DstNode->Pins[I], SrcNode->Pins[I]);
+						// Don't record pins to the backtrack map if this is an excluded node.
+						if (bIsExcludedNode)
+						{
+							// Patch any input links to pass through to the other side.
+							if (DstPin->Direction == EGPD_Input && DstPin->LinkedTo.Num() > 0)
+							{
+								UEdGraphPin* PassThroughPin = DstNode->GetPassThroughPin(DstPin);
+								if (PassThroughPin != nullptr && PassThroughPin->LinkedTo.Num() > 0)
+								{
+									for (UEdGraphPin* OutputPin : DstPin->LinkedTo)
+									{
+										for (UEdGraphPin* InputPin : PassThroughPin->LinkedTo)
+										{
+											InputPin->LinkedTo.Add(OutputPin);
+											OutputPin->LinkedTo.Add(InputPin);
+										}
+									}
+								}
+							}
+						}
+						else if (MessageLog)
+						{
+							MessageLog->NotifyIntermediatePinCreation(DstPin, SrcPin);
+						}
 					}
 				}
 
-				if (bCloningForCompile)
+				if (bIsExcludedNode)
 				{
-					DstNode->SetEnabledState(SrcNode->IsNodeEnabled() ? ENodeEnabledState::Enabled : ENodeEnabledState::Disabled);
+					// Break remaining node links, if any exist.
+					DstNode->BreakAllNodeLinks();
+
+					// Remove the node from the cloned graph, if valid.
+					if (ClonedGraph)
+					{
+						ClonedGraph->Nodes.Remove(DstNode);
+					}
 				}
 			}
+
+			// Don't record excluded nodes to the backtrack map.
+			if (MessageLog && !bIsExcludedNode)
+			{
+				MessageLog->NotifyIntermediateObjectCreation(Dest, Source);
+			}
 		}
 	}
 
diff --git a/Engine/Source/Runtime/Engine/Private/EdGraph/EdGraph.cpp b/Engine/Source/Runtime/Engine/Private/EdGraph/EdGraph.cpp
--- a/Engine/Source/Runtime/Engine/Private/EdGraph/EdGraph.cpp
+++ b/Engine/Source/Runtime/Engine/Private/EdGraph/EdGraph.cpp
@@ -293,41 +293,9 @@ void UEdGraph::MoveNodesToAnotherGraph(UEdGraph* DestinationGraph, bool bIsLoadi
 	{
 		if (UEdGraphNode* Node = Nodes.Pop(/*bAllowShrinking=*/ false))
 		{
-#if WITH_EDITOR
-			// During compilation, do not move ghost nodes, they are not used during compilation.
-			if (bInIsCompiling && !Node->IsNodeEnabled())
-			{
-				// Pass existing connections through non-enabled nodes
-				for (UEdGraphPin* Pin : Node->Pins)
-				{
-					if (Pin->Direction == EGPD_Input && Pin->LinkedTo.Num() > 0)
-					{
-						UEdGraphPin* PassThroughPin = Node->GetPassThroughPin(Pin);
-						if (PassThroughPin != nullptr && PassThroughPin->LinkedTo.Num() > 0)
-						{
-							for (UEdGraphPin* OutputPin : Pin->LinkedTo)
-							{
-								for (UEdGraphPin* InputPin : PassThroughPin->LinkedTo)
-								{
-									InputPin->LinkedTo.Add(OutputPin);
-									OutputPin->LinkedTo.Add(InputPin);
-								}
-							}
-						}
-					}
-				}
-
-				// Break all node links, if any exist, do not move the node
-				Node->BreakAllNodeLinks();
-				continue;
-			}
-#endif
-
 			// Let the name be autogenerated, to automatically avoid naming conflicts
 			// Since this graph is always going to come from a cloned source graph, user readable names can come from the remap stored in a MessageLog.
-	
-			//@todo:  The bIsLoading check is to force no reset loaders when blueprints are compiling on load.  This might not catch all cases though!
-			Node->Rename(/*NewName=*/ NULL, /*NewOuter=*/ DestinationGraph, REN_DontCreateRedirectors | (bIsLoading ? REN_ForceNoResetLoaders : 0) | (bInIsCompiling ? REN_NonTransactional : 0));
+			Node->Rename(/*NewName=*/ nullptr, /*NewOuter=*/ DestinationGraph, REN_DontCreateRedirectors | (bIsLoading ? REN_ForceNoResetLoaders : 0) | (bInIsCompiling ? REN_NonTransactional : 0));
 	
 			DestinationGraph->Nodes.Add(Node);
 		}
