EditorDomain: Fix missing feedback in error message.
#rnx

diff --git a/Engine/Source/Editor/UnrealEd/Private/EditorDomain/EditorDomain.cpp b/Engine/Source/Editor/UnrealEd/Private/EditorDomain/EditorDomain.cpp
--- a/Engine/Source/Editor/UnrealEd/Private/EditorDomain/EditorDomain.cpp
+++ b/Engine/Source/Editor/UnrealEd/Private/EditorDomain/EditorDomain.cpp
@@ -169,7 +169,6 @@ bool FEditorDomain::TryFindOrAddPackageSource(FName PackageName,
 		return true;
 	}
 
-	FString ErrorMessage;
 	FPackageDigest PackageDigest = CalculatePackageDigest(*AssetRegistry, PackageName);
 	switch (PackageDigest.Status)
 	{
@@ -194,8 +193,8 @@ bool FEditorDomain::TryFindOrAddPackageSource(FName PackageName,
 		return false;
 	default:
 		UE_LOG(LogEditorDomain, Warning,
-			TEXT("Could not load package from EditorDomain; it will be loaded from the WorkspaceDomain: %s."),
-			*ErrorMessage)
+			TEXT("Could not load package %s from EditorDomain; it will be loaded from the WorkspaceDomain: %s"),
+			*WriteToString<256>(PackageName), *PackageDigest.GetStatusString());
 		PackageSource = new FPackageSource();
 		PackageSource->Source = EPackageSource::Workspace;
 		OutSource = PackageSource;
@@ -636,9 +635,8 @@ void FEditorDomain::BatchDownload(TArrayView<FName> PackageNames)
 			{
 				FScopeLock ScopeLock(&Locks->Lock);
 				TRefCountPtr<FPackageSource> PackageSource;
-				FPackageDigest ErrorDigest;
 				FName PackageName = FName(*Response.Name);
-				if (TryFindOrAddPackageSource(PackageName, PackageSource, &ErrorDigest))
+				if (TryFindOrAddPackageSource(PackageName, PackageSource))
 				{
 					PackageSource->bHasQueriedCatalog = true;
 				}
