Clearing the garbage / pendingkill flag manually instead of with all flags to prevent asserts when duplicating actors

diff --git a/Engine/Source/Editor/UnrealEd/Private/EditorObject.cpp b/Engine/Source/Editor/UnrealEd/Private/EditorObject.cpp
--- a/Engine/Source/Editor/UnrealEd/Private/EditorObject.cpp
+++ b/Engine/Source/Editor/UnrealEd/Private/EditorObject.cpp
@@ -597,6 +597,7 @@ static const TCHAR* ImportProperties(
 					}
 
 					// Make sure desired flags are set - existing object could be pending kill
+					ComponentTemplate->ClearGarbage();
 					ComponentTemplate->ClearFlags(RF_AllFlags);
 					ComponentTemplate->ClearInternalFlags(EInternalObjectFlags::AllFlags);
 					ComponentTemplate->SetFlags(NewFlags);
