Temporarily remove the new ExtraReleaseVersionAssets block until it can be brought in line with the changes to FPackageNameCache

diff --git a/Engine/Source/Editor/UnrealEd/Private/CookOnTheFlyServer.cpp b/Engine/Source/Editor/UnrealEd/Private/CookOnTheFlyServer.cpp
--- a/Engine/Source/Editor/UnrealEd/Private/CookOnTheFlyServer.cpp
+++ b/Engine/Source/Editor/UnrealEd/Private/CookOnTheFlyServer.cpp
@@ -7713,7 +7713,7 @@ void UCookOnTheFlyServer::StartCookByTheBook( const FCookByTheBookStartupOptions
 			}
 		}
 
-		FString ExtraReleaseVersionAssetsFile;
+		/*FString ExtraReleaseVersionAssetsFile;
 		const bool bUsingExtraReleaseVersionAssets = FParse::Value(FCommandLine::Get(), TEXT("ExtraReleaseVersionAssets="), ExtraReleaseVersionAssetsFile);
 		if (bUsingExtraReleaseVersionAssets)
 		{
@@ -7748,7 +7748,7 @@ void UCookOnTheFlyServer::StartCookByTheBook( const FCookByTheBookStartupOptions
 			}
 			// Add these packages to the cache such that they're resolvable later
 			GetPackageNameCache().AppendCacheResults(MoveTemp(PackageToStandardFileNames));
-		}
+		}*/
 
 		for ( const ITargetPlatform* TargetPlatform: TargetPlatforms )
 		{
