ZenCOTF:
*Let the game be in full control of what is getting cooked
*Removed artificial delays when waiting for DDC operations
*No longer relies on all requests being urgent
#rnx

diff --git a/Engine/Source/Developer/IoStoreUtilities/Private/ZenStoreWriter.cpp b/Engine/Source/Developer/IoStoreUtilities/Private/ZenStoreWriter.cpp
--- a/Engine/Source/Developer/IoStoreUtilities/Private/ZenStoreWriter.cpp
+++ b/Engine/Source/Developer/IoStoreUtilities/Private/ZenStoreWriter.cpp
@@ -158,6 +158,7 @@ FZenStoreWriter::FZenStoreWriter(
 	const ITargetPlatform* InTargetPlatform
 )
 	: TargetPlatform(*InTargetPlatform)
+	, TargetPlatformFName(*InTargetPlatform->PlatformName())
 	, OutputPath(InOutputPath)
 	, MetadataDirectoryPath(InMetadataDirectoryPath)
 	, PackageStoreManifest(InOutputPath)
@@ -259,6 +260,16 @@ void FZenStoreWriter::WritePackageData(const FPackageInfo& Info, FLargeMemoryWri
 	Entry.ChunkId			= ChunkOid;
 	Entry.PackageStoreEntry = PackageStoreOptimizer->CreatePackageStoreEntry(Package.Get());
 	Entry.IsValid			= true;
+
+	if (EntryCreatedEvent.IsBound())
+	{
+		IPackageStoreWriter::FEntryCreatedEventArgs EntryCreatedEventArgs
+		{
+			TargetPlatformFName,
+			Entry.PackageStoreEntry
+		};
+		EntryCreatedEvent.Broadcast(EntryCreatedEventArgs);
+	}
 }
 
 void FZenStoreWriter::WriteIoStorePackageData(const FPackageInfo& Info, const FIoBuffer& PackageData, const FPackageStoreEntryResource& PackageStoreEntry, const TArray<FFileRegion>& FileRegions)
@@ -600,7 +611,7 @@ void FZenStoreWriter::CommitPackageInternal(FCommitPackageInfo&& CommitInfo)
 
 	IPackageStoreWriter::FCommitEventArgs CommitEventArgs;
 
-	CommitEventArgs.PlatformName	= FName(*TargetPlatform.PlatformName());
+	CommitEventArgs.PlatformName	= TargetPlatformFName;
 	CommitEventArgs.PackageName		= CommitInfo.PackageName;
 	CommitEventArgs.EntryIndex		= INDEX_NONE;
 	
@@ -1089,7 +1100,7 @@ void FZenStoreWriter::BroadcastMarkUpToDate(IPackageStoreWriter::FMarkUpToDateEv
 	if (MarkUpToDateEvent.IsBound())
 	{
 		FReadScopeLock _(EntriesLock);
-		EventArgs.PlatformName = FName(*TargetPlatform.PlatformName());
+		EventArgs.PlatformName = TargetPlatformFName;
 		EventArgs.Entries = PackageStoreEntries;
 		EventArgs.CookInfos = CookedPackagesInfo;
 		MarkUpToDateEvent.Broadcast(EventArgs);
diff --git a/Engine/Source/Developer/IoStoreUtilities/Public/PackageStoreWriter.h b/Engine/Source/Developer/IoStoreUtilities/Public/PackageStoreWriter.h
--- a/Engine/Source/Developer/IoStoreUtilities/Public/PackageStoreWriter.h
+++ b/Engine/Source/Developer/IoStoreUtilities/Public/PackageStoreWriter.h
@@ -32,6 +32,15 @@ public:
 	 */
 	virtual void GetEntries(TFunction<void(TArrayView<const FPackageStoreEntryResource>, TArrayView<const FOplogCookInfo>)>&&) = 0;
 
+	struct FEntryCreatedEventArgs
+	{
+		FName PlatformName;
+		const FPackageStoreEntryResource& Entry;
+	};
+
+	DECLARE_EVENT_OneParam(IPackageStoreWriter, FEntryCreatedEvent, const FEntryCreatedEventArgs&);
+	virtual FEntryCreatedEvent& OnEntryCreated() = 0;
+
 	/**
 	 * Package commit event arguments
 	 */
diff --git a/Engine/Source/Developer/IoStoreUtilities/Public/ZenStoreWriter.h b/Engine/Source/Developer/IoStoreUtilities/Public/ZenStoreWriter.h
--- a/Engine/Source/Developer/IoStoreUtilities/Public/ZenStoreWriter.h
+++ b/Engine/Source/Developer/IoStoreUtilities/Public/ZenStoreWriter.h
@@ -54,6 +54,11 @@ public:
 
 	IOSTOREUTILITIES_API virtual void GetEntries(TFunction<void(TArrayView<const FPackageStoreEntryResource>, TArrayView<const FOplogCookInfo>)>&& Callback) override;
 
+	DECLARE_DERIVED_EVENT(FZenStoreWriter, IPackageStoreWriter::FEntryCreatedEvent, FEntryCreatedEvent);
+	IOSTOREUTILITIES_API virtual FEntryCreatedEvent& OnEntryCreated() override
+	{
+		return EntryCreatedEvent;
+	}
 	DECLARE_DERIVED_EVENT(FZenStoreWriter, IPackageStoreWriter::FCommitEvent, FCommitEvent);
 	IOSTOREUTILITIES_API virtual FCommitEvent& OnCommit() override
 	{
@@ -143,6 +148,7 @@ private:
 	TUniquePtr<UE::FZenStoreHttpClient>	HttpClient;
 
 	const ITargetPlatform&				TargetPlatform;
+	const FName							TargetPlatformFName;
 	FString								OutputPath;
 	FString								MetadataDirectoryPath;
 	FIoContainerId						ContainerId = FIoContainerId::FromName(TEXT("global"));
@@ -157,6 +163,7 @@ private:
 
 	TUniquePtr<FZenFileSystemManifest>	ZenFileSystemManifest;
 	
+	FEntryCreatedEvent					EntryCreatedEvent;
 	FCriticalSection					CommitEventCriticalSection;
 	FCommitEvent						CommitEvent;
 	FMarkUpToDateEvent					MarkUpToDateEvent;
diff --git a/Engine/Source/Editor/UnrealEd/Classes/CookOnTheSide/CookOnTheFlyServer.h b/Engine/Source/Editor/UnrealEd/Classes/CookOnTheSide/CookOnTheFlyServer.h
--- a/Engine/Source/Editor/UnrealEd/Classes/CookOnTheSide/CookOnTheFlyServer.h
+++ b/Engine/Source/Editor/UnrealEd/Classes/CookOnTheSide/CookOnTheFlyServer.h
@@ -678,6 +678,8 @@ public:
 
 private:
 
+	bool IsUsingLegacyCookOnTheFlyScheduling() const;
+
 	//////////////////////////////////////////////////////////////////////////
 	// cook by the book specific functions
 	/**
diff --git a/Engine/Source/Editor/UnrealEd/Private/CookOnTheFlyServer.cpp b/Engine/Source/Editor/UnrealEd/Private/CookOnTheFlyServer.cpp
--- a/Engine/Source/Editor/UnrealEd/Private/CookOnTheFlyServer.cpp
+++ b/Engine/Source/Editor/UnrealEd/Private/CookOnTheFlyServer.cpp
@@ -1004,6 +1004,11 @@ bool UCookOnTheFlyServer::IsCookOnTheFlyMode() const
 	return CurrentCookMode == ECookMode::CookOnTheFly || CurrentCookMode == ECookMode::CookOnTheFlyFromTheEditor; 
 }
 
+bool UCookOnTheFlyServer::IsUsingLegacyCookOnTheFlyScheduling() const
+{
+	return CookOnTheFlyRequestManager && CookOnTheFlyRequestManager->ShouldUseLegacyScheduling();
+}
+
 bool UCookOnTheFlyServer::IsCreatingReleaseVersion()
 {
 	if (CookByTheBookOptions)
@@ -1605,22 +1610,44 @@ UCookOnTheFlyServer::ECookAction UCookOnTheFlyServer::DecideNextCookAction(UE::C
 		return ECookAction::Load;
 	}
 
-	if (bSaveBusy & (NumSaves > 0))
+	if (IsCookOnTheFlyMode() && !IsRealtimeMode())
 	{
-		const float CurrentTime = FPlatformTime::Seconds();
-		if (CurrentTime - SaveBusyTimeLastRetry > GCookProgressRetryBusyTime)
+		if (NumSaves > 0 && bLoadBusy)
 		{
-			SaveBusyTimeLastRetry = CurrentTime;
 			return ECookAction::Save;
 		}
+		if (NumLoads > 0 && bSaveBusy)
+		{
+			return ECookAction::Load;
+		}
+		if (NumSaves > 0)
+		{
+			return ECookAction::Save;
+		}
+		if (NumLoads > 0)
+		{
+			return ECookAction::Load;
+		}
 	}
-	if (bLoadBusy & (NumLoads > 0))
+	else
 	{
-		const float CurrentTime = FPlatformTime::Seconds();
-		if (CurrentTime - LoadBusyTimeLastRetry > GCookProgressRetryBusyTime)
+		if (bSaveBusy & (NumSaves > 0))
 		{
-			LoadBusyTimeLastRetry = CurrentTime;
-			return ECookAction::Load;
+			const float CurrentTime = FPlatformTime::Seconds();
+			if (CurrentTime - SaveBusyTimeLastRetry > GCookProgressRetryBusyTime)
+			{
+				SaveBusyTimeLastRetry = CurrentTime;
+				return ECookAction::Save;
+			}
+		}
+		if (bLoadBusy & (NumLoads > 0))
+		{
+			const float CurrentTime = FPlatformTime::Seconds();
+			if (CurrentTime - LoadBusyTimeLastRetry > GCookProgressRetryBusyTime)
+			{
+				LoadBusyTimeLastRetry = CurrentTime;
+				return ECookAction::Load;
+			}
 		}
 	}
 
@@ -1670,7 +1697,7 @@ void UCookOnTheFlyServer::PumpExternalRequests(const UE::Cook::FCookerTimer& Coo
 				NetworkRequestEvent->Trigger();
 			}
 #endif
-			bool bRequestsAreUrgent = IsCookOnTheFlyMode();
+			bool bRequestsAreUrgent = IsCookOnTheFlyMode() && IsUsingLegacyCookOnTheFlyScheduling();
 			TRingBuffer<UE::Cook::FRequestCluster>& RequestClusters = PackageDatas->GetRequestQueue().GetRequestClusters();
 			UE::Cook::FRequestCluster::AddClusters(*this, MoveTemp(BuildRequests), bRequestsAreUrgent, RequestClusters);
 		}
@@ -2045,9 +2072,16 @@ void UCookOnTheFlyServer::QueueDiscoveredPackageData(UE::Cook::FPackageData& Pac
 		return;
 	}
 
-	if (CookOnTheFlyRequestManager && PackageData.IsGenerated())
+	if (CookOnTheFlyRequestManager)
 	{
-		CookOnTheFlyRequestManager->OnPackageGenerated(PackageData.GetPackageName());
+		if (PackageData.IsGenerated())
+		{
+			CookOnTheFlyRequestManager->OnPackageGenerated(PackageData.GetPackageName());
+		}
+		if (!CookOnTheFlyRequestManager->ShouldUseLegacyScheduling())
+		{
+			return;
+		}
 	}
 	
 	if (!PackageData.IsInProgress() &&
@@ -2730,7 +2764,7 @@ void UCookOnTheFlyServer::ReleaseCookedPlatformData(UE::Cook::FPackageData& Pack
 
 void UCookOnTheFlyServer::TickCancels()
 {
-	PackageDatas->PollPendingCookedPlatformDatas();
+	PackageDatas->PollPendingCookedPlatformDatas(false);
 }
 
 bool UCookOnTheFlyServer::LoadPackageForCooking(UE::Cook::FPackageData& PackageData, UPackage*& OutPackage, FString* LoadFromFilename,
@@ -2983,7 +3017,7 @@ void UCookOnTheFlyServer::PumpSaves(UE::Cook::FTickStackData& StackData, uint32
 		bool bShouldFinishTick = false;
 		if (IsCookOnTheFlyMode())
 		{
-			if (!PackageData.GetIsUrgent())
+			if (IsUsingLegacyCookOnTheFlyScheduling() && !PackageData.GetIsUrgent())
 			{
 				if (ExternalRequests->HasRequests() || PackageDatas->GetMonitor().GetNumUrgent() > 0)
 				{
@@ -3028,7 +3062,8 @@ void UCookOnTheFlyServer::PumpSaves(UE::Cook::FTickStackData& StackData, uint32
 		}
 
 		// Release any completed pending CookedPlatformDatas, so that slots in the per-class limits on calls to BeginCacheForCookedPlatformData are freed up for new objects to use
-		PackageDatas->PollPendingCookedPlatformDatas();
+		bool bForce = IsCookOnTheFlyMode() && !IsRealtimeMode();
+		PackageDatas->PollPendingCookedPlatformDatas(bForce);
 
 		// Always wait for FinishPrepareSave before attempting to save the package
 		bool AllObjectsCookedDataCached = FinishPrepareSave(PackageData, StackData.Timer);
@@ -3075,7 +3110,7 @@ void UCookOnTheFlyServer::PumpSaves(UE::Cook::FTickStackData& StackData, uint32
 					// sleep for a bit
 					FPlatformProcess::Sleep(0.0f);
 					// Poll the results again and check whether we are now done
-					PackageDatas->PollPendingCookedPlatformDatas();
+					PackageDatas->PollPendingCookedPlatformDatas(true);
 					AllObjectsCookedDataCached = FinishPrepareSave(PackageData, StackData.Timer);
 				} while (!StackData.Timer.IsTimeUp() && !AllObjectsCookedDataCached);
 			}
diff --git a/Engine/Source/Editor/UnrealEd/Private/Cooker/CookOnTheFlyServerInterface.h b/Engine/Source/Editor/UnrealEd/Private/Cooker/CookOnTheFlyServerInterface.h
--- a/Engine/Source/Editor/UnrealEd/Private/Cooker/CookOnTheFlyServerInterface.h
+++ b/Engine/Source/Editor/UnrealEd/Private/Cooker/CookOnTheFlyServerInterface.h
@@ -122,6 +122,9 @@ public:
 
 	/** Called when a new package is generated */
 	virtual void OnPackageGenerated(const FName& PackageName) = 0;
+
+	/** Returns true if the cooker should use package discovery and urgency to schedule requests */
+	virtual bool ShouldUseLegacyScheduling() = 0;
 };
 
 }} // namespace UE::Cook
diff --git a/Engine/Source/Editor/UnrealEd/Private/Cooker/CookPackageData.cpp b/Engine/Source/Editor/UnrealEd/Private/Cooker/CookPackageData.cpp
--- a/Engine/Source/Editor/UnrealEd/Private/Cooker/CookPackageData.cpp
+++ b/Engine/Source/Editor/UnrealEd/Private/Cooker/CookPackageData.cpp
@@ -2110,21 +2110,24 @@ TArray<FPendingCookedPlatformData>& FPackageDatas::GetPendingCookedPlatformDatas
 	return PendingCookedPlatformDatas;
 }
 
-void FPackageDatas::PollPendingCookedPlatformDatas()
+void FPackageDatas::PollPendingCookedPlatformDatas(bool bForce)
 {
 	if (PendingCookedPlatformDatas.Num() == 0)
 	{
 		return;
 	}
 
-	// ProcessAsyncResults and IsCachedCookedPlatformDataLoaded can be expensive to call
-	// Cap the frequency at which we call them.
-	double CurrentTime = FPlatformTime::Seconds();
-	if (CurrentTime < LastPollAsyncTime + GPollAsyncPeriod)
+	if (!bForce)
 	{
-		return;
+		// ProcessAsyncResults and IsCachedCookedPlatformDataLoaded can be expensive to call
+		// Cap the frequency at which we call them.
+		double CurrentTime = FPlatformTime::Seconds();
+		if (CurrentTime < LastPollAsyncTime + GPollAsyncPeriod)
+		{
+			return;
+		}
+		LastPollAsyncTime = CurrentTime;
 	}
-	LastPollAsyncTime = CurrentTime;
 
 	GShaderCompilingManager->ProcessAsyncResults(true /* bLimitExecutionTime */,
 		false /* bBlockOnGlobalShaderCompletion */);
diff --git a/Engine/Source/Editor/UnrealEd/Private/Cooker/CookPackageData.h b/Engine/Source/Editor/UnrealEd/Private/Cooker/CookPackageData.h
--- a/Engine/Source/Editor/UnrealEd/Private/Cooker/CookPackageData.h
+++ b/Engine/Source/Editor/UnrealEd/Private/Cooker/CookPackageData.h
@@ -1042,7 +1042,7 @@ public:
 	 * Iterate over all elements in PendingCookedPlatformDatas and check whether they have completed,
 	 * releasing their resources and pending count if so.
 	 */
-	void PollPendingCookedPlatformDatas();
+	void PollPendingCookedPlatformDatas(bool bForce);
 
 	/** RangedFor methods for iterating over all FPackageData managed by this FPackageDatas */
 	TArray<FPackageData*>::RangedForIteratorType begin();
diff --git a/Engine/Source/Editor/UnrealEd/Private/Cooker/IoStoreCookOnTheFlyRequestManager.cpp b/Engine/Source/Editor/UnrealEd/Private/Cooker/IoStoreCookOnTheFlyRequestManager.cpp
--- a/Engine/Source/Editor/UnrealEd/Private/Cooker/IoStoreCookOnTheFlyRequestManager.cpp
+++ b/Engine/Source/Editor/UnrealEd/Private/Cooker/IoStoreCookOnTheFlyRequestManager.cpp
@@ -736,6 +736,11 @@ private:
 		TickRecookPackages();
 	}
 
+	virtual bool ShouldUseLegacyScheduling() override
+	{
+		return false;
+	}
+
 private:
 	bool HandleClientConnection(const FName& PlatformName, FIoStoreCookOnTheFlyNetworkServer::EConnectionStatus ConnectionStatus)
 	{
@@ -763,6 +768,7 @@ private:
 						Context->AddExistingPackages(Entries, CookInfos);
 					});
 
+					PackageWriter->OnEntryCreated().AddRaw(this, &FIoStoreCookOnTheFlyRequestManager::OnPackageStoreEntryCreated);
 					PackageWriter->OnCommit().AddRaw(this, &FIoStoreCookOnTheFlyRequestManager::OnPackageCooked);
 					PackageWriter->OnMarkUpToDate().AddRaw(this, &FIoStoreCookOnTheFlyRequestManager::OnPackagesMarkedUpToDate);
 				}
@@ -898,6 +904,23 @@ private:
 		return true;
 	}
 
+	void OnPackageStoreEntryCreated(const IPackageStoreWriter::FEntryCreatedEventArgs& EventArgs)
+	{
+		FPlatformContext& Context = GetContext(EventArgs.PlatformName);
+
+		FScopeLock _(&Context.GetLock());
+		for (const FPackageId& ImportedPackageId : EventArgs.Entry.ImportedPackageIds)
+		{
+			FPackageStoreEntryResource DummyEntry;
+			auto GetPackageNameFunc = [this, &ImportedPackageId]()
+			{
+				FScopeLock _(&AllKnownPackagesCriticalSection);
+				return AllKnownPackagesMap.FindRef(ImportedPackageId);
+			};
+			EPackageStoreEntryStatus PackageStatus = Context.RequestCook(CookOnTheFlyServer, ImportedPackageId, GetPackageNameFunc, DummyEntry);
+		}
+	}
+
 	void OnPackageCooked(const IPackageStoreWriter::FCommitEventArgs& EventArgs)
 	{
 		using namespace UE::Cook;
diff --git a/Engine/Source/Editor/UnrealEd/Private/Cooker/NetworkFileCookOnTheFlyRequestManager.cpp b/Engine/Source/Editor/UnrealEd/Private/Cooker/NetworkFileCookOnTheFlyRequestManager.cpp
--- a/Engine/Source/Editor/UnrealEd/Private/Cooker/NetworkFileCookOnTheFlyRequestManager.cpp
+++ b/Engine/Source/Editor/UnrealEd/Private/Cooker/NetworkFileCookOnTheFlyRequestManager.cpp
@@ -51,6 +51,11 @@ public:
 	{
 	}
 
+	virtual bool ShouldUseLegacyScheduling() override
+	{
+		return true;
+	}
+
 private:
 	bool OnNewConnection(const FString& VersionInfo, const FString& PlatformName)
 	{
