Merge ReleaseVersionAssets into the new version of the cooker that has eliminated FPackageNameCache in favor of FPackageData.
#rnx

diff --git a/Engine/Source/Editor/UnrealEd/Private/CookOnTheFlyServer.cpp b/Engine/Source/Editor/UnrealEd/Private/CookOnTheFlyServer.cpp
--- a/Engine/Source/Editor/UnrealEd/Private/CookOnTheFlyServer.cpp
+++ b/Engine/Source/Editor/UnrealEd/Private/CookOnTheFlyServer.cpp
@@ -7748,42 +7748,25 @@ void UCookOnTheFlyServer::StartCookByTheBook( const FCookByTheBookStartupOptions
 			}
 		}
 
-		/*FString ExtraReleaseVersionAssetsFile;
+		FString ExtraReleaseVersionAssetsFile;
 		const bool bUsingExtraReleaseVersionAssets = FParse::Value(FCommandLine::Get(), TEXT("ExtraReleaseVersionAssets="), ExtraReleaseVersionAssetsFile);
 		if (bUsingExtraReleaseVersionAssets)
 		{
-			// read the file 
+			// read AssetPaths out of the file and add them as already-cooked PackageDatas
 			TArray<FString> OutAssetPaths;
 			FFileHelper::LoadFileToStringArray(OutAssetPaths, *ExtraReleaseVersionAssetsFile);
-
-			TArray<TTuple<FName, FPackageNameCache::FCachedPackageFilename>> PackageToStandardFileNames;
 			for (const FString& AssetPath : OutAssetPaths)
 			{
-				// register with package cache
-				FName AssetPathFName(*AssetPath);
-				if (UE::Cook::FPackageData* PackageData = PackageDatas->TryAddPackageDataByFileName(AssetPathFName))
+				if (UE::Cook::FPackageData* PackageData = PackageDatas->TryAddPackageDataByFileName(FName(*AssetPath)))
 				{
-					// build standard filename
-					FString StandardPackageFilename = AssetPath;
-					FPaths::MakeStandardFilename(StandardPackageFilename);
-
-					// add to override list
-					FName StandardPackageFileFName(*StandardPackageFilename);
-					OverridePackageList.Add(StandardPackageFileFName);
-
-					PackageToStandardFileNames.Add(
-						TTuple<FName, FPackageNameCache::FCachedPackageFilename>(
-							PackageData->GetPackageName(),
-							FPackageNameCache::FCachedPackageFilename(MoveTemp(StandardPackageFilename), StandardPackageFileFName)));
+					PackageData->SetPlatformsCooked(TargetPlatforms, true /* Succeeded */);
 				}
 				else
 				{
-					UE_LOG(LogCook, Error, TEXT("Failed to resolve package data for package [%s]"), *AssetPathFName.ToString());
+					UE_LOG(LogCook, Error, TEXT("Failed to resolve package data for package [%s]"), *AssetPath);
 				}
 			}
-			// Add these packages to the cache such that they're resolvable later
-			GetPackageNameCache().AppendCacheResults(MoveTemp(PackageToStandardFileNames));
-		}*/
+		}
 
 		for ( const ITargetPlatform* TargetPlatform: TargetPlatforms )
 		{
