back out 19021295

diff --git a/Engine/Source/Editor/UnrealEd/Private/CookOnTheFlyServer.cpp b/Engine/Source/Editor/UnrealEd/Private/CookOnTheFlyServer.cpp
--- a/Engine/Source/Editor/UnrealEd/Private/CookOnTheFlyServer.cpp
+++ b/Engine/Source/Editor/UnrealEd/Private/CookOnTheFlyServer.cpp
@@ -7748,26 +7748,6 @@ void UCookOnTheFlyServer::StartCookByTheBook( const FCookByTheBookStartupOptions
 			}
 		}
 
-		FString ExtraReleaseVersionAssetsFile;
-		const bool bUsingExtraReleaseVersionAssets = FParse::Value(FCommandLine::Get(), TEXT("ExtraReleaseVersionAssets="), ExtraReleaseVersionAssetsFile);
-		if (bUsingExtraReleaseVersionAssets)
-		{
-			// read AssetPaths out of the file and add them as already-cooked PackageDatas
-			TArray<FString> OutAssetPaths;
-			FFileHelper::LoadFileToStringArray(OutAssetPaths, *ExtraReleaseVersionAssetsFile);
-			for (const FString& AssetPath : OutAssetPaths)
-			{
-				if (UE::Cook::FPackageData* PackageData = PackageDatas->TryAddPackageDataByFileName(FName(*AssetPath)))
-				{
-					PackageData->SetPlatformsCooked(TargetPlatforms, true /* Succeeded */);
-				}
-				else
-				{
-					UE_LOG(LogCook, Error, TEXT("Failed to resolve package data for package [%s]"), *AssetPath);
-				}
-			}
-		}
-
 		for ( const ITargetPlatform* TargetPlatform: TargetPlatforms )
 		{
 			SCOPED_BOOT_TIMING("AddCookedPlatforms");
