Fix for CIS localization errors caused by the text gathering commandlet returning an error code.
The issue is caused by EditorTutorials.ini failing to load a dependent module: IntroTutorials. This module has been moved into a plugin.
- Updated EditorTutorials.ini to remove IntroTutorials as a required preload dependency
- Added logging to GatherTextCommandlets.cpp to detect when certain modules that are required to be preloaded fail to load.
#rnx

diff --git a/Engine/Config/Localization/EditorTutorials.ini b/Engine/Config/Localization/EditorTutorials.ini
--- a/Engine/Config/Localization/EditorTutorials.ini
+++ b/Engine/Config/Localization/EditorTutorials.ini
@@ -17,7 +17,9 @@ CulturesToGenerate=ja
 CulturesToGenerate=ko
 CulturesToGenerate=pt-BR
 CulturesToGenerate=zh-Hans
-ModulesToPreload=IntroTutorials
+
+;NOTE: IntroTutorials is now moved into the GuidedTutorial plugin.
+; That module needs to be loaded first so make sure the plugin is enabled before performing these gathers.
 
 ;Gather text from assets
 [GatherTextStep0]
diff --git a/Engine/Source/Editor/UnrealEd/Private/Commandlets/GatherTextFromAssetsCommandlet.cpp b/Engine/Source/Editor/UnrealEd/Private/Commandlets/GatherTextFromAssetsCommandlet.cpp
--- a/Engine/Source/Editor/UnrealEd/Private/Commandlets/GatherTextFromAssetsCommandlet.cpp
+++ b/Engine/Source/Editor/UnrealEd/Private/Commandlets/GatherTextFromAssetsCommandlet.cpp
@@ -474,6 +474,7 @@ int32 UGatherTextFromAssetsCommandlet::Main(const FString& Params)
 
 			if (ModuleLoadResult != EModuleLoadResult::Success)
 			{
+				UE_LOG(LogGatherTextFromAssetsCommandlet, Warning, TEXT("Failed to preload dependent module %s. Please check if the modules have been renamed or moved to another folder."), *ModuleName);
 				HasFailedToPreloadAnyModules = true;
 				continue;
 			}
