When we export an anim sequence to fbx, we now also export the metadata of the preview mesh
#rnx

diff --git a/Engine/Source/Editor/UnrealEd/Private/Fbx/FbxAnimationExport.cpp b/Engine/Source/Editor/UnrealEd/Private/Fbx/FbxAnimationExport.cpp
--- a/Engine/Source/Editor/UnrealEd/Private/Fbx/FbxAnimationExport.cpp
+++ b/Engine/Source/Editor/UnrealEd/Private/Fbx/FbxAnimationExport.cpp
@@ -534,6 +534,8 @@ FbxNode* FFbxExporter::ExportAnimSequence( const UAnimSequence* AnimSeq, const U
 		{
 			TmpNodeNoTransform->RemoveChild(MeshRootNode);
 			RootNode->AddChild(MeshRootNode);
+			//Export the preview mesh metadata
+			ExportObjectMetadata(SkelMesh, MeshRootNode);
 		}
 	}
 	
