Cooker: Create PumpPollables to handle more performantly the processes in CookCommandlet's while loop that need to be infrequently polled.
Move FlushRenderingCommands into PumpPollables.
#rnx

diff --git a/Engine/Source/Editor/UnrealEd/Classes/CookOnTheSide/CookOnTheFlyServer.h b/Engine/Source/Editor/UnrealEd/Classes/CookOnTheSide/CookOnTheFlyServer.h
--- a/Engine/Source/Editor/UnrealEd/Classes/CookOnTheSide/CookOnTheFlyServer.h
+++ b/Engine/Source/Editor/UnrealEd/Classes/CookOnTheSide/CookOnTheFlyServer.h
@@ -11,6 +11,7 @@
 #include "INetworkFileSystemModule.h"
 #include "IPlatformFileSandboxWrapper.h"
 #include "Misc/EnumClassFlags.h"
+#include "Templates/Function.h"
 #include "Templates/UniquePtr.h"
 #include "TickableEditorObject.h"
 #include "UObject/Object.h"
@@ -157,6 +158,18 @@ public:
 	};
 
 private:
+	using FPollFunction = TUniqueFunction<void()>;
+	struct FPollable
+	{
+		FPollable(float InPeriodSeconds, float InPeriodIdleSeconds, FPollFunction&& InFunction);
+
+		FPollFunction PollFunction;
+		double NextTimeSeconds = 0.;
+		double NextTimeIdleSeconds = 0.;
+		float PeriodSeconds = 60.f;
+		float PeriodIdleSeconds = 5.f;
+	};
+
 	//////////////////////////////////////////////////////////////////////////
 	// Cook on the fly server interface adapter
 	class FCookOnTheFlyServerInterface;
@@ -311,6 +324,8 @@ private:
 		LoadLimited,	// Process the LoadQueue, stopping when loadqueuelength reaches the desired population level
 		Save,			// Process the SaveQueue
 		SaveLimited,	// Process the SaveQueue, stopping when savequeuelength reaches the desired population level
+		Poll,			// Execute pollables which have exceeded their period
+		PollIdle,		// Execute pollables which have exceeded their idle period
 		YieldTick,		// Progress is blocked by an async result. Temporarily exit TickCookOnTheSide.
 		Cancel,			// Cancel the current CookByTheBook
 	};
@@ -364,6 +379,10 @@ private:
 	 */
 	void OnRemoveSessionPlatform(const ITargetPlatform* TargetPlatform);
 
+	void InitializePollables();
+	void PumpPollables(UE::Cook::FTickStackData& StackData, bool bIsIdle);
+	void PollFlushRenderingCommands();
+
 public:
 
 	enum ECookOnTheSideResult
@@ -1135,6 +1154,13 @@ private:
 	/** Helper struct for running cooking in diagnostic modes */
 	TUniquePtr<FDiffModeCookServerUtils> DiffModeHelper;
 
+	TArray<FPollable> Pollables;
+	double PollNextTimeSeconds = 0.;
+	double PollNextTimeIdleSeconds = 0.;
+	int32 PollStartIndex = 0;
+	float PumpPollablesTimeSlice = 0.f;
+	float PumpPollablesMinPeriod = 0.f;
+
 	friend UE::Cook::FPackageData;
 	friend UE::Cook::FPendingCookedPlatformData;
 	friend UE::Cook::FPlatformManager;
diff --git a/Engine/Source/Editor/UnrealEd/Private/Commandlets/CookCommandlet.cpp b/Engine/Source/Editor/UnrealEd/Private/Commandlets/CookCommandlet.cpp
--- a/Engine/Source/Editor/UnrealEd/Private/Commandlets/CookCommandlet.cpp
+++ b/Engine/Source/Editor/UnrealEd/Private/Commandlets/CookCommandlet.cpp
@@ -1162,14 +1162,6 @@ bool UCookCommandlet::CookByTheBook( const TArray<ITargetPlatform*>& Platforms)
 					UE_SCOPED_COOKTIMER_AND_DURATION(CookByTheBook_ProcessDeferredCommands, DetailedCookStats::TickLoopProcessDeferredCommandsTimeSec);
 					ProcessDeferredCommands();
 				}
-
-				{
-					UE_SCOPED_COOKTIMER_AND_DURATION(CookByTheBook_TickCommandletStats, DetailedCookStats::TickLoopFlushRenderingCommandsTimeSec);
-
-					// Flush rendering commands to release any RHI resources (shaders and shader maps).
-					// Delete any FPendingCleanupObjects (shader maps).
-					FlushRenderingCommands();
-				}
 			}
 
 			{
diff --git a/Engine/Source/Editor/UnrealEd/Private/CookOnTheFlyServer.cpp b/Engine/Source/Editor/UnrealEd/Private/CookOnTheFlyServer.cpp
--- a/Engine/Source/Editor/UnrealEd/Private/CookOnTheFlyServer.cpp
+++ b/Engine/Source/Editor/UnrealEd/Private/CookOnTheFlyServer.cpp
@@ -1234,13 +1234,19 @@ uint32 UCookOnTheFlyServer::TickCookOnTheSide(const float TimeSlice, uint32 &Coo
 				PumpSaves(StackData, DesiredSaveQueueLength, NumPushed, bBusy);
 				SetSaveBusy(bBusy && NumPushed == 0); // Mark as busy if pump was blocked and we did not make any progress
 				break;
-			case ECookAction::Done:
-				bContinueTick = false;
-				bCookComplete = true;
+			case ECookAction::Poll:
+				PumpPollables(StackData, false /* bIsIdle */);
+				break;
+			case ECookAction::PollIdle:
+				PumpPollables(StackData, true /* bIsIdle */);
 				break;
 			case ECookAction::YieldTick:
 				bContinueTick = false;
 				break;
+			case ECookAction::Done:
+				bContinueTick = false;
+				bCookComplete = true;
+				break;
 			case ECookAction::Cancel:
 				CancelCookByTheBook();
 				bContinueTick = false;
@@ -1533,6 +1539,103 @@ void UCookOnTheFlyServer::UpdateDisplay(ECookTickFlags TickFlags, bool bForceDis
 	}
 }
 
+UCookOnTheFlyServer::FPollable::FPollable(float InPeriodSeconds, float InPeriodIdleSeconds, UCookOnTheFlyServer::FPollFunction&& InFunction)
+	: PollFunction(MoveTemp(InFunction))
+	, NextTimeSeconds(0)
+	, NextTimeIdleSeconds(0)
+	, PeriodSeconds(InPeriodSeconds)
+	, PeriodIdleSeconds(InPeriodIdleSeconds)
+{
+}
+
+void UCookOnTheFlyServer::InitializePollables()
+{
+	Pollables.Reset();
+
+	if (IsCookByTheBookMode() && !IsCookingInEditor())
+	{
+		if (Pollables.IsEmpty())
+		{
+			Pollables.Add(FPollable(60.f, 5.f, [this] { PollFlushRenderingCommands(); }));
+		}
+
+		for (FPollable& Pollable : Pollables)
+		{
+			Pollable.NextTimeSeconds = Pollable.NextTimeIdleSeconds = 0.0;
+		}
+	}
+
+	PollStartIndex = 0;
+	if (Pollables.Num())
+	{
+		PollNextTimeSeconds = 0.;
+		PollNextTimeIdleSeconds = 0.;
+	}
+	else
+	{
+		PollNextTimeSeconds = MAX_flt;
+		PollNextTimeIdleSeconds = MAX_flt;
+	}
+}
+
+void UCookOnTheFlyServer::PumpPollables(UE::Cook::FTickStackData& StackData, bool bIsIdle)
+{
+	UE_SCOPED_HIERARCHICAL_COOKTIMER(PumpPollables);
+	PollNextTimeSeconds = MAX_flt;
+	PollNextTimeIdleSeconds = MAX_flt;
+	int32 NumPollables = Pollables.Num();
+	if (NumPollables == 0)
+	{
+		return;
+	}
+
+	double CurrentTime = FPlatformTime::Seconds();
+	double MaxTime = CurrentTime + PumpPollablesTimeSlice;
+
+	PollStartIndex = PollStartIndex % NumPollables;
+	int32 PollEndIndex = PollStartIndex;
+	do
+	{
+		FPollable& Pollable = Pollables[PollStartIndex++];
+		PollStartIndex = PollStartIndex >= NumPollables ? 0 : PollStartIndex;
+
+		double NextTimeSeconds = bIsIdle ? Pollable.NextTimeIdleSeconds : Pollable.NextTimeSeconds;
+		if (NextTimeSeconds <= CurrentTime)
+		{
+			Pollable.PollFunction();
+			double NewCurrentTime = FPlatformTime::Seconds();
+			Pollable.NextTimeSeconds = NewCurrentTime + Pollable.PeriodSeconds;
+			Pollable.NextTimeIdleSeconds = NewCurrentTime + Pollable.PeriodIdleSeconds;
+			CurrentTime = NewCurrentTime;
+		}
+		PollNextTimeSeconds = FMath::Min(Pollable.NextTimeSeconds, PollNextTimeSeconds);
+		PollNextTimeIdleSeconds = FMath::Min(Pollable.NextTimeIdleSeconds, PollNextTimeIdleSeconds);
+		if (CurrentTime >= MaxTime)
+		{
+			break;
+		}
+	} while (PollStartIndex != PollEndIndex);
+
+	// If we early exited, finish calculating PollNextTimeSeconds from the remaining members we didn't reach
+	for (int32 Index = PollStartIndex; Index != PollEndIndex; Index = (Index+1) % NumPollables)
+	{
+		FPollable& Pollable = Pollables[Index++];
+		PollNextTimeSeconds = FMath::Min(Pollable.NextTimeSeconds, PollNextTimeSeconds);
+		PollNextTimeIdleSeconds = FMath::Min(Pollable.NextTimeIdleSeconds, PollNextTimeIdleSeconds);
+	}
+	PollNextTimeSeconds = FMath::Max(PollNextTimeSeconds, CurrentTime + PumpPollablesMinPeriod);
+	// PollNextTimeIdleSeconds does not have a min period
+}
+
+void UCookOnTheFlyServer::PollFlushRenderingCommands()
+{
+	UE_SCOPED_COOKTIMER_AND_DURATION(CookByTheBook_TickCommandletStats, DetailedCookStats::TickLoopFlushRenderingCommandsTimeSec);
+
+	// Flush rendering commands to release any RHI resources (shaders and shader maps).
+	// Delete any FPendingCleanupObjects (shader maps).
+	FlushRenderingCommands();
+}
+
 UCookOnTheFlyServer::ECookAction UCookOnTheFlyServer::DecideNextCookAction(UE::Cook::FTickStackData& StackData)
 {
 	if (IsCookByTheBookMode() && CookByTheBookOptions->bCancel)
@@ -1545,10 +1648,16 @@ UCookOnTheFlyServer::ECookAction UCookOnTheFlyServer::DecideNextCookAction(UE::C
 		// if we just cooked a map then don't process anything the rest of this tick
 		return ECookAction::YieldTick;
 	}
-	else if (StackData.Timer.IsTimeUp())
+
+	double CurrentTime = FPlatformTime::Seconds();
+	if (StackData.Timer.IsTimeUp(CurrentTime))
 	{
 		return ECookAction::YieldTick;
 	}
+	else if (CurrentTime >= PollNextTimeSeconds)
+	{
+		return ECookAction::Poll;
+	}
 
 	UE::Cook::FRequestQueue& RequestQueue = PackageDatas->GetRequestQueue();
 	if (RequestQueue.GetRequestClusters().Num() != 0 || RequestQueue.GetUnclusteredRequests().Num() != 0)
@@ -1633,7 +1742,6 @@ UCookOnTheFlyServer::ECookAction UCookOnTheFlyServer::DecideNextCookAction(UE::C
 	{
 		if (bSaveBusy & (NumSaves > 0))
 		{
-			const float CurrentTime = FPlatformTime::Seconds();
 			if (CurrentTime - SaveBusyTimeLastRetry > GCookProgressRetryBusyTime)
 			{
 				SaveBusyTimeLastRetry = CurrentTime;
@@ -1642,7 +1750,6 @@ UCookOnTheFlyServer::ECookAction UCookOnTheFlyServer::DecideNextCookAction(UE::C
 		}
 		if (bLoadBusy & (NumLoads > 0))
 		{
-			const float CurrentTime = FPlatformTime::Seconds();
 			if (CurrentTime - LoadBusyTimeLastRetry > GCookProgressRetryBusyTime)
 			{
 				LoadBusyTimeLastRetry = CurrentTime;
@@ -1653,6 +1760,10 @@ UCookOnTheFlyServer::ECookAction UCookOnTheFlyServer::DecideNextCookAction(UE::C
 
 	if (PackageDatas->GetMonitor().GetNumInProgress() > 0)
 	{
+		if (CurrentTime >= PollNextTimeIdleSeconds)
+		{
+			return ECookAction::PollIdle;
+		}
 		return ECookAction::YieldTick;
 	}
 
@@ -4463,6 +4574,11 @@ void UCookOnTheFlyServer::Initialize( ECookMode::Type DesiredCookMode, ECookInit
 	}
 	GConfig->GetBool(TEXT("CookSettings"), TEXT("PreexploreDependenciesEnabled"), bPreexploreDependenciesEnabled, GEditorIni);
 
+	PumpPollablesTimeSlice = 1.0f;
+	PumpPollablesMinPeriod = 0.1f;
+	PollNextTimeSeconds = MAX_flt;
+	PollNextTimeIdleSeconds = MAX_flt;
+
 	{
 		const FConfigSection* CacheSettings = GConfig->GetSectionPrivate(TEXT("CookPlatformDataCacheSettings"), false, true, GEditorIni);
 		if ( CacheSettings )
@@ -7748,6 +7864,26 @@ void UCookOnTheFlyServer::StartCookByTheBook( const FCookByTheBookStartupOptions
 			}
 		}
 
+		FString ExtraReleaseVersionAssetsFile;
+		const bool bUsingExtraReleaseVersionAssets = FParse::Value(FCommandLine::Get(), TEXT("ExtraReleaseVersionAssets="), ExtraReleaseVersionAssetsFile);
+		if (bUsingExtraReleaseVersionAssets)
+		{
+			// read AssetPaths out of the file and add them as already-cooked PackageDatas
+			TArray<FString> OutAssetPaths;
+			FFileHelper::LoadFileToStringArray(OutAssetPaths, *ExtraReleaseVersionAssetsFile);
+			for (const FString& AssetPath : OutAssetPaths)
+			{
+				if (UE::Cook::FPackageData* PackageData = PackageDatas->TryAddPackageDataByFileName(FName(*AssetPath)))
+				{
+					PackageData->SetPlatformsCooked(TargetPlatforms, true /* Succeeded */);
+				}
+				else
+				{
+					UE_LOG(LogCook, Error, TEXT("Failed to resolve package data for package [%s]"), *AssetPath);
+				}
+			}
+		}
+
 		for ( const ITargetPlatform* TargetPlatform: TargetPlatforms )
 		{
 			SCOPED_BOOT_TIMING("AddCookedPlatforms");
@@ -7951,6 +8087,8 @@ void UCookOnTheFlyServer::StartCookByTheBook( const FCookByTheBookStartupOptions
 	{
 		FindOrCreatePackageWriter(TargetPlatform).BeginCook();
 	}
+
+	InitializePollables();
 }
 
 TArray<FName> UCookOnTheFlyServer::GetNeverCookPackageFileNames(TArrayView<const FString> ExtraNeverCookDirectories)
diff --git a/Engine/Source/Editor/UnrealEd/Private/Cooker/CookProfiling.h b/Engine/Source/Editor/UnrealEd/Private/Cooker/CookProfiling.h
--- a/Engine/Source/Editor/UnrealEd/Private/Cooker/CookProfiling.h
+++ b/Engine/Source/Editor/UnrealEd/Private/Cooker/CookProfiling.h
@@ -126,6 +126,7 @@ namespace DetailedCookStats
 	extern double TickCookOnTheSideBeginPrepareSaveTimeSec;
 	extern double TickCookOnTheSideFinishPrepareSaveTimeSec;
 	extern double GameCookModificationDelegateTimeSec;
+	extern double TickLoopFlushRenderingCommandsTimeSec;
 
 	// Stats tracked through FAutoRegisterCallback
 	extern uint32 NumPreloadedDependencies;
diff --git a/Engine/Source/Editor/UnrealEd/Private/Cooker/CookTypes.cpp b/Engine/Source/Editor/UnrealEd/Private/Cooker/CookTypes.cpp
--- a/Engine/Source/Editor/UnrealEd/Private/Cooker/CookTypes.cpp
+++ b/Engine/Source/Editor/UnrealEd/Private/Cooker/CookTypes.cpp
@@ -30,12 +30,17 @@ double FCookerTimer::GetTimeTillNow() const
 {
 	return FPlatformTime::Seconds() - StartTime;
 }
-	
+
 bool FCookerTimer::IsTimeUp() const
+{
+	return IsTimeUp(FPlatformTime::Seconds());
+}
+
+bool FCookerTimer::IsTimeUp(double CurrentTimeSeconds) const
 {
 	if (bIsRealtimeMode)
 	{
-		if ((FPlatformTime::Seconds() - StartTime) > TimeSlice)
+		if ((CurrentTimeSeconds - StartTime) > TimeSlice)
 		{
 			return true;
 		}
diff --git a/Engine/Source/Editor/UnrealEd/Private/Cooker/CookTypes.h b/Engine/Source/Editor/UnrealEd/Private/Cooker/CookTypes.h
--- a/Engine/Source/Editor/UnrealEd/Private/Cooker/CookTypes.h
+++ b/Engine/Source/Editor/UnrealEd/Private/Cooker/CookTypes.h
@@ -138,6 +138,7 @@ namespace UE::Cook
 
 		double GetTimeTillNow() const;
 		bool IsTimeUp() const;
+		bool IsTimeUp(double CurrentTimeSeconds) const;
 		void SavedPackage();
 		double GetTimeRemain() const;
 
