PackageSourceControlHelper: Add optional bErrorsAsWarnings param to AddToSourceContro/Checkout/Delete operations (defaults to false)
WorldPartitionHelpers: Add optional callback when releasing ActorReferences to allow caller to batch operations
WorldPartitionBuilder: Add SavePackages/DeletePackages statics that can be used by the different Builders (takes care of the SCC batching)
WorldPartitionRenameDuplicateBuilder & WorldPartitionResaveActorsBuilder use the new new SavePackages/DeletePackages

#rnx

diff --git a/Engine/Source/Editor/UnrealEd/Private/PackageSourceControlHelper.cpp b/Engine/Source/Editor/UnrealEd/Private/PackageSourceControlHelper.cpp
--- a/Engine/Source/Editor/UnrealEd/Private/PackageSourceControlHelper.cpp
+++ b/Engine/Source/Editor/UnrealEd/Private/PackageSourceControlHelper.cpp
@@ -16,9 +16,25 @@
 #include "ISourceControlState.h"
 #include "ISourceControlProvider.h"
 #include "SourceControlOperations.h"
+#include "Misc/PackageName.h"
 
 DEFINE_LOG_CATEGORY_STATIC(LogCommandletPackageHelper, Log, All);
 
+namespace FPackageSourceControlHelperLog
+{
+	void Error(bool bErrorsAsWarnings, const FString& Msg)
+	{
+		if (bErrorsAsWarnings)
+		{
+			UE_LOG(LogCommandletPackageHelper, Warning, TEXT("%s"), *Msg);
+		}
+		else
+		{
+			UE_LOG(LogCommandletPackageHelper, Error, TEXT("%s"), *Msg);
+		}
+	}
+}
+
 bool FPackageSourceControlHelper::UseSourceControl() const
 {
 	return GetSourceControlProvider().IsEnabled();
@@ -35,8 +51,10 @@ bool FPackageSourceControlHelper::Delete(const FString& PackageName) const
 	return Delete(PackageNames);
 }
 
-bool FPackageSourceControlHelper::Delete(const TArray<FString>& PackageNames) const
+bool FPackageSourceControlHelper::Delete(const TArray<FString>& PackageNames, bool bErrorsAsWarnings) const
 {
+	bool bSuccess = true;
+	
 	// Early out when not using source control
 	if (!UseSourceControl())
 	{
@@ -47,38 +65,33 @@ bool FPackageSourceControlHelper::Delete(const TArray<FString>& PackageNames) co
 			if (!IPlatformFile::GetPlatformPhysical().SetReadOnly(*Filename, false) ||
 				!IPlatformFile::GetPlatformPhysical().DeleteFile(*Filename))
 			{
-				UE_LOG(LogCommandletPackageHelper, Error, TEXT("Error deleting %s"), *Filename);
-				return false;
+				bSuccess = false;
+				FPackageSourceControlHelperLog::Error(bErrorsAsWarnings, FString::Printf(TEXT("Error deleting %s"), *Filename));
+				if (!bErrorsAsWarnings)
+				{
+					return false;
+				}
 			}
 		}
 
-		return true;
+		return bSuccess;
 	}
 
 	TArray<FString> FilesToRevert;
 	TArray<FString> FilesToDeleteFromDisk;
 	TArray<FString> FilesToDeleteFromSCC;
-
-	bool bSCCErrorsFound = false;
-
+		
 	// First: get latest state from source control
-	TArray<FString> Filenames;
+	TArray<FString> Filenames = SourceControlHelpers::PackageFilenames(PackageNames);
 	TArray<FSourceControlStateRef> SourceControlStates;
-	Filenames.Reset(PackageNames.Num());
-
-	for (const FString& PackageName : PackageNames)
-	{
-		Filenames.Emplace(SourceControlHelpers::PackageFilename(PackageName));
-	}
-
-	ECommandResult::Type UpdateState = GetSourceControlProvider().GetState(Filenames, SourceControlStates, EStateCacheUsage::ForceUpdate);
-
-	if (UpdateState != ECommandResult::Succeeded)
+	
+	if (GetSourceControlProvider().GetState(Filenames, SourceControlStates, EStateCacheUsage::ForceUpdate) != ECommandResult::Succeeded)
 	{
-		UE_LOG(LogCommandletPackageHelper, Error, TEXT("Could not get source control state for packages"));
+		// Nothing we can do if SCCStates fail
+		FPackageSourceControlHelperLog::Error(false, TEXT("Could not get source control state for packages"));
 		return false;
 	}
-
+		
 	for(FSourceControlStateRef& SourceControlState : SourceControlStates)
 	{
 		const FString& Filename = SourceControlState->GetFilename();
@@ -90,13 +103,13 @@ bool FPackageSourceControlHelper::Delete(const TArray<FString>& PackageNames) co
 			FString OtherCheckedOutUser;
 			if (SourceControlState->IsCheckedOutOther(&OtherCheckedOutUser))
 			{
-				UE_LOG(LogCommandletPackageHelper, Error, TEXT("Overwriting package %s already checked out by %s, will not submit"), *Filename, *OtherCheckedOutUser);
-				bSCCErrorsFound = true;
+				FPackageSourceControlHelperLog::Error(bErrorsAsWarnings, FString::Printf(TEXT("File %s already checked out by %s, will not delete"), *Filename, *OtherCheckedOutUser));
+				bSuccess = false;
 			}
 			else if (!SourceControlState->IsCurrent())
 			{
-				UE_LOG(LogCommandletPackageHelper, Error, TEXT("Overwriting package %s (not at head revision), will not submit"), *Filename);
-				bSCCErrorsFound = true;
+				FPackageSourceControlHelperLog::Error(bErrorsAsWarnings, FString::Printf(TEXT("File %s (not at head revision), will not delete"), *Filename));
+				bSuccess = false;
 			}
 			else if (SourceControlState->IsAdded())
 			{
@@ -119,7 +132,7 @@ bool FPackageSourceControlHelper::Delete(const TArray<FString>& PackageNames) co
 		}
 	}
 
-	if (bSCCErrorsFound)
+	if (!bSuccess && !bErrorsAsWarnings)
 	{
 		// Errors were found, we'll cancel everything
 		return false;
@@ -142,8 +155,12 @@ bool FPackageSourceControlHelper::Delete(const TArray<FString>& PackageNames) co
 	{
 		if (GetSourceControlProvider().Execute(ISourceControlOperation::Create<FRevert>(), FilesToRevert) != ECommandResult::Succeeded)
 		{
-			UE_LOG(LogCommandletPackageHelper, Error, TEXT("Error reverting packages from source control"));
-			return false;
+			bSuccess = false;
+			FPackageSourceControlHelperLog::Error(bErrorsAsWarnings, TEXT("Error reverting packages from source control"));
+			if (!bErrorsAsWarnings)
+			{
+				return false;
+			}
 		}
 	}
 
@@ -152,23 +169,30 @@ bool FPackageSourceControlHelper::Delete(const TArray<FString>& PackageNames) co
 	{
 		if (GetSourceControlProvider().Execute(ISourceControlOperation::Create<FDelete>(), FilesToDeleteFromSCC) != ECommandResult::Succeeded)
 		{
-			UE_LOG(LogCommandletPackageHelper, Error, TEXT("Error deleting packages from source control"));
-			return false;
+			bSuccess = false;
+			FPackageSourceControlHelperLog::Error(bErrorsAsWarnings, TEXT("Error deleting packages from source control"));
+			if(!bErrorsAsWarnings)
+			{
+				return false;
+			}
 		}
 	}
 
 	// Then delete files on disk
-	bool bDeleteOnDiskOk = true;
 	for (const FString& Filename : FilesToDeleteFromDisk)
 	{
 		if (!IFileManager::Get().Delete(*Filename, false, true))
 		{
-			UE_LOG(LogCommandletPackageHelper, Error, TEXT("Error deleting package %s locally"), *Filename);
-			bDeleteOnDiskOk = false;
+			bSuccess = false;
+			FPackageSourceControlHelperLog::Error(bErrorsAsWarnings, FString::Printf(TEXT("Error deleting %s locally"), *Filename));
+			if(!bErrorsAsWarnings)
+			{
+				return false;
+			}
 		}
 	}
 
-	return bDeleteOnDiskOk;
+	return bSuccess;
 }
 
 bool FPackageSourceControlHelper::Delete(UPackage* Package) const
@@ -207,7 +231,7 @@ bool FPackageSourceControlHelper::AddToSourceControl(UPackage* Package) const
 	return true;
 }
 
-bool FPackageSourceControlHelper::AddToSourceControl(const TArray<FString>& PackageNames) const
+bool FPackageSourceControlHelper::AddToSourceControl(const TArray<FString>& PackageNames, bool bErrorsAsWarnings) const
 {
 	if (!UseSourceControl())
 	{
@@ -223,14 +247,13 @@ bool FPackageSourceControlHelper::AddToSourceControl(const TArray<FString>& Pack
 	bool bSuccess = true;
 	
 	TArray<FSourceControlStateRef> SourceControlStates;
-	ECommandResult::Type UpdateState = GetSourceControlProvider().GetState(PackageFilenames, SourceControlStates, EStateCacheUsage::ForceUpdate);
-
-	if (UpdateState != ECommandResult::Succeeded)
+	if (GetSourceControlProvider().GetState(PackageFilenames, SourceControlStates, EStateCacheUsage::ForceUpdate) != ECommandResult::Succeeded)
 	{
-		UE_LOG(LogCommandletPackageHelper, Error, TEXT("Could not get source control state for packages"));
+		// Nothing we can do if SCCStates fail
+		FPackageSourceControlHelperLog::Error(false, TEXT("Could not get source control state for packages"));
 		return false;
 	}
-	
+		
 	for (FSourceControlStateRef& SourceControlState : SourceControlStates)
 	{
 		const FString& PackageFilename = SourceControlState->GetFilename();
@@ -238,12 +261,12 @@ bool FPackageSourceControlHelper::AddToSourceControl(const TArray<FString>& Pack
 		FString OtherCheckedOutUser;
 		if (SourceControlState->IsCheckedOutOther(&OtherCheckedOutUser))
 		{
-			UE_LOG(LogCommandletPackageHelper, Error, TEXT("Overwriting package %s already checked out by %s, will not add"), *PackageFilename, *OtherCheckedOutUser);
+			FPackageSourceControlHelperLog::Error(bErrorsAsWarnings, FString::Printf(TEXT("File %s already checked out by %s, will not add"), *PackageFilename, *OtherCheckedOutUser));
 			bSuccess = false;
 		}
 		else if (!SourceControlState->IsCurrent())
 		{
-			UE_LOG(LogCommandletPackageHelper, Error, TEXT("Overwriting package %s (not at head revision), will not add"), *PackageFilename);
+			FPackageSourceControlHelperLog::Error(bErrorsAsWarnings, FString::Printf(TEXT("File %s (not at head revision), will not add"), *PackageFilename));
 			bSuccess = false;
 		}
 		else if (SourceControlState->IsAdded())
@@ -257,17 +280,21 @@ bool FPackageSourceControlHelper::AddToSourceControl(const TArray<FString>& Pack
 	}
 
 	// Any error up to here will be an early out
-	if (!bSuccess)
+	if (!bSuccess && !bErrorsAsWarnings)
 	{
 		return false;
 	}
 
 	if (PackagesToAdd.Num())
 	{
-		return (GetSourceControlProvider().Execute(ISourceControlOperation::Create<FMarkForAdd>(), PackagesToAdd) == ECommandResult::Succeeded);
+		if (GetSourceControlProvider().Execute(ISourceControlOperation::Create<FMarkForAdd>(), PackagesToAdd) != ECommandResult::Succeeded)
+		{
+			FPackageSourceControlHelperLog::Error(bErrorsAsWarnings, TEXT("Error adding packages to source control"));
+			return false;
+		}
 	}
 
-	return true;
+	return bSuccess;
 }
 
 bool FPackageSourceControlHelper::Checkout(UPackage* Package) const
@@ -275,7 +302,7 @@ bool FPackageSourceControlHelper::Checkout(UPackage* Package) const
 	return !Package || Checkout({ Package->GetName() });
 }
 
-bool FPackageSourceControlHelper::Checkout(const TArray<FString>& PackageNames) const
+bool FPackageSourceControlHelper::Checkout(const TArray<FString>& PackageNames, bool bErrorsAsWarnings) const
 {
 	const bool bUseSourceControl = UseSourceControl();
 
@@ -295,7 +322,8 @@ bool FPackageSourceControlHelper::Checkout(const TArray<FString>& PackageNames)
 
 		if (UpdateState != ECommandResult::Succeeded)
 		{
-			UE_LOG(LogCommandletPackageHelper, Error, TEXT("Could not get source control state for packages"));
+			// Nothing we can do if SCCStates fail
+			FPackageSourceControlHelperLog::Error(false, TEXT("Could not get source control state for packages"));
 			return false;
 		}
 
@@ -306,12 +334,12 @@ bool FPackageSourceControlHelper::Checkout(const TArray<FString>& PackageNames)
 			FString OtherCheckedOutUser;
 			if (SourceControlState->IsCheckedOutOther(&OtherCheckedOutUser))
 			{
-				UE_LOG(LogCommandletPackageHelper, Error, TEXT("Overwriting package %s already checked out by %s, will not checkout"), *PackageFilename, *OtherCheckedOutUser);
+				FPackageSourceControlHelperLog::Error(bErrorsAsWarnings, FString::Printf(TEXT("File %s already checked out by %s, will not checkout"), *PackageFilename, *OtherCheckedOutUser));
 				bSuccess = false;
 			}
 			else if (!SourceControlState->IsCurrent())
 			{
-				UE_LOG(LogCommandletPackageHelper, Error, TEXT("Overwriting package %s (not at head revision), will not checkout"), *PackageFilename);
+				FPackageSourceControlHelperLog::Error(bErrorsAsWarnings, FString::Printf(TEXT("File %s (not at head revision), will not checkout"), *PackageFilename));
 				bSuccess = false;
 			}
 			else if (SourceControlState->IsCheckedOut() || SourceControlState->IsAdded())
@@ -330,7 +358,7 @@ bool FPackageSourceControlHelper::Checkout(const TArray<FString>& PackageNames)
 		{
 			if (!IPlatformFile::GetPlatformPhysical().FileExists(*PackageFilename))
 			{
-				UE_LOG(LogCommandletPackageHelper, Error, TEXT("File %s cannot be checked out as it does not exist"), *PackageFilename);
+				FPackageSourceControlHelperLog::Error(bErrorsAsWarnings, FString::Printf(TEXT("File %s cannot be checked out as it does not exist"), *PackageFilename));
 				bSuccess = false;
 			}
 			else if (IPlatformFile::GetPlatformPhysical().IsReadOnly(*PackageFilename))
@@ -341,7 +369,7 @@ bool FPackageSourceControlHelper::Checkout(const TArray<FString>& PackageNames)
 	}
 
 	// Any error up to here will be an early out
-	if (!bSuccess)
+	if (!bSuccess && !bErrorsAsWarnings)
 	{
 		return false;
 	}
@@ -349,11 +377,16 @@ bool FPackageSourceControlHelper::Checkout(const TArray<FString>& PackageNames)
 	// In the second pass, we will perform the checkout operation
 	if (PackagesToCheckout.Num() == 0)
 	{
-		return true;
+		return bSuccess;
 	}
 	else if (bUseSourceControl)
 	{
-		return GetSourceControlProvider().Execute(ISourceControlOperation::Create<FCheckOut>(), PackagesToCheckout) == ECommandResult::Succeeded;
+		if (GetSourceControlProvider().Execute(ISourceControlOperation::Create<FCheckOut>(), PackagesToCheckout) != ECommandResult::Succeeded)
+		{
+			// If operation didn't succeed. Get the list of invalid states when provided with a OutFailedPackages
+			FPackageSourceControlHelperLog::Error(bErrorsAsWarnings, TEXT("Error checking out packages from source control"));
+			return false;
+		}
 	}
 	else
 	{
@@ -363,7 +396,7 @@ bool FPackageSourceControlHelper::Checkout(const TArray<FString>& PackageNames)
 		{
 			if (!IPlatformFile::GetPlatformPhysical().SetReadOnly(*PackagesToCheckout[PackageIndex], false))
 			{
-				UE_LOG(LogCommandletPackageHelper, Error, TEXT("Error setting %s writable"), *PackagesToCheckout[PackageIndex]);
+				FPackageSourceControlHelperLog::Error(bErrorsAsWarnings, FString::Printf(TEXT("Error setting %s writable"), *PackagesToCheckout[PackageIndex]));
 				bSuccess = false;
 				--PackageIndex;
 				break;
@@ -371,14 +404,83 @@ bool FPackageSourceControlHelper::Checkout(const TArray<FString>& PackageNames)
 		}
 
 		// If a file couldn't be made writeable, put back the files to their original state
-		if (!bSuccess)
+		if (!bSuccess && !bErrorsAsWarnings)
 		{
 			for (; PackageIndex >= 0; --PackageIndex)
 			{
 				IPlatformFile::GetPlatformPhysical().SetReadOnly(*PackagesToCheckout[PackageIndex], true);
 			}
 		}
+	}
 
-		return bSuccess;
+	return bSuccess;
+}
+
+bool FPackageSourceControlHelper::GetDesiredStatesForModification(const TArray<FString>& PackageNames, TArray<FString>& OutPackagesToCheckout, TArray<FString>& OutPackagesToAdd, bool bErrorsAsWarnings) const
+{
+	// Convert package names to package filenames
+	TMap<FString, FString> PackageFilenamesToPackageName;
+	TArray<FString> PackageFilenames;
+	PackageFilenames.Reserve(PackageNames.Num());
+
+	for (const FString& PackageName : PackageNames)
+	{
+		FString PackageFilename = SourceControlHelpers::PackageFilename(PackageName);
+		PackageFilenamesToPackageName.Add(PackageFilename, PackageName);
+		PackageFilenames.Add(PackageFilename);
+	}
+
+	if (!UseSourceControl())
+	{
+		for (const FString& PackageFilename : PackageFilenames)
+		{
+			if (IPlatformFile::GetPlatformPhysical().FileExists(*PackageFilename) && IPlatformFile::GetPlatformPhysical().IsReadOnly(*PackageFilename))
+			{
+				OutPackagesToCheckout.Add(PackageFilenamesToPackageName.FindChecked(PackageFilename));
+			}
+		}
+		return true;
+	}
+		
+	TArray<FSourceControlStateRef> SourceControlStates;
+	if (GetSourceControlProvider().GetState(PackageFilenames, SourceControlStates, EStateCacheUsage::ForceUpdate) != ECommandResult::Succeeded)
+	{
+		// Nothing we can do if SCCStates fail
+		FPackageSourceControlHelperLog::Error(false, TEXT("Could not get source control state for packages"));
+		return false;
+	}
+
+	bool bSuccess = true;
+
+
+	for (FSourceControlStateRef& SourceControlState : SourceControlStates)
+	{
+		const FString& PackageFilename = SourceControlState->GetFilename();
+
+		FString OtherCheckedOutUser;
+		if (SourceControlState->IsCheckedOutOther(&OtherCheckedOutUser))
+		{
+			FPackageSourceControlHelperLog::Error(bErrorsAsWarnings, FString::Printf(TEXT("File %s already checked out by %s, will not checkout"), *PackageFilename, *OtherCheckedOutUser));
+			bSuccess = false;
+		}
+		else if (!SourceControlState->IsCurrent())
+		{
+			FPackageSourceControlHelperLog::Error(bErrorsAsWarnings, FString::Printf(TEXT("File %s (not at head revision), will not checkout"), *PackageFilename));
+			bSuccess = false;
+		}
+		else if (SourceControlState->IsCheckedOut() || SourceControlState->IsAdded())
+		{
+			// Nothing to do
+		}
+		else if (SourceControlState->IsSourceControlled())
+		{
+			OutPackagesToCheckout.Add(PackageFilenamesToPackageName.FindChecked(PackageFilename));
+		}
+		else
+		{
+			OutPackagesToAdd.Add(PackageFilenamesToPackageName.FindChecked(PackageFilename));
+		}
 	}
+
+	return bSuccess;
 }
\ No newline at end of file
diff --git a/Engine/Source/Editor/UnrealEd/Private/WorldPartition/WorldPartitionBuilder.cpp b/Engine/Source/Editor/UnrealEd/Private/WorldPartition/WorldPartitionBuilder.cpp
--- a/Engine/Source/Editor/UnrealEd/Private/WorldPartition/WorldPartitionBuilder.cpp
+++ b/Engine/Source/Editor/UnrealEd/Private/WorldPartition/WorldPartitionBuilder.cpp
@@ -16,9 +16,26 @@
 #include "WorldPartition/DataLayer/DataLayer.h"
 #include "LevelInstance/LevelInstanceSubsystem.h"
 #include "Math/IntVector.h"
+#include "UObject/SavePackage.h"
+#include "Algo/Transform.h"
 
 DEFINE_LOG_CATEGORY_STATIC(LogWorldPartitionBuilder, Log, All);
 
+namespace FWorldPartitionBuilderLog
+{
+	void Error(bool bErrorsAsWarnings, const FString& Msg)
+	{
+		if (bErrorsAsWarnings)
+		{
+			UE_LOG(LogWorldPartitionBuilder, Warning, TEXT("%s"), *Msg);
+		}
+		else
+		{
+			UE_LOG(LogWorldPartitionBuilder, Error, TEXT("%s"), *Msg);
+		}
+	}
+}
+
 FCellInfo::FCellInfo()
 	: Location(ForceInitToZero)
 	, Bounds(ForceInitToZero)
@@ -261,4 +278,85 @@ bool UWorldPartitionBuilder::Run(UWorld* World, FPackageSourceControlHelper& Pac
 	}
 
 	return PostRun(World, PackageHelper, bResult);
+}
+
+bool UWorldPartitionBuilder::SavePackages(const TArray<UPackage*>& Packages, FPackageSourceControlHelper& PackageHelper, bool bErrorsAsWarnings)
+{
+	if (Packages.IsEmpty())
+	{
+		return true;
+	}
+	UE_LOG(LogWorldPartitionBuilder, Display, TEXT("Saving %d packages..."), Packages.Num());
+
+	const TArray<FString> PackageFilenames = SourceControlHelpers::PackageFilenames(Packages);
+	
+	TArray<FString> PackagesToCheckout;
+	TArray<FString> PackagesToAdd;
+	TArray<FString> PackageNames;
+
+	PackageNames.Reserve(Packages.Num());
+	Algo::Transform(Packages, PackageNames, [](UPackage* InPackage) { return InPackage->GetName(); });
+	bool bSuccess = PackageHelper.GetDesiredStatesForModification(PackageNames, PackagesToCheckout, PackagesToAdd, bErrorsAsWarnings);
+	if (!bSuccess && !bErrorsAsWarnings)
+	{
+		return false;
+	}
+		
+	if (PackagesToCheckout.Num())
+	{
+		if (!PackageHelper.Checkout(PackagesToCheckout, bErrorsAsWarnings))
+		{
+			bSuccess = false;
+			if (!bErrorsAsWarnings)
+			{
+				return false;
+			}
+		}
+	}
+
+	for (int PackageIndex = 0; PackageIndex < Packages.Num(); ++PackageIndex)
+	{
+		const FString& PackageFilename = PackageFilenames[PackageIndex];
+		// Check readonly flag in case some checkouts failed
+		if (!IPlatformFile::GetPlatformPhysical().IsReadOnly(*PackageFilename))
+		{
+			// Save package
+			FSavePackageArgs SaveArgs;
+			SaveArgs.TopLevelFlags = RF_Standalone;
+			if (!UPackage::SavePackage(Packages[PackageIndex], nullptr, *PackageFilenames[PackageIndex], SaveArgs))
+			{
+				bSuccess = false;
+				FWorldPartitionBuilderLog::Error(bErrorsAsWarnings, FString::Printf(TEXT("Error saving package %s."), *Packages[PackageIndex]->GetName()));
+				if(!bErrorsAsWarnings)
+				{
+					return false;
+				}
+			}
+		}
+		else
+		{
+			check(bErrorsAsWarnings);
+		}
+	}
+
+	if (PackagesToAdd.Num())
+	{
+		if (!PackageHelper.AddToSourceControl(PackagesToAdd, bErrorsAsWarnings))
+		{
+			return false;
+		}
+	}
+
+	return bSuccess;
+}
+
+bool UWorldPartitionBuilder::DeletePackages(const TArray<FString>& PackageNames, FPackageSourceControlHelper& PackageHelper, bool bErrorsAsWarnings)
+{
+	if (PackageNames.Num() > 0)
+	{
+		UE_LOG(LogWorldPartitionBuilder, Display, TEXT("Deleting %d packages..."), PackageNames.Num());
+		return PackageHelper.Delete(PackageNames, bErrorsAsWarnings);
+	}
+
+	return true;
 }
\ No newline at end of file
diff --git a/Engine/Source/Editor/UnrealEd/Private/WorldPartition/WorldPartitionFoliageBuilder.cpp b/Engine/Source/Editor/UnrealEd/Private/WorldPartition/WorldPartitionFoliageBuilder.cpp
--- a/Engine/Source/Editor/UnrealEd/Private/WorldPartition/WorldPartitionFoliageBuilder.cpp
+++ b/Engine/Source/Editor/UnrealEd/Private/WorldPartition/WorldPartitionFoliageBuilder.cpp
@@ -201,7 +201,7 @@ bool UWorldPartitionFoliageBuilder::RunInternal(UWorld* World, const FCellInfo&
         }
 					
 		return true;
-	}, /*bGCPerActor=*/ false);
+	});
 
 	check(NumInstances == NumInstancesProcessed);
 
diff --git a/Engine/Source/Editor/UnrealEd/Private/WorldPartition/WorldPartitionRenameDuplicateBuilder.cpp b/Engine/Source/Editor/UnrealEd/Private/WorldPartition/WorldPartitionRenameDuplicateBuilder.cpp
--- a/Engine/Source/Editor/UnrealEd/Private/WorldPartition/WorldPartitionRenameDuplicateBuilder.cpp
+++ b/Engine/Source/Editor/UnrealEd/Private/WorldPartition/WorldPartitionRenameDuplicateBuilder.cpp
@@ -16,7 +16,7 @@
 #include "ProfilingDebugging/ScopedTimers.h"
 #include "HAL/FileManager.h"
 
-DEFINE_LOG_CATEGORY_STATIC(LogWorldPartitionCopyWorldBuilder, All, All);
+DEFINE_LOG_CATEGORY_STATIC(LogWorldPartitionRenameDuplicateBuilder, All, All);
 
 class FReplaceObjectRefsArchive : public FArchiveUObject
 {
@@ -70,7 +70,7 @@ private:
 
 static bool DeleteExistingMapPackages(const FString& ExistingPackageName, FPackageSourceControlHelper& PackageHelper)
 {
-	UE_SCOPED_TIMER(TEXT("Delete existing destination packages"), LogWorldPartitionCopyWorldBuilder, Display);
+	UE_SCOPED_TIMER(TEXT("Delete existing destination packages"), LogWorldPartitionRenameDuplicateBuilder, Display);
 	TArray<FString> PackagesToDelete;
 
 	FString ExistingMapPackageFilePath = FPackageName::LongPackageNameToFilename(ExistingPackageName, FPackageName::GetMapPackageExtension());
@@ -109,79 +109,13 @@ static bool DeleteExistingMapPackages(const FString& ExistingPackageName, FPacka
 
 			if (!bSuccess)
 			{
-				UE_LOG(LogWorldPartitionCopyWorldBuilder, Error, TEXT("Failed to iterate existing external actors path: %s"), *ExternalPackagesPath);
+				UE_LOG(LogWorldPartitionRenameDuplicateBuilder, Error, TEXT("Failed to iterate existing external actors path: %s"), *ExternalPackagesPath);
 				return false;
 			}
 		}
 	}
 
-	UE_LOG(LogWorldPartitionCopyWorldBuilder, Display, TEXT("Deleting %d packages..."), PackagesToDelete.Num());
-	if (PackagesToDelete.Num() > 0)
-	{
-		if (!PackageHelper.Delete(PackagesToDelete))
-		{
-			UE_LOG(LogWorldPartitionCopyWorldBuilder, Error, TEXT("Failed to delete existing destination packages:"));
-			for (const FString& PackageToDelete : PackagesToDelete)
-			{
-				UE_LOG(LogWorldPartitionCopyWorldBuilder, Error, TEXT("    Package: %s"), *PackageToDelete);
-			}
-			return false;
-		}
-	}
-
-	return true;
-}
-
-static bool SavePackages(const TArray<UPackage*>& Packages, FPackageSourceControlHelper& PackageHelper)
-{
-	const TArray<FString> PackageFilenames = SourceControlHelpers::PackageFilenames(Packages);
-	TArray<FString> PackagesToCheckout;
-	TArray<FString> PackagesToAdd;
-
-	for (int PackageIndex = 0; PackageIndex < Packages.Num(); ++PackageIndex)
-	{
-		const FString& PackageFilename = PackageFilenames[PackageIndex];
-		const bool bFileExists = IPlatformFile::GetPlatformPhysical().FileExists(*PackageFilename);
-
-		if (bFileExists)
-		{
-			PackagesToCheckout.Add(PackageFilename);
-		}
-		else
-		{
-			PackagesToAdd.Add(PackageFilename);
-		}
-	}
-
-	if (PackagesToCheckout.Num())
-	{
-		if (!PackageHelper.Checkout(PackagesToCheckout))
-		{
-			return false;
-		}
-	}
-
-	for(int PackageIndex = 0; PackageIndex < Packages.Num(); ++PackageIndex)
-	{
-		// Save package
-		FSavePackageArgs SaveArgs;
-		SaveArgs.TopLevelFlags = RF_Standalone;
-		if (!UPackage::SavePackage(Packages[PackageIndex], nullptr, *PackageFilenames[PackageIndex], SaveArgs))
-		{
-			UE_LOG(LogWorldPartitionCopyWorldBuilder, Error, TEXT("Error saving package %s."), *Packages[PackageIndex]->GetName());
-			return false;
-		}
-	}
-
-	if (PackagesToAdd.Num())
-	{
-		if (!PackageHelper.AddToSourceControl(PackagesToAdd))
-		{
-			return false;
-		}
-	}
-
-	return true;
+	return UWorldPartitionBuilder::DeletePackages(PackagesToDelete, PackageHelper);
 }
 
 UWorldPartitionRenameDuplicateBuilder::UWorldPartitionRenameDuplicateBuilder(const FObjectInitializer& ObjectInitializer)
@@ -196,7 +130,7 @@ bool UWorldPartitionRenameDuplicateBuilder::RunInternal(UWorld* World, const FCe
 	UWorldPartition* WorldPartition = World->GetWorldPartition();
 	if (!WorldPartition)
 	{
-		UE_LOG(LogWorldPartitionCopyWorldBuilder, Error, TEXT("Failed to retrieve WorldPartition."));
+		UE_LOG(LogWorldPartitionRenameDuplicateBuilder, Error, TEXT("Failed to retrieve WorldPartition."));
 		return false;
 	}
 		
@@ -212,7 +146,7 @@ bool UWorldPartitionRenameDuplicateBuilder::RunInternal(UWorld* World, const FCe
 
 	TArray<FActorCluster> ActorClusters;
 	{
-		UE_SCOPED_TIMER(TEXT("Create actor clusters"), LogWorldPartitionCopyWorldBuilder, Display);
+		UE_SCOPED_TIMER(TEXT("Create actor clusters"), LogWorldPartitionRenameDuplicateBuilder, Display);
 		FActorClusterContext::CreateActorClusters(World, ActorDescViewMap, ActorClusters);
 	}
 
@@ -241,7 +175,7 @@ bool UWorldPartitionRenameDuplicateBuilder::RunInternal(UWorld* World, const FCe
 	// Delete destination if it exists
 	if (!DeleteExistingMapPackages(NewPackageName, PackageHelper))
 	{
-		UE_LOG(LogWorldPartitionCopyWorldBuilder, Error, TEXT("Failed to delete existing destination package."));
+		UE_LOG(LogWorldPartitionRenameDuplicateBuilder, Error, TEXT("Failed to delete existing destination package."));
 		return false;
 	}
 				
@@ -250,7 +184,7 @@ bool UWorldPartitionRenameDuplicateBuilder::RunInternal(UWorld* World, const FCe
 	TArray<UPackage*> DuplicatedPackagesToSave;
 	UWorld* NewWorld = nullptr;
 	{
-		UE_SCOPED_TIMER(TEXT("Duplicating world"), LogWorldPartitionCopyWorldBuilder, Display);
+		UE_SCOPED_TIMER(TEXT("Duplicating world"), LogWorldPartitionRenameDuplicateBuilder, Display);
 		FObjectDuplicationParameters DuplicationParameters(World, NewPackage);
 		DuplicationParameters.DuplicateMode = EDuplicateMode::World;
 
@@ -298,7 +232,7 @@ bool UWorldPartitionRenameDuplicateBuilder::RunInternal(UWorld* World, const FCe
 		FSoftObjectPathFixupArchive SoftObjectPathFixupArchive(OriginalPackageName + TEXT(".") + OriginalWorldName, NewPackageName + TEXT(".") + NewWorldName);
 
 		{
-			UE_SCOPED_TIMER(TEXT("Saving actors"), LogWorldPartitionCopyWorldBuilder, Display);
+			UE_SCOPED_TIMER(TEXT("Saving actors"), LogWorldPartitionRenameDuplicateBuilder, Display);
 								
 			auto ProcessLoadedActors = [&](TArray<FWorldPartitionReference>& ActorReferences) -> bool
 			{
@@ -347,13 +281,13 @@ bool UWorldPartitionRenameDuplicateBuilder::RunInternal(UWorld* World, const FCe
 					ActorPackages.Add(Actor->GetPackage());
 				}
 
-				UE_LOG(LogWorldPartitionCopyWorldBuilder, Display, TEXT("Saving %d actor(s)"), ActorPackages.Num());
-				if (!SavePackages(ActorPackages, PackageHelper))
+				UE_LOG(LogWorldPartitionRenameDuplicateBuilder, Display, TEXT("Saving %d actor(s)"), ActorPackages.Num());
+				if (!UWorldPartitionBuilder::SavePackages(ActorPackages, PackageHelper))
 				{
-					UE_LOG(LogWorldPartitionCopyWorldBuilder, Error, TEXT("Failed to save actor packages:"));
+					UE_LOG(LogWorldPartitionRenameDuplicateBuilder, Error, TEXT("Failed to save actor packages:"));
 					for (UPackage* ActorPackage : ActorPackages)
 					{
-						UE_LOG(LogWorldPartitionCopyWorldBuilder, Error, TEXT("    Package: %s"), *ActorPackage->GetName());
+						UE_LOG(LogWorldPartitionRenameDuplicateBuilder, Error, TEXT("    Package: %s"), *ActorPackage->GetName());
 					}
 					return false;
 				}
@@ -392,7 +326,7 @@ bool UWorldPartitionRenameDuplicateBuilder::RunInternal(UWorld* World, const FCe
 			TArray<FWorldPartitionReference> ActorReferences;
 			for (const FActorCluster& ActorCluster : ActorClusters)
 			{
-				UE_LOG(LogWorldPartitionCopyWorldBuilder, Display, TEXT("Processing cluster with %d actor(s)"), ActorCluster.Actors.Num());
+				UE_LOG(LogWorldPartitionRenameDuplicateBuilder, Display, TEXT("Processing cluster with %d actor(s)"), ActorCluster.Actors.Num());
 				for (const FGuid& ActorGuid : ActorCluster.Actors)
 				{
 					// Duplicated actors don't need to be processed
@@ -432,8 +366,8 @@ bool UWorldPartitionRenameDuplicateBuilder::RunInternal(UWorld* World, const FCe
 
 		{
 			// Save all duplicated packages
-			UE_SCOPED_TIMER(TEXT("Saving new map packages"), LogWorldPartitionCopyWorldBuilder, Display);
-			if (!SavePackages(DuplicatedPackagesToSave, PackageHelper))
+			UE_SCOPED_TIMER(TEXT("Saving new map packages"), LogWorldPartitionRenameDuplicateBuilder, Display);
+			if (!UWorldPartitionBuilder::SavePackages(DuplicatedPackagesToSave, PackageHelper))
 			{
 				return false;
 			}
@@ -442,7 +376,7 @@ bool UWorldPartitionRenameDuplicateBuilder::RunInternal(UWorld* World, const FCe
 
 		{
 			// Validate results
-			UE_SCOPED_TIMER(TEXT("Validating actors"), LogWorldPartitionCopyWorldBuilder, Display);
+			UE_SCOPED_TIMER(TEXT("Validating actors"), LogWorldPartitionRenameDuplicateBuilder, Display);
 			for (UActorDescContainer::TConstIterator<> ActorDescIterator(WorldPartition); ActorDescIterator; ++ActorDescIterator)
 			{
 				const FWorldPartitionActorDesc* SourceActorDesc = *ActorDescIterator;
@@ -450,13 +384,13 @@ bool UWorldPartitionRenameDuplicateBuilder::RunInternal(UWorld* World, const FCe
 				const FWorldPartitionActorDesc* NewActorDesc = NewWorld->GetWorldPartition()->GetActorDesc(DuplicatedGuid ? *DuplicatedGuid : SourceActorDesc->GetGuid());
 				if (!NewActorDesc)
 				{
-					UE_LOG(LogWorldPartitionCopyWorldBuilder, Warning, TEXT("Failed to find source actor for Actor: %s"), *SourceActorDesc->GetActorPath().ToString());
+					UE_LOG(LogWorldPartitionRenameDuplicateBuilder, Warning, TEXT("Failed to find source actor for Actor: %s"), *SourceActorDesc->GetActorPath().ToString());
 				}
 				else
 				{
 					if (NewActorDesc->GetReferences().Num() != SourceActorDesc->GetReferences().Num())
 					{
-						UE_LOG(LogWorldPartitionCopyWorldBuilder, Warning, TEXT("Actor: %s and Source Actor: %s have mismatching reference count"), *NewActorDesc->GetActorPath().ToString(), *SourceActorDesc->GetActorPath().ToString());
+						UE_LOG(LogWorldPartitionRenameDuplicateBuilder, Warning, TEXT("Actor: %s and Source Actor: %s have mismatching reference count"), *NewActorDesc->GetActorPath().ToString(), *SourceActorDesc->GetActorPath().ToString());
 					}
 					else
 					{
@@ -465,7 +399,7 @@ bool UWorldPartitionRenameDuplicateBuilder::RunInternal(UWorld* World, const FCe
 							FGuid* DuplicateReferenceGuid = DuplicatedActorGuids.Find(ReferenceGuid);
 							if (!NewActorDesc->GetReferences().Contains(DuplicateReferenceGuid ? *DuplicateReferenceGuid : ReferenceGuid))
 							{
-								UE_LOG(LogWorldPartitionCopyWorldBuilder, Warning, TEXT("Actor: %s and Source Actor: %s have mismatching reference"), *NewActorDesc->GetActorPath().ToString(), *SourceActorDesc->GetActorPath().ToString());
+								UE_LOG(LogWorldPartitionRenameDuplicateBuilder, Warning, TEXT("Actor: %s and Source Actor: %s have mismatching reference"), *NewActorDesc->GetActorPath().ToString(), *SourceActorDesc->GetActorPath().ToString());
 							}
 						}
 					}
@@ -479,15 +413,15 @@ bool UWorldPartitionRenameDuplicateBuilder::RunInternal(UWorld* World, const FCe
 		
 	if (PackagesToDelete.Num() > 0)
 	{
-		UE_SCOPED_TIMER(TEXT("Delete source packages (-rename switch)"), LogWorldPartitionCopyWorldBuilder, Display);
+		UE_SCOPED_TIMER(TEXT("Delete source packages (-rename switch)"), LogWorldPartitionRenameDuplicateBuilder, Display);
 		
-		UE_LOG(LogWorldPartitionCopyWorldBuilder, Display, TEXT("Deleting %d packages"), PackagesToDelete.Num());
+		UE_LOG(LogWorldPartitionRenameDuplicateBuilder, Display, TEXT("Deleting %d packages"), PackagesToDelete.Num());
 		if (!PackageHelper.Delete(PackagesToDelete.Array()))
 		{
-			UE_LOG(LogWorldPartitionCopyWorldBuilder, Error, TEXT("Failed to delete source packages:"));
+			UE_LOG(LogWorldPartitionRenameDuplicateBuilder, Error, TEXT("Failed to delete source packages:"));
 			for (const FString& PackageToDelete : PackagesToDelete)
 			{
-				UE_LOG(LogWorldPartitionCopyWorldBuilder, Error, TEXT("    Package: %s"), *PackageToDelete);
+				UE_LOG(LogWorldPartitionRenameDuplicateBuilder, Error, TEXT("    Package: %s"), *PackageToDelete);
 			}
 			return false;
 		}
@@ -524,10 +458,10 @@ bool UWorldPartitionRenameDuplicateBuilder::PostWorldTeardown(FPackageSourceCont
 		RedirectorPackage->MarkAsFullyLoaded();
 
 		// Saving the NewPackage will save the duplicated external packages
-		UE_SCOPED_TIMER(TEXT("Saving new redirector"), LogWorldPartitionCopyWorldBuilder, Display);
-		if (!SavePackages({ RedirectorPackage }, PackageHelper))
+		UE_SCOPED_TIMER(TEXT("Saving new redirector"), LogWorldPartitionRenameDuplicateBuilder, Display);
+		if (!UWorldPartitionBuilder::SavePackages({ RedirectorPackage }, PackageHelper))
 		{
-			UE_LOG(LogWorldPartitionCopyWorldBuilder, Error, TEXT("Failed to save redirector package: %s"), *RedirectorPackage->GetName());
+			UE_LOG(LogWorldPartitionRenameDuplicateBuilder, Error, TEXT("Failed to save redirector package: %s"), *RedirectorPackage->GetName());
 			return false;
 		}
 
@@ -544,7 +478,7 @@ bool UWorldPartitionRenameDuplicateBuilder::PostWorldTeardown(FPackageSourceCont
 		UWorld* RedirectedWorld = CastChecked<UWorld>(RedirectorPath.TryLoad());
 		if (!RedirectedWorld)
 		{
-			UE_LOG(LogWorldPartitionCopyWorldBuilder, Error, TEXT("Failed to validate redirector package: %s"), *RedirectorPackage->GetName());
+			UE_LOG(LogWorldPartitionRenameDuplicateBuilder, Error, TEXT("Failed to validate redirector package: %s"), *RedirectorPackage->GetName());
 			return false;
 		}
 	}
diff --git a/Engine/Source/Editor/UnrealEd/Private/WorldPartition/WorldPartitionResaveActorsBuilder.cpp b/Engine/Source/Editor/UnrealEd/Private/WorldPartition/WorldPartitionResaveActorsBuilder.cpp
--- a/Engine/Source/Editor/UnrealEd/Private/WorldPartition/WorldPartitionResaveActorsBuilder.cpp
+++ b/Engine/Source/Editor/UnrealEd/Private/WorldPartition/WorldPartitionResaveActorsBuilder.cpp
@@ -90,14 +90,8 @@ bool UWorldPartitionResaveActorsBuilder::PreRun(UWorld* World, FPackageSourceCon
 
 bool UWorldPartitionResaveActorsBuilder::RunInternal(UWorld* World, const FCellInfo& InCellInfo, FPackageSourceControlHelper& PackageHelper)
 {
-	FPackageSourceControlHelper SCCHelper;
-
 	UPackage* WorldPackage = World->GetPackage();
 
-	int32 LoadCount = 0;
-	int32 SaveCount = 0;
-	int32 FailCount = 0;
-
 	// Actor Class Filter
 	UClass* ActorClass = AActor::StaticClass();
 
@@ -150,66 +144,6 @@ bool UWorldPartitionResaveActorsBuilder::RunInternal(UWorld* World, const FCellI
 
 	TArray<FString> PackagesToDelete;
 
-	auto AddAndSavePackage = [&SaveCount, &FailCount, &SCCHelper](UPackage* InPackage, const FString& InPackageType)
-	{
-		// Save the new package and add to source control if needed
-		FString NewPackageFileName = SourceControlHelpers::PackageFilename(InPackage);
-		FSavePackageArgs SaveArgs;
-		SaveArgs.TopLevelFlags = RF_Standalone;
-		if (!UPackage::SavePackage(InPackage, nullptr, *NewPackageFileName, SaveArgs))
-		{
-			UE_LOG(LogWorldPartitionResaveActorsBuilder, Error, TEXT("Error saving %s package %s."), *InPackageType, *InPackage->GetName());
-			++FailCount;
-			return true;
-		}
-
-		if (!SCCHelper.AddToSourceControl(InPackage))
-		{
-			UE_LOG(LogWorldPartitionResaveActorsBuilder, Error, TEXT("Error adding %s package to source control %s."), *InPackageType, *InPackage->GetName());
-			++FailCount;
-			return true;
-		}
-
-		UE_LOG(LogWorldPartitionResaveActorsBuilder, Display, TEXT("Saved %s package %s."), *InPackageType, *InPackage->GetName());
-		++SaveCount;
-		return true;
-	};
-
-	auto CheckOutAndSavePackage = [&SaveCount, &FailCount, &SCCHelper](UPackage* InPackage, const FString& InPackageType, bool bErrorAtFailCheckout)
-	{
-		if (SCCHelper.Checkout(InPackage))
-		{
-			// Save package
-			FString PackageFileName = SourceControlHelpers::PackageFilename(InPackage);
-			FSavePackageArgs SaveArgs;
-			SaveArgs.TopLevelFlags = RF_Standalone;
-			if (!UPackage::SavePackage(InPackage, nullptr, *PackageFileName, SaveArgs))
-			{
-				UE_LOG(LogWorldPartitionResaveActorsBuilder, Error, TEXT("Error saving %s package %s."), *InPackageType, *InPackage->GetName());
-				++FailCount;
-				return true;
-			}
-
-			UE_LOG(LogWorldPartitionResaveActorsBuilder, Display, TEXT("Saved %s package %s."), *InPackageType, *InPackage->GetName());
-			++SaveCount;
-			return true;
-		}
-		else
-		{
-			if (bErrorAtFailCheckout)
-			{
-				UE_LOG(LogWorldPartitionResaveActorsBuilder, Error, TEXT("Error checking out %s package %s."), *InPackageType, *InPackage->GetName());
-			}
-			else
-			{
-				UE_LOG(LogWorldPartitionResaveActorsBuilder, Warning, TEXT("Error checking out %s package %s."), *InPackageType, *InPackage->GetName());
-			}
-
-			++FailCount;
-			return true;
-		}
-	};
-
 	if (bSwitchActorPackagingSchemeToReduced)
 	{
 		World->PersistentLevel->ActorPackagingScheme = EActorPackagingScheme::Reduced;
@@ -223,16 +157,18 @@ bool UWorldPartitionResaveActorsBuilder::RunInternal(UWorld* World, const FCellI
 
 		TArray<TArray<FGuid>> Clusters = GenerateObjectsClusters(ActorsWithRefs);
 
+		TSet<FWorldPartitionReference> LoadedReferences;
+		TArray<UPackage*> PackagesToSave;
 		for (const TArray<FGuid>& Cluster : Clusters)
 		{
-			TArray<UPackage*> PackagesToSave;
-
 			// Load actor clusters
 			TArray<FWorldPartitionReference> ActorReferences;
+			ActorReferences.Reserve(Cluster.Num());
 			for (const FGuid& ActorGuid : Cluster)
 			{
 				ActorReferences.Emplace(WorldPartition, ActorGuid);
 			}
+			LoadedReferences.Append(ActorReferences);
 
 			// Change packaging of all actors in the current cluster
 			for (FWorldPartitionReference& ActorReference : ActorReferences)
@@ -243,29 +179,24 @@ bool UWorldPartitionResaveActorsBuilder::RunInternal(UWorld* World, const FCellI
 				if (!Actor)
 				{
 					PackagesToDelete.Add(ActorDesc->GetActorPackage().ToString());
-					++FailCount;
 					continue;
 				}
 				else
 				{
-					++LoadCount;
-
 					UPackage* Package = Actor->GetExternalPackage();
 					check(Package);
 
 					if (!bReportOnly)
 					{
-						if (SCCHelper.Checkout(Package))
-						{
-							// Always mark this package for deletion, as it will contain a temporary redirector to fixup references
-							PackagesToDelete.Add(ActorDesc->GetActorPackage().ToString());
+						// Always mark this package for deletion, as it will contain a temporary redirector to fixup references
+						PackagesToDelete.Add(ActorDesc->GetActorPackage().ToString());
 
-							// Move actor back into world's package
-							Actor->SetPackageExternal(false);
+						// Move actor back into world's package
+						Actor->SetPackageExternal(false);
 
-							// Gather dependant objects that also needs to be moved along with the actor
-							TArray<UObject*> DependantObjects;
-							ForEachObjectWithPackage(Package, [&DependantObjects](UObject* Object)
+						// Gather dependant objects that also needs to be moved along with the actor
+						TArray<UObject*> DependantObjects;
+						ForEachObjectWithPackage(Package, [&DependantObjects](UObject* Object)
 							{
 								if (!Cast<UMetaData>(Object))
 								{
@@ -274,69 +205,69 @@ bool UWorldPartitionResaveActorsBuilder::RunInternal(UWorld* World, const FCellI
 								return true;
 							}, false);
 
-							// Move dependant objects into the new world package temporarily
-							for (UObject* DependantObject : DependantObjects)
-							{
-								DependantObject->Rename(nullptr, WorldPackage, REN_NonTransactional | REN_DontCreateRedirectors | REN_ForceNoResetLoaders);
-							}
-
-							// Move actor in its new package
-							Actor->SetPackageExternal(true);
+						// Move dependant objects into the new world package temporarily
+						for (UObject* DependantObject : DependantObjects)
+						{
+							DependantObject->Rename(nullptr, WorldPackage, REN_NonTransactional | REN_DontCreateRedirectors | REN_ForceNoResetLoaders);
+						}
 
-							// Also move dependant objects into the new package
-							UPackage* NewActorPackage = Actor->GetExternalPackage();
-							for (UObject* DependantObject : DependantObjects)
-							{
-								DependantObject->Rename(nullptr, NewActorPackage, REN_NonTransactional | REN_DontCreateRedirectors | REN_ForceNoResetLoaders);
-							}
+						// Move actor in its new package
+						Actor->SetPackageExternal(true);
 
-							PackagesToSave.Add(NewActorPackage);
-						}
-						else
+						// Also move dependant objects into the new package
+						UPackage* NewActorPackage = Actor->GetExternalPackage();
+						for (UObject* DependantObject : DependantObjects)
 						{
-							// It is possible the resave can't checkout everything. Continue processing.
-							UE_LOG(LogWorldPartitionResaveActorsBuilder, Warning, TEXT("Error checking out package %s."), *Package->GetName());
-							++FailCount;
-							continue;
+							DependantObject->Rename(nullptr, NewActorPackage, REN_NonTransactional | REN_DontCreateRedirectors | REN_ForceNoResetLoaders);
 						}
+
+						PackagesToSave.Add(NewActorPackage);
 					}
 				}
 			}
 
-			// Save modified actors
-			for (UPackage* PackageToSave : PackagesToSave)
+			if (FWorldPartitionHelpers::HasExceededMaxMemory())
 			{
-				AddAndSavePackage(PackageToSave, TEXT("Actor"));
+				if (!UWorldPartitionBuilder::SavePackages(PackagesToSave, PackageHelper))
+				{
+					return false;
+				}
+
+				PackagesToSave.Empty();
+				LoadedReferences.Empty();
+				FWorldPartitionHelpers::DoCollectGarbage();
 			}
 		}
+
+		// Save last batch
+		if (!UWorldPartitionBuilder::SavePackages(PackagesToSave, PackageHelper))
+		{
+			return false;
+		}
+		PackagesToSave.Empty();
+		LoadedReferences.Empty();
+		FWorldPartitionHelpers::DoCollectGarbage();
 	}
 	else
 	{
-		FWorldPartitionHelpers::ForEachActorWithLoading(WorldPartition, ActorClass, [this, &LoadCount, &SaveCount, &FailCount, &SCCHelper, &PackagesToDelete, &CheckOutAndSavePackage](const FWorldPartitionActorDesc* ActorDesc)
+		TArray<UPackage*> PackagesToCheckout;
+		FWorldPartitionHelpers::ForEachActorWithLoading(WorldPartition, ActorClass, 
+		[this, &PackagesToDelete, &PackagesToCheckout](const FWorldPartitionActorDesc* ActorDesc)
 		{
-			ON_SCOPE_EXIT
-			{
-				UE_LOG(LogWorldPartitionResaveActorsBuilder, Display, TEXT("Processed %d packages (%d Saved / %d Failed)"), LoadCount, SaveCount, FailCount);
-			};
-
 			AActor* Actor = ActorDesc->GetActor();
 
 			if (!Actor)
 			{
 				PackagesToDelete.Add(ActorDesc->GetActorPackage().ToString());
-				++FailCount;
 				return true;
 			}
 			else
 			{
-				++LoadCount;
-
 				if (bEnableActorFolders)
 				{
 					if (!Actor->CreateOrUpdateActorFolder())
 					{
 						UE_LOG(LogWorldPartitionResaveActorsBuilder, Error, TEXT("Failed to create actor folder for actor %s."), *Actor->GetName());
-						++FailCount;
 						return true;
 					}
 				}
@@ -358,32 +289,48 @@ bool UWorldPartitionResaveActorsBuilder::RunInternal(UWorld* World, const FCellI
 
 				if (!bReportOnly)
 				{
-					return CheckOutAndSavePackage(Package, TEXT("Actor"), /*bErrorAtFailCheckout*/ false);
+					PackagesToCheckout.Add(Package);
+					return true;
 				}
 			}
 
 			return true;
+		},
+		[&PackagesToCheckout, &PackageHelper]()
+		{
+			UWorldPartitionBuilder::SavePackages(PackagesToCheckout, PackageHelper, /*bErrorsAsWarnings*/true);
+			PackagesToCheckout.Empty();
 		});
 	}
 
 	if (!bReportOnly)
 	{
-		SCCHelper.Delete(PackagesToDelete);
+		if (!UWorldPartitionBuilder::DeletePackages(PackagesToDelete, PackageHelper, !bSwitchActorPackagingSchemeToReduced) && bSwitchActorPackagingSchemeToReduced)
+		{
+			return false;
+		}
 
 		if (bEnableActorFolders)
 		{
-			World->PersistentLevel->ForEachActorFolder([&AddAndSavePackage](UActorFolder* ActorFolder)
+			TArray<UPackage*> FolderPackages;
+			World->PersistentLevel->ForEachActorFolder([&FolderPackages](UActorFolder* ActorFolder)
 			{
 				UPackage* NewActorFolderPackage = ActorFolder->GetExternalPackage();
 				check(NewActorFolderPackage);
-				return AddAndSavePackage(NewActorFolderPackage, TEXT("Actor Folder"));
+				FolderPackages.Add(NewActorFolderPackage);
+				return true;
 			});
+
+			if (!UWorldPartitionBuilder::SavePackages(FolderPackages, PackageHelper))
+			{
+				return false;
+			}
 		}
 
 		const bool bNeedWorldResave = bSwitchActorPackagingSchemeToReduced || bEnableActorFolders;
 		if (bNeedWorldResave)
 		{
-			return CheckOutAndSavePackage(WorldPackage, TEXT("World"), /*bErrorAtFailCheckout*/ true);
+			return UWorldPartitionBuilder::SavePackages({ WorldPackage }, PackageHelper);
 		}
 	}
 
diff --git a/Engine/Source/Editor/UnrealEd/Public/PackageSourceControlHelper.h b/Engine/Source/Editor/UnrealEd/Public/PackageSourceControlHelper.h
--- a/Engine/Source/Editor/UnrealEd/Public/PackageSourceControlHelper.h
+++ b/Engine/Source/Editor/UnrealEd/Public/PackageSourceControlHelper.h
@@ -11,14 +11,14 @@ class UNREALED_API FPackageSourceControlHelper
 public:
 	bool UseSourceControl() const;
 	bool Delete(const FString& PackageName) const;
-	bool Delete(const TArray<FString>& PackageNames) const;
+	bool Delete(const TArray<FString>& PackageNames, bool bErrorsAsWarnings = false) const;
 	bool Delete(UPackage* Package) const;
 	bool Delete(const TArray<UPackage*>& Packages) const;
 	bool AddToSourceControl(UPackage* Package) const;
-	bool AddToSourceControl(const TArray<FString>& PackageNames) const;
+	bool AddToSourceControl(const TArray<FString>& PackageNames, bool bErrorsAsWarnings = false) const;
 	bool Checkout(UPackage* Package) const;
-	bool Checkout(const TArray<FString>& PackageNames) const;
-
+	bool Checkout(const TArray<FString>& PackageNames, bool bErrorsAsWarnings = false) const;
+	bool GetDesiredStatesForModification(const TArray<FString>& PackageNames, TArray<FString>& OutPackagesToCheckout, TArray<FString>& OutPackagesToAdd, bool bErrorsAsWarnings = false) const;
 private:
 	ISourceControlProvider& GetSourceControlProvider() const;
 	FScopedSourceControl SourceControl;	
diff --git a/Engine/Source/Editor/UnrealEd/Public/WorldPartition/WorldPartitionBuilder.h b/Engine/Source/Editor/UnrealEd/Public/WorldPartition/WorldPartitionBuilder.h
--- a/Engine/Source/Editor/UnrealEd/Public/WorldPartition/WorldPartitionBuilder.h
+++ b/Engine/Source/Editor/UnrealEd/Public/WorldPartition/WorldPartitionBuilder.h
@@ -58,6 +58,9 @@ public:
 
 	virtual bool PreWorldInitialization(FPackageSourceControlHelper& PackageHelper) { return true; }
 
+	static bool SavePackages(const TArray<UPackage*>& Packages, FPackageSourceControlHelper& PackageHelper, bool bErrorsAsWarnings = false);
+	static bool DeletePackages(const TArray<FString>& PackageNames, FPackageSourceControlHelper& PackageHelper, bool bErrorsAsWarnings = false);
+
 protected:
 	/**
 	 * Overridable method for derived classes to perform operations when world builder process starts.
diff --git a/Engine/Source/Runtime/Engine/Private/WorldPartition/WorldPartitionHelpers.cpp b/Engine/Source/Runtime/Engine/Private/WorldPartition/WorldPartitionHelpers.cpp
--- a/Engine/Source/Runtime/Engine/Private/WorldPartition/WorldPartitionHelpers.cpp
+++ b/Engine/Source/Runtime/Engine/Private/WorldPartition/WorldPartitionHelpers.cpp
@@ -56,7 +56,7 @@ namespace WorldPartitionHelpers
 		}
 	}
 
-	bool ForEachActorWithLoadingBody(const FGuid& ActorGuid, UWorldPartition* WorldPartition, TFunctionRef<bool(const FWorldPartitionActorDesc*)> Predicate, bool bGCPerActor, TMap<FGuid, FWorldPartitionReference>& ActorReferences)
+	bool ForEachActorWithLoadingBody(const FGuid& ActorGuid, UWorldPartition* WorldPartition, TFunctionRef<bool(const FWorldPartitionActorDesc*)> Predicate, TFunctionRef<void()> OnReleasingActorReferences, bool bGCPerActor, TMap<FGuid, FWorldPartitionReference>& ActorReferences)
 	{
 		WorldPartitionHelpers::LoadReferences(WorldPartition, ActorGuid, ActorReferences);
 
@@ -68,6 +68,7 @@ namespace WorldPartitionHelpers
 
 		if (bGCPerActor || FWorldPartitionHelpers::HasExceededMaxMemory())
 		{
+			OnReleasingActorReferences();
 			ActorReferences.Empty();
 			FWorldPartitionHelpers::DoCollectGarbage();
 		}
@@ -76,31 +77,33 @@ namespace WorldPartitionHelpers
 	}
 }
 
-void FWorldPartitionHelpers::ForEachActorWithLoading(UWorldPartition* WorldPartition, TSubclassOf<AActor> ActorClass, TFunctionRef<bool(const FWorldPartitionActorDesc*)> Predicate, bool bGCPerActor)
+void FWorldPartitionHelpers::ForEachActorWithLoading(UWorldPartition* WorldPartition, TSubclassOf<AActor> ActorClass, TFunctionRef<bool(const FWorldPartitionActorDesc*)> Predicate, TFunctionRef<void()> OnReleasingActorReferences, bool bGCPerActor)
 {
 	TMap<FGuid, FWorldPartitionReference> ActorReferences;
-		
+
 	ForEachActorDesc(WorldPartition, ActorClass, [&](const FWorldPartitionActorDesc* ActorDesc)
 	{
-		return WorldPartitionHelpers::ForEachActorWithLoadingBody(ActorDesc->GetGuid(), WorldPartition, Predicate, bGCPerActor, ActorReferences);
+		return WorldPartitionHelpers::ForEachActorWithLoadingBody(ActorDesc->GetGuid(), WorldPartition, Predicate, OnReleasingActorReferences, bGCPerActor, ActorReferences);
 	});
 
+	OnReleasingActorReferences();
 	ActorReferences.Empty();
 	DoCollectGarbage();
 }
 
-void FWorldPartitionHelpers::ForEachActorWithLoading(UWorldPartition* WorldPartition, const TArray<FGuid>& ActorGuids, TFunctionRef<bool(const FWorldPartitionActorDesc*)> Predicate, bool bGCPerActor)
+void FWorldPartitionHelpers::ForEachActorWithLoading(UWorldPartition* WorldPartition, const TArray<FGuid>& ActorGuids, TFunctionRef<bool(const FWorldPartitionActorDesc*)> Predicate, TFunctionRef<void()> OnReleasingActorReferences, bool bGCPerActor)
 {
 	TMap<FGuid, FWorldPartitionReference> ActorReferences;
 
 	for(const FGuid& ActorGuid : ActorGuids)
 	{
-		if (!WorldPartitionHelpers::ForEachActorWithLoadingBody(ActorGuid, WorldPartition, Predicate, bGCPerActor, ActorReferences))
+		if (!WorldPartitionHelpers::ForEachActorWithLoadingBody(ActorGuid, WorldPartition, Predicate, OnReleasingActorReferences, bGCPerActor, ActorReferences))
 		{
 			break;
 		}
 	}
 
+	OnReleasingActorReferences();
 	ActorReferences.Empty();
 	DoCollectGarbage();
 }
diff --git a/Engine/Source/Runtime/Engine/Public/WorldPartition/WorldPartitionHelpers.h b/Engine/Source/Runtime/Engine/Public/WorldPartition/WorldPartitionHelpers.h
--- a/Engine/Source/Runtime/Engine/Public/WorldPartition/WorldPartitionHelpers.h
+++ b/Engine/Source/Runtime/Engine/Public/WorldPartition/WorldPartitionHelpers.h
@@ -33,9 +33,9 @@ public:
 
 	static void ForEachIntersectingActorDesc(UWorldPartition* WorldPartition, const FBox& Box, TSubclassOf<AActor> ActorClass, TFunctionRef<bool(const FWorldPartitionActorDesc*)> Func);
 	static void ForEachActorDesc(UWorldPartition* WorldPartition, TSubclassOf<AActor> ActorClass, TFunctionRef<bool(const FWorldPartitionActorDesc*)> Func);
-	static void ForEachActorWithLoading(UWorldPartition* WorldPartition, TSubclassOf<AActor> ActorClass, TFunctionRef<bool(const FWorldPartitionActorDesc*)> Func, bool bGCPerActor = false);
-	static void ForEachActorWithLoading(UWorldPartition* WorldPartition, const TArray<FGuid>& ActorGuids, TFunctionRef<bool(const FWorldPartitionActorDesc*)> Func, bool bGCPerActor = false);
-
+	static void ForEachActorWithLoading(UWorldPartition* WorldPartition, TSubclassOf<AActor> ActorClass, TFunctionRef<bool(const FWorldPartitionActorDesc*)> Func, TFunctionRef<void()> OnReleasingActorReferences = [](){}, bool bGCPerActor = false);
+	static void ForEachActorWithLoading(UWorldPartition* WorldPartition, const TArray<FGuid>& ActorGuids, TFunctionRef<bool(const FWorldPartitionActorDesc*)> Func, TFunctionRef<void()> OnReleasingActorReferences = [](){}, bool bGCPerActor = false);
+	
 	static bool HasExceededMaxMemory();
 	static void DoCollectGarbage();
 
