Remove DataTable Factory Class Picker support

diff --git a/Engine/Source/Editor/UnrealEd/Classes/Factories/DataTableFactory.h b/Engine/Source/Editor/UnrealEd/Classes/Factories/DataTableFactory.h
--- a/Engine/Source/Editor/UnrealEd/Classes/Factories/DataTableFactory.h
+++ b/Engine/Source/Editor/UnrealEd/Classes/Factories/DataTableFactory.h
@@ -16,8 +16,6 @@ class UNREALED_API UDataTableFactory : public UFactory
 
 	UPROPERTY(BlueprintReadWrite, Category = "Data Table Factory")
 	TObjectPtr<const class UScriptStruct> Struct;
-	UPROPERTY(BlueprintReadWrite, Category = "Data Table Factory")
-	TObjectPtr<const UClass> TableClass;
 
 	//~ Begin UFactory Interface
 	virtual bool ConfigureProperties() override;
diff --git a/Engine/Source/Editor/UnrealEd/Private/Factories/DataTableFactory.cpp b/Engine/Source/Editor/UnrealEd/Private/Factories/DataTableFactory.cpp
--- a/Engine/Source/Editor/UnrealEd/Private/Factories/DataTableFactory.cpp
+++ b/Engine/Source/Editor/UnrealEd/Private/Factories/DataTableFactory.cpp
@@ -8,8 +8,6 @@
 #include "Modules/ModuleManager.h"
 #include "StructViewerModule.h"
 #include "StructViewerFilter.h"
-#include "ClassViewerModule.h"
-#include "ClassViewerFilter.h"
 
 #include "Widgets/SWindow.h"
 #include "Widgets/Input/SButton.h"
@@ -47,29 +45,12 @@ bool UDataTableFactory::ConfigureProperties()
 		}
 	};
 
-	class FDataTableClassFilter : public IClassViewerFilter
-	{
-	public:
-		virtual bool IsClassAllowed(const FClassViewerInitializationOptions& InInitOptions, const UClass* InClass, TSharedRef< class FClassViewerFilterFuncs > InFilterFuncs ) override
-		{
-			return InClass && !InClass->HasAnyClassFlags(EClassFlags::CLASS_Abstract) && InClass->IsChildOf(UDataTable::StaticClass());
-		}
-
-		virtual bool IsUnloadedClassAllowed(const FClassViewerInitializationOptions& InInitOptions, const TSharedRef< const class IUnloadedBlueprintData > InUnloadedClassData, TSharedRef< class FClassViewerFilterFuncs > InFilterFuncs) override
-		{
-			// DataTable cannot have Blueprint sub-classes.
-			return false;
-		}
-	};
-
 	class FDataTableFactoryUI : public TSharedFromThis<FDataTableFactoryUI>
 	{
 	public:
 		FReply OnCreate()
 		{
 			check(ResultStruct);
-			check(ResultClass);
-
 			if (PickerWindow.IsValid())
 			{
 				PickerWindow->RequestDestroyWindow();
@@ -80,8 +61,6 @@ bool UDataTableFactory::ConfigureProperties()
 		FReply OnCancel()
 		{
 			ResultStruct = nullptr;
-			ResultClass = nullptr;
-
 			if (PickerWindow.IsValid())
 			{
 				PickerWindow->RequestDestroyWindow();
@@ -89,9 +68,9 @@ bool UDataTableFactory::ConfigureProperties()
 			return FReply::Handled();
 		}
 
-		bool IsValidSelection() const
+		bool IsStructSelected() const
 		{
-			return ResultStruct && ResultClass;
+			return ResultStruct != nullptr;
 		}
 
 		void OnPickedStruct(const UScriptStruct* ChosenStruct)
@@ -100,26 +79,13 @@ bool UDataTableFactory::ConfigureProperties()
 			StructPickerAnchor->SetIsOpen(false);
 		}
 
-		void OnPickedClass(UClass* ChosenClass)
-		{
-			ResultClass = ChosenClass;
-			ClassPickerAnchor->SetIsOpen(false);
-		}
-
-		FText OnGetStructComboTextValue() const
+		FText OnGetComboTextValue() const
 		{
 			return ResultStruct
 				? FText::AsCultureInvariant(ResultStruct->GetName())
 				: LOCTEXT("None", "None");
 		}
 
-		FText OnGetClassComboTextValue() const
-		{
-			return ResultClass
-				? FText::AsCultureInvariant(ResultClass->GetName())
-				: LOCTEXT("None", "None");
-		}
-
 		TSharedRef<SWidget> GenerateStructPicker()
 		{
 			FStructViewerModule& StructViewerModule = FModuleManager::LoadModuleChecked<FStructViewerModule>("StructViewer");
@@ -149,38 +115,18 @@ bool UDataTableFactory::ConfigureProperties()
 				];
 		}
 
-		TSharedRef<SWidget> GenerateClassPicker()
+		const UScriptStruct* OpenStructSelector()
 		{
-			FClassViewerModule& ClassViewerModule = FModuleManager::LoadModuleChecked<FClassViewerModule>("ClassViewer");
-
-			FClassViewerInitializationOptions Options;
-			Options.Mode = EClassViewerMode::ClassPicker;
-			Options.ClassFilters.Add(MakeShared<FDataTableClassFilter>());
-			Options.bShowNoneOption = false;
-
-			return SNew(SBox)
-				.WidthOverride(330)
-				[
-					SNew(SVerticalBox)
+			FStructViewerModule& StructViewerModule = FModuleManager::LoadModuleChecked<FStructViewerModule>("StructViewer");
+			ResultStruct = nullptr;
 
-					+SVerticalBox::Slot()
-					.FillHeight(1.0f)
-					.MaxHeight(500)
-					[
-						SNew(SBorder)
-						.Padding(4)
-						.BorderImage(FEditorStyle::GetBrush("ToolPanel.GroupBorder"))
-						[
-							ClassViewerModule.CreateClassViewer(Options, FOnClassPicked::CreateSP(this, &FDataTableFactoryUI::OnPickedClass))
-						]
-					]
-				];
-		}
+			// Fill in options
+			FStructViewerInitializationOptions Options;
+			Options.Mode = EStructViewerMode::StructPicker;
+			Options.StructFilter = MakeShared<FDataTableStructFilter>();
 
-		bool OpenSelectorDialog(TObjectPtr<const UClass> ClassResult, TObjectPtr<const UScriptStruct> StructResult)
-		{
 			PickerWindow = SNew(SWindow)
-				.Title(LOCTEXT("DataTableFactoryOptions", "Pick Class & Row Structure"))
+				.Title(LOCTEXT("DataTableFactoryOptions", "Pick Row Structure"))
 				.ClientSize(FVector2D(350, 100))
 				.SupportsMinimize(false)
 				.SupportsMaximize(false)
@@ -190,21 +136,6 @@ bool UDataTableFactory::ConfigureProperties()
 					.Padding(10)
 					[
 						SNew(SVerticalBox)
-						// Class Picker
-						+SVerticalBox::Slot()
-						.AutoHeight()
-						[
-							SAssignNew(ClassPickerAnchor, SComboButton)
-							.ContentPadding(FMargin(2,2,2,1))
-							.MenuPlacement(MenuPlacement_BelowAnchor)
-							.ButtonContent()
-							[
-								SNew(STextBlock)
-								.Text(this, &FDataTableFactoryUI::OnGetClassComboTextValue)
-							]
-							.OnGetMenuContent(this, &FDataTableFactoryUI::GenerateClassPicker)
-						]
-						// Struct Picker
 						+SVerticalBox::Slot()
 						.AutoHeight()
 						[
@@ -214,7 +145,7 @@ bool UDataTableFactory::ConfigureProperties()
 							.ButtonContent()
 							[
 								SNew(STextBlock)
-								.Text(this, &FDataTableFactoryUI::OnGetStructComboTextValue)
+								.Text(this, &FDataTableFactoryUI::OnGetComboTextValue)
 							]
 							.OnGetMenuContent(this, &FDataTableFactoryUI::GenerateStructPicker)
 						]
@@ -228,7 +159,7 @@ bool UDataTableFactory::ConfigureProperties()
 							[
 								SNew(SButton)
 								.Text(LOCTEXT("OK", "OK"))
-								.IsEnabled(this, &FDataTableFactoryUI::IsValidSelection)
+								.IsEnabled(this, &FDataTableFactoryUI::IsStructSelected)
 								.OnClicked(this, &FDataTableFactoryUI::OnCreate)
 							]
 							+SHorizontalBox::Slot()
@@ -245,31 +176,25 @@ bool UDataTableFactory::ConfigureProperties()
 			GEditor->EditorAddModalWindow(PickerWindow.ToSharedRef());
 			PickerWindow.Reset();
 
-			ClassResult = ResultClass;
-			StructResult = ResultStruct;
-
-			return IsValidSelection();
+			return ResultStruct;
 		}
 
 	private:
 		TSharedPtr<SWindow> PickerWindow;
 		TSharedPtr<SComboButton> StructPickerAnchor;
-		TSharedPtr<SComboButton> ClassPickerAnchor;
 		const UScriptStruct* ResultStruct = nullptr;
-		const UClass* ResultClass = UDataTable::StaticClass();
 	};
 
-	FDataTableFactoryUI ConfigSelector = FDataTableFactoryUI();
+	TSharedRef<FDataTableFactoryUI> StructSelector = MakeShareable(new FDataTableFactoryUI());
+	Struct = StructSelector->OpenStructSelector();
 
-	TableClass = nullptr;
-	Struct = nullptr;
-	return ConfigSelector.OpenSelectorDialog(TableClass, Struct);
+	return Struct != nullptr;
 }
 
 UObject* UDataTableFactory::FactoryCreateNew(UClass* Class, UObject* InParent, FName Name, EObjectFlags Flags, UObject* Context, FFeedbackContext* Warn)
 {
 	UDataTable* DataTable = nullptr;
-	if (Struct && TableClass && (SupportedClass == Class))
+	if (Struct && ensure(SupportedClass == Class))
 	{
 		ensure(0 != (RF_Public & Flags));
 		DataTable = MakeNewDataTable(InParent, Name, Flags);
@@ -283,7 +208,7 @@ UObject* UDataTableFactory::FactoryCreateNew(UClass* Class, UObject* InParent, F
 
 UDataTable* UDataTableFactory::MakeNewDataTable(UObject* InParent, FName Name, EObjectFlags Flags)
 {
-	return NewObject<UDataTable>(InParent, TableClass, Name, Flags);
+	return NewObject<UDataTable>(InParent, Name, Flags);
 }
 
 #undef LOCTEXT_NAMESPACE // "DataTableFactory"
