Cooker: Disable experimental feature TryCollectKeyAndDependencies; it is only useful when other experimental feature EditorDomain is enabled and it costs 5 minutes in cook by the book.
#rnx

diff --git a/Engine/Source/Editor/UnrealEd/Private/TargetDomain/TargetDomainUtils.cpp b/Engine/Source/Editor/UnrealEd/Private/TargetDomain/TargetDomainUtils.cpp
--- a/Engine/Source/Editor/UnrealEd/Private/TargetDomain/TargetDomainUtils.cpp
+++ b/Engine/Source/Editor/UnrealEd/Private/TargetDomain/TargetDomainUtils.cpp
@@ -128,6 +128,12 @@ bool TryCollectKeyAndDependencies(UPackage* Package, const ITargetPlatform* Targ
 		if (OutErrorMessage) *OutErrorMessage = TEXT("Invalid null package.");
 		return false;
 	}
+	FEditorDomain* EditorDomain = FEditorDomain::Get();
+	if (!EditorDomain)
+	{
+		if (OutErrorMessage) *OutErrorMessage = TEXT("EditorDomain is unavailable.");
+		return false;
+	}
 	IAssetRegistry* AssetRegistry = IAssetRegistry::Get();
 	if (!AssetRegistry)
 	{
