Fix bad auto-merge.
#rnx

diff --git a/Engine/Source/Editor/UnrealEd/Private/TargetDomain/TargetDomainUtils.cpp b/Engine/Source/Editor/UnrealEd/Private/TargetDomain/TargetDomainUtils.cpp
--- a/Engine/Source/Editor/UnrealEd/Private/TargetDomain/TargetDomainUtils.cpp
+++ b/Engine/Source/Editor/UnrealEd/Private/TargetDomain/TargetDomainUtils.cpp
@@ -128,12 +128,6 @@ bool TryCollectKeyAndDependencies(UPackage* Package, const ITargetPlatform* Targ
 		if (OutErrorMessage) *OutErrorMessage = TEXT("Invalid null package.");
 		return false;
 	}
-	FEditorDomain* EditorDomain = FEditorDomain::Get();
-	if (!EditorDomain)
-	{
-		if (OutErrorMessage) *OutErrorMessage = TEXT("EditorDomain is unavailable.");
-		return false;
-	}
 	IAssetRegistry* AssetRegistry = IAssetRegistry::Get();
 	if (!AssetRegistry)
 	{
