Multiplayer PIE when another editor is also running multiplayer PIE (even of a different game) fails semi-silently

diff --git a/Engine/Source/Editor/UnrealEd/Private/PlayLevel.cpp b/Engine/Source/Editor/UnrealEd/Private/PlayLevel.cpp
--- a/Engine/Source/Editor/UnrealEd/Private/PlayLevel.cpp
+++ b/Engine/Source/Editor/UnrealEd/Private/PlayLevel.cpp
@@ -53,6 +53,7 @@
 #include "Components/AudioComponent.h"
 #include "Engine/Note.h"
 #include "Engine/Selection.h"
+#include "Engine/NetDriver.h"
 #include "UnrealEngine.h"
 #include "EngineUtils.h"
 #include "Editor.h"
@@ -2817,6 +2818,21 @@ void UEditorEngine::StartPlayInEditorSession(FRequestPlaySessionParams& InReques
 			{
 				return;
 			}
+
+			if (bClientIsServer)
+			{
+				// Grab New Created PIE Server Instance and set PlaySettings to Server's actual Port so Clients Connect To Correct Server
+				const FWorldContext *const PIEServerWorldContext = GetWorldContextFromPIEInstance(PlayInEditorSessionInfo->PIEInstanceCount - 1);
+				const UWorld *const PIEServerWorld = PIEServerWorldContext->World();
+				if (PIEServerWorld)
+				{
+					UNetDriver *const NetDriver = PIEServerWorld->GetNetDriver();
+					if (NetDriver && NetDriver->GetLocalAddr().IsValid())
+					{
+						EditorPlaySettings->SetServerPort(NetDriver->GetLocalAddr()->GetPort());
+					}
+				}
+			}
 		}
 
 		// This needs to be run before ToggleBetweenPIEandSIE as it creates the list of selected actors.
