PumpPollables: Move MemTracker and Is-it-time-to-collect-garbage pollables from CookCommandlet into UCookOnTheFlyServer::PumpPollables.
#rnx

diff --git a/Engine/Source/Editor/UnrealEd/Classes/CookOnTheSide/CookOnTheFlyServer.h b/Engine/Source/Editor/UnrealEd/Classes/CookOnTheSide/CookOnTheFlyServer.h
--- a/Engine/Source/Editor/UnrealEd/Classes/CookOnTheSide/CookOnTheFlyServer.h
+++ b/Engine/Source/Editor/UnrealEd/Classes/CookOnTheSide/CookOnTheFlyServer.h
@@ -158,7 +158,7 @@ public:
 	};
 
 private:
-	using FPollFunction = TUniqueFunction<void()>;
+	using FPollFunction = TUniqueFunction<void(UE::Cook::FTickStackData&)>;
 	struct FPollable
 	{
 		FPollable(float InPeriodSeconds, float InPeriodIdleSeconds, FPollFunction&& InFunction);
@@ -382,6 +382,7 @@ private:
 	void InitializePollables();
 	void PumpPollables(UE::Cook::FTickStackData& StackData, bool bIsIdle);
 	void PollFlushRenderingCommands();
+	void PollPackagesPerGC(UE::Cook::FTickStackData& StackData);
 
 public:
 
@@ -393,7 +394,9 @@ public:
 		COSR_ErrorLoadingPackage	= 0x00000004,
 		COSR_RequiresGC				= 0x00000008,
 		COSR_WaitingOnCache			= 0x00000010,
-		COSR_MarkedUpKeepPackages	= 0x00000040
+		COSR_MarkedUpKeepPackages	= 0x00000040,
+		COSR_RequiresGC_OOM			= 0x00000080,
+		COSR_RequiresGC_PackageCount= 0x00000100,
 	};
 
 
@@ -1160,7 +1163,7 @@ private:
 	int32 PollStartIndex = 0;
 	float PumpPollablesTimeSlice = 0.f;
 	float PumpPollablesMinPeriod = 0.f;
-
+	uint32 CookedPackageCountSinceLastGC = 0;
 	friend UE::Cook::FPackageData;
 	friend UE::Cook::FPendingCookedPlatformData;
 	friend UE::Cook::FPlatformManager;
diff --git a/Engine/Source/Editor/UnrealEd/Private/Commandlets/CookCommandlet.cpp b/Engine/Source/Editor/UnrealEd/Private/Commandlets/CookCommandlet.cpp
--- a/Engine/Source/Editor/UnrealEd/Private/Commandlets/CookCommandlet.cpp
+++ b/Engine/Source/Editor/UnrealEd/Private/Commandlets/CookCommandlet.cpp
@@ -990,14 +990,7 @@ bool UCookCommandlet::CookByTheBook( const TArray<ITargetPlatform*>& Platforms)
 	});
 
 #if ENABLE_LOW_LEVEL_MEM_TRACKER
-	FLowLevelMemTracker& MemTracker = FLowLevelMemTracker::Get();
-	if (MemTracker.IsEnabled())
-	{
-		MemTracker.UpdateStatsPerFrame();
-	}
-	FDateTime LastMemTrackerTime = FDateTime::UtcNow();
-	double MemTrackerPeriodSeconds = 120.0;
-	FParse::Value(FCommandLine::Get(), TEXT("CookLLMPeriod="), MemTrackerPeriodSeconds);
+	FLowLevelMemTracker::Get().UpdateStatsPerFrame();
 #endif
 
 	do
@@ -1006,138 +999,93 @@ bool UCookCommandlet::CookByTheBook( const TArray<ITargetPlatform*>& Platforms)
 			COOK_STAT(FScopedDurationTimer StartCookByTheBookTimer(DetailedCookStats::StartCookByTheBookTimeSec));
 			CookOnTheFlyServer->StartCookByTheBook(StartupOptions);
 		}
-
-
-		// Garbage collection should happen when either
-		//	1. We have cooked a map (configurable asset type)
-		//	2. We have cooked non-map packages and...
-		//		a. we have accumulated 50 (configurable) of these since the last GC.
-		//		b. we have been idle for 20 (configurable) seconds.
-		bool bShouldGC = false;
-		FString GCReason;
-
-		// megamoth
-		uint32 NonMapPackageCountSinceLastGC = 0;
-
-		const uint32 PackagesPerGC = CookOnTheFlyServer->GetPackagesPerGC();
-		const double IdleTimeToGC = CookOnTheFlyServer->GetIdleTimeToGC();
-		const uint32 PackagesPerPartialGC = CookOnTheFlyServer->GetPackagesPerPartialGC();
-
-		double LastCookActionTime = FPlatformTime::Seconds();
-
-		FDateTime LastConnectionTime = FDateTime::UtcNow();
-		bool bHadConnection = false;
-
 		while (CookOnTheFlyServer->IsCookByTheBookRunning())
 		{
 			DECLARE_SCOPE_CYCLE_COUNTER( TEXT( "CookByTheBook.MainLoop" ), STAT_CookByTheBook_MainLoop, STATGROUP_LoadTime );
 			{
 				uint32 TickResults = 0;
-				static const float CookOnTheSideTimeSlice = 10.0f;
+				const float CookOnTheSideTimeSlice = 10.0f;
+				uint32 UnusedVariable = 0;
 
-#if ENABLE_LOW_LEVEL_MEM_TRACKER
-				if (MemTracker.IsEnabled())
-				{
-					FDateTime CurrentTime = FDateTime::UtcNow();
-					if (CurrentTime >= LastMemTrackerTime + FTimespan::FromSeconds(MemTrackerPeriodSeconds))
-					{
-						MemTracker.UpdateStatsPerFrame();
-						LastMemTrackerTime = CurrentTime;
-					}
-				}
-#endif
-
-				TickResults = CookOnTheFlyServer->TickCookOnTheSide( CookOnTheSideTimeSlice, NonMapPackageCountSinceLastGC, ShowProgress ? ECookTickFlags::None : ECookTickFlags::HideProgressDisplay );
+				TickResults = CookOnTheFlyServer->TickCookOnTheSide( CookOnTheSideTimeSlice, UnusedVariable,
+					ShowProgress ? ECookTickFlags::None : ECookTickFlags::HideProgressDisplay );
 
 				{
 					UE_SCOPED_COOKTIMER_AND_DURATION(CookByTheBook_ShaderProcessAsync, DetailedCookStats::TickLoopShaderProcessAsyncResultsTimeSec);
 					GShaderCompilingManager->ProcessAsyncResults(true, false);
 				}
-
-				
-				// Flush the asset registry before GC
 				{
 					UE_SCOPED_COOKTIMER(CookByTheBook_TickAssetRegistry);
 					FAssetRegistryModule::TickAssetRegistry(-1.0f);
 				}
-
-				auto DumpMemStats = []()
+				
+				if ((TickResults & UCookOnTheFlyServer::COSR_RequiresGC) != 0)
 				{
-					FGenericMemoryStats MemStats;
-					GMalloc->GetAllocatorStats(MemStats);
-					for (const auto& Item : MemStats.Data)
+					FString GCReason;
+					if ((TickResults & UCookOnTheFlyServer::COSR_RequiresGC_PackageCount) != 0)
 					{
-						UE_LOG(LogCookCommandlet, Display, TEXT("Item %s = %d"), *Item.Key, Item.Value);
+						GCReason = TEXT("Exceeded packages per GC");
 					}
-				};
-
-
-				const bool bHasExceededMaxMemory = CookOnTheFlyServer->HasExceededMaxMemory();
-				// We should GC if we have packages to collect and we've been idle for some time.
-				const bool bExceededPackagesPerGC = (PackagesPerGC > 0) && (NonMapPackageCountSinceLastGC > PackagesPerGC);
-				const bool bWaitingOnObjectCache = ((TickResults & UCookOnTheFlyServer::COSR_WaitingOnCache) == 0);
-
-
-				if (!bWaitingOnObjectCache && bExceededPackagesPerGC) // if we are waiting on things to cache then ignore the exceeded packages per gc
-				{
-					bShouldGC = true;
-					GCReason = TEXT("Exceeded packages per GC");
-				}
-				else if (bHasExceededMaxMemory) // if we are exceeding memory then we need to gc (this can cause thrashing if the cooker loads the same stuff into memory next tick
-				{
-					bShouldGC = true;
-					GCReason = TEXT("Exceeded Max Memory");
+					else if ((TickResults & UCookOnTheFlyServer::COSR_RequiresGC_OOM) != 0)
+					{
+						// this can cause thrashing if the cooker loads the same stuff into memory next tick
+						GCReason = TEXT("Exceeded Max Memory");
 
-					int32 JobsToLogAt = GShaderCompilingManager->GetNumRemainingJobs();
+						int32 JobsToLogAt = GShaderCompilingManager->GetNumRemainingJobs();
 
-					UE_SCOPED_COOKTIMER(CookByTheBook_ShaderJobFlush);
-					UE_LOG(LogCookCommandlet, Display, TEXT("Detected max mem exceeded - forcing shader compilation flush"));
-					while ( true )
-					{
-						int32 NumRemainingJobs = GShaderCompilingManager->GetNumRemainingJobs();
-						if ( NumRemainingJobs < 1000)
+						UE_SCOPED_COOKTIMER(CookByTheBook_ShaderJobFlush);
+						UE_LOG(LogCookCommandlet, Display, TEXT("Detected max mem exceeded - forcing shader compilation flush"));
+						while (true)
 						{
-							UE_LOG(LogCookCommandlet, Display, TEXT("Finished flushing shader jobs at %d"), NumRemainingJobs);
-							break;
-						}
+							int32 NumRemainingJobs = GShaderCompilingManager->GetNumRemainingJobs();
+							if (NumRemainingJobs < 1000)
+							{
+								UE_LOG(LogCookCommandlet, Display, TEXT("Finished flushing shader jobs at %d"), NumRemainingJobs);
+								break;
+							}
 
-						if (NumRemainingJobs < JobsToLogAt )
-						{
-							UE_LOG(LogCookCommandlet, Display, TEXT("Flushing shader jobs, remaining jobs %d"), NumRemainingJobs);
-						}
+							if (NumRemainingJobs < JobsToLogAt)
+							{
+								UE_LOG(LogCookCommandlet, Display, TEXT("Flushing shader jobs, remaining jobs %d"), NumRemainingJobs);
+							}
 
-						GShaderCompilingManager->ProcessAsyncResults(false, false);
+							GShaderCompilingManager->ProcessAsyncResults(false, false);
 
-						FPlatformProcess::Sleep(0.05);
+							FPlatformProcess::Sleep(0.05);
 
-						// GShaderCompilingManager->FinishAllCompilation();
+							// GShaderCompilingManager->FinishAllCompilation();
+						}
+					}
+					else 
+					{
+						// cooker loaded some object which needs to be cleaned up before the cooker can proceed so force gc
+						GCReason = TEXT("COSR_RequiresGC");
 					}
-				}
-				else if ((TickResults & UCookOnTheFlyServer::COSR_RequiresGC) != 0) // cooker loaded some object which needs to be cleaned up before the cooker can proceed so force gc
-				{
-					GCReason = TEXT("COSR_RequiresGC");
-					bShouldGC = true;
-				}
-
-				bShouldGC |= bTestCook; // testing cooking / gc path
-
 
-				if (bShouldGC )
-				{
+					// Flush the asset registry before GC
+					{
+						UE_SCOPED_COOKTIMER(CookByTheBook_TickAssetRegistry);
+						FAssetRegistryModule::TickAssetRegistry(-1.0f);
+					}
 					UE_SCOPED_COOKTIMER_AND_DURATION(CookByTheBook_GC, DetailedCookStats::TickLoopGCTimeSec);
-					bShouldGC = false;
 
 					int32 NumObjectsBeforeGC = GUObjectArray.GetObjectArrayNumMinusAvailable();
 					int32 NumObjectsAvailableBeforeGC = GUObjectArray.GetObjectArrayEstimatedAvailable();
 					UE_LOG(LogCookCommandlet, Display, TEXT("GarbageCollection...%s (%s)"), (bPartialGC? TEXT(" partial gc") : TEXT("")), *GCReason);
-					GCReason = FString();
 
 #if ENABLE_LOW_LEVEL_MEM_TRACKER
-					if (MemTracker.IsEnabled())
-					{
-						MemTracker.UpdateStatsPerFrame();
-					}
+					FLowLevelMemTracker::Get().UpdateStatsPerFrame();
 #endif
+					auto DumpMemStats = []()
+					{
+						FGenericMemoryStats MemStats;
+						GMalloc->GetAllocatorStats(MemStats);
+						for (const auto& Item : MemStats.Data)
+						{
+							UE_LOG(LogCookCommandlet, Display, TEXT("Item %s = %d"), *Item.Key, Item.Value);
+						}
+					};
+
 					DumpMemStats();
 
 					CollectGarbage(RF_NoFlags);
@@ -1147,8 +1095,6 @@ bool UCookCommandlet::CookByTheBook( const TArray<ITargetPlatform*>& Platforms)
 					UE_LOG(LogCookCommandlet, Display, TEXT("%s GC before %d available %d after %d available %d"), (bPartialGC ? TEXT("Partial") : TEXT("Full")), NumObjectsBeforeGC, NumObjectsAvailableBeforeGC, NumObjectsAfterGC, NumObjectsAvailableAfterGC);
 
 					DumpMemStats();
-
-					NonMapPackageCountSinceLastGC = 0;
 				}
 
 				{
@@ -1172,10 +1118,7 @@ bool UCookCommandlet::CookByTheBook( const TArray<ITargetPlatform*>& Platforms)
 	} while (bTestCook);
 
 #if ENABLE_LOW_LEVEL_MEM_TRACKER
-	if (MemTracker.IsEnabled())
-	{
-		MemTracker.UpdateStatsPerFrame();
-	}
+	FLowLevelMemTracker::Get().UpdateStatsPerFrame();
 #endif
 
 	if (StartupOptions.DLCName.IsEmpty())
diff --git a/Engine/Source/Editor/UnrealEd/Private/CookOnTheFlyServer.cpp b/Engine/Source/Editor/UnrealEd/Private/CookOnTheFlyServer.cpp
--- a/Engine/Source/Editor/UnrealEd/Private/CookOnTheFlyServer.cpp
+++ b/Engine/Source/Editor/UnrealEd/Private/CookOnTheFlyServer.cpp
@@ -1551,18 +1551,36 @@ UCookOnTheFlyServer::FPollable::FPollable(float InPeriodSeconds, float InPeriodI
 
 void UCookOnTheFlyServer::InitializePollables()
 {
+	using namespace UE::Cook;
+
 	Pollables.Reset();
 
 	if (IsCookByTheBookMode() && !IsCookingInEditor())
 	{
-		if (Pollables.IsEmpty())
-		{
-			Pollables.Add(FPollable(60.f, 5.f, [this] { PollFlushRenderingCommands(); }));
-		}
+		Pollables.Add(FPollable(60.f, 5.f, [this](FTickStackData&) { PollFlushRenderingCommands(); }));
 
-		for (FPollable& Pollable : Pollables)
+#if ENABLE_LOW_LEVEL_MEM_TRACKER
+		double MemTrackerPeriodSeconds = 120.0;
+		FParse::Value(FCommandLine::Get(), TEXT("-CookLLMPeriod="), MemTrackerPeriodSeconds);
+		Pollables.Add(FPollable(MemTrackerPeriodSeconds, MemTrackerPeriodSeconds, [] (FTickStackData&) { FLowLevelMemTracker::Get().UpdateStatsPerFrame(); }));
+#endif
+		bool bTestCook = IsCookFlagSet(ECookInitializationFlags::TestCook);
+		if (PackagesPerGC > 0 || bTestCook)
 		{
-			Pollable.NextTimeSeconds = Pollable.NextTimeIdleSeconds = 0.0;
+			// PackagesPerGC is usually used only to debug; max memory counts are commonly used instead
+			// Since it's not commonly used, we make a concession to support it: we check on a timer rather than checking after every saved package.
+			// For large values, check less frequently.
+			uint32 NumPackagesForTimeEstimate = UINT32_MAX;
+			if (bTestCook)
+			{
+				NumPackagesForTimeEstimate = FMath::Min(NumPackagesForTimeEstimate, 50U);
+			}
+			if (PackagesPerGC > 0)
+			{
+				NumPackagesForTimeEstimate = FMath::Min(NumPackagesForTimeEstimate, PackagesPerGC);
+			}
+			float PeriodSeconds = FMath::Min(NumPackagesForTimeEstimate * .01f, 10.f);
+			Pollables.Add(FPollable(PeriodSeconds, PeriodSeconds, [this](FTickStackData& StackData) { PollPackagesPerGC(StackData); }));
 		}
 	}
 
@@ -1603,7 +1621,7 @@ void UCookOnTheFlyServer::PumpPollables(UE::Cook::FTickStackData& StackData, boo
 		double NextTimeSeconds = bIsIdle ? Pollable.NextTimeIdleSeconds : Pollable.NextTimeSeconds;
 		if (NextTimeSeconds <= CurrentTime)
 		{
-			Pollable.PollFunction();
+			Pollable.PollFunction(StackData);
 			double NewCurrentTime = FPlatformTime::Seconds();
 			Pollable.NextTimeSeconds = NewCurrentTime + Pollable.PeriodSeconds;
 			Pollable.NextTimeIdleSeconds = NewCurrentTime + Pollable.PeriodIdleSeconds;
@@ -1620,7 +1638,7 @@ void UCookOnTheFlyServer::PumpPollables(UE::Cook::FTickStackData& StackData, boo
 	// If we early exited, finish calculating PollNextTimeSeconds from the remaining members we didn't reach
 	for (int32 Index = PollStartIndex; Index != PollEndIndex; Index = (Index+1) % NumPollables)
 	{
-		FPollable& Pollable = Pollables[Index++];
+		FPollable& Pollable = Pollables[Index];
 		PollNextTimeSeconds = FMath::Min(Pollable.NextTimeSeconds, PollNextTimeSeconds);
 		PollNextTimeIdleSeconds = FMath::Min(Pollable.NextTimeIdleSeconds, PollNextTimeIdleSeconds);
 	}
@@ -1637,6 +1655,22 @@ void UCookOnTheFlyServer::PollFlushRenderingCommands()
 	FlushRenderingCommands();
 }
 
+void UCookOnTheFlyServer::PollPackagesPerGC(UE::Cook::FTickStackData& StackData)
+{
+	if (IsCookFlagSet(ECookInitializationFlags::TestCook))
+	{
+		StackData.ResultFlags |= COSR_RequiresGC;
+	}
+	else if (PackagesPerGC > 0 && CookedPackageCountSinceLastGC > PackagesPerGC)
+	{
+		// if we are waiting on things to cache then ignore the PackagesPerGC
+		if (!bSaveBusy)
+		{
+			StackData.ResultFlags |= COSR_RequiresGC | COSR_RequiresGC_PackageCount;
+		}
+	}
+}
+
 UCookOnTheFlyServer::ECookAction UCookOnTheFlyServer::DecideNextCookAction(UE::Cook::FTickStackData& StackData)
 {
 	if (IsCookByTheBookMode() && CookByTheBookOptions->bCancel)
@@ -1964,7 +1998,7 @@ void UCookOnTheFlyServer::PumpLoads(UE::Cook::FTickStackData& StackData, uint32
 
 		if (HasExceededMaxMemory())
 		{
-			StackData.ResultFlags |= COSR_RequiresGC;
+			StackData.ResultFlags |= COSR_RequiresGC | COSR_RequiresGC_OOM;
 			return;
 		}
 	}
@@ -3070,7 +3104,7 @@ void UCookOnTheFlyServer::PumpSaves(UE::Cook::FTickStackData& StackData, uint32
 	{
 		if (HasExceededMaxMemory())
 		{
-			StackData.ResultFlags |= COSR_RequiresGC;
+			StackData.ResultFlags |= COSR_RequiresGC | COSR_RequiresGC_OOM;
 		}
 	};
 
@@ -3992,6 +4026,8 @@ void UCookOnTheFlyServer::PostGarbageCollect()
 			GeneratorPackage->PostGarbageCollect();
 		}
 	}
+
+	CookedPackageCountSinceLastGC = 0;
 }
 
 void UCookOnTheFlyServer::BeginDestroy()
@@ -4361,6 +4397,7 @@ void FSaveCookedPackageContext::FinishPlatform()
 		else
 		{
 			++StackData.CookedPackageCount;
+			++COTFS.CookedPackageCountSinceLastGC;
 			StackData.ResultFlags |= UCookOnTheFlyServer::COSR_CookedPackage;
 		}
 	}
