CookCommandlet: Fix typos
#rnx

diff --git a/Engine/Source/Editor/UnrealEd/Private/Commandlets/CookCommandlet.cpp b/Engine/Source/Editor/UnrealEd/Private/Commandlets/CookCommandlet.cpp
--- a/Engine/Source/Editor/UnrealEd/Private/Commandlets/CookCommandlet.cpp
+++ b/Engine/Source/Editor/UnrealEd/Private/Commandlets/CookCommandlet.cpp
@@ -462,7 +462,7 @@ bool UCookCommandlet::CookOnTheFly( FGuid InstanceId, int32 Timeout, bool bForce
 			, LastCookActionTime(FPlatformTime::Seconds())
 		{}
 
-		/** Tntended to be called with stats from a UCookOnTheFlyServer::TickCookOnTheSide() call. Determines if we should be calling GC after TickCookOnTheSide(). */
+		/** Intended to be called with stats from a UCookOnTheFlyServer::TickCookOnTheSide() call. Determines if we should be calling GC after TickCookOnTheSide(). */
 		void Update(uint32 CookedCount, UCookOnTheFlyServer::ECookOnTheSideResult ResultFlags)
 		{
 			if (ResultFlags & (UCookOnTheFlyServer::COSR_CookedMap | UCookOnTheFlyServer::COSR_CookedPackage | UCookOnTheFlyServer::COSR_WaitingOnCache))
@@ -484,7 +484,7 @@ bool UCookCommandlet::CookOnTheFly( FGuid InstanceId, int32 Timeout, bool bForce
 			}
 
 			// we don't want to gc if we are waiting on cache of objects. this could clean up objects which we will need to reload next frame
-			bPosteponeGC = (ResultFlags & UCookOnTheFlyServer::COSR_WaitingOnCache) != 0;
+			bPostponeGC = (ResultFlags & UCookOnTheFlyServer::COSR_WaitingOnCache) != 0;
 		}
 
 		/** Runs GC if Update() determined it should happen. Also checks the idle time against the limit, and runs GC then if packages have been loaded. */
@@ -509,7 +509,7 @@ bool UCookCommandlet::CookOnTheFly( FGuid InstanceId, int32 Timeout, bool bForce
 				}
 			}
 
-			if (bShouldGC && !bPosteponeGC)
+			if (bShouldGC && !bPostponeGC)
 			{	
 				Reset();
 				UE_LOG(LogCookCommandlet, Display, TEXT("GC..."));
@@ -532,7 +532,7 @@ bool UCookCommandlet::CookOnTheFly( FGuid InstanceId, int32 Timeout, bool bForce
 		bool   bShouldGC;
 		uint32 PackagesCookedSinceLastGC;
 		double LastCookActionTime;
-		bool   bPosteponeGC;
+		bool   bPostponeGC;
 
 	} CookOnTheFlyGCController(CookOnTheFlyServer);
 
