Update tooltip for disabled targets to mention Launcher projects too, to avoid potential confusion.

#rnx

diff --git a/Engine/Source/Editor/UnrealEd/Private/Settings/ProjectPackagingSettingsCustomization.h b/Engine/Source/Editor/UnrealEd/Private/Settings/ProjectPackagingSettingsCustomization.h
--- a/Engine/Source/Editor/UnrealEd/Private/Settings/ProjectPackagingSettingsCustomization.h
+++ b/Engine/Source/Editor/UnrealEd/Private/Settings/ProjectPackagingSettingsCustomization.h
@@ -147,7 +147,7 @@ protected:
 	{
 		TArray<EProjectPackagingBuildConfigurations> PackagingConfigurations = UProjectPackagingSettings::GetValidPackageConfigurations();
 
-		TSharedRef<FPropertyRestriction> BuildConfigurationRestriction = MakeShareable(new FPropertyRestriction(LOCTEXT("ConfigurationRestrictionReason", "This configuration is not valid for this project. DebugGame configurations require a code project, and client/server configurations require the appropriate targets.")));
+		TSharedRef<FPropertyRestriction> BuildConfigurationRestriction = MakeShareable(new FPropertyRestriction(LOCTEXT("ConfigurationRestrictionReason", "This configuration is not valid for this project. DebugGame configurations are not available in Content-Only or Launcher projects, and client/server configurations require the appropriate targets..")));
 
 		const UEnum* const ProjectPackagingBuildConfigurationsEnum = StaticEnum<EProjectPackagingBuildConfigurations>();
 		for (int Idx = 0; Idx < (int)EProjectPackagingBuildConfigurations::PPBC_MAX; Idx++)
