PumpPollables: Move pollable construction steps into separate functions to reduce complexity in InitializePollables.
#rnx

diff --git a/Engine/Source/Editor/UnrealEd/Classes/CookOnTheSide/CookOnTheFlyServer.h b/Engine/Source/Editor/UnrealEd/Classes/CookOnTheSide/CookOnTheFlyServer.h
--- a/Engine/Source/Editor/UnrealEd/Classes/CookOnTheSide/CookOnTheFlyServer.h
+++ b/Engine/Source/Editor/UnrealEd/Classes/CookOnTheSide/CookOnTheFlyServer.h
@@ -382,6 +382,8 @@ private:
 	void InitializePollables();
 	void PumpPollables(UE::Cook::FTickStackData& StackData, bool bIsIdle);
 	void PollFlushRenderingCommands();
+	TOptional<FPollable> CreatePollableLLM();
+	TOptional<FPollable> CreatePollableTriggerGC();
 	void PollPackagesPerGC(UE::Cook::FTickStackData& StackData);
 
 public:
diff --git a/Engine/Source/Editor/UnrealEd/Private/CookOnTheFlyServer.cpp b/Engine/Source/Editor/UnrealEd/Private/CookOnTheFlyServer.cpp
--- a/Engine/Source/Editor/UnrealEd/Private/CookOnTheFlyServer.cpp
+++ b/Engine/Source/Editor/UnrealEd/Private/CookOnTheFlyServer.cpp
@@ -1549,54 +1549,6 @@ UCookOnTheFlyServer::FPollable::FPollable(float InPeriodSeconds, float InPeriodI
 {
 }
 
-void UCookOnTheFlyServer::InitializePollables()
-{
-	using namespace UE::Cook;
-
-	Pollables.Reset();
-
-	if (IsCookByTheBookMode() && !IsCookingInEditor())
-	{
-		Pollables.Add(FPollable(60.f, 5.f, [this](FTickStackData&) { PollFlushRenderingCommands(); }));
-
-#if ENABLE_LOW_LEVEL_MEM_TRACKER
-		double MemTrackerPeriodSeconds = 120.0;
-		FParse::Value(FCommandLine::Get(), TEXT("-CookLLMPeriod="), MemTrackerPeriodSeconds);
-		Pollables.Add(FPollable(MemTrackerPeriodSeconds, MemTrackerPeriodSeconds, [] (FTickStackData&) { FLowLevelMemTracker::Get().UpdateStatsPerFrame(); }));
-#endif
-		bool bTestCook = IsCookFlagSet(ECookInitializationFlags::TestCook);
-		if (PackagesPerGC > 0 || bTestCook)
-		{
-			// PackagesPerGC is usually used only to debug; max memory counts are commonly used instead
-			// Since it's not commonly used, we make a concession to support it: we check on a timer rather than checking after every saved package.
-			// For large values, check less frequently.
-			uint32 NumPackagesForTimeEstimate = UINT32_MAX;
-			if (bTestCook)
-			{
-				NumPackagesForTimeEstimate = FMath::Min(NumPackagesForTimeEstimate, 50U);
-			}
-			if (PackagesPerGC > 0)
-			{
-				NumPackagesForTimeEstimate = FMath::Min(NumPackagesForTimeEstimate, PackagesPerGC);
-			}
-			float PeriodSeconds = FMath::Min(NumPackagesForTimeEstimate * .01f, 10.f);
-			Pollables.Add(FPollable(PeriodSeconds, PeriodSeconds, [this](FTickStackData& StackData) { PollPackagesPerGC(StackData); }));
-		}
-	}
-
-	PollStartIndex = 0;
-	if (Pollables.Num())
-	{
-		PollNextTimeSeconds = 0.;
-		PollNextTimeIdleSeconds = 0.;
-	}
-	else
-	{
-		PollNextTimeSeconds = MAX_flt;
-		PollNextTimeIdleSeconds = MAX_flt;
-	}
-}
-
 void UCookOnTheFlyServer::PumpPollables(UE::Cook::FTickStackData& StackData, bool bIsIdle)
 {
 	UE_SCOPED_HIERARCHICAL_COOKTIMER(PumpPollables);
@@ -1655,6 +1607,44 @@ void UCookOnTheFlyServer::PollFlushRenderingCommands()
 	FlushRenderingCommands();
 }
 
+TOptional<UCookOnTheFlyServer::FPollable> UCookOnTheFlyServer::CreatePollableLLM()
+{
+#if ENABLE_LOW_LEVEL_MEM_TRACKER
+	if (FLowLevelMemTracker::Get().IsEnabled())
+	{
+		float PeriodSeconds = 120.0f;
+		FParse::Value(FCommandLine::Get(), TEXT("-CookLLMPeriod="), PeriodSeconds);
+		return TOptional<FPollable>(FPollable(PeriodSeconds, PeriodSeconds,
+			[](UE::Cook::FTickStackData&) { FLowLevelMemTracker::Get().UpdateStatsPerFrame(); }));
+	}
+#endif
+	return TOptional<FPollable>();
+}
+
+TOptional<UCookOnTheFlyServer::FPollable> UCookOnTheFlyServer::CreatePollableTriggerGC()
+{
+	bool bTestCook = IsCookFlagSet(ECookInitializationFlags::TestCook);
+	if (PackagesPerGC > 0 || bTestCook)
+	{
+		// PackagesPerGC is usually used only to debug; max memory counts are commonly used instead
+		// Since it's not commonly used, we make a concession to support it: we check on a timer rather than checking after every saved package.
+		// For large values, check less frequently.
+		uint32 NumPackagesForTimeEstimate = UINT32_MAX;
+		if (bTestCook)
+		{
+			NumPackagesForTimeEstimate = FMath::Min(NumPackagesForTimeEstimate, 50U);
+		}
+		if (PackagesPerGC > 0)
+		{
+			NumPackagesForTimeEstimate = FMath::Min(NumPackagesForTimeEstimate, PackagesPerGC);
+		}
+		float PeriodSeconds = FMath::Min(NumPackagesForTimeEstimate * .01f, 10.f);
+		return TOptional<FPollable>(FPollable(PeriodSeconds, PeriodSeconds,
+			[this](UE::Cook::FTickStackData& StackData) { PollPackagesPerGC(StackData); }));
+	}
+	return TOptional<FPollable>();
+}
+
 void UCookOnTheFlyServer::PollPackagesPerGC(UE::Cook::FTickStackData& StackData)
 {
 	if (IsCookFlagSet(ECookInitializationFlags::TestCook))
@@ -1671,6 +1661,38 @@ void UCookOnTheFlyServer::PollPackagesPerGC(UE::Cook::FTickStackData& StackData)
 	}
 }
 
+void UCookOnTheFlyServer::InitializePollables()
+{
+	using namespace UE::Cook;
+
+	Pollables.Reset();
+
+	if (IsCookByTheBookMode() && !IsCookingInEditor())
+	{
+		Pollables.Add(FPollable(60.f, 5.f, [this](FTickStackData&) { PollFlushRenderingCommands(); }));
+		if (TOptional<FPollable> Pollable = CreatePollableLLM())
+		{
+			Pollables.Add(MoveTemp(*Pollable));
+		}
+		if (TOptional<FPollable> Pollable = CreatePollableTriggerGC())
+		{
+			Pollables.Add(MoveTemp(*Pollable));
+		}
+	}
+
+	PollStartIndex = 0;
+	if (Pollables.Num())
+	{
+		PollNextTimeSeconds = 0.;
+		PollNextTimeIdleSeconds = 0.;
+	}
+	else
+	{
+		PollNextTimeSeconds = MAX_flt;
+		PollNextTimeIdleSeconds = MAX_flt;
+	}
+}
+
 UCookOnTheFlyServer::ECookAction UCookOnTheFlyServer::DecideNextCookAction(UE::Cook::FTickStackData& StackData)
 {
 	if (IsCookByTheBookMode() && CookByTheBookOptions->bCancel)
