"testprops generator" command added to test property generator.

"testprops struct=StructName" parsing added to use a specific struct, much like a specific class can be specified.

Added various test cases to PropertyEditorTestObject.

diff --git a/Engine/Source/Editor/UnrealEd/Classes/Editor/PropertyEditorTestObject.h b/Engine/Source/Editor/UnrealEd/Classes/Editor/PropertyEditorTestObject.h
--- a/Engine/Source/Editor/UnrealEd/Classes/Editor/PropertyEditorTestObject.h
+++ b/Engine/Source/Editor/UnrealEd/Classes/Editor/PropertyEditorTestObject.h
@@ -5,6 +5,7 @@
 #include "CoreMinimal.h"
 #include "UObject/ObjectMacros.h"
 #include "UObject/Object.h"
+#include "UObject/PrimaryAssetId.h"
 #include "UObject/ScriptInterface.h"
 #include "UObject/SoftObjectPath.h"
 #include "Blueprint/UserWidget.h"
@@ -178,7 +179,7 @@ struct FPropertyEditorTestEditCondition
 	int32 EnabledAndVisibleWhenOne = 0;
 };
 
-UCLASS(EditInlineNew)
+UCLASS(EditInlineNew, Abstract)
 class UPropertyEditorTestInstancedObject : public UObject
 {
 	GENERATED_BODY()
@@ -187,6 +188,24 @@ class UPropertyEditorTestInstancedObject : public UObject
 	int32 Number;
 };
 
+UCLASS()
+class UFirstDerivedPropertyEditorTestObject : public UPropertyEditorTestInstancedObject
+{
+	GENERATED_BODY()
+
+	UPROPERTY(EditAnywhere, Category = Foo)
+	FString String;
+};
+
+UCLASS()
+class USecondDerivedPropertyEditorTestObject : public UPropertyEditorTestInstancedObject
+{
+	GENERATED_BODY()
+
+	UPROPERTY(EditAnywhere, Category = Bar)
+	bool Bool;
+};
+
 USTRUCT()
 struct FPropertyEditorTestInstancedStruct
 {
@@ -196,10 +215,10 @@ struct FPropertyEditorTestInstancedStruct
 	TObjectPtr<UPropertyEditorTestInstancedObject> Object { nullptr };
 };
 
-UCLASS(transient)
+UCLASS(transient, BlueprintType, EditInlineNew)
 class UPropertyEditorTestObject : public UObject
 {
-    GENERATED_BODY()
+	GENERATED_BODY()
 
 	UPROPERTY(EditAnywhere, Category=BasicProperties)
 	int8 Int8Property;
@@ -371,6 +390,9 @@ class UPropertyEditorTestObject : public UObject
 	UPROPERTY(EditAnywhere, Category=ArraysOfProperties)
 	TArray<FPropertyEditorTestInstancedStruct> InstancedStructArray;
 
+	UPROPERTY(EditAnywhere, Category=ArraysOfProperties, Instanced, meta=(TitleProperty="Number"))
+	TArray<TObjectPtr<UPropertyEditorTestInstancedObject>> ObjectPropertyArrayWithTitle;
+
 	UPROPERTY(EditAnywhere, Instanced, Category=ArraysOfProperties)
 	TArray<TObjectPtr<UPropertyEditorTestInstancedObject>> InstancedUObjectArray;
 
@@ -427,7 +449,13 @@ class UPropertyEditorTestObject : public UObject
 	FRichCurve RichCurve;
 
 	UPROPERTY(EditAnywhere, Category=Assets)
-	FSoftObjectPath AssetReferenceCustomStruct;
+	FSoftObjectPath SoftObjectPath;
+
+	UPROPERTY(EditAnywhere, Category=Assets)
+	FPrimaryAssetId PrimaryAssetId;
+
+	UPROPERTY(EditAnywhere, Category=Assets, meta=(DisplayThumbnail=false))
+	FPrimaryAssetId PrimaryAssetIdWithoutThumbnail;
 
 	UPROPERTY(EditAnywhere, Category=Assets, meta=(DisplayThumbnail="true"))
 	FSoftObjectPath AssetReferenceCustomStructWithThumbnail;
@@ -519,6 +547,9 @@ class UPropertyEditorTestObject : public UObject
 	UPROPERTY(EditAnywhere, Category="TMap Tests")
 	TMap<int32, FString> Int32ToStringMap;
 
+	UPROPERTY(EditAnywhere, Category = "TMap Tests", meta=(MultiLine=true))
+	TMap<FString, FText> StringToMultilineTextMap;
+
 	UPROPERTY(EditAnywhere, Category="TMap Tests")
 	TMap<FString, FLinearColor> StringToColorMap;
 
@@ -678,6 +709,15 @@ class UPropertyEditorTestObject : public UObject
 	UPROPERTY(EditAnywhere, Category=EditCondition, meta=(EditCondition="EditConditionFlags & ETestEnumFlags::One == false"))
 	bool bDisabledWhenFlagsIsOdd;
 
+	UPROPERTY(EditAnywhere, Category=EditCondition, meta=(EditCondition=false))
+	int32 AlwaysDisabled;
+
+	UPROPERTY(EditAnywhere, Category="Category Inline Edit Condition", meta=(InlineCategoryProperty))
+	bool bCategoryInlineEditCondition;
+
+	UPROPERTY(EditAnywhere, Category="Category Inline Edit Condition", meta=(EditCondition="bCategoryInlineEditCondition"))
+	float EnabledWhenCategoryChecked;
+
 	UPROPERTY(EditAnywhere, Category=OnlyInlineProperty, meta=(InlineCategoryProperty))
 	TEnumAsByte<EComponentMobility::Type> InlineProperty;
 
@@ -738,7 +778,28 @@ class UPropertyEditorTestObject : public UObject
 	UPROPERTY(EditAnywhere, Category="Inline Edit Conditions")
 	TArray<FPropertyEditorTestEditCondition> ArrayOfStructsWithInlineCondition;
 
-	bool CanEditChange(const FProperty* InProperty) const;
+	UPROPERTY(EditAnywhere, Category = ArraysOfProperties)
+	int32 NestedArrayOfInts[5];
+};
+
+UCLASS(HideCategories=(ShownByDerived))
+class UHideCategoriesBase : public UObject
+{
+	GENERATED_BODY()
+public:
+
+	UPROPERTY(EditAnywhere, Category=ShownByDerived)
+	int32 HiddenInBase;
+};
+
+UCLASS(ShowCategories=(ShownByDerived))
+class UShowCategoriesTest : public UHideCategoriesBase
+{
+	GENERATED_BODY()
+public:
+
+	UPROPERTY(EditAnywhere, Category=InDerived)
+	int32 InDerived;
 };
 
 UCLASS(EditInlineNew, Blueprintable)
@@ -751,8 +812,14 @@ public:
 
 	UPROPERTY(EditAnywhere, Category="Visible")
 	int32 ShouldBeVisible;
-};
 
+	UPROPERTY(EditAnywhere, Instanced, Category=Default)
+	TObjectPtr<USoundBase> ObjectA;
+
+	UPROPERTY(EditAnywhere, Instanced, Category=Default)
+	TObjectPtr<USoundBase> ObjectB;
+};
+ 
 UCLASS(Blueprintable)
 class UBlueprintPropertyContainerTestObject : public UObject
 {
@@ -788,9 +855,49 @@ class APropertyEditorTestActor : public AActor
 {
 	GENERATED_BODY()
 
+	UPROPERTY(EditAnywhere, Instanced, Category = ArraysOfProperties)
+	TArray<TObjectPtr<UPropertyEditorTestInstancedObject>> InstancedUObjectArray;
+
 	UPROPERTY(EditAnywhere, Category=Default, meta=(GetOptions=GetOptionsFunc))
-	FName GetOptionsValue;
+	FName GetOptionsValue; 
+
+	UPROPERTY(EditDefaultsOnly, Category="Defaults Only")
+	float DefaultsOnly;
+
+	UPROPERTY(EditDefaultsOnly, Category="Defaults Only|Subcategory")
+	float DefaultsOnlySubcategory;
+
+	UPROPERTY(EditInstanceOnly, Category="Instance Only")
+	float InstanceOnly;
+
+	UPROPERTY(EditInstanceOnly, Category="Instance Only|Subcategory")
+	float InstanceOnlySubcategory;
+
+	UPROPERTY(EditAnywhere, Category="Map")
+	TMap<int32, FText> MultiLineMap;
 
 	UFUNCTION()
 	TArray<FString> GetOptionsFunc() const;
 };
+
+class IDetailTreeNode;
+class IPropertyRowGenerator;
+
+UCLASS()
+class UPropertyEditorRowGeneratorTest : public UObject
+{
+	GENERATED_BODY()
+
+public:
+
+	TSharedRef<SWidget> GenerateWidget();
+
+private:
+	void OnRowsRefreshed();
+	TSharedRef<ITableRow> GenerateListRow(TSharedPtr<IDetailTreeNode> InItem, const TSharedRef<STableViewBase>& OwnerTable);
+
+private:
+	TArray<TSharedPtr<IDetailTreeNode>> DetailsNodes;
+	TSharedPtr<SListView<TSharedPtr<IDetailTreeNode>>> ListView;
+	TSharedPtr<IPropertyRowGenerator> PropertyRowGenerator;
+};
\ No newline at end of file
diff --git a/Engine/Source/Editor/UnrealEd/Private/EditorServer.cpp b/Engine/Source/Editor/UnrealEd/Private/EditorServer.cpp
--- a/Engine/Source/Editor/UnrealEd/Private/EditorServer.cpp
+++ b/Engine/Source/Editor/UnrealEd/Private/EditorServer.cpp
@@ -6025,17 +6025,6 @@ bool UEditorEngine::HandleCallbackCommand( UWorld* InWorld, const TCHAR* Str , F
 
 bool UEditorEngine::HandleTestPropsCommand( const TCHAR* Str, FOutputDevice& Ar )
 {
-	UObject* Object;
-	UClass* Class = NULL;
-	if( ParseObject<UClass>( Str, TEXT("CLASS="), Class, ANY_PACKAGE ) != false )
-	{ 
-		Object = NewObject<UObject>(GetTransientPackage(), Class);
-	}
-	else
-	{
-		Object = NewObject<UPropertyEditorTestObject>();
-	}
-
 	TSharedRef<SWindow> Window = SNew(SWindow)
 		.Title( NSLOCTEXT("UnrealEd", "PropertyEditorTestWindowTitle", "Property Editor Test") )
 		.ClientSize(FVector2D(500,1000));
@@ -6044,6 +6033,17 @@ bool UEditorEngine::HandleTestPropsCommand( const TCHAR* Str, FOutputDevice& Ar
 
 	if( FParse::Command(&Str,TEXT("TREE")) )
 	{
+		UObject* Object;
+		UClass* Class = NULL;
+		if( ParseObject<UClass>( Str, TEXT("CLASS="), Class, ANY_PACKAGE ) != false )
+		{ 
+			Object = NewObject<UObject>(GetTransientPackage(), Class);
+		}
+		else
+		{
+			Object = NewObject<UPropertyEditorTestObject>();
+		}
+
 		FDetailsViewArgs Args;
 		Args.bHideSelectionTip = true;
 
@@ -6098,7 +6098,13 @@ bool UEditorEngine::HandleTestPropsCommand( const TCHAR* Str, FOutputDevice& Ar
 		StructArgs.bShowObjects = true;
 		StructArgs.bShowInterfaces = true;
 
-		TSharedRef<IStructureDetailsView> StructDetails = Module.CreateStructureDetailView(DetailsArgs, StructArgs, MakeShared<FStructOnScope>(FPropertyEditorTestBasicStruct::StaticStruct()));
+		UStruct* Struct = nullptr;
+		if (!ParseObject<UStruct>(Str, TEXT("STRUCT="), Struct, ANY_PACKAGE))
+		{
+			Struct = FPropertyEditorTestBasicStruct::StaticStruct();
+		}
+
+		TSharedRef<IStructureDetailsView> StructDetails = Module.CreateStructureDetailView(DetailsArgs, StructArgs, MakeShared<FStructOnScope>(Struct));
 
 		Window->SetContent(
 			SNew(SBorder)
@@ -6108,8 +6114,31 @@ bool UEditorEngine::HandleTestPropsCommand( const TCHAR* Str, FOutputDevice& Ar
 			]
 		);
 	}
+	else if ( FParse::Command(&Str, TEXT("GENERATOR")) )
+	{
+		UPropertyEditorRowGeneratorTest* TestGenerator = NewObject<UPropertyEditorRowGeneratorTest>();
+		
+		Window->SetContent(
+			SNew(SBorder)
+			.BorderImage(FEditorStyle::GetBrush("ToolPanel.GroupBorder"))
+			[
+				TestGenerator->GenerateWidget()
+			]
+		);
+	}
 	else
 	{
+		UObject* Object;
+		UClass* Class = nullptr;
+		if( ParseObject<UClass>( Str, TEXT("CLASS="), Class, ANY_PACKAGE ) != false )
+		{ 
+			Object = NewObject<UObject>(GetTransientPackage(), Class);
+		}
+		else
+		{
+			Object = NewObject<UPropertyEditorTestObject>();
+		}
+
 		//Details
 		TArray<UObject*> Objects;
 		Objects.Add( Object );
diff --git a/Engine/Source/Editor/UnrealEd/Private/PropertyEditorTestObject.cpp b/Engine/Source/Editor/UnrealEd/Private/PropertyEditorTestObject.cpp
--- a/Engine/Source/Editor/UnrealEd/Private/PropertyEditorTestObject.cpp
+++ b/Engine/Source/Editor/UnrealEd/Private/PropertyEditorTestObject.cpp
@@ -2,17 +2,91 @@
 
 #include "Editor/PropertyEditorTestObject.h"
 
-bool UPropertyEditorTestObject::CanEditChange(const FProperty* InProperty) const
+#include "IDetailTreeNode.h"
+#include "IPropertyRowGenerator.h"
+#include "Modules/ModuleManager.h"
+#include "Widgets/Views/STableViewBase.h"
+#include "Widgets/Views/SListView.h"
+
+
+TArray<FString> APropertyEditorTestActor::GetOptionsFunc() const
+{
+	return TArray<FString> { TEXT("One"), TEXT("Two"), TEXT("Three") };
+}
+
+TSharedRef<SWidget> UPropertyEditorRowGeneratorTest::GenerateWidget()
 {
-	if (InProperty->GetFName() == GET_MEMBER_NAME_CHECKED(UPropertyEditorTestObject, DisabledByCanEditChange))
+	FPropertyRowGeneratorArgs GeneratorArgs;
+	GeneratorArgs.bAllowMultipleTopLevelObjects = true;
+
+	FPropertyEditorModule& Module = FModuleManager::LoadModuleChecked<FPropertyEditorModule>( "PropertyEditor" );
+	PropertyRowGenerator = Module.CreatePropertyRowGenerator(GeneratorArgs);
+	PropertyRowGenerator->OnRowsRefreshed().AddUObject(this, &UPropertyEditorRowGeneratorTest::OnRowsRefreshed);
+
+	TArray<UObject*> Objects;
+	for (int32 Idx = 0; Idx < 10; ++Idx)
 	{
-		return false;
+		Objects.Add(NewObject<UPropertyEditorTestObject>());
 	}
 
-	return true;
+	ListView = SNew(SListView<TSharedPtr<IDetailTreeNode>>)
+		.ItemHeight(24)
+		.ListItemsSource(&DetailsNodes)
+		.OnGenerateRow_UObject(this, &UPropertyEditorRowGeneratorTest::GenerateListRow);
+
+	PropertyRowGenerator->SetObjects(Objects);
+
+	return ListView.ToSharedRef();
 }
 
-TArray<FString> APropertyEditorTestActor::GetOptionsFunc() const
+static void AddChildrenRecursive(TSharedPtr<IDetailTreeNode> Node, TArray<TSharedPtr<IDetailTreeNode>>& OutNodes)
 {
-	return TArray<FString> { TEXT("One"), TEXT("Two"), TEXT("Three") };
-}
\ No newline at end of file
+	TArray<TSharedRef<IDetailTreeNode>> Children;
+	Node->GetChildren(Children);
+			
+	for (const TSharedRef<IDetailTreeNode>& Child : Children)
+	{
+		OutNodes.Add(Child);
+
+		AddChildrenRecursive(Child, OutNodes);
+	}
+}
+
+void UPropertyEditorRowGeneratorTest::OnRowsRefreshed()
+{
+	DetailsNodes.Reset();
+
+	for (const TSharedRef<IDetailTreeNode>& RootNode : PropertyRowGenerator->GetRootTreeNodes())
+	{
+		AddChildrenRecursive(RootNode, DetailsNodes);
+	}
+
+	ListView->RebuildList();
+}
+
+TSharedRef<ITableRow> UPropertyEditorRowGeneratorTest::GenerateListRow(TSharedPtr<IDetailTreeNode> InItem, const TSharedRef<STableViewBase>& OwnerTable)
+{ 
+	TSharedRef<STableRow<TSharedPtr<IDetailTreeNode>>> TableRow = SNew(STableRow<TSharedPtr<IDetailTreeNode>>, OwnerTable);
+
+	FNodeWidgets NodeWidgets = InItem->CreateNodeWidgets();
+	if (NodeWidgets.WholeRowWidget.IsValid())
+	{
+		TableRow->SetContent(NodeWidgets.WholeRowWidget.ToSharedRef());
+	}
+	else if (NodeWidgets.NameWidget.IsValid() && NodeWidgets.ValueWidget.IsValid())
+	{
+		TableRow->SetContent(
+			SNew(SHorizontalBox)
+			+ SHorizontalBox::Slot()
+			[
+				NodeWidgets.NameWidget.ToSharedRef()
+			]
+			+ SHorizontalBox::Slot()
+			[
+				NodeWidgets.ValueWidget.ToSharedRef()
+			]
+		);
+	}
+
+	return TableRow;
+}
