CookCommandlet: Remove TickRecompileShaderRequests from CookByTheBook, since RecompileShaderRequests are only populated during CookOnTheFly, and we want to remove unnecessary polling from CookByTheBook for performance improvement.
#rnx

diff --git a/Engine/Source/Editor/UnrealEd/Private/Commandlets/CookCommandlet.cpp b/Engine/Source/Editor/UnrealEd/Private/Commandlets/CookCommandlet.cpp
--- a/Engine/Source/Editor/UnrealEd/Private/Commandlets/CookCommandlet.cpp
+++ b/Engine/Source/Editor/UnrealEd/Private/Commandlets/CookCommandlet.cpp
@@ -1098,13 +1098,6 @@ bool UCookCommandlet::CookByTheBook( const TArray<ITargetPlatform*>& Platforms)
 					DumpMemStats();
 				}
 
-				{
-					UE_SCOPED_COOKTIMER_AND_DURATION(CookByTheBook_RecompileShaderRequests, DetailedCookStats::TickLoopRecompileShaderRequestsTimeSec);
-					CookOnTheFlyServer->TickRecompileShaderRequests();
-
-					FPlatformProcess::Sleep( 0.0f );
-				}
-
 				{
 					UE_SCOPED_COOKTIMER_AND_DURATION(CookByTheBook_ProcessDeferredCommands, DetailedCookStats::TickLoopProcessDeferredCommandsTimeSec);
 					ProcessDeferredCommands();
