CookCommandlet: Remove remaining CookByTheBookInCommandlet pollables into PumpPollables, and add a sleep command when we are busy with no pollables ready.

diff --git a/Engine/Source/Editor/UnrealEd/Classes/CookOnTheSide/CookOnTheFlyServer.h b/Engine/Source/Editor/UnrealEd/Classes/CookOnTheSide/CookOnTheFlyServer.h
--- a/Engine/Source/Editor/UnrealEd/Classes/CookOnTheSide/CookOnTheFlyServer.h
+++ b/Engine/Source/Editor/UnrealEd/Classes/CookOnTheSide/CookOnTheFlyServer.h
@@ -328,6 +328,7 @@ private:
 		SaveLimited,	// Process the SaveQueue, stopping when savequeuelength reaches the desired population level
 		Poll,			// Execute pollables which have exceeded their period
 		PollIdle,		// Execute pollables which have exceeded their idle period
+		WaitForAsync,	// Sleep for a time slice while we wait for async tasks to complete
 		YieldTick,		// Progress is blocked by an async result. Temporarily exit TickCookOnTheSide.
 		Cancel,			// Cancel the current CookByTheBook
 	};
@@ -387,6 +388,7 @@ private:
 	TOptional<FPollable> CreatePollableLLM();
 	TOptional<FPollable> CreatePollableTriggerGC();
 	void PollPackagesPerGC(UE::Cook::FTickStackData& StackData);
+	void WaitForAsync(UE::Cook::FTickStackData& StackData);
 
 public:
 
@@ -512,7 +514,7 @@ public:
 	* 
 	* @return returns ECookOnTheSideResult
 	*/
-	uint32 TickCookOnTheSide( const float TimeSlice, uint32 &CookedPackagesCount, ECookTickFlags TickFlags = ECookTickFlags::None );
+	uint32 TickCookOnTheSide( const float TimeSlice, uint32 &CookedPackagesCount, ECookTickFlags TickFlags = ECookTickFlags::None, int32 InMaxNumPackagesToSave = 50);
 
 	/**
 	* Clear all the previously cooked data all cook requests from now on will be considered recook requests
@@ -1168,6 +1170,7 @@ private:
 	float PumpPollablesTimeSlice = 0.f;
 	float PumpPollablesMinPeriod = 0.f;
 	uint32 CookedPackageCountSinceLastGC = 0;
+	float WaitForAsyncSleepSeconds = 0.0f;
 	friend UE::Cook::FPackageData;
 	friend UE::Cook::FPendingCookedPlatformData;
 	friend UE::Cook::FPlatformManager;
diff --git a/Engine/Source/Editor/UnrealEd/Private/Commandlets/CookCommandlet.cpp b/Engine/Source/Editor/UnrealEd/Private/Commandlets/CookCommandlet.cpp
--- a/Engine/Source/Editor/UnrealEd/Private/Commandlets/CookCommandlet.cpp
+++ b/Engine/Source/Editor/UnrealEd/Private/Commandlets/CookCommandlet.cpp
@@ -1005,21 +1005,11 @@ bool UCookCommandlet::CookByTheBook( const TArray<ITargetPlatform*>& Platforms)
 			DECLARE_SCOPE_CYCLE_COUNTER( TEXT( "CookByTheBook.MainLoop" ), STAT_CookByTheBook_MainLoop, STATGROUP_LoadTime );
 			{
 				uint32 TickResults = 0;
-				const float CookOnTheSideTimeSlice = 10.0f;
 				uint32 UnusedVariable = 0;
 
-				TickResults = CookOnTheFlyServer->TickCookOnTheSide( CookOnTheSideTimeSlice, UnusedVariable,
-					ShowProgress ? ECookTickFlags::None : ECookTickFlags::HideProgressDisplay );
+				TickResults = CookOnTheFlyServer->TickCookOnTheSide( MAX_flt, UnusedVariable,
+					ShowProgress ? ECookTickFlags::None : ECookTickFlags::HideProgressDisplay, MAX_int32);
 
-				{
-					UE_SCOPED_COOKTIMER_AND_DURATION(CookByTheBook_ShaderProcessAsync, DetailedCookStats::TickLoopShaderProcessAsyncResultsTimeSec);
-					GShaderCompilingManager->ProcessAsyncResults(true, false);
-				}
-				{
-					UE_SCOPED_COOKTIMER(CookByTheBook_TickAssetRegistry);
-					FAssetRegistryModule::TickAssetRegistry(-1.0f);
-				}
-				
 				if ((TickResults & UCookOnTheFlyServer::COSR_RequiresGC) != 0)
 				{
 					FString GCReason;
@@ -1097,16 +1087,6 @@ bool UCookCommandlet::CookByTheBook( const TArray<ITargetPlatform*>& Platforms)
 
 					DumpMemStats();
 				}
-
-				{
-					UE_SCOPED_COOKTIMER_AND_DURATION(CookByTheBook_ProcessDeferredCommands, DetailedCookStats::TickLoopProcessDeferredCommandsTimeSec);
-					ProcessDeferredCommands();
-				}
-			}
-
-			{
-				UE_SCOPED_COOKTIMER_AND_DURATION(CookByTheBook_TickCommandletStats, DetailedCookStats::TickLoopTickCommandletStatsTimeSec);
-				FStats::TickCommandletStats();
 			}
 		}
 	} while (bTestCook);
diff --git a/Engine/Source/Editor/UnrealEd/Private/CookOnTheFlyServer.cpp b/Engine/Source/Editor/UnrealEd/Private/CookOnTheFlyServer.cpp
--- a/Engine/Source/Editor/UnrealEd/Private/CookOnTheFlyServer.cpp
+++ b/Engine/Source/Editor/UnrealEd/Private/CookOnTheFlyServer.cpp
@@ -1159,7 +1159,7 @@ bool UCookOnTheFlyServer::RequestPackage(const FName& StandardPackageFName, cons
 	return RequestPackage(StandardPackageFName, PlatformManager->GetSessionPlatforms(), bForceFrontOfQueue);
 }
 
-uint32 UCookOnTheFlyServer::TickCookOnTheSide(const float TimeSlice, uint32 &CookedPackageCount, ECookTickFlags TickFlags)
+uint32 UCookOnTheFlyServer::TickCookOnTheSide(const float TimeSlice, uint32 &CookedPackageCount, ECookTickFlags TickFlags, int32 InMaxNumPackagesToSave)
 {
 	LLM_SCOPE_BYTAG(Cooker);
 	TickCancels();
@@ -1188,7 +1188,7 @@ uint32 UCookOnTheFlyServer::TickCookOnTheSide(const float TimeSlice, uint32 &Coo
 		}
 	}
 
-	UE::Cook::FTickStackData StackData(TimeSlice, IsRealtimeMode(), TickFlags);
+	UE::Cook::FTickStackData StackData(TimeSlice, IsRealtimeMode(), TickFlags, InMaxNumPackagesToSave);
 	bool bCookComplete = false;
 
 	{
@@ -1241,6 +1241,9 @@ uint32 UCookOnTheFlyServer::TickCookOnTheSide(const float TimeSlice, uint32 &Coo
 			case ECookAction::PollIdle:
 				PumpPollables(StackData, true /* bIsIdle */);
 				break;
+			case ECookAction::WaitForAsync:
+				WaitForAsync(StackData);
+				break;
 			case ECookAction::YieldTick:
 				bContinueTick = false;
 				break;
@@ -1661,6 +1664,49 @@ void UCookOnTheFlyServer::PollPackagesPerGC(UE::Cook::FTickStackData& StackData)
 	}
 }
 
+namespace UE::Cook
+{
+
+static void CBTBProcessDeferredCommands()
+{
+	UE_SCOPED_COOKTIMER_AND_DURATION(CookByTheBook_ProcessDeferredCommands, DetailedCookStats::TickLoopProcessDeferredCommandsTimeSec);
+#if PLATFORM_MAC
+	// On Mac we need to process Cocoa events so that the console window for CookOnTheFlyServer is interactive
+	FPlatformApplicationMisc::PumpMessages(true);
+#endif
+
+	// update task graph
+	FTaskGraphInterface::Get().ProcessThreadUntilIdle(ENamedThreads::GameThread);
+
+	// execute deferred commands
+	for (const FString& DeferredCommand : GEngine->DeferredCommands)
+	{
+		GEngine->Exec(GWorld, *DeferredCommand, *GLog);
+	}
+
+	GEngine->DeferredCommands.Empty();
+}
+
+static void CBTBTickCommandletStats()
+{
+	UE_SCOPED_COOKTIMER_AND_DURATION(CookByTheBook_TickCommandletStats, DetailedCookStats::TickLoopTickCommandletStatsTimeSec);
+	FStats::TickCommandletStats();
+}
+
+static void CBTBTickShaderCompilingManager()
+{
+	UE_SCOPED_COOKTIMER_AND_DURATION(CookByTheBook_ShaderProcessAsync, DetailedCookStats::TickLoopShaderProcessAsyncResultsTimeSec);
+	GShaderCompilingManager->ProcessAsyncResults(true, false);
+}
+
+static void CBTBTickAssetRegistry()
+{
+	UE_SCOPED_COOKTIMER(CookByTheBook_TickAssetRegistry);
+	FAssetRegistryModule::TickAssetRegistry(-1.0f);
+}
+
+}
+
 void UCookOnTheFlyServer::InitializePollables()
 {
 	using namespace UE::Cook;
@@ -1678,6 +1724,10 @@ void UCookOnTheFlyServer::InitializePollables()
 		{
 			Pollables.Add(MoveTemp(*Pollable));
 		}
+		Pollables.Add(FPollable(60.f, 5.f, [this](FTickStackData&) { CBTBProcessDeferredCommands(); }));
+		Pollables.Add(FPollable(60.f, 5.f, [this](FTickStackData&) { CBTBTickCommandletStats(); }));
+		Pollables.Add(FPollable(60.f, 5.f, [this](FTickStackData&) { CBTBTickShaderCompilingManager(); }));
+		Pollables.Add(FPollable(60.f, 5.f, [this](FTickStackData&) { CBTBTickAssetRegistry(); }));
 	}
 
 	PollStartIndex = 0;
@@ -1693,6 +1743,20 @@ void UCookOnTheFlyServer::InitializePollables()
 	}
 }
 
+void UCookOnTheFlyServer::WaitForAsync(UE::Cook::FTickStackData& StackData)
+{
+	// Sleep until the next time that DecideNextCookAction will find work to do, up to a maximum of WaitForAsyncSleepSeconds
+	UE_SCOPED_COOKTIMER(WaitForAsync);
+	double CurrentTime = FPlatformTime::Seconds();
+	float SleepDuration = WaitForAsyncSleepSeconds;
+	SleepDuration = FMath::Min(SleepDuration, StackData.Timer.GetEndTimeSeconds() - CurrentTime);
+	SleepDuration = FMath::Min(SleepDuration, PollNextTimeIdleSeconds - CurrentTime);
+	SleepDuration = FMath::Min(SleepDuration, SaveBusyTimeLastRetry + GCookProgressRetryBusyTime - CurrentTime);
+	SleepDuration = FMath::Min(SleepDuration, LoadBusyTimeLastRetry + GCookProgressRetryBusyTime - CurrentTime);
+	SleepDuration = FMath::Max(SleepDuration, 0);
+	FPlatformProcess::Sleep(SleepDuration);
+}
+
 UCookOnTheFlyServer::ECookAction UCookOnTheFlyServer::DecideNextCookAction(UE::Cook::FTickStackData& StackData)
 {
 	if (IsCookByTheBookMode() && CookByTheBookOptions->bCancel)
@@ -1821,7 +1885,14 @@ UCookOnTheFlyServer::ECookAction UCookOnTheFlyServer::DecideNextCookAction(UE::C
 		{
 			return ECookAction::PollIdle;
 		}
-		return ECookAction::YieldTick;
+		else if (IsRealtimeMode() || IsCookOnTheFlyMode())
+		{
+			return ECookAction::YieldTick;
+		}
+		else
+		{
+			return ECookAction::WaitForAsync;
+		}
 	}
 
 	return ECookAction::Done;
@@ -4646,6 +4717,7 @@ void UCookOnTheFlyServer::Initialize( ECookMode::Type DesiredCookMode, ECookInit
 	PumpPollablesMinPeriod = 0.1f;
 	PollNextTimeSeconds = MAX_flt;
 	PollNextTimeIdleSeconds = MAX_flt;
+	WaitForAsyncSleepSeconds = 1.0f;
 
 	{
 		const FConfigSection* CacheSettings = GConfig->GetSectionPrivate(TEXT("CookPlatformDataCacheSettings"), false, true, GEditorIni);
diff --git a/Engine/Source/Editor/UnrealEd/Private/Cooker/CookProfiling.h b/Engine/Source/Editor/UnrealEd/Private/Cooker/CookProfiling.h
--- a/Engine/Source/Editor/UnrealEd/Private/Cooker/CookProfiling.h
+++ b/Engine/Source/Editor/UnrealEd/Private/Cooker/CookProfiling.h
@@ -126,6 +126,9 @@ namespace DetailedCookStats
 	extern double TickCookOnTheSideBeginPrepareSaveTimeSec;
 	extern double TickCookOnTheSideFinishPrepareSaveTimeSec;
 	extern double GameCookModificationDelegateTimeSec;
+	extern double TickLoopShaderProcessAsyncResultsTimeSec;
+	extern double TickLoopProcessDeferredCommandsTimeSec;
+	extern double TickLoopTickCommandletStatsTimeSec;
 	extern double TickLoopFlushRenderingCommandsTimeSec;
 
 	// Stats tracked through FAutoRegisterCallback
diff --git a/Engine/Source/Editor/UnrealEd/Private/Cooker/CookTypes.cpp b/Engine/Source/Editor/UnrealEd/Private/Cooker/CookTypes.cpp
--- a/Engine/Source/Editor/UnrealEd/Private/Cooker/CookTypes.cpp
+++ b/Engine/Source/Editor/UnrealEd/Private/Cooker/CookTypes.cpp
@@ -31,6 +31,11 @@ double FCookerTimer::GetTimeTillNow() const
 	return FPlatformTime::Seconds() - StartTime;
 }
 
+double FCookerTimer::GetEndTimeSeconds() const
+{
+	return bIsRealtimeMode ? StartTime + TimeSlice : MAX_flt;
+}
+
 bool FCookerTimer::IsTimeUp() const
 {
 	return IsTimeUp(FPlatformTime::Seconds());
diff --git a/Engine/Source/Editor/UnrealEd/Private/Cooker/CookTypes.h b/Engine/Source/Editor/UnrealEd/Private/Cooker/CookTypes.h
--- a/Engine/Source/Editor/UnrealEd/Private/Cooker/CookTypes.h
+++ b/Engine/Source/Editor/UnrealEd/Private/Cooker/CookTypes.h
@@ -137,6 +137,7 @@ namespace UE::Cook
 		FCookerTimer(ENoWait);
 
 		double GetTimeTillNow() const;
+		double GetEndTimeSeconds() const;
 		bool IsTimeUp() const;
 		bool IsTimeUp(double CurrentTimeSeconds) const;
 		void SavedPackage();
@@ -166,8 +167,8 @@ namespace UE::Cook
 		/** CookFlags describing details of the caller's desired behavior of the current TickCookOnTheSide request. */
 		ECookTickFlags TickFlags;
 
-		explicit FTickStackData(const float& TimeSlice, const bool bIsRealtimeMode, ECookTickFlags InTickFlags)
-			:Timer(TimeSlice, bIsRealtimeMode), TickFlags(InTickFlags)
+		explicit FTickStackData(const float& TimeSlice, const bool bIsRealtimeMode, ECookTickFlags InTickFlags, int32 InMaxNumPackagesToSave)
+			:Timer(TimeSlice, bIsRealtimeMode, InMaxNumPackagesToSave), TickFlags(InTickFlags)
 		{
 		}
 	};
