Fix issue with re-import with new file when there was no original file in the asset import data.
#rnx

diff --git a/Engine/Source/Editor/UnrealEd/Private/Editor.cpp b/Engine/Source/Editor/UnrealEd/Private/Editor.cpp
--- a/Engine/Source/Editor/UnrealEd/Private/Editor.cpp
+++ b/Engine/Source/Editor/UnrealEd/Private/Editor.cpp
@@ -626,13 +626,15 @@ void FReimportManager::ValidateAllSourceFileAndReimport(TArray<UObject*> &ToImpo
 				for (int32 FileIndex : SourceIndexArray)
 				{
 					TArray<FString> SourceFilenames;
-					this->GetNewReimportPath(Asset, SourceFilenames, FileIndex);
-					if (!SourceFilenames.IsValidIndex(FileIndex) || SourceFilenames[FileIndex].IsEmpty())
+					GetNewReimportPath(Asset, SourceFilenames, FileIndex);
+					//The FileIndex can be INDEX_NONE in case the caller do not specify any source index, in that case we want to use the first index which is 0.
+					int32 RealSourceFileIndex = FileIndex == INDEX_NONE ? 0 : FileIndex;
+					if (!SourceFilenames.IsValidIndex(RealSourceFileIndex) || SourceFilenames[RealSourceFileIndex].IsEmpty())
 					{
 						continue;
 					}
 					bCancelAll = false;
-					this->UpdateReimportPath(Asset, SourceFilenames[FileIndex], FileIndex);
+					UpdateReimportPath(Asset, SourceFilenames[RealSourceFileIndex], RealSourceFileIndex);
 					//We do not want to ask again the user for a file
 					bForceNewFile = false;
 				}
