Localization: Correction to TextGatherer to match MACRO compilation rules:
- Ending a line with and escape character and continue on the next line
- Concatenate parameters on multiple lines.

diff --git a/Engine/Source/Editor/UnrealEd/Classes/Commandlets/GatherTextFromSourceCommandlet.h b/Engine/Source/Editor/UnrealEd/Classes/Commandlets/GatherTextFromSourceCommandlet.h
--- a/Engine/Source/Editor/UnrealEd/Classes/Commandlets/GatherTextFromSourceCommandlet.h
+++ b/Engine/Source/Editor/UnrealEd/Classes/Commandlets/GatherTextFromSourceCommandlet.h
@@ -84,6 +84,56 @@ private:
 		TMap<FString, FParsedStringTableEntryMetaDataMap, FDefaultSetAllocator, FLocKeyMapFuncs<FParsedStringTableEntryMetaDataMap>> MetaDataEntries;
 	};
 
+	class FMacroArgumentGatherer
+	{
+	public:
+		FMacroArgumentGatherer() {};
+
+		bool Gather(const TCHAR* Arg, int32 Count);
+		bool EndArgument();
+
+		// Return the number of argument that are completely resolved.
+		int32 GetNumberOfArguments() const;
+
+		void ExtractArguments(TArray<FString>& Arguments);
+		
+		void OpenDoubleQuotes()
+		{
+			bInDblQuotes = true;
+		}
+
+		void CloseDoubleQuotes()
+		{
+			bInDblQuotes = false;
+		}
+
+		bool IsInDoubleQuotes()  const
+		{
+			return bInDblQuotes;
+		};
+		
+		void OpenSingleQuotes() 
+		{
+			bInSglQuotes = true;
+		}
+
+		void CloseSingleQuotes()
+		{
+			bInSglQuotes = false;
+		}
+
+		bool IsInSingleQuotes()  const
+		{
+			return bInSglQuotes;
+		};
+
+	private:
+		TArray<FString> Args;
+		FString CurrentArgument;
+		bool bInDblQuotes = false;
+		bool bInSglQuotes = false;
+	};
+
 	struct FSourceFileParseContext
 	{
 		bool AddManifestText( const FString& Token, const FString& Namespace, const FString& SourceText, const FManifestContext& Context );
@@ -275,8 +325,8 @@ private:
 
 	protected:
 		bool ParseArgsFromMacro(const FString& Text, TArray<FString>& Args, FSourceFileParseContext& Context) const;
-		bool ParseArgsFromNextLines(TArray<FString>& Args, int32& BracketStack, FSourceFileParseContext& Context) const;
-		bool ParseArgumentString(const FString& Text, const int32 OpenBracketIdx, int32& BracketStack, const FSourceFileParseContext& Context, TArray<FString>& Args) const;
+		bool ParseArgsFromNextLines(FMacroArgumentGatherer& ArgsGatherer, int32& BracketStack, FSourceFileParseContext& Context) const;
+		bool ParseArgumentString(const FString& Text, const int32 OpenBracketIdx, int32& BracketStack, const FSourceFileParseContext& Context, FMacroArgumentGatherer& ArgsGatherer) const;
 
 		static bool PrepareArgument(FString& Argument, bool IsAutoText, const FString& IdentForLogging, bool& OutHasQuotes);
 
diff --git a/Engine/Source/Editor/UnrealEd/Private/Commandlets/GatherTextFromSourceCommandlet.cpp b/Engine/Source/Editor/UnrealEd/Private/Commandlets/GatherTextFromSourceCommandlet.cpp
--- a/Engine/Source/Editor/UnrealEd/Private/Commandlets/GatherTextFromSourceCommandlet.cpp
+++ b/Engine/Source/Editor/UnrealEd/Private/Commandlets/GatherTextFromSourceCommandlet.cpp
@@ -331,7 +331,7 @@ int32 UGatherTextFromSourceCommandlet::Main( const FString& Params )
 	{
 		if (ParsedStringTablePair.Value.SourceLocation.Line == INDEX_NONE)
 		{
-			UE_LOG(LogGatherTextFromSourceCommandlet, Warning, TEXT("String table with ID '%s' had %s entries parsed for it, but the table was never registered. Skipping for gather."), *ParsedStringTablePair.Key.ToString(), ParsedStringTablePair.Value.TableEntries.Num());
+			UE_LOG(LogGatherTextFromSourceCommandlet, Warning, TEXT("String table with ID '%s' had %d entries parsed for it, but the table was never registered. Skipping for gather."), *ParsedStringTablePair.Key.ToString(), ParsedStringTablePair.Value.TableEntries.Num());
 		}
 		else
 		{
@@ -1064,6 +1064,66 @@ bool UGatherTextFromSourceCommandlet::ParseSourceText(const FString& Text, const
 	return true;
 }
 
+int32 UGatherTextFromSourceCommandlet::FMacroArgumentGatherer::GetNumberOfArguments() const 
+{
+	return Args.Num();
+}
+
+
+bool UGatherTextFromSourceCommandlet::FMacroArgumentGatherer::EndArgument()
+{	
+	if (CurrentArgument.IsEmpty() || IsInDoubleQuotes())
+	{
+		return false;
+	}
+	Args.Push(CurrentArgument);
+	CurrentArgument.Empty();
+	return true;
+}
+
+bool UGatherTextFromSourceCommandlet::FMacroArgumentGatherer::Gather(const TCHAR* Arg, int32 Count)
+{
+	if(Arg == nullptr || Count == 0)
+	{
+		return false;
+	}
+	FString NewArgument = FString(Count, Arg);
+	NewArgument.TrimEndInline();
+
+	int32 CurrentArgLen = CurrentArgument.Len();
+	if (CurrentArgLen == 0)
+	{
+		CurrentArgument = NewArgument;
+		return true;
+	}
+	
+	// If the last argument is a string that ends with " and the new argument too we append it
+	// This is to support this and remove duplicate double quotes.
+	// MYMACRO("Very long string\n"
+	//         "Next part of the very long string\n");
+	// And
+	// MYMACRO("Very long string\n" \
+	//         "Next part of the very long string\n");
+	if (NewArgument[0] == '\"' && CurrentArgument[CurrentArgLen - 1] == '\"')
+	{
+		Arg++;
+		CurrentArgLen--;
+		CurrentArgument.RemoveAt(CurrentArgLen);
+		Count--;
+	}
+
+	CurrentArgument.Append(Arg, Count);
+	CurrentArgument.TrimEndInline();
+	return true;
+}
+	
+
+void UGatherTextFromSourceCommandlet::FMacroArgumentGatherer::ExtractArguments(TArray<FString>& Arguments)
+{
+	Arguments = Args;
+	Args.Empty();
+}
+
 bool UGatherTextFromSourceCommandlet::FSourceFileParseContext::AddManifestText( const FString& Token, const FString& InNamespace, const FString& SourceText, const FManifestContext& Context )
 {
 	const bool bIsEditorOnly = EvaluateEditorOnlyDefineState() == EEditorOnlyDefineState::Defined;
@@ -1447,51 +1507,76 @@ void UGatherTextFromSourceCommandlet::FEndIfDescriptor::TryParse(const FString&
 	}
 }
 
-bool UGatherTextFromSourceCommandlet::FMacroDescriptor::ParseArgumentString(const FString& Text, const int32 CursorOffset, int32& BracketStack, const FSourceFileParseContext& Context, TArray<FString>& Args) const
+bool UGatherTextFromSourceCommandlet::FMacroDescriptor::ParseArgumentString(const FString& Text, const int32 CursorOffset, int32& BracketStack, const FSourceFileParseContext& Context, FMacroArgumentGatherer& ArgsGatherer) const
 {
-	bool bInDblQuotes = false;
-	bool bInSglQuotes = false;
 	bool bEscapeNextChar = false;
 
 	const TCHAR* ArgStart = *Text + CursorOffset;
 	const TCHAR* Cursor = ArgStart;
 	for (; 0 < BracketStack; ++Cursor)
 	{
+		// First: check is we are at end of line.
+		if ('\0' == *Cursor)
+		{
+			if (Cursor - ArgStart > 0)
+			{
+				// Here, we found the end of the line, but we don't know if it's the end of the argument.
+				if (!ArgsGatherer.Gather(ArgStart, Cursor - ArgStart))
+				{
+					UE_LOG(LogGatherTextFromSourceCommandlet, Warning, TEXT("Parsing Arguments failed in %s macro while parsing %s:%d. %s"), *GetToken(), *Context.Filename, Context.LineNumber, *FLocTextHelper::SanitizeLogOutput(Context.LineText.TrimStartAndEnd()));
+					return false;
+				}
+			}
+			break;
+		}
 		// Skip this character
-		if (bEscapeNextChar)
+		else if (bEscapeNextChar)
 		{
 			bEscapeNextChar = false;
 		}
-		else if ((bInDblQuotes || bInSglQuotes) && !bEscapeNextChar && '\\' == *Cursor)
+		else if ((ArgsGatherer.IsInDoubleQuotes() || ArgsGatherer.IsInSingleQuotes()) && !bEscapeNextChar && '\\' == *Cursor)
 		{
 			bEscapeNextChar = true;
-		}
 
+			// If we hit the escape character, we must verify if we are at end of line
+			if ('\0' == *(Cursor + 1))
+			{
+				if (Cursor - ArgStart - 1 > 0)
+				{
+					if (!ArgsGatherer.Gather(ArgStart, Cursor - ArgStart - 1))
+					{
+						UE_LOG(LogGatherTextFromSourceCommandlet, Warning, TEXT("Parsing Arguments failed in %s macro while parsing %s:%d. %s"), *GetToken(), *Context.Filename, Context.LineNumber, *FLocTextHelper::SanitizeLogOutput(Context.LineText.TrimStartAndEnd()));
+						return false;
+					}
+				}
+				break;
+			}
+		}
 		// We are closing "
-		else if (bInDblQuotes)
+		else if (ArgsGatherer.IsInDoubleQuotes())
 		{
 			if ('\"' == *Cursor)
 			{
-				bInDblQuotes = false;
+				ArgsGatherer.CloseDoubleQuotes();
 			}
 		}
 		// We are closing a '
-		else if (bInSglQuotes)
+		else if (ArgsGatherer.IsInSingleQuotes())
 		{
 			if ('\'' == *Cursor)
 			{
-				bInSglQuotes = false;
+				ArgsGatherer.CloseSingleQuotes();
 			}
 		}
 		// We are opening a "
 		else if ('\"' == *Cursor)
 		{
-			bInDblQuotes = true;
+			ArgsGatherer.OpenDoubleQuotes();
 		}
 		// We are opening a '
 		else if ('\'' == *Cursor)
 		{
-			bInSglQuotes = true;
+			ArgsGatherer.OpenSingleQuotes();
 		}
 		// We found an opening bracket '(', increment the stack
 		else if ('(' == *Cursor)
@@ -1508,31 +1593,63 @@ bool UGatherTextFromSourceCommandlet::FMacroDescriptor::ParseArgumentString(cons
 				UE_LOG(LogGatherTextFromSourceCommandlet, Warning, TEXT("Unexpected bracket ')' in %s macro while parsing %s:%d. %s"), *GetToken(), *Context.Filename, Context.LineNumber, *FLocTextHelper::SanitizeLogOutput(Context.LineText.TrimStartAndEnd()));
 				return false;
 			}
+		}		
+		else if (',' == *Cursor)
+		{
+			if (1 == BracketStack)
+			{
+				if (Cursor - ArgStart > 0)
+				{
+					// We have a single bracket open and we found a ',' this is the end of the argument. If Bracket stack is > 1 it means that we are in a function call in one of the parameters.
+					if (!ArgsGatherer.Gather(ArgStart, Cursor - ArgStart))
+					{
+						UE_LOG(LogGatherTextFromSourceCommandlet, Warning, TEXT("Parsing Arguments failed in %s macro while parsing %s:%d. %s"), *GetToken(), *Context.Filename, Context.LineNumber, *FLocTextHelper::SanitizeLogOutput(Context.LineText.TrimStartAndEnd()));
+						return false;
+					}
+					if (!ArgsGatherer.EndArgument())
+					{
+						UE_LOG(LogGatherTextFromSourceCommandlet, Warning, TEXT("Parsing Arguments failed in %s macro while parsing %s:%d. %s"), *GetToken(), *Context.Filename, Context.LineNumber, *FLocTextHelper::SanitizeLogOutput(Context.LineText.TrimStartAndEnd()));
+						return false;
+					}
+
+					ArgStart = Cursor + 1;
+				}
+			}
 		}
-		// We have a single bracket open and we found a ',' this is the end of the argument
-		else if (1 == BracketStack && ',' == *Cursor)
+		else if ('\\' == *Cursor)
 		{
-			// create argument from ArgStart to Cursor and set Start next char
-			Args.Add(FString(Cursor - ArgStart, ArgStart));
-			ArgStart = Cursor + 1;
+			ensure(!ArgsGatherer.IsInDoubleQuotes());
+			// We hit a escape character outside a double quote, if the next character is end of line, we are done for this line and the next one must start with a double quote
+			if ('\0' == *(Cursor + 1))
+			{
+				if (Cursor - ArgStart - 1 > 0)
+				{
+					if (!ArgsGatherer.Gather(ArgStart, Cursor - ArgStart - 1))
+					{
+						UE_LOG(LogGatherTextFromSourceCommandlet, Warning, TEXT("Parsing Arguments failed in %s macro while parsing %s:%d. %s"), *GetToken(), *Context.Filename, Context.LineNumber, *FLocTextHelper::SanitizeLogOutput(Context.LineText.TrimStartAndEnd()));
+						return false;
+					}
+				}
+				break;
+			}
 		}
+
 		// We just closed the last ')' this is the end of all args for this macro 
 		if (0 == BracketStack)
 		{
-			// If the arg is empty it means we found a closing braket after a ',' or at the begining of a line
-			FString NewArg = FString(Cursor - ArgStart, ArgStart);
-			if (!NewArg.IsEmpty())
+			// If the arg is empty it means we found a closing braket after a ',' or at the begining of a line			
+			if (Cursor - ArgStart > 0)
 			{
-				Args.Add(NewArg);
+				if (!ArgsGatherer.Gather(ArgStart, Cursor - ArgStart))
+				{
+					UE_LOG(LogGatherTextFromSourceCommandlet, Warning, TEXT("Parsing Arguments failed in %s macro while parsing %s:%d. %s"), *GetToken(), *Context.Filename, Context.LineNumber, *FLocTextHelper::SanitizeLogOutput(Context.LineText.TrimStartAndEnd()));
+					return false;
+				}
 			}
-		}
-		// We are at the end of the line and this is the last argument on that line.
-		if (*Cursor == '\0')
-		{
-			FString NewArg = FString(Cursor - ArgStart - 1, ArgStart);
-			if (!NewArg.IsEmpty())
+			if (!ArgsGatherer.EndArgument())
 			{
-				Args.Add(NewArg);
+				UE_LOG(LogGatherTextFromSourceCommandlet, Warning, TEXT("Parsing Arguments failed in %s macro while parsing %s:%d. %s"), *GetToken(), *Context.Filename, Context.LineNumber, *FLocTextHelper::SanitizeLogOutput(Context.LineText.TrimStartAndEnd()));
+				return false;
 			}
 			break;
 		}
@@ -1563,38 +1680,37 @@ bool UGatherTextFromSourceCommandlet::FMacroDescriptor::ParseArgsFromMacro(const
 
 	Args.Empty();
 
+	FMacroArgumentGatherer ArgumentGatherer;
 	int32 BracketStack = OpenBracketIdx == INDEX_NONE ? 0:1;
 
 	// If we found a bracket, we can start parsing argument on this line
 	if (BracketStack > 0)
 	{
 		// Parse the argument that are on the same line as the macro.
-		if (!ParseArgumentString(RemainingText, OpenBracketIdx + 1, BracketStack, Context, Args))
+		if (!ParseArgumentString(RemainingText, OpenBracketIdx + 1, BracketStack, Context, ArgumentGatherer))
 		{
 			return false;
 		}
 	}
 
-	// Not all arguments were found, look on the next lines.
-	if (Args.Num() < GetMinNumberOfArgument())
+	// We didn't find the end bracket, we must continue.
+	if (BracketStack > 0)
 	{
-		if (!ParseArgsFromNextLines(Args, BracketStack, Context))
+		if (!ParseArgsFromNextLines(ArgumentGatherer, BracketStack, Context))
 		{
 			return false;
 		}
 	}
 
+	ArgumentGatherer.ExtractArguments(Args);
+
 	return 0 < Args.Num() ? true : false;
 }
 
-bool UGatherTextFromSourceCommandlet::FMacroDescriptor::ParseArgsFromNextLines(TArray<FString>& Args, int32& BracketStack, FSourceFileParseContext& Context) const
+bool UGatherTextFromSourceCommandlet::FMacroDescriptor::ParseArgsFromNextLines(FMacroArgumentGatherer& ArgsGatherer, int32& BracketStack, FSourceFileParseContext& Context) const
 {
-	// We allow one line for the macro,that line may contain one or multiple arguments.
-	// So we parse a line by arguments left and a line for the end bracket.
-	int32 MaxLineToParse = Context.LineNumber + GetMinNumberOfArgument() - Args.Num() + 1;
-
 	// Loop until we have all arguments and the closing bracket.
-	for (int32 i = Context.LineNumber; Args.Num() < GetMinNumberOfArgument() || BracketStack > 0; i++)
+	for (int32 i = Context.LineNumber; ArgsGatherer.GetNumberOfArguments() < GetMinNumberOfArgument() || BracketStack > 0; i++)
 	{
 		if(i >= Context.TextLines.Num())
 		{
@@ -1602,17 +1718,10 @@ bool UGatherTextFromSourceCommandlet::FMacroDescriptor::ParseArgsFromNextLines(T
 			return false;
 		}
 
-		if (i > MaxLineToParse)
-		{
-			UE_LOG(LogGatherTextFromSourceCommandlet, Warning, TEXT("The specified %s macro exceeds maximum allowed lines for arguments. %s:%d. %s"), *GetToken(), *Context.Filename, Context.LineNumber, *FLocTextHelper::SanitizeLogOutput(Context.LineText.TrimStartAndEnd()));
-			return false;
-		}
-
 		FString LineText = StripCommentsFromToken(Context.TextLines[i], Context);
 		// Ignore empty lines.
 		if (LineText.IsEmpty())
 		{
-			MaxLineToParse++;
 			continue;
 		}
 
@@ -1625,7 +1734,7 @@ bool UGatherTextFromSourceCommandlet::FMacroDescriptor::ParseArgsFromNextLines(T
 			// We did not find the opening bracket on the first line and on the second non empty line, we give up
 			if (OpenBracketIdx == INDEX_NONE)
 			{
-				UE_LOG(LogGatherTextFromSourceCommandlet, Warning, TEXT("Openaning bracket not found while parsing specified %s macro for arguments. %s:%d. %s"), *GetToken(), *Context.Filename, Context.LineNumber, *FLocTextHelper::SanitizeLogOutput(Context.LineText.TrimStartAndEnd()));
+				UE_LOG(LogGatherTextFromSourceCommandlet, Warning, TEXT("Opening bracket '(' not found while parsing specified %s macro for arguments. %s:%d. %s"), *GetToken(), *Context.Filename, Context.LineNumber, *FLocTextHelper::SanitizeLogOutput(Context.LineText.TrimStartAndEnd()));
 				return false;
 			}
 			BracketStack = 1;
@@ -1633,7 +1742,7 @@ bool UGatherTextFromSourceCommandlet::FMacroDescriptor::ParseArgsFromNextLines(T
 			ParsingOffset = OpenBracketIdx + 1;
 		}
 		
-		if (!ParseArgumentString(LineText, ParsingOffset, BracketStack, Context, Args))
+		if (!ParseArgumentString(LineText, ParsingOffset, BracketStack, Context, ArgsGatherer))
 		{
 			return false;
 		}	
