stop the physics manipulation mode from kicking the user out of fracture mode every time they sim while having the geometry collection selected (note: the physics manipulation mode does not work for geometry collections)

diff --git a/Engine/Source/Editor/UnrealEd/Private/PhysicsManipulationMode.cpp b/Engine/Source/Editor/UnrealEd/Private/PhysicsManipulationMode.cpp
--- a/Engine/Source/Editor/UnrealEd/Private/PhysicsManipulationMode.cpp
+++ b/Engine/Source/Editor/UnrealEd/Private/PhysicsManipulationMode.cpp
@@ -20,7 +20,8 @@ void FPhysicsManipulationEdModeFactory::OnSelectionChanged(FEditorModeTools& Too
 		if (SelectedActor != NULL)
 		{
 			UPrimitiveComponent* PC = Cast<UPrimitiveComponent>(SelectedActor->GetRootComponent());
-			if (PC != NULL && PC->BodyInstance.bSimulatePhysics)
+			// Note the physics handle to be grabbed will be accessed via the GetBodyInstance() function, not via direct access to the BodyInstance member
+			if (PC != NULL && PC->GetBodyInstance() != nullptr && PC->BodyInstance.bSimulatePhysics)
 			{
 				Tools.ActivateMode(FBuiltinEditorModes::EM_Physics);
 				return;
