World Partition HLOD Builder fixes

* Simulate an engine tick between each HLOD build. This ensures engine & render resources that are queued for deletion are processed regularly. Solves an out of vram issue when running with the -D3D12 RHI.
* Save HLOD actor folders created during the HLOD setup phase.
* Updated OpenWorld template to have proper HLOD folders in the Scene Outliner.

diff --git a/Engine/Source/Editor/UnrealEd/Private/WorldPartition/WorldPartitionHLODsBuilder.cpp b/Engine/Source/Editor/UnrealEd/Private/WorldPartition/WorldPartitionHLODsBuilder.cpp
--- a/Engine/Source/Editor/UnrealEd/Private/WorldPartition/WorldPartitionHLODsBuilder.cpp
+++ b/Engine/Source/Editor/UnrealEd/Private/WorldPartition/WorldPartitionHLODsBuilder.cpp
@@ -10,6 +10,7 @@
 #include "Algo/ForEach.h"
 #include "UObject/SavePackage.h"
 
+#include "ActorFolder.h"
 #include "EngineUtils.h"
 #include "SourceControlHelpers.h"
 
@@ -327,6 +328,16 @@ bool UWorldPartitionHLODsBuilder::RunInternal(UWorld* World, const FCellInfo& In
 
 bool UWorldPartitionHLODsBuilder::SetupHLODActors()
 {
+	auto ActorFolderAddedDelegateHandle = GEngine->OnActorFolderAdded().AddLambda([this](UActorFolder* InActorFolder)
+	{
+		SourceControlHelper->Save(InActorFolder->GetPackage());
+	});
+	
+	ON_SCOPE_EXIT
+	{
+		GEngine->OnActorFolderAdded().Remove(ActorFolderAddedDelegateHandle);
+	};
+
 	const bool bCreateActorsOnly = true;
 	WorldPartition->GenerateHLOD(SourceControlHelper, bCreateActorsOnly);
 
@@ -457,6 +468,9 @@ bool UWorldPartitionHLODsBuilder::BuildHLODActors()
 
 		if (HLODLevelToBuild == INDEX_NONE || HLODActor->GetLODLevel() == HLODLevelToBuild)
 		{
+			// Simulate an engine tick to make sure engine & render resources that are queued for deletion are processed.
+			FWorldPartitionHelpers::FakeEngineTick(WorldPartition->GetWorld());
+
 			HLODActor->BuildHLOD();
 		}
 
