Exclude dead package objects when rebuilding dependent cache sets on a Blueprint asset reload, which could otherwise lead to an editor crash during reinstancing in some cases.

diff --git a/Engine/Source/Editor/UnrealEd/Private/PackageTools.cpp b/Engine/Source/Editor/UnrealEd/Private/PackageTools.cpp
--- a/Engine/Source/Editor/UnrealEd/Private/PackageTools.cpp
+++ b/Engine/Source/Editor/UnrealEd/Private/PackageTools.cpp
@@ -1072,29 +1072,78 @@ UPackageTools::UPackageTools(const FObjectInitializer& ObjectInitializer)
 			{
 				FScopedSlowTask CompilingBlueprintsSlowTask(BlueprintsToRecompileThisBatch.Num(), NSLOCTEXT("UnrealEd", "CompilingBlueprints", "Compiling Blueprints"));
 
+				// Gather up all loaded BP assets.
 				TArray<UObject*> BPs;
 				GetObjectsOfClass(UBlueprint::StaticClass(), BPs);
+
+				// Keeps track of BPs with a dependent cache set that's ready to be repopulated.
+				TSet<UBlueprint*> BPsWithResetDependentCache;
+				BPsWithResetDependentCache.Reserve(BPs.Num());
+
+				// Rebuild the dependency/dependent cache sets for each loaded BP.
 				for (UObject* BP : BPs)
 				{
 					UBlueprint* AsBP = CastChecked<UBlueprint>(BP);
-					AsBP->bCachedDependenciesUpToDate = false;
-					FBlueprintEditorUtils::EnsureCachedDependenciesUpToDate(AsBP);
-					for (TWeakObjectPtr<UBlueprint> Dependent : AsBP->CachedDependents)
+
+					// If this BP has been replaced, clear its dependency cache, but don't rebuild it.
+					if (AsBP->HasAnyFlags(RF_NewerVersionExists))
+					{
+						AsBP->CachedDependencies.Empty();
+						AsBP->CachedUDSDependencies.Empty();
+						AsBP->bCachedDependenciesUpToDate = true;
+
+						// Also clear out its dependent cache, as there will no longer be any dependencies on it.
+						AsBP->CachedDependents.Empty();
+						BPsWithResetDependentCache.Add(AsBP);
+					}
+					else
 					{
-						if (UBlueprint* StillAlive = Dependent.Get())
+						// Rebuild the dependency cache for the this BP.
+						AsBP->bCachedDependenciesUpToDate = false;
+						FBlueprintEditorUtils::EnsureCachedDependenciesUpToDate(AsBP);
+					}
+
+					// For each cached dependency, add this BP into its dependent cache set. Note that replaced
+					// BPs won't have any dependencies, and so will not be included in any dependent cache sets.
+					TSet<TWeakObjectPtr<UBlueprint>> LocalCopy_CachedDependencies = AsBP->CachedDependencies;
+					for (const TWeakObjectPtr<UBlueprint>& DependencyPtr : LocalCopy_CachedDependencies)
+					{
+						if (UBlueprint* Dependency = DependencyPtr.Get())
 						{
-							StillAlive->CachedDependencies.Add(AsBP);
+							// Ensure that the dependency's cached dependent set has been cleared before we start adding to it.
+							if (!BPsWithResetDependentCache.Contains(Dependency))
+							{
+								Dependency->CachedDependents.Empty();
+								BPsWithResetDependentCache.Add(Dependency);
+							}
+
+							Dependency->CachedDependents.Add(AsBP);
+						}
+						else
+						{
+							// Remove any entries that cannot be resolved. Not likely, but just in case.
+							AsBP->CachedDependencies.Remove(DependencyPtr);
 						}
 					}
-					for (TWeakObjectPtr<UBlueprint> Dependent : AsBP->CachedDependencies)
+				}
+
+				// Clear out remaining dependent cache sets that weren't already covered above (those not considered as a dependency and/or any replaced BPs).
+				if (BPs.Num() > BPsWithResetDependentCache.Num())
+				{
+					for (UObject* BP : BPs)
 					{
-						if (UBlueprint* StillAlive = Dependent.Get())
+						UBlueprint* AsBP = CastChecked<UBlueprint>(BP);
+						if (!BPsWithResetDependentCache.Contains(AsBP))
 						{
-							StillAlive->CachedDependents.Add(AsBP);
+							AsBP->CachedDependents.Empty();
+							BPsWithResetDependentCache.Add(AsBP);
 						}
 					}
 				}
 
+				// Sanity check that all BPs have reset their dependent cache sets, to assert that there are no stale references left behind.
+				check(BPs.Num() == BPsWithResetDependentCache.Num());
+
 				for (UBlueprint* BlueprintToRecompile : BlueprintsToRecompileThisBatch)
 				{
 					CompilingBlueprintsSlowTask.EnterProgressFrame(1.0f);
