Added support for native FProperty setters and getters.

Setters and getters are native functions called by FProperties when setting property values with *_InContainer functions.
Setters and getter function names can be manually specified with Setter = Func and Getter = Func keywords inside of UPROEPRTY macro but they will also be automatically parsed if the name is not explicitly specified if the setter or getter function name matches SetPropertyName and GetPropertyName pattern.
The latter behavior can be disabled in UHT's DefaultEngine.ini by setting AutomaticSettersAndGetters=False.
ImportText and ExportTextItem functions have been deprecated and should be replaced with *_InContainer or *_Direct variants.

diff --git a/Engine/Plugins/Animation/ControlRig/Source/ControlRigDeveloper/Private/AnimGraphNode_ControlRig.cpp b/Engine/Plugins/Animation/ControlRig/Source/ControlRigDeveloper/Private/AnimGraphNode_ControlRig.cpp
--- a/Engine/Plugins/Animation/ControlRig/Source/ControlRigDeveloper/Private/AnimGraphNode_ControlRig.cpp
+++ b/Engine/Plugins/Animation/ControlRig/Source/ControlRigDeveloper/Private/AnimGraphNode_ControlRig.cpp
@@ -172,7 +172,7 @@ void UAnimGraphNode_ControlRig::CreateCustomPins(TArray<UEdGraphPin*>* OldPins)
 							FString DefaultValue;
 							
 							// Variable.Memory here points to the corresponding property in the Control Rig BP CDO, it was initialized in UAnimGraphNode_ControlRig::RebuildExposedProperties 
-							PropertyIt->ExportTextItem(DefaultValue, Variable.Memory, nullptr, nullptr, PPF_None, nullptr);
+							PropertyIt->ExportTextItem_Direct(DefaultValue, Variable.Memory, nullptr, nullptr, PPF_None, nullptr);
 							
 							if (!DefaultValue.IsEmpty())
 							{
diff --git a/Engine/Plugins/Animation/ControlRig/Source/ControlRigEditor/Private/ControlRigEditorModule.cpp b/Engine/Plugins/Animation/ControlRig/Source/ControlRigEditor/Private/ControlRigEditorModule.cpp
--- a/Engine/Plugins/Animation/ControlRig/Source/ControlRigEditor/Private/ControlRigEditorModule.cpp
+++ b/Engine/Plugins/Animation/ControlRig/Source/ControlRigEditor/Private/ControlRigEditorModule.cpp
@@ -1159,7 +1159,7 @@ void FControlRigEditorModule::GetTypeActions(UControlRigBlueprint* CRB, FBluepri
 			}
 
 			TArray<FControlRigPublicFunctionData> PublicFunctions;
-			PublicFunctionsProperty->ImportText(*PublicFunctionsString, &PublicFunctions, EPropertyPortFlags::PPF_None, nullptr);
+			PublicFunctionsProperty->ImportText_Direct(*PublicFunctionsString, &PublicFunctions, nullptr, EPropertyPortFlags::PPF_None);
 
 			for(const FControlRigPublicFunctionData& PublicFunction : PublicFunctions)
 			{
diff --git a/Engine/Plugins/Animation/ControlRig/Source/ControlRigEditor/Private/Editor/ControlRigEditor.cpp b/Engine/Plugins/Animation/ControlRig/Source/ControlRigEditor/Private/Editor/ControlRigEditor.cpp
--- a/Engine/Plugins/Animation/ControlRig/Source/ControlRigEditor/Private/Editor/ControlRigEditor.cpp
+++ b/Engine/Plugins/Animation/ControlRig/Source/ControlRigEditor/Private/Editor/ControlRigEditor.cpp
@@ -2694,7 +2694,7 @@ void FControlRigEditor::HandleModifiedEvent(ERigVMGraphNotifType InNotifType, UR
 								{
 									const FProperty* TargetProperty = WrapperObject->GetClass()->FindPropertyByName(RootPin->GetFName());
 									uint8* PropertyStorage = TargetProperty->ContainerPtrToValuePtr<uint8>(WrapperObject);
-									TargetProperty->ImportText(*DefaultValue, PropertyStorage, PPF_None, nullptr);
+									TargetProperty->ImportText_Direct(*DefaultValue, PropertyStorage, nullptr, PPF_None);
 								}
 							}
 						}
diff --git a/Engine/Plugins/Animation/LiveLink/Source/LiveLink/Private/LiveLinkBlueprintLibrary.cpp b/Engine/Plugins/Animation/LiveLink/Source/LiveLink/Private/LiveLinkBlueprintLibrary.cpp
--- a/Engine/Plugins/Animation/LiveLink/Source/LiveLink/Private/LiveLinkBlueprintLibrary.cpp
+++ b/Engine/Plugins/Animation/LiveLink/Source/LiveLink/Private/LiveLinkBlueprintLibrary.cpp
@@ -374,6 +374,7 @@ DEFINE_FUNCTION(ULiveLinkBlueprintLibrary::execEvaluateLiveLinkFrame)
 	P_GET_STRUCT(FLiveLinkSubjectRepresentation, SubjectRepresentation);
 
 	Stack.MostRecentPropertyAddress = nullptr;
+	Stack.MostRecentPropertyContainer = nullptr;
 	Stack.StepCompiledIn<FStructProperty>(NULL);
 	void* OutBlueprintDataPtr = Stack.MostRecentPropertyAddress;
 	FStructProperty* BlueprintDataStructProp = CastField<FStructProperty>(Stack.MostRecentProperty);
@@ -392,6 +393,7 @@ DEFINE_FUNCTION(ULiveLinkBlueprintLibrary::execEvaluateLiveLinkFrameWithSpecific
 	Stack.StepCompiledIn<FClassProperty>(&RoleClass);
 
 	Stack.MostRecentPropertyAddress = nullptr;
+	Stack.MostRecentPropertyContainer = nullptr;
 	Stack.StepCompiledIn<FStructProperty>(NULL);
 	void* OutBlueprintDataPtr = Stack.MostRecentPropertyAddress;
 	FStructProperty* BlueprintDataStructProp = CastField<FStructProperty>(Stack.MostRecentProperty);
@@ -414,6 +416,7 @@ DEFINE_FUNCTION(ULiveLinkBlueprintLibrary::execEvaluateLiveLinkFrameAtWorldTimeO
 	Stack.StepCompiledIn<FFloatProperty>(&WorldTimeOffset);
 
 	Stack.MostRecentPropertyAddress = nullptr;
+	Stack.MostRecentPropertyContainer = nullptr;
 	Stack.StepCompiledIn<FStructProperty>(NULL);
 	void* OutBlueprintDataPtr = Stack.MostRecentPropertyAddress;
 	FStructProperty* BlueprintDataStructProp = CastField<FStructProperty>(Stack.MostRecentProperty);
@@ -433,6 +436,7 @@ DEFINE_FUNCTION(ULiveLinkBlueprintLibrary::execEvaluateLiveLinkFrameAtSceneTime)
 	P_GET_STRUCT(FTimecode, SceneTime);
 
 	Stack.MostRecentPropertyAddress = nullptr;
+	Stack.MostRecentPropertyContainer = nullptr;
 	Stack.StepCompiledIn<FStructProperty>(NULL);
 	void* OutBlueprintDataPtr = Stack.MostRecentPropertyAddress;
 	FStructProperty* BlueprintDataStructProp = CastField<FStructProperty>(Stack.MostRecentProperty);
diff --git a/Engine/Plugins/Developer/Concert/ConcertMain/Source/Concert/Private/ConcertLogger.cpp b/Engine/Plugins/Developer/Concert/ConcertMain/Source/Concert/Private/ConcertLogger.cpp
--- a/Engine/Plugins/Developer/Concert/ConcertMain/Source/Concert/Private/ConcertLogger.cpp
+++ b/Engine/Plugins/Developer/Concert/ConcertMain/Source/Concert/Private/ConcertLogger.cpp
@@ -728,13 +728,12 @@ void FConcertLogger::LogEntry(FConcertLog& Log)
 		}
 
 		// Skip exporting the MessageOrderIndex value for unreliable messages as they add a lot of noise to the log
-		const void* PropertyValuePtr = PropertyIt->ContainerPtrToValuePtr<void>(&Log);
 		const bool bExportValue = bIsReliable || PropertyIt->GetFName() != MessageOrderIndexPropertyName;
 
 		FString PropertyValue;
 		if (bExportValue)
 		{
-			PropertyIt->ExportTextItem(PropertyValue, PropertyValuePtr, nullptr, nullptr, PPF_None);
+			PropertyIt->ExportTextItem_InContainer(PropertyValue, &Log, nullptr, nullptr, PPF_None);
 		}
 		PropertyValue.ReplaceInline(TEXT("\""), TEXT("\"\""));
 
diff --git a/Engine/Plugins/Editor/GameplayTagsEditor/Source/GameplayTagsEditor/Private/SGameplayTagQueryGraphPin.cpp b/Engine/Plugins/Editor/GameplayTagsEditor/Source/GameplayTagsEditor/Private/SGameplayTagQueryGraphPin.cpp
--- a/Engine/Plugins/Editor/GameplayTagsEditor/Source/GameplayTagsEditor/Private/SGameplayTagQueryGraphPin.cpp
+++ b/Engine/Plugins/Editor/GameplayTagsEditor/Source/GameplayTagsEditor/Private/SGameplayTagQueryGraphPin.cpp
@@ -48,7 +48,7 @@ void SGameplayTagQueryGraphPin::ParseDefaultValueData()
 	if (TQProperty)
 	{
 		FGameplayTagQuery* const TQ = TagQuery.Get();
-		TQProperty->ImportText(*TagQueryString, TQ, 0, nullptr, GLog);
+		TQProperty->ImportText_Direct(*TagQueryString, TQ, nullptr, 0, GLog);
 	}
 }
 
diff --git a/Engine/Plugins/Enterprise/DatasmithImporter/Source/DatasmithImporter/Private/DatasmithActorImporter.cpp b/Engine/Plugins/Enterprise/DatasmithImporter/Source/DatasmithImporter/Private/DatasmithActorImporter.cpp
--- a/Engine/Plugins/Enterprise/DatasmithImporter/Source/DatasmithImporter/Private/DatasmithActorImporter.cpp
+++ b/Engine/Plugins/Enterprise/DatasmithImporter/Source/DatasmithImporter/Private/DatasmithActorImporter.cpp
@@ -390,7 +390,7 @@ AActor* FDatasmithActorImporter::ImportCustomActor(FDatasmithImportContext& Impo
 
 		if ( Property )
 		{
-			Property->ImportText( KeyValueProperty->GetValue(), Property->ContainerPtrToValuePtr< uint8 >( Actor ), PPF_None, nullptr );
+			Property->ImportText_InContainer( KeyValueProperty->GetValue(), Actor, nullptr, PPF_None );
 		}
 	}
 
diff --git a/Engine/Plugins/Enterprise/DatasmithImporter/Source/DatasmithImporter/Private/DatasmithImportContext.cpp b/Engine/Plugins/Enterprise/DatasmithImporter/Source/DatasmithImporter/Private/DatasmithImportContext.cpp
--- a/Engine/Plugins/Enterprise/DatasmithImporter/Source/DatasmithImporter/Private/DatasmithImportContext.cpp
+++ b/Engine/Plugins/Enterprise/DatasmithImporter/Source/DatasmithImporter/Private/DatasmithImportContext.cpp
@@ -151,7 +151,7 @@ void FDatasmithImportOptionHelper::LoadOptions(const TArray<UObject*>& ImportOpt
 					FString PropertyValue = OptionDataJsonObject->GetStringField((*It)->GetNameCPP());
 					uint8* PropertyAddr = (*It)->ContainerPtrToValuePtr<uint8>(Object);
 
-					(*It)->ImportText(*PropertyValue, PropertyAddr, 0, Object);
+					(*It)->ImportText_Direct(*PropertyValue, PropertyAddr, Object, 0);
 				}
 			}
 		}
diff --git a/Engine/Plugins/Enterprise/DatasmithImporter/Source/DatasmithImporter/Private/DatasmithMaterialExpressions.cpp b/Engine/Plugins/Enterprise/DatasmithImporter/Source/DatasmithImporter/Private/DatasmithMaterialExpressions.cpp
--- a/Engine/Plugins/Enterprise/DatasmithImporter/Source/DatasmithImporter/Private/DatasmithMaterialExpressions.cpp
+++ b/Engine/Plugins/Enterprise/DatasmithImporter/Source/DatasmithImporter/Private/DatasmithMaterialExpressions.cpp
@@ -3019,7 +3019,7 @@ UMaterialExpression* FDatasmithMaterialExpressions::CreateGenericExpression( IDa
 			{
 				if ( Property )
 				{
-					Property->ImportText( KeyValueProperty->GetValue(), Property->ContainerPtrToValuePtr< void >( MaterialExpression ), PPF_None, nullptr );
+					Property->ImportText_Direct( KeyValueProperty->GetValue(), Property->ContainerPtrToValuePtr< void >( MaterialExpression ), nullptr, PPF_None);
 				}
 			}
 		}
diff --git a/Engine/Plugins/Enterprise/VariantManagerContent/Source/VariantManagerContent/Private/PropertyValue.cpp b/Engine/Plugins/Enterprise/VariantManagerContent/Source/VariantManagerContent/Private/PropertyValue.cpp
--- a/Engine/Plugins/Enterprise/VariantManagerContent/Source/VariantManagerContent/Private/PropertyValue.cpp
+++ b/Engine/Plugins/Enterprise/VariantManagerContent/Source/VariantManagerContent/Private/PropertyValue.cpp
@@ -1461,7 +1461,7 @@ void UPropertyValue::ApplyViaFunctionSetter(UObject* TargetObject)
 		FString* Defaults = PropertySetterParameterDefaults.Find(PropertyParam->GetName());
 		if (Defaults)
 		{
-			const TCHAR* Result = PropertyParam->ImportText( **Defaults, PropertyParam->ContainerPtrToValuePtr<uint8>(Parms), ExportFlags, NULL );
+			const TCHAR* Result = PropertyParam->ImportText_Direct( **Defaults, PropertyParam->ContainerPtrToValuePtr<uint8>(Parms), NULL, ExportFlags );
 			if (!Result)
 			{
 				UE_LOG(LogVariantContent, Error, TEXT("Failed at applying the default value for parameter '%s' of PropertyValue '%s'"), *PropertyParam->GetName(), *GetFullDisplayString());
diff --git a/Engine/Plugins/Enterprise/VariantManagerContent/Source/VariantManagerContent/Private/PropertyValueSoftObject.cpp b/Engine/Plugins/Enterprise/VariantManagerContent/Source/VariantManagerContent/Private/PropertyValueSoftObject.cpp
--- a/Engine/Plugins/Enterprise/VariantManagerContent/Source/VariantManagerContent/Private/PropertyValueSoftObject.cpp
+++ b/Engine/Plugins/Enterprise/VariantManagerContent/Source/VariantManagerContent/Private/PropertyValueSoftObject.cpp
@@ -157,7 +157,7 @@ void UPropertyValueSoftObject::ApplyViaFunctionSetter(UObject* TargetObject)
 		FString* Defaults = PropertySetterParameterDefaults.Find(PropertyParam->GetName());
 		if (Defaults)
 		{
-			const TCHAR* Result = PropertyParam->ImportText( **Defaults, PropertyParam->ContainerPtrToValuePtr<uint8>(Parms), ExportFlags, NULL );
+			const TCHAR* Result = PropertyParam->ImportText_Direct( **Defaults, PropertyParam->ContainerPtrToValuePtr<uint8>(Parms), NULL, ExportFlags );
 			if (!Result)
 			{
 				UE_LOG(LogVariantContent, Error, TEXT("Failed at applying the default value for parameter '%s' of PropertyValue '%s'"), *PropertyParam->GetName(), *GetFullDisplayString());
diff --git a/Engine/Plugins/Experimental/CommonConversation/Source/CommonConversationRuntime/Private/ConversationRegistry.cpp b/Engine/Plugins/Experimental/CommonConversation/Source/CommonConversationRuntime/Private/ConversationRegistry.cpp
--- a/Engine/Plugins/Experimental/CommonConversation/Source/CommonConversationRuntime/Private/ConversationRegistry.cpp
+++ b/Engine/Plugins/Experimental/CommonConversation/Source/CommonConversationRuntime/Private/ConversationRegistry.cpp
@@ -416,7 +416,7 @@ void UConversationRegistry::BuildDependenciesGraph()
 			TArray<FConversationEntryList> EntryTags;
 
 			FArrayProperty* ArrayProperty = FindFProperty<FArrayProperty>(UConversationDatabase::StaticClass(), GET_MEMBER_NAME_CHECKED(UConversationDatabase, EntryTags));
-			ArrayProperty->ImportText(*EntryTagsString, &EntryTags, 0, nullptr);
+			ArrayProperty->ImportText_Direct(*EntryTagsString, &EntryTags, nullptr, 0);
 
 			for (const FConversationEntryList& Entry : EntryTags)
 			{
@@ -431,7 +431,7 @@ void UConversationRegistry::BuildDependenciesGraph()
 			TArray<FGuid> NodeIds;
 
 			FArrayProperty* ArrayProperty = FindFProperty<FArrayProperty>(UConversationDatabase::StaticClass(), GET_MEMBER_NAME_CHECKED(UConversationDatabase, InternalNodeIds));
-			ArrayProperty->ImportText(*InternalNodeIds, &NodeIds, 0, nullptr);
+			ArrayProperty->ImportText_Direct(*InternalNodeIds, &NodeIds, nullptr, 0);
 
 			for (FGuid& NodeId : NodeIds)
 			{
@@ -468,7 +468,7 @@ void UConversationRegistry::BuildDependenciesGraph()
 			TArray<FGuid> NodeIds;
 
 			FArrayProperty* ArrayProperty = FindFProperty<FArrayProperty>(UConversationDatabase::StaticClass(), GET_MEMBER_NAME_CHECKED(UConversationDatabase, LinkedToNodeIds));
-			ArrayProperty->ImportText(*LinkedToNodeIds, &NodeIds, 0, nullptr);
+			ArrayProperty->ImportText_Direct(*LinkedToNodeIds, &NodeIds, nullptr, 0);
 
 			//@TODO: CONVERSATION: Register that we need to link to other graphs here.
 		}
diff --git a/Engine/Plugins/Experimental/PythonScriptPlugin/Source/PythonScriptPlugin/Private/PyGenUtil.cpp b/Engine/Plugins/Experimental/PythonScriptPlugin/Source/PythonScriptPlugin/Private/PyGenUtil.cpp
--- a/Engine/Plugins/Experimental/PythonScriptPlugin/Source/PythonScriptPlugin/Private/PyGenUtil.cpp
+++ b/Engine/Plugins/Experimental/PythonScriptPlugin/Source/PythonScriptPlugin/Private/PyGenUtil.cpp
@@ -1163,6 +1163,7 @@ bool InvokePythonCallableFromUnrealFunctionThunk(FPyObjectPtr InSelf, PyObject*
 
 			// Step the property data to populate the local value
 			Stack.MostRecentPropertyAddress = nullptr;
+			Stack.MostRecentPropertyContainer = nullptr;
 			void* LocalValue = Param->ContainerPtrToValuePtr<void>(LocalStruct);
 			Stack.StepCompiledIn(LocalValue, Param->GetClass());
 
diff --git a/Engine/Plugins/Experimental/PythonScriptPlugin/Source/PythonScriptPlugin/Private/PyUtil.cpp b/Engine/Plugins/Experimental/PythonScriptPlugin/Source/PythonScriptPlugin/Private/PyUtil.cpp
--- a/Engine/Plugins/Experimental/PythonScriptPlugin/Source/PythonScriptPlugin/Private/PyUtil.cpp
+++ b/Engine/Plugins/Experimental/PythonScriptPlugin/Source/PythonScriptPlugin/Private/PyUtil.cpp
@@ -520,7 +520,7 @@ void ImportDefaultValue(const FProperty* InProp, void* InPropValue, const FStrin
 			}
 		}
 
-		InProp->ImportText(*InDefaultValue, InPropValue, PPF_None, nullptr);
+		InProp->ImportText_Direct(*InDefaultValue, InPropValue, nullptr, PPF_None);
 	}
 }
 
@@ -1184,7 +1184,7 @@ FString GetFriendlyPropertyValue(const FProperty* InProp, const void* InPropValu
 	}
 	
 	FString FriendlyPropertyValue;
-	InProp->ExportTextItem(FriendlyPropertyValue, InPropValue, InPropValue, nullptr, InPortFlags, nullptr);
+	InProp->ExportTextItem_Direct(FriendlyPropertyValue, InPropValue, InPropValue, nullptr, InPortFlags, nullptr);
 	return FriendlyPropertyValue;
 }
 
diff --git a/Engine/Plugins/Experimental/PythonScriptPlugin/Source/PythonScriptPlugin/Private/PyWrapperArray.cpp b/Engine/Plugins/Experimental/PythonScriptPlugin/Source/PythonScriptPlugin/Private/PyWrapperArray.cpp
--- a/Engine/Plugins/Experimental/PythonScriptPlugin/Source/PythonScriptPlugin/Private/PyWrapperArray.cpp
+++ b/Engine/Plugins/Experimental/PythonScriptPlugin/Source/PythonScriptPlugin/Private/PyWrapperArray.cpp
@@ -342,7 +342,7 @@ FPyWrapperArray* FPyWrapperArray::CastPyObject(PyObject* InPyObject, PyTypeObjec
 					return nullptr;
 				}
 
-				if (!NewArray->ArrayProp->Inner->ImportText(*ExportedEntry, NewScriptArrayHelper.GetRawPtr(ElementIndex), PPF_None, nullptr))
+				if (!NewArray->ArrayProp->Inner->ImportText_Direct(*ExportedEntry, NewScriptArrayHelper.GetRawPtr(ElementIndex), nullptr, PPF_None))
 				{
 					PyUtil::SetPythonError(PyExc_Exception, Self, *FString::Printf(TEXT("Failed to import text '%s' element for property '%s' (%s) at index %d"), *ExportedEntry, *NewArray->ArrayProp->Inner->GetName(), *NewArray->ArrayProp->Inner->GetClass()->GetName(), ElementIndex));
 					return nullptr;
diff --git a/Engine/Plugins/Experimental/PythonScriptPlugin/Source/PythonScriptPlugin/Private/PyWrapperFixedArray.cpp b/Engine/Plugins/Experimental/PythonScriptPlugin/Source/PythonScriptPlugin/Private/PyWrapperFixedArray.cpp
--- a/Engine/Plugins/Experimental/PythonScriptPlugin/Source/PythonScriptPlugin/Private/PyWrapperFixedArray.cpp
+++ b/Engine/Plugins/Experimental/PythonScriptPlugin/Source/PythonScriptPlugin/Private/PyWrapperFixedArray.cpp
@@ -310,7 +310,7 @@ FPyWrapperFixedArray* FPyWrapperFixedArray::CastPyObject(PyObject* InPyObject, P
 				return nullptr;
 			}
 
-			if (!NewArray->ArrayProp->ImportText(*ExportedEntry, GetItemPtr(NewArray, ArrIndex), PPF_None, nullptr))
+			if (!NewArray->ArrayProp->ImportText_Direct(*ExportedEntry, GetItemPtr(NewArray, ArrIndex), nullptr, PPF_None))
 			{
 				PyUtil::SetPythonError(PyExc_Exception, Self, *FString::Printf(TEXT("Failed to import text '%s' for property '%s' (%s) at index %d"), *ExportedEntry, *NewArray->ArrayProp->GetName(), *NewArray->ArrayProp->GetClass()->GetName(), ArrIndex));
 				return nullptr;
diff --git a/Engine/Plugins/Experimental/PythonScriptPlugin/Source/PythonScriptPlugin/Private/PyWrapperMap.cpp b/Engine/Plugins/Experimental/PythonScriptPlugin/Source/PythonScriptPlugin/Private/PyWrapperMap.cpp
--- a/Engine/Plugins/Experimental/PythonScriptPlugin/Source/PythonScriptPlugin/Private/PyWrapperMap.cpp
+++ b/Engine/Plugins/Experimental/PythonScriptPlugin/Source/PythonScriptPlugin/Private/PyWrapperMap.cpp
@@ -690,12 +690,12 @@ FPyWrapperMap* FPyWrapperMap::CastPyObject(PyObject* InPyObject, PyTypeObject* I
 				}
 
 				const int32 NewElementIndex = NewScriptMapHelper.AddDefaultValue_Invalid_NeedsRehash();
-				if (!NewScriptMapHelper.GetKeyProperty()->ImportText(*ExportedKey, NewScriptMapHelper.GetKeyPtr(NewElementIndex), PPF_None, nullptr))
+				if (!NewScriptMapHelper.GetKeyProperty()->ImportText_Direct(*ExportedKey, NewScriptMapHelper.GetKeyPtr(NewElementIndex), nullptr, PPF_None))
 				{
 					PyUtil::SetPythonError(PyExc_Exception, Self, *FString::Printf(TEXT("Failed to import text '%s' key for property '%s' (%s) at index %d"), *ExportedKey, *NewScriptMapHelper.GetKeyProperty()->GetName(), *NewScriptMapHelper.GetKeyProperty()->GetClass()->GetName(), ElementIndex));
 					return nullptr;
 				}
-				if (!NewScriptMapHelper.GetValueProperty()->ImportText(*ExportedValue, NewScriptMapHelper.GetValuePtr(NewElementIndex), PPF_None, nullptr))
+				if (!NewScriptMapHelper.GetValueProperty()->ImportText_Direct(*ExportedValue, NewScriptMapHelper.GetValuePtr(NewElementIndex), nullptr, PPF_None))
 				{
 					PyUtil::SetPythonError(PyExc_Exception, Self, *FString::Printf(TEXT("Failed to import text '%s' value for property '%s' (%s) at index %d"), *ExportedValue, *NewScriptMapHelper.GetValueProperty()->GetName(), *NewScriptMapHelper.GetValueProperty()->GetClass()->GetName(), ElementIndex));
 					return nullptr;
diff --git a/Engine/Plugins/Experimental/PythonScriptPlugin/Source/PythonScriptPlugin/Private/PyWrapperSet.cpp b/Engine/Plugins/Experimental/PythonScriptPlugin/Source/PythonScriptPlugin/Private/PyWrapperSet.cpp
--- a/Engine/Plugins/Experimental/PythonScriptPlugin/Source/PythonScriptPlugin/Private/PyWrapperSet.cpp
+++ b/Engine/Plugins/Experimental/PythonScriptPlugin/Source/PythonScriptPlugin/Private/PyWrapperSet.cpp
@@ -399,7 +399,7 @@ FPyWrapperSet* FPyWrapperSet::CastPyObject(PyObject* InPyObject, PyTypeObject* I
 				}
 
 				const int32 NewElementIndex = NewScriptSetHelper.AddDefaultValue_Invalid_NeedsRehash();
-				if (!NewScriptSetHelper.GetElementProperty()->ImportText(*ExportedEntry, NewScriptSetHelper.GetElementPtr(NewElementIndex), PPF_None, nullptr))
+				if (!NewScriptSetHelper.GetElementProperty()->ImportText_Direct(*ExportedEntry, NewScriptSetHelper.GetElementPtr(NewElementIndex), nullptr, PPF_None))
 				{
 					PyUtil::SetPythonError(PyExc_Exception, Self, *FString::Printf(TEXT("Failed to import text '%s' element for property '%s' (%s) at index %d"), *ExportedEntry, *NewScriptSetHelper.GetElementProperty()->GetName(), *NewScriptSetHelper.GetElementProperty()->GetClass()->GetName(), ElementIndex));
 					return nullptr;
diff --git a/Engine/Plugins/FX/Niagara/Source/Niagara/Private/NiagaraRendererComponents.cpp b/Engine/Plugins/FX/Niagara/Source/Niagara/Private/NiagaraRendererComponents.cpp
--- a/Engine/Plugins/FX/Niagara/Source/Niagara/Private/NiagaraRendererComponents.cpp
+++ b/Engine/Plugins/FX/Niagara/Source/Niagara/Private/NiagaraRendererComponents.cpp
@@ -145,7 +145,7 @@ void InvokeSetterFunction(UObject* InRuntimeObject, UFunction* Setter, const uin
 				const FString* DefaultValue = SetterDefaultValues.Find(Property->GetName());
 				if (DefaultValue)
 				{
-					Property->ImportText(**DefaultValue, Property->ContainerPtrToValuePtr<uint8>(Params), PPF_None, InRuntimeObject);
+					Property->ImportText_Direct(**DefaultValue, Property->ContainerPtrToValuePtr<uint8>(Params), InRuntimeObject, PPF_None);
 				}
 				else
 				{
diff --git a/Engine/Plugins/NetcodeUnitTest/NetcodeUnitTest/Source/NetcodeUnitTest/Private/NUTActor.cpp b/Engine/Plugins/NetcodeUnitTest/NetcodeUnitTest/Source/NetcodeUnitTest/Private/NUTActor.cpp
--- a/Engine/Plugins/NetcodeUnitTest/NetcodeUnitTest/Source/NetcodeUnitTest/Private/NUTActor.cpp
+++ b/Engine/Plugins/NetcodeUnitTest/NetcodeUnitTest/Source/NetcodeUnitTest/Private/NUTActor.cpp
@@ -684,7 +684,7 @@ void ANUTActor::ExecuteOnServer(UObject* InTargetObj, FString InTargetFunc)
 					FDelegateProperty* DelProp = FindFProperty<FDelegateProperty>(GetClass(), TEXT("TempDelegate"));
 
 					FString DelString;
-					DelProp->ExportTextItem(DelString, DelProp->ContainerPtrToValuePtr<uint8>(this), nullptr, this, 0, nullptr);
+					DelProp->ExportTextItem_InContainer(DelString, this, nullptr, this, 0, nullptr);
 
 					ServerExecute(DelString);
 				}
@@ -722,7 +722,7 @@ void ANUTActor::ServerExecute_Implementation(const FString& InDelegate)
 	const TCHAR* InDelText = *InDelegate;
 
 	TempDelegate.Unbind();
-	DelProp->ImportText(InDelText, DelProp->ContainerPtrToValuePtr<uint8>(this), 0, NULL);
+	DelProp->ImportText_InContainer(InDelText, this, NULL, 0);
 
 	if (TempDelegate.IsBound())
 	{
diff --git a/Engine/Plugins/NetcodeUnitTest/NetcodeUnitTest/Source/NetcodeUnitTest/Private/NUTUtilReflection.cpp b/Engine/Plugins/NetcodeUnitTest/NetcodeUnitTest/Source/NetcodeUnitTest/Private/NUTUtilReflection.cpp
--- a/Engine/Plugins/NetcodeUnitTest/NetcodeUnitTest/Source/NetcodeUnitTest/Private/NUTUtilReflection.cpp
+++ b/Engine/Plugins/NetcodeUnitTest/NetcodeUnitTest/Source/NetcodeUnitTest/Private/NUTUtilReflection.cpp
@@ -1337,7 +1337,7 @@ TValueOrError<FString, FString> FVMReflection::GetValueAsString()
 		{
 			FString Result;
 
-			Prop->ExportTextItem(Result, FieldAddress, FieldAddress, Obj, PPF_None);
+			Prop->ExportTextItem_Direct(Result, FieldAddress, FieldAddress, Obj, PPF_None);
 
 			ReturnVal = MakeValue(Result);
 		}
@@ -1515,7 +1515,7 @@ FString NUTUtilRefl::FunctionParmsToString(UFunction* InFunction, void* Parms)
 		FString CurPropText;
 
 		// Set PortFlags to maximize property coverage
-		It->ExportTextItem(CurPropText, It->ContainerPtrToValuePtr<void>(Parms), nullptr, nullptr,
+		It->ExportTextItem_Direct(CurPropText, It->ContainerPtrToValuePtr<void>(Parms), nullptr, nullptr,
 							PPF_IncludeTransient | PPF_DuplicateForPIE);
 
 		if (ReturnVal.Len() > 0)
diff --git a/Engine/Plugins/Online/OnlineSubsystem/Source/Private/OnlineKeyValuePair.cpp b/Engine/Plugins/Online/OnlineSubsystem/Source/Private/OnlineKeyValuePair.cpp
--- a/Engine/Plugins/Online/OnlineSubsystem/Source/Private/OnlineKeyValuePair.cpp
+++ b/Engine/Plugins/Online/OnlineSubsystem/Source/Private/OnlineKeyValuePair.cpp
@@ -1476,7 +1476,7 @@ bool FVariantDataConverter::ConvertScalarVariantToFProperty(const FVariantData*
 				if (!TheCppStructOps->ImportTextItem(ImportTextPtr, OutValue, PPF_None, nullptr, (FOutputDevice*)GWarn))
 				{
 					// Fall back to trying the tagged property approach if custom ImportTextItem couldn't get it done
-					Property->ImportText(ImportTextPtr, OutValue, PPF_None, nullptr);
+					Property->ImportText_Direct(ImportTextPtr, OutValue, nullptr, PPF_None);
 				}
 			}
 			else
@@ -1505,7 +1505,7 @@ bool FVariantDataConverter::ConvertScalarVariantToFProperty(const FVariantData*
 		{
 			FString StrValue;
 			Variant->GetValue(StrValue);
-			if (Property->ImportText(*StrValue, OutValue, 0, nullptr) == nullptr)
+			if (Property->ImportText_Direct(*StrValue, OutValue, nullptr, 0) == nullptr)
 			{
 				UE_LOG_ONLINE(Error, TEXT("ConvertScalarVariantToFProperty - Unable import property type %s from string value for property %s"), *Property->GetClass()->GetName(), *Property->GetNameCPP());
 				return false;
@@ -1670,7 +1670,7 @@ bool FVariantDataConverter::ConvertScalarFPropertyToVariant(FProperty* Property,
 	{
 		// Default to export as string for everything else
 		FString StringValue;
-		Property->ExportTextItem(StringValue, Value, nullptr, nullptr, PPF_None);
+		Property->ExportTextItem_Direct(StringValue, Value, nullptr, nullptr, PPF_None);
 		OutVariantData.SetValue(StringValue);
 	}
 
diff --git a/Engine/Plugins/Runtime/GameplayAbilities/Source/GameplayAbilities/Private/AbilitySystemDebugHUD.cpp b/Engine/Plugins/Runtime/GameplayAbilities/Source/GameplayAbilities/Private/AbilitySystemDebugHUD.cpp
--- a/Engine/Plugins/Runtime/GameplayAbilities/Source/GameplayAbilities/Private/AbilitySystemDebugHUD.cpp
+++ b/Engine/Plugins/Runtime/GameplayAbilities/Source/GameplayAbilities/Private/AbilitySystemDebugHUD.cpp
@@ -122,7 +122,7 @@ void AAbilitySystemDebugHUD::DrawDebugAbilitySystemComponent(UAbilitySystemCompo
 
 			FString ValueString;
 			const void *PropertyValue = Prop->ContainerPtrToValuePtr<void>(Set);
-			Prop->ExportTextItem(ValueString, PropertyValue, NULL, NULL, 0);
+			Prop->ExportTextItem_Direct(ValueString, PropertyValue, NULL, NULL, 0);
 
 			String = FString::Printf(TEXT("%s: %s"), *Prop->GetName(), *ValueString);
 			DrawWithBackground(Font, String, Color, EAlignHorizontal::Left, X, EAlignVertical::Top, Y);
diff --git a/Engine/Plugins/Runtime/Metasound/Source/MetasoundFrontend/Private/MetasoundJsonBackend.h b/Engine/Plugins/Runtime/Metasound/Source/MetasoundFrontend/Private/MetasoundJsonBackend.h
--- a/Engine/Plugins/Runtime/Metasound/Source/MetasoundFrontend/Private/MetasoundJsonBackend.h
+++ b/Engine/Plugins/Runtime/Metasound/Source/MetasoundFrontend/Private/MetasoundJsonBackend.h
@@ -131,7 +131,7 @@ namespace Metasound
 			else if (State.KeyProperty != nullptr)
 			{
 				FString KeyString;
-				State.KeyProperty->ExportTextItem(KeyString, State.KeyData, nullptr, nullptr, PPF_None);
+				State.KeyProperty->ExportTextItem_Direct(KeyString, State.KeyData, nullptr, nullptr, PPF_None);
 				JsonWriter->WriteValue(KeyString, Value);
 			}
 			else
@@ -153,7 +153,7 @@ namespace Metasound
 			else if (State.KeyProperty != nullptr)
 			{
 				FString KeyString;
-				State.KeyProperty->ExportTextItem(KeyString, State.KeyData, nullptr, nullptr, PPF_None);
+				State.KeyProperty->ExportTextItem_Direct(KeyString, State.KeyData, nullptr, nullptr, PPF_None);
 				JsonWriter->WriteNull(KeyString);
 			}
 			else
@@ -212,7 +212,7 @@ namespace Metasound
 			else if (State.KeyProperty != nullptr)
 			{
 				FString KeyString;
-				State.KeyProperty->ExportTextItem(KeyString, State.KeyData, nullptr, nullptr, PPF_None);
+				State.KeyProperty->ExportTextItem_Direct(KeyString, State.KeyData, nullptr, nullptr, PPF_None);
 				JsonWriter->WriteArrayStart(KeyString);
 			}
 			else
@@ -232,7 +232,7 @@ namespace Metasound
 				else if (State.KeyProperty != nullptr)
 				{
 					FString KeyString;
-					State.KeyProperty->ExportTextItem(KeyString, State.KeyData, nullptr, nullptr, PPF_None);
+					State.KeyProperty->ExportTextItem_Direct(KeyString, State.KeyData, nullptr, nullptr, PPF_None);
 					JsonWriter->WriteObjectStart(KeyString);
 				}
 				else
diff --git a/Engine/Plugins/Tests/EditorTests/Source/EditorTests/Private/UnrealEd/EditorBuildPromotionTests.cpp b/Engine/Plugins/Tests/EditorTests/Source/EditorTests/Private/UnrealEd/EditorBuildPromotionTests.cpp
--- a/Engine/Plugins/Tests/EditorTests/Source/EditorTests/Private/UnrealEd/EditorBuildPromotionTests.cpp
+++ b/Engine/Plugins/Tests/EditorTests/Source/EditorTests/Private/UnrealEd/EditorBuildPromotionTests.cpp
@@ -382,7 +382,7 @@ namespace EditorBuildPromotionTestUtils
 		{
 			FString ValueString;
 			const uint8* PropertyAddr = FoundProperty->ContainerPtrToValuePtr<uint8>(TargetObject);
-			FoundProperty->ExportTextItem(ValueString, PropertyAddr, NULL, NULL, PPF_None);
+			FoundProperty->ExportTextItem_Direct(ValueString, PropertyAddr, NULL, NULL, PPF_None);
 			return ValueString;
 		}
 		return TEXT("");
diff --git a/Engine/Plugins/Tests/EditorTests/Source/EditorTests/Private/UnrealEd/LightingTests.cpp b/Engine/Plugins/Tests/EditorTests/Source/EditorTests/Private/UnrealEd/LightingTests.cpp
--- a/Engine/Plugins/Tests/EditorTests/Source/EditorTests/Private/UnrealEd/LightingTests.cpp
+++ b/Engine/Plugins/Tests/EditorTests/Source/EditorTests/Private/UnrealEd/LightingTests.cpp
@@ -101,7 +101,7 @@ namespace LightingTestHelpers
 			TargetObject->Modify();
 
 			TargetObject->PreEditChange(FoundProperty);
-			FoundProperty->ImportText(*NewValueString, FoundProperty->ContainerPtrToValuePtr<uint8>(TargetObject), 0, TargetObject);
+			FoundProperty->ImportText_Direct(*NewValueString, FoundProperty->ContainerPtrToValuePtr<uint8>(TargetObject), TargetObject, 0);
 			FPropertyChangedEvent PropertyChangedEvent(FoundProperty, EPropertyChangeType::ValueSet);
 			TargetObject->PostEditChangeProperty(PropertyChangedEvent);
 		}
diff --git a/Engine/Plugins/VirtualProduction/LevelSnapshots/Source/LevelSnapshotsEditor/Private/Misc/LevelSnapshotsEditorCustomWidgetGenerator.cpp b/Engine/Plugins/VirtualProduction/LevelSnapshots/Source/LevelSnapshotsEditor/Private/Misc/LevelSnapshotsEditorCustomWidgetGenerator.cpp
--- a/Engine/Plugins/VirtualProduction/LevelSnapshots/Source/LevelSnapshotsEditor/Private/Misc/LevelSnapshotsEditorCustomWidgetGenerator.cpp
+++ b/Engine/Plugins/VirtualProduction/LevelSnapshots/Source/LevelSnapshotsEditor/Private/Misc/LevelSnapshotsEditorCustomWidgetGenerator.cpp
@@ -146,7 +146,7 @@ TSharedPtr<SWidget> LevelSnapshotsEditorCustomWidgetGenerator::GenerateGenericPr
 	if (ensure(InPropertyValue && InObject))
 	{
 		FString ValueText;
-		InProperty->ExportTextItem(ValueText, InPropertyValue, 0, InObject, PPF_None);
+		InProperty->ExportTextItem_Direct(ValueText, InPropertyValue, 0, InObject, PPF_None);
 
 		if (InWidgetTextEditLambda)
 		{
diff --git a/Engine/Plugins/VirtualProduction/RemoteControl/Source/RemoteControl/Private/RemoteControlField.cpp b/Engine/Plugins/VirtualProduction/RemoteControl/Source/RemoteControl/Private/RemoteControlField.cpp
--- a/Engine/Plugins/VirtualProduction/RemoteControl/Source/RemoteControl/Private/RemoteControlField.cpp
+++ b/Engine/Plugins/VirtualProduction/RemoteControl/Source/RemoteControl/Private/RemoteControlField.cpp
@@ -549,7 +549,7 @@ void FRemoteControlFunction::AssignDefaultFunctionArguments()
 				const FString& PropertyDefaultValue = Function->GetMetaData(DefaultPropertyKey);
 				if (!PropertyDefaultValue.IsEmpty())
 				{
-					It->ImportText(*PropertyDefaultValue, It->ContainerPtrToValuePtr<uint8>(FunctionArguments->GetStructMemory()), PPF_None, NULL);
+					It->ImportText_Direct(*PropertyDefaultValue, It->ContainerPtrToValuePtr<uint8>(FunctionArguments->GetStructMemory()), NULL, PPF_None);
 				}
 			}
 		}
diff --git a/Engine/Programs/UnrealHeaderTool/Config/DefaultEngine.ini b/Engine/Programs/UnrealHeaderTool/Config/DefaultEngine.ini
--- a/Engine/Programs/UnrealHeaderTool/Config/DefaultEngine.ini
+++ b/Engine/Programs/UnrealHeaderTool/Config/DefaultEngine.ini
@@ -36,3 +36,4 @@ EngineNativePointerMemberBehavior=Disallow
 EngineObjectPtrMemberBehavior=AllowSilently
 NonEngineNativePointerMemberBehavior=AllowSilently
 NonEngineObjectPtrMemberBehavior=AllowSilently
+AutomaticSettersAndGetters=True
\ No newline at end of file
diff --git a/Engine/Source/Developer/RigVMDeveloper/Private/RigVMModel/RigVMController.cpp b/Engine/Source/Developer/RigVMDeveloper/Private/RigVMModel/RigVMController.cpp
--- a/Engine/Source/Developer/RigVMDeveloper/Private/RigVMModel/RigVMController.cpp
+++ b/Engine/Source/Developer/RigVMDeveloper/Private/RigVMModel/RigVMController.cpp
@@ -6719,7 +6719,7 @@ FString URigVMController::GetPinInitialDefaultValueFromStruct(UScriptStruct* Scr
 		if (Memory)
 		{
 			check(Property);
-			Property->ExportTextItem(DefaultValue, Memory, nullptr, nullptr, PPF_None);
+			Property->ExportTextItem_Direct(DefaultValue, Memory, nullptr, nullptr, PPF_None);
 		}
 	}
 
diff --git a/Engine/Source/Editor/BlueprintGraph/Private/CallFunctionHandler.cpp b/Engine/Source/Editor/BlueprintGraph/Private/CallFunctionHandler.cpp
--- a/Engine/Source/Editor/BlueprintGraph/Private/CallFunctionHandler.cpp
+++ b/Engine/Source/Editor/BlueprintGraph/Private/CallFunctionHandler.cpp
@@ -211,7 +211,7 @@ void FKCHandler_CallFunction::CreateFunctionCallStatement(FKismetFunctionContext
 
 											// Import the literal text to a dummy struct to verify it's well-formed
 											FImportTextErrorContext ErrorPipe(CompilerContext.MessageLog, Node);
-											StructProperty->ImportText(*((*Term)->Name), StructData, 0, nullptr, &ErrorPipe);
+											StructProperty->ImportText_Direct(*((*Term)->Name), StructData, nullptr, 0, &ErrorPipe);
 											if(ErrorPipe.NumErrors > 0)
 											{
 												bMatchedAllParams = false;
diff --git a/Engine/Source/Editor/DetailCustomizations/Private/ConfigEditorPropertyDetails.cpp b/Engine/Source/Editor/DetailCustomizations/Private/ConfigEditorPropertyDetails.cpp
--- a/Engine/Source/Editor/DetailCustomizations/Private/ConfigEditorPropertyDetails.cpp
+++ b/Engine/Source/Editor/DetailCustomizations/Private/ConfigEditorPropertyDetails.cpp
@@ -127,7 +127,7 @@ void FConfigPropertyHelperDetails::AddEditablePropertyForConfig(IDetailLayoutBui
 	FString PropertyName = ConfigEditorCopyOfEditProperty->GetName();
 	if (GConfig->GetString(*SectionName, *PropertyName, ExistingConfigEntryValue, ConfigFilePropertyRowObj->ConfigFileName))
 	{
-		ConfigEditorCopyOfEditProperty->ImportText(*ExistingConfigEntryValue, ConfigEditorCopyOfEditProperty->ContainerPtrToValuePtr<uint8>(ConfigEntryObject), 0, nullptr);
+		ConfigEditorCopyOfEditProperty->ImportText_InContainer(*ExistingConfigEntryValue, ConfigEntryObject, nullptr, 0);
 	}
 
 	// Cache a reference for future usage.
diff --git a/Engine/Source/Editor/EditorConfig/Private/EditorConfig.cpp b/Engine/Source/Editor/EditorConfig/Private/EditorConfig.cpp
--- a/Engine/Source/Editor/EditorConfig/Private/EditorConfig.cpp
+++ b/Engine/Source/Editor/EditorConfig/Private/EditorConfig.cpp
@@ -293,7 +293,7 @@ void FEditorConfig::ReadValue(TSharedPtr<FJsonValue> JsonValue, const FProperty*
 		FString PathString;
 		if (JsonValue->TryGetString(PathString))
 		{
-			Property->ImportText(*PathString, DataPtr, 0, Owner);
+			Property->ImportText_Direct(*PathString, DataPtr, Owner, 0);
 		}
 	}
 	else if (const FStructProperty* StructProperty = CastField<FStructProperty>(Property))
@@ -371,7 +371,7 @@ void FEditorConfig::ReadValue(TSharedPtr<FJsonValue> JsonValue, const FProperty*
 			for (const TPair<FString, TSharedPtr<FJsonValue>>& JsonPair : (*JsonObjectValue)->Values)
 			{
 				KeyProperty->InitializeValue(TempKey.GetData());
-				KeyProperty->ImportText(*JsonPair.Key, TempKey.GetData(), 0, Owner);
+				KeyProperty->ImportText_Direct(*JsonPair.Key, TempKey.GetData(), Owner, 0);
 
 				ValueProperty->InitializeValue(TempValue.GetData());
 				ReadValue(JsonPair.Value, ValueProperty, TempValue.GetData(), Owner); 
@@ -512,7 +512,7 @@ TSharedPtr<FJsonValue> FEditorConfig::WriteValue(const FProperty* Property, cons
 	else if (const FObjectPropertyBase* ObjectProperty = CastField<FObjectPropertyBase>(Property))
 	{
 		FString ObjectPath;
-		ObjectProperty->ExportTextItem(ObjectPath, DataPtr, nullptr, nullptr, PPF_None, nullptr);
+		ObjectProperty->ExportTextItem_Direct(ObjectPath, DataPtr, nullptr, nullptr, PPF_None, nullptr);
 		ResultValue = MakeShared<FJsonValueString>(ObjectPath);
 	}
 	else if (const FStructProperty* StructProperty = CastField<FStructProperty>(Property))
diff --git a/Engine/Source/Editor/Kismet/Private/FindInBlueprintManager.cpp b/Engine/Source/Editor/Kismet/Private/FindInBlueprintManager.cpp
--- a/Engine/Source/Editor/Kismet/Private/FindInBlueprintManager.cpp
+++ b/Engine/Source/Editor/Kismet/Private/FindInBlueprintManager.cpp
@@ -2492,7 +2492,7 @@ void FFindInBlueprintSearchManager::AddOrUpdateBlueprintSearchMetadata(UBlueprin
 	// Build the search data
 	if (FProperty* ParentClassProp = InBlueprint->GetClass()->FindPropertyByName(GET_MEMBER_NAME_CHECKED(UBlueprint, ParentClass)))
 	{
-		ParentClassProp->ExportTextItem(SearchData.ParentClass, ParentClassProp->ContainerPtrToValuePtr<uint8>(InBlueprint), nullptr, InBlueprint, 0);
+		ParentClassProp->ExportTextItem_Direct(SearchData.ParentClass, ParentClassProp->ContainerPtrToValuePtr<uint8>(InBlueprint), nullptr, InBlueprint, 0);
 	}
 
 	if (SearchData.IsValid())
diff --git a/Engine/Source/Editor/Kismet/Private/FindInBlueprints.cpp b/Engine/Source/Editor/Kismet/Private/FindInBlueprints.cpp
--- a/Engine/Source/Editor/Kismet/Private/FindInBlueprints.cpp
+++ b/Engine/Source/Editor/Kismet/Private/FindInBlueprints.cpp
@@ -1066,7 +1066,7 @@ void SFindInBlueprints::MakeSearchQuery(FString InSearchString, bool bInIsFindWi
 			FString ParentClass;
 			if (FProperty* ParentClassProp = Blueprint->GetClass()->FindPropertyByName(GET_MEMBER_NAME_CHECKED(UBlueprint, ParentClass)))
 			{
-				ParentClassProp->ExportTextItem(ParentClass, ParentClassProp->ContainerPtrToValuePtr<uint8>(Blueprint), nullptr, Blueprint, 0);
+				ParentClassProp->ExportTextItem_Direct(ParentClass, ParentClassProp->ContainerPtrToValuePtr<uint8>(Blueprint), nullptr, Blueprint, 0);
 			}
 
 			TArray<FString> Interfaces;
diff --git a/Engine/Source/Editor/Kismet/Private/SMyBlueprint.cpp b/Engine/Source/Editor/Kismet/Private/SMyBlueprint.cpp
--- a/Engine/Source/Editor/Kismet/Private/SMyBlueprint.cpp
+++ b/Engine/Source/Editor/Kismet/Private/SMyBlueprint.cpp
@@ -3357,7 +3357,7 @@ void SMyBlueprint::OnMoveToParent()
 						{
 							// If there is a property for variable, it means the original default value was already copied, so it can be safely overridden
 							VarDesc.DefaultValue.Empty();
-							TargetProperty->ExportTextItem(VarDesc.DefaultValue, OldPropertyAddr, OldPropertyAddr, nullptr, PPF_SerializedAsImportText);
+							TargetProperty->ExportTextItem_Direct(VarDesc.DefaultValue, OldPropertyAddr, OldPropertyAddr, nullptr, PPF_SerializedAsImportText);
 						}
 					}
 				}
@@ -3449,7 +3449,7 @@ void SMyBlueprint::OnCopy()
 				void* OldPropertyAddr = TargetProperty->ContainerPtrToValuePtr<void>(GeneratedCDO);
 				if (OldPropertyAddr)
 				{
-					TargetProperty->ExportTextItem(Description.DefaultValue, OldPropertyAddr, OldPropertyAddr, nullptr, PPF_SerializedAsImportText);
+					TargetProperty->ExportTextItem_Direct(Description.DefaultValue, OldPropertyAddr, OldPropertyAddr, nullptr, PPF_SerializedAsImportText);
 				}
 			}
 
diff --git a/Engine/Source/Editor/KismetCompiler/Private/KismetCompilerVMBackend.cpp b/Engine/Source/Editor/KismetCompiler/Private/KismetCompilerVMBackend.cpp
--- a/Engine/Source/Editor/KismetCompiler/Private/KismetCompilerVMBackend.cpp
+++ b/Engine/Source/Editor/KismetCompiler/Private/KismetCompilerVMBackend.cpp
@@ -935,7 +935,7 @@ public:
 				FProperty* InnerProp = ArrayPropr->Inner;
 				ensure(InnerProp);
 				FScriptArray ScriptArray;
-				ArrayPropr->ImportText(*Term->Name, &ScriptArray, 0, NULL, GLog);
+				ArrayPropr->ImportText_Direct(*Term->Name, &ScriptArray, NULL, 0, GLog);
 
 				FScriptArrayHelper ScriptArrayHelper(ArrayPropr, &ScriptArray);
 				int32 ElementNum = ScriptArrayHelper.Num();
@@ -956,7 +956,7 @@ public:
 				ensure(InnerProp);
 
 				FScriptSet ScriptSet;
-				SetPropr->ImportText(*Term->Name, &ScriptSet, 0, NULL, GLog);
+				SetPropr->ImportText_Direct(*Term->Name, &ScriptSet, NULL, 0, GLog);
 				int32 ElementNum = ScriptSet.Num();
 
 				FScriptSetHelper ScriptSetHelper(SetPropr, &ScriptSet);
@@ -984,7 +984,7 @@ public:
 				ensure(KeyProp && ValProp);
 
 				FScriptMap ScriptMap;
-				MapPropr->ImportText(*Term->Name, &ScriptMap, 0, NULL, GLog);
+				MapPropr->ImportText_Direct(*Term->Name, &ScriptMap, NULL, 0, GLog);
 				int32 ElementNum = ScriptMap.Num();
 
 				FScriptMapHelper ScriptMapHelper(MapPropr, &ScriptMap);
@@ -1135,7 +1135,7 @@ public:
 		StructProperty->InitializeValue(StructData);
 
 		// Assume that any errors on the import of the name string have been caught in the function call generation
-		StructProperty->ImportText(*Term->Name, StructData, 0, NULL, GLog);
+		StructProperty->ImportText_Direct(*Term->Name, StructData, NULL, 0, GLog);
 
 		Writer << EX_StructConst;
 		Writer << LatentInfoStruct;
diff --git a/Engine/Source/Editor/MaterialEditor/Private/MaterialEditorInstanceDetailCustomization.cpp b/Engine/Source/Editor/MaterialEditor/Private/MaterialEditorInstanceDetailCustomization.cpp
--- a/Engine/Source/Editor/MaterialEditor/Private/MaterialEditorInstanceDetailCustomization.cpp
+++ b/Engine/Source/Editor/MaterialEditor/Private/MaterialEditorInstanceDetailCustomization.cpp
@@ -1086,8 +1086,7 @@ void FMaterialInstanceParameterDetails::OnPasteParameterValues(int32 ParameterGr
 						FProperty* ParameterValueProperty = Parameter->GetClass()->FindPropertyByName("ParameterValue");
 						if (ParameterValueProperty != nullptr)
 						{
-							void* ParameterValuePtr = ParameterValueProperty->ContainerPtrToValuePtr<void>(Parameter);
-							ParameterValueProperty->ImportText(*ParsedValueString, ParameterValuePtr, PPF_Copy, Parameter);
+							ParameterValueProperty->ImportText_InContainer(*ParsedValueString, Parameter, Parameter, PPF_Copy);
 						}
 					}
 				}
diff --git a/Engine/Source/Editor/Persona/Private/SAnimNotifyPanel.cpp b/Engine/Source/Editor/Persona/Private/SAnimNotifyPanel.cpp
--- a/Engine/Source/Editor/Persona/Private/SAnimNotifyPanel.cpp
+++ b/Engine/Source/Editor/Persona/Private/SAnimNotifyPanel.cpp
@@ -308,7 +308,7 @@ struct FNotifyNodeInterface : public INodeObjectInterface
 		uint8* PropertyData = Seq->FindNotifyPropertyData(Index, ArrayProperty);
 		if (PropertyData && ArrayProperty)
 		{
-			ArrayProperty->Inner->ExportTextItem(StrValue, PropertyData, PropertyData, Seq, PPF_Copy);
+			ArrayProperty->Inner->ExportTextItem_Direct(StrValue, PropertyData, PropertyData, Seq, PPF_Copy);
 		}
 	}
 
@@ -418,7 +418,7 @@ struct FSyncMarkerNodeInterface : public INodeObjectInterface
 			uint8* PropertyData = Seq->FindSyncMarkerPropertyData(Index, ArrayProperty);
 			if (PropertyData && ArrayProperty)
 			{
-				ArrayProperty->Inner->ExportTextItem(StrValue, PropertyData, PropertyData, Seq, PPF_Copy);
+				ArrayProperty->Inner->ExportTextItem_Direct(StrValue, PropertyData, PropertyData, Seq, PPF_Copy);
 			}
 		}
 	}
@@ -2486,7 +2486,7 @@ FAnimNotifyEvent& SAnimNotifyTrack::CreateNewNotify(FString NewNotifyName, UClas
 				if( Asset )
 				{
 					uint8* Offset = (*PropIt)->ContainerPtrToValuePtr<uint8>(NewEvent.Notify);
-					(*PropIt)->ImportText( *Asset->GetAsset()->GetPathName(), Offset, 0, NewEvent.Notify );
+					(*PropIt)->ImportText_Direct( *Asset->GetAsset()->GetPathName(), Offset, NewEvent.Notify, 0 );
 					break;
 				}
 			}
@@ -3553,7 +3553,7 @@ void SAnimNotifyTrack::PasteSingleNotify(FString& NotifyString, float PasteTime)
 
 	if(PropertyData && ArrayProperty)
 	{
-		ArrayProperty->Inner->ImportText(*NotifyString, PropertyData, PPF_Copy, NULL);
+		ArrayProperty->Inner->ImportText_Direct(*NotifyString, PropertyData, NULL, PPF_Copy);
 
 		FAnimNotifyEvent& NewNotify = Sequence->Notifies[NewIdx];
 
@@ -3629,7 +3629,7 @@ void SAnimNotifyTrack::PasteSingleSyncMarker(FString& MarkerString, float PasteT
 
 		if (PropertyData && ArrayProperty)
 		{
-			ArrayProperty->Inner->ImportText(*MarkerString, PropertyData, PPF_Copy, NULL);
+			ArrayProperty->Inner->ImportText_Direct(*MarkerString, PropertyData, NULL, PPF_Copy);
 
 			FAnimSyncMarker& SyncMarker = AnimSeq->AuthoredSyncMarkers[NewIdx];
 
diff --git a/Engine/Source/Editor/PropertyEditor/Private/PropertyNode.cpp b/Engine/Source/Editor/PropertyEditor/Private/PropertyNode.cpp
--- a/Engine/Source/Editor/PropertyEditor/Private/PropertyNode.cpp
+++ b/Engine/Source/Editor/PropertyEditor/Private/PropertyNode.cpp
@@ -1440,6 +1440,11 @@ public:
 	 */
 	uint8* GetPropertyDefaultAddress() const { return PropertyDefaultAddress; }
 
+	/**
+	 * @return The address of the property's owner object.
+	 */
+	uint8* GetPropertyRootAddress() const { return PropertyValueRoot.ValueAddress; }
+
 private:
 		/**
 	 * Determines whether the property bound to this struct exists in the owning object's archetype.
@@ -2077,7 +2082,7 @@ FString FPropertyNode::GetDefaultValueAsStringForObject( FPropertyItemValueDataT
 			else
 			{
 				// Port flags will cause enums to display correctly
-				InProperty->ExportTextItem( DefaultValue, ValueTracker.GetPropertyDefaultAddress(), ValueTracker.GetPropertyDefaultAddress(), nullptr, PortFlags, nullptr );
+				InProperty->ExportTextItem_InContainer( DefaultValue, ValueTracker.GetPropertyRootAddress(), ValueTracker.GetPropertyDefaultAddress(), nullptr, PortFlags, nullptr );
 			}
 		}
 	}
@@ -3329,15 +3334,15 @@ void FPropertyNode::PropagatePropertyChange( UObject* ModifiedObject, const TCHA
 					// This ensures that shared state keeps the correct value, even if the destination property itself isn't imported (or only partly imported, as is the case with arrays/maps/sets)
 					FString CurrentValue;
 					ComplexProperty->ExportText_Direct(CurrentValue, ModifiedComplexPropAddr, ModifiedComplexPropAddr, ModifiedObject, PPF_None);
-					ComplexProperty->ImportText(*PreviousValue, TempComplexPropAddr, PPF_None, ModifiedObject);
+					ComplexProperty->ImportText_Direct(*PreviousValue, TempComplexPropAddr, ModifiedObject, PPF_None);
 					bShouldImport = ComplexProperty->Identical(DestComplexPropAddr, TempComplexPropAddr, PPF_DeepComparison);
-					ComplexProperty->ImportText(*CurrentValue, TempComplexPropAddr, PPF_None, ModifiedObject);
+					ComplexProperty->ImportText_Direct(*CurrentValue, TempComplexPropAddr, ModifiedObject, PPF_None);
 				}
 
 				// Only import if the value matches the previous value of the property that changed
 				if (bShouldImport)
 				{
-					Prop->ImportText(NewValue, DestSimplePropAddr, PPF_InstanceSubobjects, ActualObjToChange);
+					Prop->ImportText_Direct(NewValue, DestSimplePropAddr, ActualObjToChange, PPF_InstanceSubobjects);
 				}
 			}
 		}
diff --git a/Engine/Source/Editor/PropertyEditor/Private/PropertyTextUtilities.cpp b/Engine/Source/Editor/PropertyEditor/Private/PropertyTextUtilities.cpp
--- a/Engine/Source/Editor/PropertyEditor/Private/PropertyTextUtilities.cpp
+++ b/Engine/Source/Editor/PropertyEditor/Private/PropertyTextUtilities.cpp
@@ -18,14 +18,21 @@ void FPropertyTextUtilities::PropertyToTextHelper(FString& OutString, const FPro
 
 void FPropertyTextUtilities::PropertyToTextHelper(FString& OutString, const FPropertyNode* InPropertyNode, const FProperty* Property, const FObjectBaseAddress& ObjectAddress, EPropertyPortFlags PortFlags)
 {
-	PropertyToTextHelper(OutString, InPropertyNode, Property, ObjectAddress.BaseAddress, ObjectAddress.Object, PortFlags);
+	if (InPropertyNode->GetArrayIndex() != INDEX_NONE || Property->ArrayDim == 1)
+	{
+		Property->ExportTextItem_InContainer(OutString, ObjectAddress.Object, nullptr, ObjectAddress.Object, PortFlags);
+	}
+	else
+	{
+		FArrayProperty::ExportTextInnerItem(OutString, Property, ObjectAddress.BaseAddress, Property->ArrayDim, ObjectAddress.BaseAddress, Property->ArrayDim, ObjectAddress.Object, PortFlags);
+	}
 }
 
 void FPropertyTextUtilities::TextToPropertyHelper(const TCHAR* Buffer, const FPropertyNode* InPropertyNode, const FProperty* Property, uint8* ValueAddress, UObject* Object, EPropertyPortFlags PortFlags)
 {
 	if (InPropertyNode->GetArrayIndex() != INDEX_NONE || Property->ArrayDim == 1)
 	{
-		Property->ImportText(Buffer, ValueAddress, PortFlags, Object);
+		Property->ImportText_Direct(Buffer, ValueAddress, Object, PortFlags);
 	}
 	else
 	{
@@ -35,5 +42,12 @@ void FPropertyTextUtilities::TextToPropertyHelper(const TCHAR* Buffer, const FPr
 
 void FPropertyTextUtilities::TextToPropertyHelper(const TCHAR* Buffer, const FPropertyNode* InPropertyNode, const FProperty* Property, const FObjectBaseAddress& ObjectAddress, EPropertyPortFlags PortFlags)
 {
-	TextToPropertyHelper(Buffer, InPropertyNode, Property, ObjectAddress.BaseAddress, ObjectAddress.Object, PortFlags);
+	if (InPropertyNode->GetArrayIndex() != INDEX_NONE || Property->ArrayDim == 1)
+	{
+		Property->ImportText_InContainer(Buffer, ObjectAddress.Object, ObjectAddress.Object, PortFlags);
+	}
+	else
+	{
+		FArrayProperty::ImportTextInnerItem(Buffer, Property, ObjectAddress.BaseAddress, PortFlags, ObjectAddress.Object);
+	}
 }
diff --git a/Engine/Source/Editor/UMGEditor/Private/WidgetBlueprintEditorUtils.cpp b/Engine/Source/Editor/UMGEditor/Private/WidgetBlueprintEditorUtils.cpp
--- a/Engine/Source/Editor/UMGEditor/Private/WidgetBlueprintEditorUtils.cpp
+++ b/Engine/Source/Editor/UMGEditor/Private/WidgetBlueprintEditorUtils.cpp
@@ -1893,7 +1893,7 @@ void FWidgetBlueprintEditorUtils::ImportPropertiesFromText(UObject* Object, cons
 				PropertyChain.AddHead(Property);
 				Object->PreEditChange(PropertyChain);
 
-				Property->ImportText(*Entry.Value, Property->ContainerPtrToValuePtr<uint8>(Object), 0, Object);
+				Property->ImportText_InContainer(*Entry.Value, Object, Object, 0);
 
 				FPropertyChangedEvent ChangedEvent(Property);
 				Object->PostEditChangeProperty(ChangedEvent);
diff --git a/Engine/Source/Editor/UnrealEd/Private/Commandlets/ExportDialogueScriptCommandlet.cpp b/Engine/Source/Editor/UnrealEd/Private/Commandlets/ExportDialogueScriptCommandlet.cpp
--- a/Engine/Source/Editor/UnrealEd/Private/Commandlets/ExportDialogueScriptCommandlet.cpp
+++ b/Engine/Source/Editor/UnrealEd/Private/Commandlets/ExportDialogueScriptCommandlet.cpp
@@ -359,7 +359,7 @@ FString UExportDialogueScriptCommandlet::GenerateCSVRow(const FDialogueScriptEnt
 		}
 
 		FString PropertyValue;
-		PropertyIt->ExportTextItem(PropertyValue, PropertyIt->ContainerPtrToValuePtr<void>(&InDialogueScriptEntry), nullptr, nullptr, PPF_None);
+		PropertyIt->ExportTextItem_InContainer(PropertyValue, &InDialogueScriptEntry, nullptr, nullptr, PPF_None);
 
 		CSVRow += TEXT("\"");
 		CSVRow += PropertyValue.Replace(TEXT("\""), TEXT("\"\""));
diff --git a/Engine/Source/Editor/UnrealEd/Private/Commandlets/ImportDialogueScriptCommandlet.cpp b/Engine/Source/Editor/UnrealEd/Private/Commandlets/ImportDialogueScriptCommandlet.cpp
--- a/Engine/Source/Editor/UnrealEd/Private/Commandlets/ImportDialogueScriptCommandlet.cpp
+++ b/Engine/Source/Editor/UnrealEd/Private/Commandlets/ImportDialogueScriptCommandlet.cpp
@@ -227,7 +227,7 @@ bool UImportDialogueScriptCommandlet::ImportDialogueScriptForCulture(FLocTextHel
 		// Parse the SpokenDialogue data
 		{
 			const TCHAR* const CellData = RowData[SpokenDialogueColumnIndex];
-			if (SpokenDialogueProperty->ImportText(CellData, SpokenDialogueProperty->ContainerPtrToValuePtr<void>(&ParsedScriptEntry), PPF_None, nullptr) == nullptr)
+			if (SpokenDialogueProperty->ImportText_InContainer(CellData, &ParsedScriptEntry, nullptr, PPF_None) == nullptr)
 			{
 				UE_LOG(LogImportDialogueScriptCommandlet, Error, TEXT("Failed to parse the required column '%s' for row '%d' for culture '%s'."), *SpokenDialogueProperty->GetName(), RowIndex, *InCultureName);
 				continue;
@@ -237,7 +237,7 @@ bool UImportDialogueScriptCommandlet::ImportDialogueScriptForCulture(FLocTextHel
 		// Parse the LocalizationKeys data
 		{
 			const TCHAR* const CellData = RowData[LocalizationKeysColumnIndex];
-			if (LocalizationKeysProperty->ImportText(CellData, LocalizationKeysProperty->ContainerPtrToValuePtr<void>(&ParsedScriptEntry), PPF_None, nullptr) == nullptr)
+			if (LocalizationKeysProperty->ImportText_InContainer(CellData, &ParsedScriptEntry, nullptr, PPF_None) == nullptr)
 			{
 				UE_LOG(LogImportDialogueScriptCommandlet, Error, TEXT("Failed to parse the required column '%s' for row '%d' for culture '%s'."), *LocalizationKeysProperty->GetName(), RowIndex, *InCultureName);
 				continue;
diff --git a/Engine/Source/Editor/UnrealEd/Private/Factories/PackFactory.cpp b/Engine/Source/Editor/UnrealEd/Private/Factories/PackFactory.cpp
--- a/Engine/Source/Editor/UnrealEd/Private/Factories/PackFactory.cpp
+++ b/Engine/Source/Editor/UnrealEd/Private/Factories/PackFactory.cpp
@@ -205,7 +205,7 @@ namespace PackFactoryHelper
 				if (SettingPair.Key.ToString().Contains("ActionMappings"))
 				{
 					FInputActionKeyMapping ActionKeyMapping;
-					ActionMappingsProp->Inner->ImportText(*SettingPair.Value.GetValue(), &ActionKeyMapping, PPF_None, nullptr);
+					ActionMappingsProp->Inner->ImportText_Direct(*SettingPair.Value.GetValue(), &ActionKeyMapping, nullptr, PPF_None);
 
 					if (!InputSettingsCDO->DoesActionExist(ActionKeyMapping.ActionName))
 					{
@@ -215,7 +215,7 @@ namespace PackFactoryHelper
 				else if (SettingPair.Key.ToString().Contains("AxisMappings"))
 				{
 					FInputAxisKeyMapping AxisKeyMapping;
-					AxisMappingsProp->Inner->ImportText(*SettingPair.Value.GetValue(), &AxisKeyMapping, PPF_None, nullptr);
+					AxisMappingsProp->Inner->ImportText_Direct(*SettingPair.Value.GetValue(), &AxisKeyMapping, nullptr, PPF_None);
 
 					if (!InputSettingsCDO->DoesAxisExist(AxisKeyMapping.AxisName))
 					{
diff --git a/Engine/Source/Editor/UnrealEd/Private/Fbx/FbxExportOption.cpp b/Engine/Source/Editor/UnrealEd/Private/Fbx/FbxExportOption.cpp
--- a/Engine/Source/Editor/UnrealEd/Private/Fbx/FbxExportOption.cpp
+++ b/Engine/Source/Editor/UnrealEd/Private/Fbx/FbxExportOption.cpp
@@ -66,7 +66,7 @@ void UFbxExportOption::LoadOptions()
 					ArrayHelper.EmptyAndAddValues(List.Num());
 					for (int32 i = List.Num() - 1, c = 0; i >= 0; i--, c++)
 					{
-						Array->Inner->ImportText(*List[i].GetValue(), ArrayHelper.GetRawPtr(c), PortFlags, this);
+						Array->Inner->ImportText_Direct(*List[i].GetValue(), ArrayHelper.GetRawPtr(c), this, PortFlags);
 					}
 				}
 				else
@@ -91,7 +91,7 @@ void UFbxExportOption::LoadOptions()
 						{
 							// expand the array if necessary so that Index is a valid element
 							ArrayHelper.ExpandForIndex(Index);
-							Array->Inner->ImportText(*ElementValue->GetValue(), ArrayHelper.GetRawPtr(Index), PortFlags, this);
+							Array->Inner->ImportText_Direct(*ElementValue->GetValue(), ArrayHelper.GetRawPtr(Index), this, PortFlags);
 						}
 
 						Index++;
@@ -113,7 +113,7 @@ void UFbxExportOption::LoadOptions()
 
 				if (bFoundValue)
 				{
-					if (Property->ImportText(*Value, Property->ContainerPtrToValuePtr<uint8>(this, i), PortFlags, this) == NULL)
+					if (Property->ImportText_Direct(*Value, Property->ContainerPtrToValuePtr<uint8>(this, i), this, PortFlags) == NULL)
 					{
 						// this should be an error as the properties from the .ini / .int file are not correctly being read in and probably are affecting things in subtle ways
 					}
@@ -151,7 +151,7 @@ void UFbxExportOption::SaveOptions()
 			for (int32 i = 0; i < ArrayHelper.Num(); i++)
 			{
 				FString	Buffer;
-				Array->Inner->ExportTextItem(Buffer, ArrayHelper.GetRawPtr(i), ArrayHelper.GetRawPtr(i), this, PortFlags);
+				Array->Inner->ExportTextItem_Direct(Buffer, ArrayHelper.GetRawPtr(i), ArrayHelper.GetRawPtr(i), this, PortFlags);
 				Sec->Add(*Key, *Buffer);
 			}
 		}
diff --git a/Engine/Source/Editor/UnrealEd/Private/Fbx/FbxFactory.cpp b/Engine/Source/Editor/UnrealEd/Private/Fbx/FbxFactory.cpp
--- a/Engine/Source/Editor/UnrealEd/Private/Fbx/FbxFactory.cpp
+++ b/Engine/Source/Editor/UnrealEd/Private/Fbx/FbxFactory.cpp
@@ -1838,7 +1838,7 @@ void UFbxImportUI::LoadOptions(UObject* ObjectToLoadOptions)
 					ArrayHelper.EmptyAndAddValues(List.Num());
 					for (int32 i = List.Num() - 1, c = 0; i >= 0; i--, c++)
 					{
-						Array->Inner->ImportText(*List[i].GetValue(), ArrayHelper.GetRawPtr(c), PortFlags, ObjectToLoadOptions);
+						Array->Inner->ImportText_Direct(*List[i].GetValue(), ArrayHelper.GetRawPtr(c), ObjectToLoadOptions, PortFlags);
 					}
 				}
 				else
@@ -1863,7 +1863,7 @@ void UFbxImportUI::LoadOptions(UObject* ObjectToLoadOptions)
 						{
 							// expand the array if necessary so that Index is a valid element
 							ArrayHelper.ExpandForIndex(Index);
-							Array->Inner->ImportText(*ElementValue->GetValue(), ArrayHelper.GetRawPtr(Index), PortFlags, ObjectToLoadOptions);
+							Array->Inner->ImportText_Direct(*ElementValue->GetValue(), ArrayHelper.GetRawPtr(Index), ObjectToLoadOptions, PortFlags);
 						}
 
 						Index++;
@@ -1885,7 +1885,7 @@ void UFbxImportUI::LoadOptions(UObject* ObjectToLoadOptions)
 
 				if (bFoundValue)
 				{
-					if (Property->ImportText(*Value, Property->ContainerPtrToValuePtr<uint8>(ObjectToLoadOptions, i), PortFlags, ObjectToLoadOptions) == NULL)
+					if (Property->ImportText_Direct(*Value, Property->ContainerPtrToValuePtr<uint8>(ObjectToLoadOptions, i), ObjectToLoadOptions, PortFlags) == NULL)
 					{
 						// this should be an error as the properties from the .ini / .int file are not correctly being read in and probably are affecting things in subtle ways
 						UE_LOG(LogFbx, Error, TEXT("FBX Options LoadOptions (%s): import failed for %s in: %s"), *ObjectToLoadOptions->GetPathName(), *Property->GetName(), *Value);
@@ -1925,7 +1925,7 @@ void UFbxImportUI::SaveOptions(UObject* ObjectToSaveOptions)
 			for (int32 i = 0; i < ArrayHelper.Num(); i++)
 			{
 				FString	Buffer;
-				Array->Inner->ExportTextItem(Buffer, ArrayHelper.GetRawPtr(i), ArrayHelper.GetRawPtr(i), ObjectToSaveOptions, PortFlags);
+				Array->Inner->ExportTextItem_Direct(Buffer, ArrayHelper.GetRawPtr(i), ArrayHelper.GetRawPtr(i), ObjectToSaveOptions, PortFlags);
 				Sec->Add(*Key, *Buffer);
 			}
 		}
diff --git a/Engine/Source/Editor/UnrealEd/Private/Kismet2/BlueprintEditorUtils.cpp b/Engine/Source/Editor/UnrealEd/Private/Kismet2/BlueprintEditorUtils.cpp
--- a/Engine/Source/Editor/UnrealEd/Private/Kismet2/BlueprintEditorUtils.cpp
+++ b/Engine/Source/Editor/UnrealEd/Private/Kismet2/BlueprintEditorUtils.cpp
@@ -5377,7 +5377,7 @@ FName FBlueprintEditorUtils::DuplicateMemberVariable(UBlueprint* InFromBlueprint
 				{
 					// if there is a property for variable, it means the original default value was already copied, so it can be safely overridden
 					NewVar.DefaultValue.Empty();
-					TargetProperty->ExportTextItem(NewVar.DefaultValue, OldPropertyAddr, OldPropertyAddr, nullptr, PPF_SerializedAsImportText);
+					TargetProperty->ExportTextItem_Direct(NewVar.DefaultValue, OldPropertyAddr, OldPropertyAddr, nullptr, PPF_SerializedAsImportText);
 				}
 			}
 
@@ -9121,7 +9121,7 @@ bool FBlueprintEditorUtils::PropertyValueFromString_Direct(const FProperty* Prop
 		else if (Property->IsA(FTextProperty::StaticClass()))
 		{
 			FStringOutputDevice ImportError;
-			const TCHAR* EndOfParsedBuff = Property->ImportText(*StrValue, DirectValue, PPF_SerializedAsImportText, OwningObject, &ImportError);
+			const TCHAR* EndOfParsedBuff = Property->ImportText_Direct(*StrValue, DirectValue, OwningObject, PPF_SerializedAsImportText, &ImportError);
 			bParseSucceeded = EndOfParsedBuff && ImportError.IsEmpty();
 		}
 		else
@@ -9132,7 +9132,7 @@ bool FBlueprintEditorUtils::PropertyValueFromString_Direct(const FProperty* Prop
 				: *StrValue;
 
 			FStringOutputDevice ImportError;
-			const TCHAR* EndOfParsedBuff = Property->ImportText(*StrValue, DirectValue, PPF_SerializedAsImportText, OwningObject, &ImportError);
+			const TCHAR* EndOfParsedBuff = Property->ImportText_Direct(*StrValue, DirectValue, OwningObject, PPF_SerializedAsImportText, &ImportError);
 			bParseSucceeded = EndOfParsedBuff && ImportError.IsEmpty();
 		}
 	}
@@ -9179,7 +9179,7 @@ bool FBlueprintEditorUtils::PropertyValueFromString_Direct(const FProperty* Prop
 			ensure(1 == StructProperty->ArrayDim);
 
 			FStringOutputDevice ImportError;
-			const TCHAR* EndOfParsedBuff = StructProperty->ImportText(StrValue.IsEmpty() ? TEXT("()") : *StrValue, DirectValue, PPF_SerializedAsImportText, OwningObject, &ImportError);
+			const TCHAR* EndOfParsedBuff = StructProperty->ImportText_Direct(StrValue.IsEmpty() ? TEXT("()") : *StrValue, DirectValue, OwningObject, PPF_SerializedAsImportText, &ImportError);
 			bParseSucceeded &= EndOfParsedBuff && ImportError.IsEmpty();
 		}
 	}
diff --git a/Engine/Source/Editor/UnrealEd/Private/Kismet2/KismetDebugUtilities.cpp b/Engine/Source/Editor/UnrealEd/Private/Kismet2/KismetDebugUtilities.cpp
--- a/Engine/Source/Editor/UnrealEd/Private/Kismet2/KismetDebugUtilities.cpp
+++ b/Engine/Source/Editor/UnrealEd/Private/Kismet2/KismetDebugUtilities.cpp
@@ -2005,7 +2005,7 @@ FPropertyInstanceInfo::FPropertyInstanceInfo(FPropertyInstance PropertyInstance)
 	else if (const FStructProperty* StructProperty = CastField<FStructProperty>(ResolvedProperty))
 	{
 		FString WatchText;
-		StructProperty->ExportTextItem(WatchText, PropertyInstance.Value, PropertyInstance.Value, nullptr, PPF_PropertyWindow | PPF_BlueprintDebugView, nullptr);
+		StructProperty->ExportTextItem_Direct(WatchText, PropertyInstance.Value, PropertyInstance.Value, nullptr, PPF_PropertyWindow | PPF_BlueprintDebugView, nullptr);
 		Value = FText::FromString(WatchText);
 		return;
 	}
@@ -2185,7 +2185,7 @@ void FPropertyInstanceInfo::PopulateChildren(FPropertyInstance PropertyInstance,
 				const TSharedPtr<FPropertyInstanceInfo> ChildInfo = FindOrMake(ChildProperty, VisitedNodes);
 				
 				FString NameStr = TEXT("[");
-				MapProperty->KeyProp->ExportTextItem(
+				MapProperty->KeyProp->ExportTextItem_Direct(
 					NameStr,
 					PropData,
 					nullptr,
diff --git a/Engine/Source/Editor/UnrealEd/Private/MaterialGraphSchema.cpp b/Engine/Source/Editor/UnrealEd/Private/MaterialGraphSchema.cpp
--- a/Engine/Source/Editor/UnrealEd/Private/MaterialGraphSchema.cpp
+++ b/Engine/Source/Editor/UnrealEd/Private/MaterialGraphSchema.cpp
@@ -864,7 +864,7 @@ void UMaterialGraphSchema::GetMaterialFunctionActions(FGraphActionMenuBuilder& A
 						if (FArrayProperty* LibraryCategoriesProperty = FindFieldChecked<FArrayProperty>(UMaterialFunction::StaticClass(), TEXT("LibraryCategories")))
 						{
 							uint8* DestAddr = (uint8*)(&LibraryCategories);
-							LibraryCategoriesProperty->ImportText(*LibraryCategoriesString, DestAddr, PPF_None, NULL, GWarn);
+							LibraryCategoriesProperty->ImportText_Direct(*LibraryCategoriesString, DestAddr, NULL, PPF_None, GWarn);
 						}
 					}
 				}
@@ -875,7 +875,7 @@ void UMaterialGraphSchema::GetMaterialFunctionActions(FGraphActionMenuBuilder& A
 					{
 						FArrayProperty* LibraryCategoriesProperty = FindFieldChecked<FArrayProperty>(UMaterialFunction::StaticClass(), GET_MEMBER_NAME_CHECKED(UMaterialFunction, LibraryCategoriesText));
 						uint8* DestAddr = (uint8*)(&LibraryCategoriesText);
-						LibraryCategoriesProperty->ImportText(*LibraryCategoriesString, DestAddr, PPF_None, NULL, GWarn);
+						LibraryCategoriesProperty->ImportText_Direct(*LibraryCategoriesString, DestAddr, NULL, PPF_None, GWarn);
 					}
 
 					for (const FString& Category : LibraryCategories)
diff --git a/Engine/Source/Editor/UnrealEd/Private/Tests/AutomationEditorCommon.cpp b/Engine/Source/Editor/UnrealEd/Private/Tests/AutomationEditorCommon.cpp
--- a/Engine/Source/Editor/UnrealEd/Private/Tests/AutomationEditorCommon.cpp
+++ b/Engine/Source/Editor/UnrealEd/Private/Tests/AutomationEditorCommon.cpp
@@ -243,7 +243,7 @@ void FAutomationEditorCommonUtils::ApplyCustomFactorySetting(UObject* InObject,
 	{
 		if (PropertyChain.Num() == 0)
 		{
-			TargetProperty->ImportText(*Value, TargetProperty->ContainerPtrToValuePtr<uint8>(InObject), 0, InObject);
+			TargetProperty->ImportText_InContainer(*Value, InObject, InObject, 0);
 		}
 		else
 		{
diff --git a/Engine/Source/Editor/UnrealEd/Private/Tests/AutomationEditorPromotionCommon.cpp b/Engine/Source/Editor/UnrealEd/Private/Tests/AutomationEditorPromotionCommon.cpp
--- a/Engine/Source/Editor/UnrealEd/Private/Tests/AutomationEditorPromotionCommon.cpp
+++ b/Engine/Source/Editor/UnrealEd/Private/Tests/AutomationEditorPromotionCommon.cpp
@@ -199,8 +199,7 @@ FString FEditorPromotionTestUtilities::GetPropertyByName(UObject* TargetObject,
 	if (FoundProperty)
 	{
 		FString ValueString;
-		const uint8* PropertyAddr = FoundProperty->ContainerPtrToValuePtr<uint8>(TargetObject);
-		FoundProperty->ExportTextItem(ValueString, PropertyAddr, NULL, NULL, PPF_None);
+		FoundProperty->ExportTextItem_InContainer(ValueString, TargetObject, NULL, NULL, PPF_None);
 		return ValueString;
 	}
 	return TEXT("");
@@ -223,7 +222,7 @@ void FEditorPromotionTestUtilities::SetPropertyByName(UObject* TargetObject, con
 		TargetObject->Modify();
 
 		TargetObject->PreEditChange(FoundProperty);
-		FoundProperty->ImportText(*NewValueString, FoundProperty->ContainerPtrToValuePtr<uint8>(TargetObject), 0, TargetObject);
+		FoundProperty->ImportText_InContainer(*NewValueString, TargetObject, TargetObject, 0);
 		FPropertyChangedEvent PropertyChangedEvent(FoundProperty, EPropertyChangeType::ValueSet);
 		TargetObject->PostEditChangeProperty(PropertyChangedEvent);
 	}
diff --git a/Engine/Source/Programs/UnrealHeaderTool/Private/BaseParser.h b/Engine/Source/Programs/UnrealHeaderTool/Private/BaseParser.h
--- a/Engine/Source/Programs/UnrealHeaderTool/Private/BaseParser.h
+++ b/Engine/Source/Programs/UnrealHeaderTool/Private/BaseParser.h
@@ -59,6 +59,9 @@ struct FUHTConfig
 	EPointerMemberBehavior NonEngineNativePointerMemberBehavior = EPointerMemberBehavior::AllowSilently;
 	EPointerMemberBehavior NonEngineObjectPtrMemberBehavior = EPointerMemberBehavior::AllowSilently;
 
+	/** If true setters and getters will be automatically (without specifying their function names on a property) parsed and generated if a function with matching signature is found */
+	bool bAllowAutomaticSettersAndGetters = false;
+
 private:
 	FUHTConfig();
 };
