Refresh ClassViewers when global filter changes

diff --git a/Engine/Source/Editor/ClassViewer/Private/ClassViewerModule.cpp b/Engine/Source/Editor/ClassViewer/Private/ClassViewerModule.cpp
--- a/Engine/Source/Editor/ClassViewer/Private/ClassViewerModule.cpp
+++ b/Engine/Source/Editor/ClassViewer/Private/ClassViewerModule.cpp
@@ -109,6 +109,7 @@ TSharedRef<FClassViewerFilterFuncs> FClassViewerModule::CreateFilterFuncs()
 void FClassViewerModule::RegisterGlobalClassViewerFilter(const TSharedRef<IClassViewerFilter>& Filter)
 {
 	GlobalClassViewerFilter = Filter;
+	OnGlobalClassViewerFilterModified.Broadcast();
 }
 
 const TSharedPtr<IClassViewerFilter>& FClassViewerModule::GetGlobalClassViewerFilter()
diff --git a/Engine/Source/Editor/ClassViewer/Public/ClassViewerModule.h b/Engine/Source/Editor/ClassViewer/Public/ClassViewerModule.h
--- a/Engine/Source/Editor/ClassViewer/Public/ClassViewerModule.h
+++ b/Engine/Source/Editor/ClassViewer/Public/ClassViewerModule.h
@@ -199,7 +199,11 @@ public:
 	/** Returns the global filter that affects all class viewer instances */
 	virtual const TSharedPtr<IClassViewerFilter>& GetGlobalClassViewerFilter();
 
+	FSimpleMulticastDelegate& GetOnGlobalClassViewerFilterModified() { return OnGlobalClassViewerFilterModified; }
+
 private:
 
 	TSharedPtr<IClassViewerFilter> GlobalClassViewerFilter;
+
+	FSimpleMulticastDelegate OnGlobalClassViewerFilterModified;
 };
diff --git a/Engine/Source/Editor/UnrealEd/Private/SComponentClassCombo.cpp b/Engine/Source/Editor/UnrealEd/Private/SComponentClassCombo.cpp
--- a/Engine/Source/Editor/UnrealEd/Private/SComponentClassCombo.cpp
+++ b/Engine/Source/Editor/UnrealEd/Private/SComponentClassCombo.cpp
@@ -57,8 +57,10 @@ void SComponentClassCombo::Construct(const FArguments& InArgs)
 
 	ComponentClassFilterData.InitOptions = MakeShared<FClassViewerInitializationOptions>();
 	ComponentClassFilterData.InitOptions->ClassFilters.Append(InArgs._CustomClassFilters);
-	ComponentClassFilterData.ClassFilter = FModuleManager::LoadModuleChecked<FClassViewerModule>("ClassViewer").CreateClassFilter(*ComponentClassFilterData.InitOptions);
-	ComponentClassFilterData.FilterFuncs = FModuleManager::LoadModuleChecked<FClassViewerModule>("ClassViewer").CreateFilterFuncs();
+	FClassViewerModule& ClassViewerModule = FModuleManager::LoadModuleChecked<FClassViewerModule>("ClassViewer");
+	ComponentClassFilterData.ClassFilter = ClassViewerModule.CreateClassFilter(*ComponentClassFilterData.InitOptions);
+	ComponentClassFilterData.FilterFuncs = ClassViewerModule.CreateFilterFuncs();
+	ClassViewerModule.GetOnGlobalClassViewerFilterModified().AddRaw(this, &SComponentClassCombo::UpdateComponentClassList);
 
 	FComponentTypeRegistry::Get().SubscribeToComponentList(ComponentClassList).AddRaw(this, &SComponentClassCombo::UpdateComponentClassList);
 
@@ -142,6 +144,11 @@ void SComponentClassCombo::Construct(const FArguments& InArgs)
 SComponentClassCombo::~SComponentClassCombo()
 {
 	FComponentTypeRegistry::Get().GetOnComponentTypeListChanged().RemoveAll(this);
+	
+	if (FClassViewerModule* ClassViewerModule = FModuleManager::Get().GetModulePtr<FClassViewerModule>("ClassViewer"))
+	{
+		ClassViewerModule->GetOnGlobalClassViewerFilterModified().RemoveAll(this);
+	}
 }
 
 void SComponentClassCombo::ClearSelection()
