Prevent  packages in /Temp that were created during the cook and do not exist on disk from being cooked unintentionally. FPackageDatas::TryLookupFileNameOnDisk is supposed to return false for packages that do not exist on disk.
#rnx

diff --git a/Engine/Source/Editor/UnrealEd/Private/Cooker/CookPackageData.cpp b/Engine/Source/Editor/UnrealEd/Private/Cooker/CookPackageData.cpp
--- a/Engine/Source/Editor/UnrealEd/Private/Cooker/CookPackageData.cpp
+++ b/Engine/Source/Editor/UnrealEd/Private/Cooker/CookPackageData.cpp
@@ -1879,6 +1879,22 @@ bool FPackageDatas::TryLookupFileNameOnDisk(FName PackageName, FString& OutFileN
 			}
 			return false;
 		}
+		else
+		{
+			// Temporary fix for packages added during cook: GetAssetsByPackageName is returning true for these generated
+			// files that do not exist on disk, since they get added into the AssetRegistryState in UAssetRegistryImpl::ProcessLoadedAssetsToUpdateCache
+			// The current known cases where this is a problem is when cooking WorldPartition maps, which create temporary World packages in
+			// //Temp. So for /Temp files, use the slower disk check rather than the AssetRegistry.
+			bool bForceDiskCheck = PackageNameStr.StartsWith(TEXT("/Temp/"));
+			if (bForceDiskCheck)
+			{
+				if (!FPackageName::DoesPackageExist(PackageNameStr, &OutFileName, false /* InAllowTextFormats */))
+				{
+					return false;
+				}
+				return true;
+			}
+		}
 
 		FName ClassRedirector = UObjectRedirector::StaticClass()->GetFName();
 		bool bContainsMap = false;
