CleanupPhysicsSceneThatWasInitializedForSave: Bring back CleanupWorld call to pre CL 16574665.
- Fix crash because of Subsystem Init call done through InitializePhysicsSceneForSaveIfNecessary not being matched with Deinit in CleanupPhysicsSceneThatWasInitializedForSave.
  This was delaying it to StaticExit() IncrementalPurgeGarbage causing some Subsystem objects to become invalid before the FSubsystemCollectionBase that contained them.
- Original change no longer necessary because CleanupWorldPartition is no longer called from within CleanupLevel

#rnx

diff --git a/Engine/Source/Editor/UnrealEd/Private/EditorEngine.cpp b/Engine/Source/Editor/UnrealEd/Private/EditorEngine.cpp
--- a/Engine/Source/Editor/UnrealEd/Private/EditorEngine.cpp
+++ b/Engine/Source/Editor/UnrealEd/Private/EditorEngine.cpp
@@ -4273,8 +4273,7 @@ void UEditorEngine::CleanupPhysicsSceneThatWasInitializedForSave(UWorld* World,
 
 	if(bForceInitialized)
 	{
-		// Cleanup the world, but with bCleanupResources=false so that we don't try to remove the world and dependent packages from memory and break future saving
-		World->CleanupWorld(true /* bSessionEnded=default */, false /* bCleanupResources */);
+		World->CleanupWorld(true, true, World);
 	}
 
 	World->SetPhysicsScene(nullptr);
