BulkDataRegistry: Allow the BulkDataRegistry to store its data in DDC keyed by EditorDomain PackageDigests even when EditorDomain is not enabled as the PackageResourceManager.
#rnx

diff --git a/Engine/Source/Editor/UnrealEd/Private/EditorDomain/EditorDomain.cpp b/Engine/Source/Editor/UnrealEd/Private/EditorDomain/EditorDomain.cpp
--- a/Engine/Source/Editor/UnrealEd/Private/EditorDomain/EditorDomain.cpp
+++ b/Engine/Source/Editor/UnrealEd/Private/EditorDomain/EditorDomain.cpp
@@ -32,7 +32,11 @@
 
 DEFINE_LOG_CATEGORY(LogEditorDomain);
 
-/** Add a hook to the PackageResourceManager's startup delegate to use the EditorDomain as the IPackageResourceManager */
+/**
+ * A pre-Main global registration struct that decides whether to create the EditorDomain, and initializes EditorDomainUtils
+ * for usage by any system that needs PackageDigests, with or without FEditorDomain.
+ * If the EditorDomain is enabled, also sets PackageResourceManager to use the EditorDomain as the IPackageResourceManager
+ */
 class FEditorDomainRegisterAsPackageResourceManager
 {
 public:
@@ -43,27 +47,31 @@ public:
 
 	static IPackageResourceManager* SetPackageResourceManager()
 	{
+		using namespace UE::EditorDomain;
+
 		bool bEditorDomainEnabled = IsEditorDomainEnabled();
 		if (GIsEditor)
 		{
 			UE_LOG(LogEditorDomain, Display, TEXT("EditorDomain is %s"), bEditorDomainEnabled ? TEXT("Enabled") : TEXT("Disabled"));
 		}
+		UtilsInitialize();
 		if (bEditorDomainEnabled)
 		{
-			UE::EditorDomain::UtilsInitialize();
 			UE::TargetDomain::UtilsInitialize(bEditorDomainEnabled);
-			if (bEditorDomainEnabled)
-			{
-				// Set values for config settings EditorDomain depends on
-				GAllowUnversionedContentInEditor = 1;
-
-				// Create the editor domain and return it as the package resource manager
-				check(FEditorDomain::RegisteredEditorDomain == nullptr);
-				FEditorDomain::RegisteredEditorDomain = new FEditorDomain();
-				return FEditorDomain::RegisteredEditorDomain;
-			}
+			// Set values for config settings EditorDomain depends on
+			GAllowUnversionedContentInEditor = 1;
+
+			// Create the editor domain, register it as the IPackageDigestCache, and return it as the package resource manager
+			check(FEditorDomain::RegisteredEditorDomain == nullptr);
+			FEditorDomain::RegisteredEditorDomain = new FEditorDomain();
+			IPackageDigestCache::Set(FEditorDomain::RegisteredEditorDomain);
+			return FEditorDomain::RegisteredEditorDomain;
+		}
+		else
+		{
+			IPackageDigestCache::SetDefault();
+			return nullptr;
 		}
-		return nullptr;
 	}
 } GRegisterAsPackageResourceManager;
 
@@ -125,6 +133,10 @@ FEditorDomain::~FEditorDomain()
 	{
 		RegisteredEditorDomain = nullptr;
 	}
+	if (UE::EditorDomain::IPackageDigestCache::Get() == this)
+	{
+		UE::EditorDomain::IPackageDigestCache::Set(nullptr);
+	}
 }
 
 FEditorDomain* FEditorDomain::Get()
diff --git a/Engine/Source/Editor/UnrealEd/Private/EditorDomain/EditorDomain.h b/Engine/Source/Editor/UnrealEd/Private/EditorDomain/EditorDomain.h
--- a/Engine/Source/Editor/UnrealEd/Private/EditorDomain/EditorDomain.h
+++ b/Engine/Source/Editor/UnrealEd/Private/EditorDomain/EditorDomain.h
@@ -73,6 +73,26 @@ struct FPackageDigest
 	FName StatusArg;
 };
 
+/**
+ * An interface to cache repeated calls to GetPackageDigest. This would be just a global TMap without an interface,
+ * except that EditorDomain if enabled needs to store other data along with the digest so it takes over the cache duty
+ * using its own storage.
+ */
+class IPackageDigestCache
+{
+public:
+	virtual ~IPackageDigestCache() {}
+	/**
+	 * Calculate the PackageDigest for the given PackageName.
+	 * Callable from any thread, but will return !IsSuccessful digest if game-thread data is required and not cached.
+	 */
+	virtual FPackageDigest GetPackageDigest(FName PackageName) = 0;
+
+	static IPackageDigestCache* Get();
+	static void Set(IPackageDigestCache* Cache);
+	static void SetDefault();
+};
+
 }
 
 extern FString LexToString(UE::EditorDomain::FPackageDigest::EStatus Status, FName StatusArg);
@@ -93,7 +113,7 @@ DECLARE_LOG_CATEGORY_EXTERN(LogEditorDomain, Log, All);
  * the WorkspaceDomain (through ordinary IFileManager operations on the Root/Game/Content folders) and creates
  * the EditorDomain version for next time.
  */
-class FEditorDomain : public IPackageResourceManager, public FTickableEditorObject
+class FEditorDomain : public IPackageResourceManager, public FTickableEditorObject, public UE::EditorDomain::IPackageDigestCache
 {
 public:
 	/** Different options for which domain a package comes from. */
@@ -143,14 +163,12 @@ public:
 	virtual ETickableTickType GetTickableTickType() const override { return ETickableTickType::Always; }
 	virtual TStatId GetStatId() const override { return TStatId(); }
 
+	// IPackageDigestCache interface
+	virtual UE::EditorDomain::FPackageDigest GetPackageDigest(FName PackageDigest) override;
+
 	// EditorDomain interface
 	/** Fetch data from game-thread sources that is required to calculate the PackageDigest of the given PackageName. */
 	void PrecachePackageDigest(FName PackageName);
-	/**
-	 * Calculate the PackageDigest for the given PackageName.
-	 * Callable from any thread, but will fail if game-thread data is required and not cached.
-	 */
-	UE::EditorDomain::FPackageDigest GetPackageDigest(FName PackageDigest);
 
 	/** Request the download of the given packages from the upstream DDC server. */
 	void BatchDownload(TArrayView<FName> PackageNames);
diff --git a/Engine/Source/Editor/UnrealEd/Private/EditorDomain/EditorDomainUtils.cpp b/Engine/Source/Editor/UnrealEd/Private/EditorDomain/EditorDomainUtils.cpp
--- a/Engine/Source/Editor/UnrealEd/Private/EditorDomain/EditorDomainUtils.cpp
+++ b/Engine/Source/Editor/UnrealEd/Private/EditorDomain/EditorDomainUtils.cpp
@@ -8,6 +8,7 @@
 #include "AssetRegistry/AssetData.h"
 #include "AssetRegistry/IAssetRegistry.h"
 #include "Containers/Array.h"
+#include "Containers/Map.h"
 #include "DerivedDataBuildDefinition.h"
 #include "DerivedDataCache.h"
 #include "DerivedDataCacheKey.h"
@@ -16,6 +17,7 @@
 #include "DerivedDataValue.h"
 #include "Editor.h"
 #include "Engine/StaticMesh.h"
+#include "HAL/CriticalSection.h"
 #include "HAL/FileManager.h"
 #include "Hash/Blake3.h"
 #include "Memory/SharedBuffer.h"
@@ -944,6 +946,59 @@ TMultiMap<FName, FName> ConstructConstructClasses()
 FDelegateHandle GUtilsPostInitDelegate;
 void UtilsPostEngineInit();
 
+/** A default implementation of IPackageDigestCache that stores the Digests in a TMap. */
+class FDefaultPackageDigestCache : public IPackageDigestCache
+{
+public:
+	virtual ~FDefaultPackageDigestCache()
+	{
+		if (this == IPackageDigestCache::Get())
+		{
+			IPackageDigestCache::Set(nullptr);
+		}
+	}
+
+	virtual UE::EditorDomain::FPackageDigest GetPackageDigest(FName PackageName) override
+	{
+		uint32 TypeHash = GetTypeHash(PackageName);
+		{
+			FReadScopeLock ScopeLock(Lock);
+			if (FPackageDigest* PackageDigest = PackageDigests.FindByHash(TypeHash, PackageName))
+			{
+				return *PackageDigest;
+			}
+		}
+		FPackageDigest PackageDigest = CalculatePackageDigest(*IAssetRegistry::Get(), PackageName);
+		{
+			FWriteScopeLock ScopeLock(Lock);
+			PackageDigests.AddByHash(TypeHash, PackageName, PackageDigest);
+		}
+		return PackageDigest;
+	}
+
+private:
+	FRWLock Lock;
+	TMap<FName, FPackageDigest> PackageDigests;
+};
+static TOptional<FDefaultPackageDigestCache> GDefaultPackageDigestCache;
+static IPackageDigestCache* GPackageDigestCache = nullptr;
+IPackageDigestCache* IPackageDigestCache::Get()
+{
+	return GPackageDigestCache;
+}
+
+void IPackageDigestCache::Set(IPackageDigestCache* Cache)
+{
+	// Set is called only before Main or after Main, during single-threaded startup/shutdown, so we do not need a lock
+	GPackageDigestCache = Cache;
+}
+
+void IPackageDigestCache::SetDefault()
+{
+	GDefaultPackageDigestCache.Emplace();
+	Set(GDefaultPackageDigestCache.GetPtrOrNull());
+}
+
 void UtilsInitialize()
 {
 	GClassBlockedUses = ConstructClassBlockedUses();
@@ -1761,14 +1816,10 @@ bool TrySavePackage(UPackage* Package)
 
 void GetBulkDataList(FName PackageName, UE::DerivedData::IRequestOwner& Owner, TUniqueFunction<void(FSharedBuffer Buffer)>&& Callback)
 {
-	FEditorDomain* EditorDomain = FEditorDomain::Get();
-	if (!EditorDomain)
-	{
-		Callback(FSharedBuffer());
-		return;
-	}
-
-	FPackageDigest PackageDigest = EditorDomain->GetPackageDigest(PackageName);
+	IPackageDigestCache* DigestCache = IPackageDigestCache::Get();
+	FPackageDigest PackageDigest = DigestCache
+		? DigestCache->GetPackageDigest(PackageName)
+		: CalculatePackageDigest(*IAssetRegistry::Get(), PackageName);
 	if (!PackageDigest.IsSuccessful())
 	{
 		Callback(FSharedBuffer());
@@ -1793,13 +1844,10 @@ void GetBulkDataList(FName PackageName, UE::DerivedData::IRequestOwner& Owner, T
 
 void PutBulkDataList(FName PackageName, FSharedBuffer Buffer)
 {
-	FEditorDomain* EditorDomain = FEditorDomain::Get();
-	if (!EditorDomain)
-	{
-		return;
-	}
-
-	FPackageDigest PackageDigest = EditorDomain->GetPackageDigest(PackageName);
+	IPackageDigestCache* DigestCache = IPackageDigestCache::Get();
+	FPackageDigest PackageDigest = DigestCache
+		? DigestCache->GetPackageDigest(PackageName)
+		: CalculatePackageDigest(*IAssetRegistry::Get(), PackageName);
 	if (!PackageDigest.IsSuccessful())
 	{
 		return;
@@ -1830,14 +1878,10 @@ FIoHash GetPackageAndGuidHash(FIoHash& EditorDomainHash, const FGuid& BulkDataId
 void GetBulkDataPayloadId(FName PackageName, const FGuid& BulkDataId, UE::DerivedData::IRequestOwner& Owner,
 	TUniqueFunction<void(FSharedBuffer Buffer)>&& Callback)
 {
-	FEditorDomain* EditorDomain = FEditorDomain::Get();
-	if (!EditorDomain)
-	{
-		Callback(FSharedBuffer());
-		return;
-	}
-
-	FPackageDigest PackageDigest = EditorDomain->GetPackageDigest(PackageName);
+	IPackageDigestCache* DigestCache = IPackageDigestCache::Get();
+	FPackageDigest PackageDigest = DigestCache
+		? DigestCache->GetPackageDigest(PackageName)
+		: CalculatePackageDigest(*IAssetRegistry::Get(), PackageName);
 	if (!PackageDigest.IsSuccessful())
 	{
 		Callback(FSharedBuffer());
@@ -1864,13 +1908,10 @@ void GetBulkDataPayloadId(FName PackageName, const FGuid& BulkDataId, UE::Derive
 
 void PutBulkDataPayloadId(FName PackageName, const FGuid& BulkDataId, FSharedBuffer Buffer)
 {
-	FEditorDomain* EditorDomain = FEditorDomain::Get();
-	if (!EditorDomain)
-	{
-		return;
-	}
-
-	FPackageDigest PackageDigest = EditorDomain->GetPackageDigest(PackageName);
+	IPackageDigestCache* DigestCache = IPackageDigestCache::Get();
+	FPackageDigest PackageDigest = DigestCache
+		? DigestCache->GetPackageDigest(PackageName)
+		: CalculatePackageDigest(*IAssetRegistry::Get(), PackageName);
 	if (!PackageDigest.IsSuccessful())
 	{
 		return;
