Disallowing using script classes as base classes for assets

diff --git a/Engine/Source/Editor/UnrealEd/Private/Factories/EditorFactories.cpp b/Engine/Source/Editor/UnrealEd/Private/Factories/EditorFactories.cpp
--- a/Engine/Source/Editor/UnrealEd/Private/Factories/EditorFactories.cpp
+++ b/Engine/Source/Editor/UnrealEd/Private/Factories/EditorFactories.cpp
@@ -358,6 +358,7 @@ public:
 	virtual bool IsClassAllowed(const FClassViewerInitializationOptions& InInitOptions, const UClass* InClass, TSharedRef< FClassViewerFilterFuncs > InFilterFuncs) override
 	{
 		bool bAllowed= !InClass->HasAnyClassFlags(DisallowedClassFlags)
+			&& InClass->CanCreateAssetOfClass()
 			&& InFilterFuncs->IfInChildOfClassesSet(AllowedChildrenOfClasses, InClass) != EFilterReturn::Failed;
 
 		if (bAllowed && bDisallowBlueprintBase)
diff --git a/Engine/Source/Runtime/CoreUObject/Public/UObject/Class.h b/Engine/Source/Runtime/CoreUObject/Public/UObject/Class.h
--- a/Engine/Source/Runtime/CoreUObject/Public/UObject/Class.h
+++ b/Engine/Source/Runtime/CoreUObject/Public/UObject/Class.h
@@ -3090,6 +3090,14 @@ public:
 	 */
 	virtual void SetupObjectInitializer(FObjectInitializer& ObjectInitializer) const {}
 
+	/**
+	 * Some classes may not support creating assets of their type
+	 */
+	virtual bool CanCreateAssetOfClass() const
+	{
+		return true;
+	}
+
 	/**
 	 * Get the name of the CDO for the this class
 	 * @return The name of the CDO
