Package Tools: Add exposed functions for converting between package names and filenames.

diff --git a/Engine/Source/Editor/UnrealEd/Private/PackageTools.cpp b/Engine/Source/Editor/UnrealEd/Private/PackageTools.cpp
--- a/Engine/Source/Editor/UnrealEd/Private/PackageTools.cpp
+++ b/Engine/Source/Editor/UnrealEd/Private/PackageTools.cpp
@@ -1376,6 +1376,20 @@ UPackageTools::UPackageTools(const FObjectInitializer& ObjectInitializer)
 		return SanitizedName;
 	}
 
+	FString UPackageTools::PackageNameToFilename(const FString& InPackageName, const FString& Extension)
+	{
+		FString Result;
+		FPackageName::TryConvertLongPackageNameToFilename(InPackageName, Result, Extension);
+		return Result;
+	}
+
+	FString UPackageTools::FilenameToPackageName(const FString& InFilename)
+	{
+		FString Result;
+    	FPackageName::TryConvertFilenameToLongPackageName(InFilename, Result);
+    	return Result;
+	}
+
 	UPackage* UPackageTools::FindOrCreatePackageForAssetType(const FName LongPackageName, UClass* AssetClass)
 	{
 		if (AssetClass)
diff --git a/Engine/Source/Editor/UnrealEd/Public/PackageTools.h b/Engine/Source/Editor/UnrealEd/Public/PackageTools.h
--- a/Engine/Source/Editor/UnrealEd/Public/PackageTools.h
+++ b/Engine/Source/Editor/UnrealEd/Public/PackageTools.h
@@ -181,6 +181,26 @@ public:
 	UFUNCTION(BlueprintCallable, Category = "Editor Scripting | Package Tools")
 	static FString SanitizePackageName(const FString& InPackageName);
 
+	/** 
+	 * Converts a long package name to a file name.
+	 * This can be called on package paths as well, provide no extension in that case.
+	 * Will return an empty string if it fails.
+	 * @param PackageName Long Package Name
+	 * @param Extension Package extension.
+	 * @return Package filename, or empty if it failed.
+	 */
+	UFUNCTION(BlueprintCallable, Category = "Editor Scripting | Package Tools", meta=(AdvancedDisplay=1))
+	static FString PackageNameToFilename(const FString& PackageName, const FString& Extension = TEXT(""));
+	
+	/** 
+	 * Tries to convert a given relative or absolute filename to a long package name or path starting with a root like /Game
+	 * This works on both package names and directories, and it does not validate that it actually exists on disk.
+	 * @param Filename Filename to convert.
+	 * @return Resulting long package name if the supplied filename properly maps to a long package root, empty string otherwise.
+	 */
+	UFUNCTION(BlueprintCallable, Category = "Editor Scripting | Package Tools")
+	static FString FilenameToPackageName(const FString& Filename);
+
 	/**
 	 * Find or create a package at the desired path
 	 * If a package already exist but was made for another type of asset,
