Quick launch fixes.
Only play quick launch sounds when they are enabled for PIE in ULevelEditorPlaySettings.
Prefix log lines piped from UAT to make it easier to view and filter the editor output log.

#rnx

diff --git a/Engine/Source/Editor/UATHelper/UATHelperModule.cpp b/Engine/Source/Editor/UATHelper/UATHelperModule.cpp
--- a/Engine/Source/Editor/UATHelper/UATHelperModule.cpp
+++ b/Engine/Source/Editor/UATHelper/UATHelperModule.cpp
@@ -91,13 +91,17 @@ public:
 	{
 		if ( NotificationItemPtr.IsValid() )
 		{
-			if ( CompletionState == SNotificationItem::CS_Fail )
+			const ULevelEditorPlaySettings* EditorPlaySettings = GetDefault<ULevelEditorPlaySettings>();
+			if (EditorPlaySettings->EnablePIEEnterAndExitSounds)
 			{
-				GEditor->PlayEditorSound(TEXT("/Engine/EditorSounds/Notifications/CompileFailed_Cue.CompileFailed_Cue"));
-			}
-			else
-			{
-				GEditor->PlayEditorSound(TEXT("/Engine/EditorSounds/Notifications/CompileSuccess_Cue.CompileSuccess_Cue"));
+				if ( CompletionState == SNotificationItem::CS_Fail )
+				{
+					GEditor->PlayEditorSound(TEXT("/Engine/EditorSounds/Notifications/CompileFailed_Cue.CompileFailed_Cue"));
+				}
+				else
+				{
+					GEditor->PlayEditorSound(TEXT("/Engine/EditorSounds/Notifications/CompileSuccess_Cue.CompileSuccess_Cue"));
+				}
 			}
 
 			TSharedPtr<SNotificationItem> NotificationItem = NotificationItemPtr.Pin();
@@ -482,7 +486,11 @@ public:
 		FEditorDelegates::OnShutdownPostPackagesSaved.Add(FSimpleDelegate::CreateStatic(&FUATHelperModule::HandleUatCancelButtonClicked, UatProcessPtr));
 
 		UatProcess->Launch();
-		GEditor->PlayEditorSound(TEXT("/Engine/EditorSounds/Notifications/CompileStart_Cue.CompileStart_Cue"));
+		const ULevelEditorPlaySettings* EditorPlaySettings = GetDefault<ULevelEditorPlaySettings>();
+		if (EditorPlaySettings->EnablePIEEnterAndExitSounds)
+		{
+			GEditor->PlayEditorSound(TEXT("/Engine/EditorSounds/Notifications/CompileStart_Cue.CompileStart_Cue"));
+		}
 	}
 
 	static void HandleUatHyperlinkNavigate()
@@ -650,7 +658,11 @@ public:
 
 	static void HandleUatLaunchFailed(TWeakPtr<SNotificationItem> NotificationItemPtr, FText PlatformDisplayName, FText TaskName, EventData Event)
 	{
-		GEditor->PlayEditorSound(TEXT("/Engine/EditorSounds/Notifications/CompileFailed_Cue.CompileFailed_Cue"));
+		const ULevelEditorPlaySettings* EditorPlaySettings = GetDefault<ULevelEditorPlaySettings>();
+		if (EditorPlaySettings->EnablePIEEnterAndExitSounds)
+		{
+			GEditor->PlayEditorSound(TEXT("/Engine/EditorSounds/Notifications/CompileFailed_Cue.CompileFailed_Cue"));
+		}
 
 		TGraphTask<FMainFrameActionsNotificationTask>::CreateTask().ConstructAndDispatchWhenReady(
 			NotificationItemPtr,
diff --git a/Engine/Source/Editor/UnrealEd/Private/PlayLevelViaLauncher.cpp b/Engine/Source/Editor/UnrealEd/Private/PlayLevelViaLauncher.cpp
--- a/Engine/Source/Editor/UnrealEd/Private/PlayLevelViaLauncher.cpp
+++ b/Engine/Source/Editor/UnrealEd/Private/PlayLevelViaLauncher.cpp
@@ -45,15 +45,15 @@ static void HandleOutputReceived(const FString& InMessage)
 {
 	if (InMessage.Contains(TEXT("Error:")))
 	{
-		UE_LOG(LogPlayLevel, Error, TEXT("%s"), *InMessage);
+		UE_LOG(LogPlayLevel, Error, TEXT("UAT: %s"), *InMessage);
 	}
 	else if (InMessage.Contains(TEXT("Warning:")))
 	{
-		UE_LOG(LogPlayLevel, Warning, TEXT("%s"), *InMessage);
+		UE_LOG(LogPlayLevel, Warning, TEXT("UAT: %s"), *InMessage);
 	}
 	else
 	{
-		UE_LOG(LogPlayLevel, Log, TEXT("%s"), *InMessage);
+		UE_LOG(LogPlayLevel, Log, TEXT("UAT: %s"), *InMessage);
 	}
 }
 
@@ -429,7 +429,10 @@ void UEditorEngine::StartPlayUsingLauncherSession(FRequestPlaySessionParams& InR
 	TWeakPtr<SNotificationItem> NotificationItemPtr(NotificationItem);
 	if (GEditor->LauncherWorker.IsValid() && GEditor->LauncherWorker->GetStatus() != ELauncherWorkerStatus::Completed)
 	{
-		GEditor->PlayEditorSound(TEXT("/Engine/EditorSounds/Notifications/CompileStart_Cue.CompileStart_Cue"));
+		if (EditorPlaySettings->EnablePIEEnterAndExitSounds)
+		{
+			GEditor->PlayEditorSound(TEXT("/Engine/EditorSounds/Notifications/CompileStart_Cue.CompileStart_Cue"));
+		}
 		GEditor->LauncherWorker->OnOutputReceived().AddStatic(HandleOutputReceived);
 		GEditor->LauncherWorker->OnStageStarted().AddUObject(this, &UEditorEngine::HandleStageStarted, NotificationItemPtr);
 		GEditor->LauncherWorker->OnStageCompleted().AddUObject(this, &UEditorEngine::HandleStageCompleted, LauncherSessionInfo->bPlayUsingLauncherHasCode, NotificationItemPtr);
@@ -439,7 +442,10 @@ void UEditorEngine::StartPlayUsingLauncherSession(FRequestPlaySessionParams& InR
 	else
 	{
 		GEditor->LauncherWorker.Reset();
-		GEditor->PlayEditorSound(TEXT("/Engine/EditorSounds/Notifications/CompileFailed_Cue.CompileFailed_Cue"));
+		if (EditorPlaySettings->EnablePIEEnterAndExitSounds)
+		{
+			GEditor->PlayEditorSound(TEXT("/Engine/EditorSounds/Notifications/CompileFailed_Cue.CompileFailed_Cue"));
+		}
 
 		NotificationItem->SetText(LOCTEXT("LauncherTaskFailedNotification", "Failed to launch task!"));
 		NotificationItem->SetCompletionState(SNotificationItem::CS_Fail);
@@ -513,13 +519,17 @@ public:
 	{
 		if (NotificationItemPtr.IsValid())
 		{
-			if (CompletionState == SNotificationItem::CS_Fail)
-			{
-				GEditor->PlayEditorSound(TEXT("/Engine/EditorSounds/Notifications/CompileFailed_Cue.CompileFailed_Cue"));
-			}
-			else if (CompletionState == SNotificationItem::CS_Success)
+			const ULevelEditorPlaySettings* EditorPlaySettings = GetDefault<ULevelEditorPlaySettings>();
+			if (EditorPlaySettings->EnablePIEEnterAndExitSounds)
 			{
-				GEditor->PlayEditorSound(TEXT("/Engine/EditorSounds/Notifications/CompileSuccess_Cue.CompileSuccess_Cue"));
+				if (CompletionState == SNotificationItem::CS_Fail)
+				{
+					GEditor->PlayEditorSound(TEXT("/Engine/EditorSounds/Notifications/CompileFailed_Cue.CompileFailed_Cue"));
+				}
+				else if (CompletionState == SNotificationItem::CS_Success)
+				{
+					GEditor->PlayEditorSound(TEXT("/Engine/EditorSounds/Notifications/CompileSuccess_Cue.CompileSuccess_Cue"));
+				}
 			}
 
 			TSharedPtr<SNotificationItem> NotificationItem = NotificationItemPtr.Pin();
