Add scoped markers to the DumpMaterialShaderTypes commandlet.

diff --git a/Engine/Source/Editor/UnrealEd/Private/Commandlets/DumpMaterialShaderTypes.cpp b/Engine/Source/Editor/UnrealEd/Private/Commandlets/DumpMaterialShaderTypes.cpp
--- a/Engine/Source/Editor/UnrealEd/Private/Commandlets/DumpMaterialShaderTypes.cpp
+++ b/Engine/Source/Editor/UnrealEd/Private/Commandlets/DumpMaterialShaderTypes.cpp
@@ -83,6 +83,8 @@ public:
 
 	void PrintHistogram(int TotalShaders)
 	{
+		TRACE_CPUPROFILER_EVENT_SCOPE(PrintHistogram);
+
 		if (ShaderTypeHistogram.Num() > 0)
 		{
 			{
@@ -150,6 +152,8 @@ public:
 
 	void PrintAlphabeticList()
 	{
+		TRACE_CPUPROFILER_EVENT_SCOPE(PrintAlphabeticList);
+
 		if (ShaderTypeHistogram.Num() > 0)
 		{
 			ShaderTypeHistogram.KeySort(TLess<FString>());
@@ -306,6 +310,8 @@ void PrintDebugShaderInfo(FShaderStatsGatheringContext& Output, const TArray<FDe
 
 int ProcessMaterials(const ITargetPlatform* TargetPlatform, const EShaderPlatform ShaderPlatform, FShaderStatsGatheringContext& Output, TArray<FAssetData>& MaterialList)
 {
+	TRACE_CPUPROFILER_EVENT_SCOPE(ProcessMaterials);
+
 	int TotalShaders = 0;
 
 	for (const FAssetData& AssetData : MaterialList)
@@ -335,6 +341,8 @@ int ProcessMaterials(const ITargetPlatform* TargetPlatform, const EShaderPlatfor
 
 int ProcessMaterialInstances(const ITargetPlatform* TargetPlatform, const EShaderPlatform ShaderPlatform, FShaderStatsGatheringContext& Output, TArray<FAssetData>& MaterialInstanceList)
 {
+	TRACE_CPUPROFILER_EVENT_SCOPE(ProcessMaterialInstances);
+
 	int TotalShaders = 0;
 
 	int StaticPermutations = 0;
@@ -389,6 +397,8 @@ int ProcessMaterialInstances(const ITargetPlatform* TargetPlatform, const EShade
 
 int ProcessGlobalShaders(const ITargetPlatform* TargetPlatform, const EShaderPlatform ShaderPlatform, FShaderStatsGatheringContext& Output)
 {
+	TRACE_CPUPROFILER_EVENT_SCOPE(ProcessGlobalShaders);
+
 	Output.Log(TEXT(""));
 	Output.Log(TEXT("Global Shaders"));
 
@@ -439,6 +449,8 @@ int ProcessGlobalShaders(const ITargetPlatform* TargetPlatform, const EShaderPla
 
 void ProcessForTargetAndShaderPlatform(const ITargetPlatform* TargetPlatform, const EShaderPlatform ShaderPlatform, const FString& Params, TArray<FAssetData>& MaterialList, TArray<FAssetData> MaterialInstanceList)
 {
+	TRACE_CPUPROFILER_EVENT_SCOPE(ProcessForTargetAndShaderPlatform);
+
 	const double StartTime = FPlatformTime::Seconds();
 
 	const FString TimeNow = FDateTime::Now().ToString();
@@ -493,47 +505,50 @@ int32 UDumpMaterialShaderTypesCommandlet::Main(const FString& Params)
 
 	const double AssetRegistryStart = FPlatformTime::Seconds();
 
-	UE_LOG(LogDumpMaterialShaderTypesCommandlet, Display, TEXT("Searching the asset registry for all assets..."));
-	IAssetRegistry& AssetRegistry = FModuleManager::Get().LoadModuleChecked<FAssetRegistryModule>("AssetRegistry").Get();
-	AssetRegistry.SearchAllAssets(true);
-
 	TArray<FAssetData> MaterialList;
 	TArray<FAssetData> MaterialInstanceList;
 
-	// Parse collection
-	FString CollectionName;
-	if (FParse::Value(*Params, TEXT("collection="), CollectionName, true))
 	{
-		if (!CollectionName.IsEmpty())
+		TRACE_CPUPROFILER_EVENT_SCOPE(UDumpMaterialShaderTypesCommandlet.AssetRegistryScan);
+
+		UE_LOG(LogDumpMaterialShaderTypesCommandlet, Display, TEXT("Searching the asset registry for all assets..."));
+		IAssetRegistry& AssetRegistry = FModuleManager::Get().LoadModuleChecked<FAssetRegistryModule>("AssetRegistry").Get();
+		AssetRegistry.SearchAllAssets(true);
+
+		// Parse collection
+		FString CollectionName;
+		if (FParse::Value(*Params, TEXT("collection="), CollectionName, true))
 		{
-			// Get the list of materials from a collection
-			FARFilter Filter;
-			Filter.PackagePaths.Add(FName(TEXT("/Game")));
-			Filter.bRecursivePaths = true;
-			Filter.ClassNames.Add(UMaterial::StaticClass()->GetFName());
+			if (!CollectionName.IsEmpty())
+			{
+				// Get the list of materials from a collection
+				FARFilter Filter;
+				Filter.PackagePaths.Add(FName(TEXT("/Game")));
+				Filter.bRecursivePaths = true;
+				Filter.ClassNames.Add(UMaterial::StaticClass()->GetFName());
 
-			FCollectionManagerModule& CollectionManagerModule = FCollectionManagerModule::GetModule();
-			CollectionManagerModule.Get().GetObjectsInCollection(FName(*CollectionName), ECollectionShareType::CST_All, Filter.ObjectPaths, ECollectionRecursionFlags::SelfAndChildren);
+				FCollectionManagerModule& CollectionManagerModule = FCollectionManagerModule::GetModule();
+				CollectionManagerModule.Get().GetObjectsInCollection(FName(*CollectionName), ECollectionShareType::CST_All, Filter.ObjectPaths, ECollectionRecursionFlags::SelfAndChildren);
 
-			AssetRegistry.GetAssets(Filter, MaterialList);
+				AssetRegistry.GetAssets(Filter, MaterialList);
 
-			Filter.ClassNames.Empty();
-			Filter.ClassNames.Add(UMaterialInstance::StaticClass()->GetFName());
-			Filter.ClassNames.Add(UMaterialInstanceConstant::StaticClass()->GetFName());
+				Filter.ClassNames.Empty();
+				Filter.ClassNames.Add(UMaterialInstance::StaticClass()->GetFName());
+				Filter.ClassNames.Add(UMaterialInstanceConstant::StaticClass()->GetFName());
 
-			AssetRegistry.GetAssets(Filter, MaterialInstanceList);
+				AssetRegistry.GetAssets(Filter, MaterialInstanceList);
+			}
 		}
-	}
-	else
-	{
-		if (!AssetRegistry.IsLoadingAssets())
+		else
 		{
-			AssetRegistry.GetAssetsByClass(UMaterial::StaticClass()->GetFName(), MaterialList, true);
-			AssetRegistry.GetAssetsByClass(UMaterialInstance::StaticClass()->GetFName(), MaterialInstanceList, true);
+			if (!AssetRegistry.IsLoadingAssets())
+			{
+				AssetRegistry.GetAssetsByClass(UMaterial::StaticClass()->GetFName(), MaterialList, true);
+				AssetRegistry.GetAssetsByClass(UMaterialInstance::StaticClass()->GetFName(), MaterialInstanceList, true);
+			}
 		}
 	}
 
-
 	const double AssetRegistryEnd = FPlatformTime::Seconds();
 	UE_LOG(LogDumpMaterialShaderTypesCommandlet, Display, TEXT("Asset scan took: %.3f"), AssetRegistryEnd - AssetRegistryStart);
 
