ZenCOTF: Added support for single threaded clients
#rnx

diff --git a/Engine/Source/Editor/UnrealEd/Private/Cooker/IoStoreCookOnTheFlyRequestManager.cpp b/Engine/Source/Editor/UnrealEd/Private/Cooker/IoStoreCookOnTheFlyRequestManager.cpp
--- a/Engine/Source/Editor/UnrealEd/Private/Cooker/IoStoreCookOnTheFlyRequestManager.cpp
+++ b/Engine/Source/Editor/UnrealEd/Private/Cooker/IoStoreCookOnTheFlyRequestManager.cpp
@@ -39,9 +39,9 @@ public:
 		Connected
 	};
 
-	using FRequestHandler = TFunction<bool(const FName&, const UE::Cook::FCookOnTheFlyRequest&, UE::Cook::FCookOnTheFlyResponse&)>;
+	using FRequestHandler = TFunction<bool(const FName&, bool, const UE::Cook::FCookOnTheFlyRequest&, UE::Cook::FCookOnTheFlyResponse&)>;
 
-	using FClientConnectionHandler = TFunction<bool(const FName&, EConnectionStatus)>;
+	using FClientConnectionHandler = TFunction<bool(const FName&, bool, EConnectionStatus)>;
 
 	using FFillRequest = TFunction<void(FArchive&)>;
 
@@ -218,6 +218,7 @@ private:
 		TAtomic<double> LastActivityTime{ 0 };
 		uint32 ClientId = 0;
 		FName PlatformName;
+		bool bIsSingleThreaded = false;
 	};
 
 	void ServerThreadEntry()
@@ -268,41 +269,12 @@ private:
 				for (auto It = Clients.CreateIterator(); It; ++It)
 				{
 					TUniquePtr<FClient>& Client = *It;
-
-					const double SecondsSinceLastActivity = FPlatformTime::Seconds() - Client->LastActivityTime.Load();
-
-					if (SecondsSinceLastActivity > HeartbeatTimeoutInSeconds && !Client->bIsProcessingRequest)
-					{
-						Client->LastActivityTime = FPlatformTime::Seconds();
-
-						UE_LOG(LogCookOnTheFly, Display, TEXT("Sending hearbeat message, ClientId='%d', Platform='%s', Address='%s', IdleTime='%.2llf's"),
-							Client->ClientId, *Client->PlatformName.ToString(), *Client->PeerAddr->ToString(true), HeartbeatTimeoutInSeconds);
-
-						FCookOnTheFlyMessage HeartbeatRequest(ECookOnTheFlyMessage::Heartbeat | ECookOnTheFlyMessage::Request);
-						FCookOnTheFlyMessageHeader& Header = HeartbeatRequest.GetHeader();
-
-						Header.MessageStatus = ECookOnTheFlyMessageStatus::Ok;
-						Header.SenderId = ServerSenderId;
-						Header.CorrelationId = Client->ClientId;
-						Header.Timestamp = FDateTime::UtcNow().GetTicks();
-
-						FBufferArchive RequestPayload;
-						RequestPayload.Reserve(HeartbeatRequest.TotalSize());
-						RequestPayload << HeartbeatRequest;
-
-						if (!FNFSMessageHeader::WrapAndSendPayload(RequestPayload, FSimpleAbstractSocket_FSocket(Client->Socket)))
-						{
-							Client->bIsRunning = false;
-							UE_LOG(LogCookOnTheFly, Display, TEXT("Heartbeat [Failed]"));
-						}
-					}
-
 					if (!Client->bIsRunning)
 					{
 						UE_LOG(LogCookOnTheFly, Display, TEXT("Closing connection to client on address '%s' (Id='%d', Platform='%s')"),
 							*Client->PeerAddr->ToString(true), Client->ClientId, *Client->PlatformName.ToString());
 
-						Options.HandleClientConnection(Client->PlatformName, EConnectionStatus::Disconnected);
+						Options.HandleClientConnection(Client->PlatformName, Client->bIsSingleThreaded, EConnectionStatus::Disconnected);
 
 						Client->Socket->Close();
 						Client->Thread.Wait();
@@ -361,23 +333,12 @@ private:
 		case ECookOnTheFlyMessage::Handshake:
 		{
 			ProcessHandshake(Client, Request, Response);
-			bRequestOk = Options.HandleClientConnection(Client.PlatformName, EConnectionStatus::Connected);
-			break;
-		}
-		case ECookOnTheFlyMessage::Heartbeat:
-		{
-			const bool bHeartbeatOk = Request.GetHeader().CorrelationId == Client.ClientId;
-
-			UE_LOG(LogCookOnTheFly, Display, TEXT("Heartbeat [%s], ClientId='%d', Platform='%s', Address='%s'"),
-				bHeartbeatOk ? TEXT("Ok") : TEXT("Failed"), Client.ClientId, *Client.PlatformName.ToString(), *Client.PeerAddr->ToString(true));
-
-			bRequestOk = bHeartbeatOk;
-			bIsResponse = true;
+			bRequestOk = Options.HandleClientConnection(Client.PlatformName, Client.bIsSingleThreaded, EConnectionStatus::Connected);
 			break;
 		}
 		default:
 		{
-			bRequestOk = Options.HandleRequest(Client.PlatformName, Request, Response);
+			bRequestOk = Options.HandleRequest(Client.PlatformName, Client.bIsSingleThreaded, Request, Response);
 			break;
 		}
 		}
@@ -412,6 +373,7 @@ private:
 			TUniquePtr<FArchive> Ar = HandshakeRequest.ReadBody();
 			*Ar << PlatformName;
 			*Ar << ProjectName;
+			*Ar << Client.bIsSingleThreaded;
 		}
 
 		if (PlatformName.Len())
@@ -491,6 +453,14 @@ private:
 		{
 		}
 
+		~FPlatformContext()
+		{
+			if (PackageCookedEvent)
+			{
+				FPlatformProcess::ReturnSynchEventToPool(PackageCookedEvent);
+			}
+		}
+
 		FCriticalSection& GetLock()
 		{
 			return CriticalSection;
@@ -583,6 +553,10 @@ private:
 			FPackage& Package = GetPackage(PackageId);
 			Package.Status = EPackageStatus::Failed;
 			OutCompletedPackages.FailedPackages.Add(PackageId);
+			if (PackageCookedEvent)
+			{
+				PackageCookedEvent->Trigger();
+			}
 		}
 
 		void MarkAsCooked(FPackageId PackageId, const FPackageStoreEntryResource& Entry, UE::ZenCookOnTheFly::Messaging::FCompletedPackages& OutCompletedPackages)
@@ -592,6 +566,10 @@ private:
 			Package.Status = EPackageStatus::Cooked;
 			Package.Entry = Entry;
 			OutCompletedPackages.CookedPackages.Add(Entry);
+			if (PackageCookedEvent)
+			{
+				PackageCookedEvent->Trigger();
+			}
 		}
 
 		void AddExistingPackages(TArrayView<const FPackageStoreEntryResource> Entries, TArrayView<const IPackageStoreWriter::FOplogCookInfo> CookInfos)
@@ -625,11 +603,38 @@ private:
 			return *Package;
 		}
 
+		void AddSingleThreadedClient()
+		{
+			++SingleThreadedClientsCount;
+			if (!PackageCookedEvent)
+			{
+				PackageCookedEvent = FPlatformProcess::GetSynchEventFromPool();
+			}
+		}
+
+		void RemoveSingleThreadedClient()
+		{
+			check(SingleThreadedClientsCount >= 0);
+			if (!SingleThreadedClientsCount)
+			{
+				FPlatformProcess::ReturnSynchEventToPool(PackageCookedEvent);
+				PackageCookedEvent = nullptr;
+			}
+		}
+
+		void WaitForCook()
+		{
+			check(PackageCookedEvent);
+			PackageCookedEvent->Wait();
+		}
+
 	private:
 		FCriticalSection CriticalSection;
 		FName PlatformName;
 		TMap<FPackageId, TUniquePtr<FPackage>> Packages;
 		IPackageStoreWriter* PackageWriter = nullptr;
+		int32 SingleThreadedClientsCount = 0;
+		FEvent* PackageCookedEvent = nullptr;
 	};
 
 	virtual bool Initialize() override
@@ -643,13 +648,13 @@ private:
 		ConnectionServer = MakeUnique<FIoStoreCookOnTheFlyNetworkServer>(FIoStoreCookOnTheFlyNetworkServer::FServerOptions
 		{
 			Port,
-			[this](const FName& PlatformName, FIoStoreCookOnTheFlyNetworkServer::EConnectionStatus ConnectionStatus)
+			[this](const FName& PlatformName, bool bIsSingleThreaded, FIoStoreCookOnTheFlyNetworkServer::EConnectionStatus ConnectionStatus)
 			{
-				return HandleClientConnection(PlatformName, ConnectionStatus);
+				return HandleClientConnection(PlatformName, bIsSingleThreaded, ConnectionStatus);
 			},
-			[this](const FName& PlatformName, const FCookOnTheFlyRequest& Request, FCookOnTheFlyResponse& Response)
+			[this](const FName& PlatformName, bool bIsSingleThreaded, const FCookOnTheFlyRequest& Request, FCookOnTheFlyResponse& Response)
 			{ 
-				return HandleClientRequest(PlatformName, Request, Response);
+				return HandleClientRequest(PlatformName, bIsSingleThreaded, Request, Response);
 			}
 		});
 
@@ -743,7 +748,7 @@ private:
 	}
 
 private:
-	bool HandleClientConnection(const FName& PlatformName, FIoStoreCookOnTheFlyNetworkServer::EConnectionStatus ConnectionStatus)
+	bool HandleClientConnection(const FName& PlatformName, bool bIsSingleThreaded, FIoStoreCookOnTheFlyNetworkServer::EConnectionStatus ConnectionStatus)
 	{
 		FScopeLock _(&ContextsCriticalSection);
 
@@ -774,6 +779,13 @@ private:
 					PackageWriter->OnMarkUpToDate().AddRaw(this, &FIoStoreCookOnTheFlyRequestManager::OnPackagesMarkedUpToDate);
 				}
 
+				if (bIsSingleThreaded)
+				{
+					FPlatformContext& Context = GetContext(PlatformName);
+					FScopeLock __(&Context.GetLock());
+					Context.AddSingleThreadedClient();
+				}
+
 				return true;
 			}
 
@@ -781,12 +793,18 @@ private:
 		}
 		else
 		{
+			if (bIsSingleThreaded)
+			{
+				FPlatformContext& Context = GetContext(PlatformName);
+				FScopeLock __(&Context.GetLock());
+				Context.RemoveSingleThreadedClient();
+			}
 			CookOnTheFlyServer.RemovePlatform(PlatformName);
 			return true;
 		}
 	}
 
-	bool HandleClientRequest(const FName& PlatformName, const UE::Cook::FCookOnTheFlyRequest& Request, UE::Cook::FCookOnTheFlyResponse& Response)
+	bool HandleClientRequest(const FName& PlatformName, bool bIsSingleThreaded, const UE::Cook::FCookOnTheFlyRequest& Request, UE::Cook::FCookOnTheFlyResponse& Response)
 	{
 		bool bRequestOk = false;
 
@@ -797,7 +815,7 @@ private:
 		switch (Request.GetHeader().MessageType)
 		{
 			case UE::Cook::ECookOnTheFlyMessage::CookPackage:
-				bRequestOk = HandleCookPackageRequest(PlatformName, Request, Response);
+				bRequestOk = HandleCookPackageRequest(PlatformName, bIsSingleThreaded, Request, Response);
 				break;
 			case UE::Cook::ECookOnTheFlyMessage::GetCookedPackages:
 				bRequestOk = HandleGetCookedPackagesRequest(PlatformName, Request, Response);
@@ -849,7 +867,7 @@ private:
 		return true;
 	}
 
-	bool HandleCookPackageRequest(const FName& PlatformName, const UE::Cook::FCookOnTheFlyRequest& Request, UE::Cook::FCookOnTheFlyResponse& Response)
+	bool HandleCookPackageRequest(const FName& PlatformName, bool bIsSingleThreaded, const UE::Cook::FCookOnTheFlyRequest& Request, UE::Cook::FCookOnTheFlyResponse& Response)
 	{
 		using namespace UE::ZenCookOnTheFly::Messaging;
 
@@ -872,9 +890,29 @@ private:
 				return AllKnownPackagesMap.FindRef(CookRequest.PackageId);
 			};
 
-			FScopeLock _(&Context.GetLock());
 			FPackageStoreEntryResource Entry;
-			EPackageStoreEntryStatus PackageStatus = Context.RequestCook(CookOnTheFlyServer, CookRequest.PackageId, GetPackageNameFunc, Entry);
+			EPackageStoreEntryStatus PackageStatus = EPackageStoreEntryStatus::Pending;
+			if (bIsSingleThreaded)
+			{
+				for (;;)
+				{
+					{
+						FScopeLock _(&Context.GetLock());
+						PackageStatus = Context.RequestCook(CookOnTheFlyServer, CookRequest.PackageId, GetPackageNameFunc, Entry);
+					}
+					if (PackageStatus != EPackageStoreEntryStatus::Pending)
+					{
+						break;
+					}
+					Context.WaitForCook();
+				}
+			}
+			else
+			{
+				FScopeLock _(&Context.GetLock());
+				PackageStatus = Context.RequestCook(CookOnTheFlyServer, CookRequest.PackageId, GetPackageNameFunc, Entry);
+			}
+
 			Response.SetBodyTo(FCookPackageResponse{ PackageStatus, MoveTemp(Entry) });
 		}
 		Response.SetStatus(UE::Cook::ECookOnTheFlyMessageStatus::Ok);
diff --git a/Engine/Source/Runtime/CookOnTheFly/Internal/CookOnTheFly.h b/Engine/Source/Runtime/CookOnTheFly/Internal/CookOnTheFly.h
--- a/Engine/Source/Runtime/CookOnTheFly/Internal/CookOnTheFly.h
+++ b/Engine/Source/Runtime/CookOnTheFly/Internal/CookOnTheFly.h
@@ -50,10 +50,8 @@ enum class ECookOnTheFlyMessage : uint32
 	PackagesCooked		= 0x50,
 	/* One way message indicating that one or more files has been added. */
 	FilesAdded			= 0x60,
-	/* Heartbeat message. */
-	Heartbeat			= 0x70,
 	/* Request to recook packages. */
-	RecookPackages		= 0x80,
+	RecookPackages		= 0x70,
 };
 ENUM_CLASS_FLAGS(ECookOnTheFlyMessage);
 
@@ -80,8 +78,6 @@ inline const TCHAR* LexToString(ECookOnTheFlyMessage Message)
 			return TEXT("PackagesCooked");
 		case ECookOnTheFlyMessage::FilesAdded:
 			return TEXT("FilesAdded");
-		case ECookOnTheFlyMessage::Heartbeat:
-			return TEXT("Heartbeat");
 		case ECookOnTheFlyMessage::RecookPackages:
 			return TEXT("RecookPackages");
 		default:
diff --git a/Engine/Source/Runtime/CookOnTheFly/Private/CookOnTheFlyServerConnection.cpp b/Engine/Source/Runtime/CookOnTheFly/Private/CookOnTheFlyServerConnection.cpp
--- a/Engine/Source/Runtime/CookOnTheFly/Private/CookOnTheFlyServerConnection.cpp
+++ b/Engine/Source/Runtime/CookOnTheFly/Private/CookOnTheFlyServerConnection.cpp
@@ -46,7 +46,9 @@ class FCookOnTheFlyServerConnection final
 {
 public:
 	FCookOnTheFlyServerConnection()
-	{ }
+		: bIsSingleThreaded(!FGenericPlatformProcess::SupportsMultithreading())
+	{
+	}
 
 	~FCookOnTheFlyServerConnection()
 	{
@@ -113,8 +115,18 @@ public:
 			return false;
 		}
 
-		Thread = AsyncThread([this] { return ThreadEntry(); },  8 * 1024, TPri_Normal);
-
+		Thread = MakeUnique<FThread>(
+			TEXT("CotfServerConnection"),
+			[this]()
+			{
+				ThreadEntry();
+			},
+			[this]()
+			{
+				SingleThreadedTickFunction();
+			},
+			8 * 1024);
+		
 		UE_LOG(LogCotfServerConnection, Log, TEXT("Connected to COTF server at '%s'"), *ServerAddr->ToString(true));
 
 		return true;
@@ -153,6 +165,13 @@ public:
 
 		if (SendMessage(RequestPayload))
 		{
+			if (bIsSingleThreaded)
+			{
+				ProcessMessagesWhile([this, CorrelationId]()
+					{
+						return PendingRequests.Contains(CorrelationId);
+					});
+			}
 			return FutureResponse;
 		}
 		else
@@ -177,7 +196,7 @@ public:
 			if (!bStopRequested.Exchange(true))
 			{
 				Socket->Close();
-				Thread.Wait();
+				Thread->Join();
 				Thread.Reset();
 				Socket.Reset();
 				ClientId  = 0;
@@ -232,6 +251,8 @@ public:
 			TUniquePtr<FArchive> Ar = HandshakeRequest.WriteBody();
 			*Ar << PlatformName;
 			*Ar << ProjectName;
+			bool bLocalIsSingleThreded = bIsSingleThreaded;
+			*Ar << bLocalIsSingleThreded;
 		}
 
 		FBufferArchive HandshakeRequestPayload;
@@ -260,17 +281,18 @@ public:
 		return ClientId > 0;
 	}
 
-	void ThreadEntry()
+	bool ProcessMessagesWhile(TFunctionRef<bool()> Condition)
 	{
 		using namespace UE::Cook;
 
-		while (!bStopRequested.Load())
+		while (Condition())
 		{
 			FArrayReader MessagePayload;
 			if (!ReceiveMessage(MessagePayload))
 			{
 				UE_LOG(LogCotfServerConnection, Warning, TEXT("Failed to receive message from '%s'"), *ServerAddr->ToString(true));
-				break;
+				ClientId = 0;
+				return false;
 			}
 
 			FCookOnTheFlyMessageHeader MessageHeader;
@@ -287,28 +309,7 @@ public:
 
 			if (bIsRequest)
 			{
-				UE_CLOG(
-					MessageHeader.MessageType != ECookOnTheFlyMessage::Heartbeat,
-					LogCotfServerConnection, Fatal, TEXT("Invalid server request message '%s'"), LexToString(MessageHeader.MessageType));
-
-				FCookOnTheFlyMessage HeartbeatResponse(ECookOnTheFlyMessage::Heartbeat | ECookOnTheFlyMessage::Response);
-				FCookOnTheFlyMessageHeader& ResponseHeader = HeartbeatResponse.GetHeader();
-
-				ResponseHeader.MessageStatus	= ECookOnTheFlyMessageStatus::Ok;
-				ResponseHeader.SenderId			= ClientId;
-				ResponseHeader.CorrelationId	= MessageHeader.CorrelationId;
-				ResponseHeader.Timestamp		= FDateTime::UtcNow().GetTicks();
-				
-				FBufferArchive ResponsePayload;
-				ResponsePayload << HeartbeatResponse;
-
-				UE_LOG(LogCotfServerConnection, Warning, TEXT("Sending heartbeat response to '%s'"), *ServerAddr->ToString(true));
-
-				if (!SendMessage(ResponsePayload))
-				{
-					UE_LOG(LogCotfServerConnection, Warning, TEXT("Failed to send heartbeat response to '%s'"), *ServerAddr->ToString(true));
-					break;
-				}
+				UE_LOG(LogCotfServerConnection, Warning, TEXT("Received request from server, ignoring"));
 			}
 			else if (bIsResponse)
 			{
@@ -337,10 +338,29 @@ public:
 			}
 		}
 
+		return true;
+	}
+
+	void ThreadEntry()
+	{
+		ProcessMessagesWhile([this]()
+			{
+				return !bStopRequested.Load();
+			});
+
 		UE_LOG(LogCotfServerConnection, Display, TEXT("Terminating connection to server '%s'"), *ServerAddr->ToString(true));
 		ClientId = 0;
 	}
 
+	void SingleThreadedTickFunction()
+	{
+		ProcessMessagesWhile([this]()
+			{
+				uint32 PendingDataSize;
+				return Socket->HasPendingData(PendingDataSize);
+			});
+	}
+
 	struct FPendingRequest
 	{
 		UE::Cook::FCookOnTheFlyMessageHeader RequestHeader;
@@ -376,8 +396,9 @@ public:
 	TSharedPtr<FInternetAddr> ServerAddr;
 	TUniquePtr<FSocket> Socket;
 	uint32 ClientId = 0;
-	TFuture<void> Thread;
+	TUniquePtr<FThread> Thread;
 	TAtomic<bool> bStopRequested { false };
+	const bool bIsSingleThreaded;
 
 	FCriticalSection RequestsCriticalSection;
 	TMap<uint32, TUniquePtr<FPendingRequest>> PendingRequests;
