Fixed -iterate never cooking iteratively for COTF
#rnx

diff --git a/Engine/Source/Editor/UnrealEd/Private/CookOnTheFlyServer.cpp b/Engine/Source/Editor/UnrealEd/Private/CookOnTheFlyServer.cpp
--- a/Engine/Source/Editor/UnrealEd/Private/CookOnTheFlyServer.cpp
+++ b/Engine/Source/Editor/UnrealEd/Private/CookOnTheFlyServer.cpp
@@ -631,6 +631,10 @@ void UCookOnTheFlyServer::AddCookOnTheFlyPlatformFromGameThread(ITargetPlatform*
 	Generator.SaveAssetRegistry(GetSandboxAssetRegistryFilename(), true);
 	check(PlatformData->bIsSandboxInitialized); // This should have been set by BeginCookSandbox, and it is what we use to determine whether a platform has been initialized
 
+	// This will miss settings that are accessed during the cook
+	// TODO: A better way of handling ini settings
+	SaveCurrentIniSettings(TargetPlatform);
+
 	FindOrCreatePackageWriter(TargetPlatform).BeginCook();
 }
 
@@ -5036,7 +5040,7 @@ void GetVersionFormatNumbersForIniVersionStrings(TMap<FString, FString>& IniVers
 }
 
 
-void GetAdditionalCurrentIniVersionStrings( const ITargetPlatform* TargetPlatform, TMap<FString, FString>& IniVersionMap )
+void GetAdditionalCurrentIniVersionStrings( const UCookOnTheFlyServer* CookOnTheFlyServer, const ITargetPlatform* TargetPlatform, TMap<FString, FString>& IniVersionMap )
 {
 	FConfigFile EngineSettings;
 	FConfigCacheIni::LoadLocalIniFile(EngineSettings, TEXT("Engine"), true, *TargetPlatform->IniPlatformName());
@@ -5148,6 +5152,9 @@ void GetAdditionalCurrentIniVersionStrings( const ITargetPlatform* TargetPlatfor
 
 	IniVersionMap.Add(TEXT("MaterialShaderMapDDCVersion"), *GetMaterialShaderMapDDCKey());
 	IniVersionMap.Add(TEXT("GlobalDDCVersion"), *GetGlobalShaderMapDDCKey());
+
+	UProjectPackagingSettings* PackagingSettings = Cast<UProjectPackagingSettings>(UProjectPackagingSettings::StaticClass()->GetDefaultObject());
+	IniVersionMap.Add(TEXT("IsUsingShaderCodeLibrary"), FString::Printf(TEXT("%d"), PackagingSettings->bShareMaterialShaderCode && CookOnTheFlyServer->IsUsingShaderCodeLibrary()));
 }
 
 
@@ -5494,7 +5501,7 @@ bool UCookOnTheFlyServer::IniSettingsOutOfDate(const ITargetPlatform* TargetPlat
 
 	// compare against current settings
 	TMap<FString, FString> CurrentAdditionalSettings;
-	GetAdditionalCurrentIniVersionStrings(TargetPlatform, CurrentAdditionalSettings);
+	GetAdditionalCurrentIniVersionStrings(this, TargetPlatform, CurrentAdditionalSettings);
 
 	for ( const auto& OldIniSetting : OldAdditionalSettings)
 	{
@@ -5626,7 +5633,7 @@ bool UCookOnTheFlyServer::SaveCurrentIniSettings(const ITargetPlatform* TargetPl
 	TGuardValue<bool> S(IniSettingRecurse, true);
 
 	TMap<FString, FString> AdditionalIniSettings;
-	GetAdditionalCurrentIniVersionStrings(TargetPlatform, AdditionalIniSettings);
+	GetAdditionalCurrentIniVersionStrings(this, TargetPlatform, AdditionalIniSettings);
 
 	FIniSettingContainer CurrentIniSettings;
 	GetCurrentIniVersionStrings(TargetPlatform, CurrentIniSettings);
