Fixed precision issue when computing the dimensions of a sphyl
  - Root cause: A faulty comparison between a float and a double was screwing thecomputation of the orientation of the sphyl
  - Solution: Compare two doubles instead

diff --git a/Engine/Source/Editor/UnrealEd/Private/GeomFitUtils.cpp b/Engine/Source/Editor/UnrealEd/Private/GeomFitUtils.cpp
--- a/Engine/Source/Editor/UnrealEd/Private/GeomFitUtils.cpp
+++ b/Engine/Source/Editor/UnrealEd/Private/GeomFitUtils.cpp
@@ -444,7 +444,7 @@ static void CalcBoundingSphyl(const FMeshDescription* MeshDescription, FSphere&
 	sphere.Center = Center;
 
 	// Work out best axis aligned orientation (longest side)
-	float Extent = Extents.GetMax();
+	double Extent = Extents.GetMax();
 	if (Extent == Extents.X)
 	{
 		rotation = FRotator(90.f, 0.f, 0.f);
