Fix UPackageTools::UnloadPackages to not potentially cause asserts when unloading a package with async compilation work running.

#rnx

- There is the potential here to call ResetLoader on a package that has async compilation work running, which currently is not thread safe. In addition, calling ResetLoader at this point can force bulkdata payloads to be loaded off disk only for the package to be immediately destroyed via garbage collection.
- The call to ::ResetLoader has been removed, we now rely on the loader being reset as part of the garbage collection process, which will not force the payload to be loaded.
- Note that if a package is garbage collected there is already code to make sure that it waits on it's async compilation tasks (if any) and so should be thread safe.
- As a safety precaution we retain an array of weak pointers to the packages and check them after the garbage collection pass to see if any of the packages survived. If they did survive we then call ::ResetLoader on them to detach the package from the file on disk, maintaining the existing behaviour.
-- Before we call ::ResetLoader we check to see if any of the packages has outstanding async compilation work and if so we flush the compilation manager to avoid any potential threading issues.

diff --git a/Engine/Source/Editor/UnrealEd/Private/PackageTools.cpp b/Engine/Source/Editor/UnrealEd/Private/PackageTools.cpp
--- a/Engine/Source/Editor/UnrealEd/Private/PackageTools.cpp
+++ b/Engine/Source/Editor/UnrealEd/Private/PackageTools.cpp
@@ -72,6 +72,73 @@ TSet<UPackage*>* UPackageTools::PackagesBeingUnloaded = nullptr;
 TSet<UObject*> UPackageTools::ObjectsThatHadFlagsCleared;
 FDelegateHandle UPackageTools::ReachabilityCallbackHandle;
 
+namespace
+{
+
+/** 
+ * Utility function that checks each UObject inside of the given UPackage to see if it is waiting 
+ * on async compilation.
+ * 
+ * @return true is the package contains at least one UObject that has compilation work running, otherwise false.
+ */
+static bool IsPackageCompiling(const UPackage* Package)
+{
+	bool bIsCompiling = false;
+	ForEachObjectWithPackage(Package, [&bIsCompiling](const UObject* Object)
+	{
+		const IInterface_AsyncCompilation* AsyncCompilationIF = Cast<IInterface_AsyncCompilation>(Object);
+		if (AsyncCompilationIF != nullptr && AsyncCompilationIF->IsCompiling())
+		{
+			bIsCompiling = true;
+			return false;
+		}
+		else
+		{
+			return true;
+		}
+	});
+
+	return bIsCompiling;
+}
+
+/**
+ * Utility function to reset the loaders for an array of packages safely.
+ * 
+ * If any of the packages contains an object with running async compilation
+ * tasks then we need to wait for all existing async compilation tasks to
+ * be flushed before we can call ResetLoaders, due to thread safety issues.
+ */
+static void ResetPackageLoaders(const TArray<TWeakObjectPtr<UPackage>>& PackagesToUnload)
+{
+	bool bHasAsyncCompilationWork = false;
+	for (const TWeakObjectPtr<UPackage>& Package : PackagesToUnload)
+	{
+		if (Package.IsValid())
+		{
+			if (IsPackageCompiling(Package.Get()))
+			{
+				bHasAsyncCompilationWork = true;
+				break;
+			}
+		}
+	}
+
+	if (bHasAsyncCompilationWork)
+	{
+		FAssetCompilingManager::Get().FinishAllCompilation();
+	}
+
+	for (const TWeakObjectPtr<UPackage>& Package : PackagesToUnload)
+	{
+		if (Package.IsValid())
+		{
+			ResetLoaders(Package.Get());
+		}
+	}
+}
+
+}
+
 UPackageTools::UPackageTools(const FObjectInitializer& ObjectInitializer)
 	: Super(ObjectInitializer)
 {
@@ -367,6 +434,14 @@ UPackageTools::UPackageTools(const FObjectInitializer& ObjectInitializer)
 			FlushAsyncLoading();
 			(*GFlushStreamingFunc)();
 
+			// Store the packages as weak pointers so that we can identify which ones 
+			// survive garbage collection later on
+			TArray<TWeakObjectPtr<UPackage>> PackagesToReset;
+			for (UPackage* Package : PackagesToUnload)
+			{
+				PackagesToReset.Add(Package);
+			}
+
 			// Remove potential references to to-be deleted objects from the GB selection set.
 			GEditor->GetSelectedObjects()->DeselectAll();
 
@@ -447,9 +522,6 @@ UPackageTools::UPackageTools(const FObjectInitializer& ObjectInitializer)
 					ObjectsInPackage.Reset();
 				}
 
-				// Reset loaders
-				ResetLoaders(PackageBeingUnloaded);
-
 				if( PackageBeingUnloaded->IsDirty() )
 				{
 					// The package was marked dirty as a result of something that happened above (e.g callbacks in CollectGarbage).  
@@ -501,6 +573,10 @@ UPackageTools::UPackageTools(const FObjectInitializer& ObjectInitializer)
 
 			CollectGarbage( GARBAGE_COLLECTION_KEEPFLAGS );
 
+			// We need to reset the loader of any of the package that survived garbage collection to  
+			// make sure that they are detached from their file on disk.
+			ResetPackageLoaders(PackagesToReset);
+			
 			// Restore the standalone flag on any metadata objects that survived the GC
 			for ( const TWeakObjectPtr<UMetaData>& WeakPackageMetaData : PackageMetaDataWithClearedStandaloneFlag )
 			{
