BulkDataRegistry: Fix bad assumption in FPendingPackage::OnBulkDataListResults that the current PendingPackage must always be present in Owner->PendingPackages. It can be removed from that list during the Owner's shutdown, before the callback is called.
#rnx

diff --git a/Engine/Source/Editor/UnrealEd/Private/EditorDomain/BulkDataRegistryEditorDomain.cpp b/Engine/Source/Editor/UnrealEd/Private/EditorDomain/BulkDataRegistryEditorDomain.cpp
--- a/Engine/Source/Editor/UnrealEd/Private/EditorDomain/BulkDataRegistryEditorDomain.cpp
+++ b/Engine/Source/Editor/UnrealEd/Private/EditorDomain/BulkDataRegistryEditorDomain.cpp
@@ -562,7 +562,10 @@ void FPendingPackage::OnBulkDataListResults(FSharedBuffer Buffer)
 		TUniquePtr<FPendingPackage> ThisPointer;
 		{
 			FScopeLock PendingPackageScopeLock(&Owner->PendingPackageLock);
-			ThisPointer = Owner->PendingPackages.FindAndRemoveChecked(PackageName);
+			Owner->PendingPackages.RemoveAndCopyValue(PackageName, ThisPointer);
+			// PackageName may have already been removed from PendingPackages if Owner is shutting down. In that case,
+			// ThisPointer will be null, and Owner holds the UniquePtr to *this.
+			// It will release that UniquePtr only after Cancel returns, which will be after this function returns.
 		}
 		WriteCache();
 
@@ -571,7 +574,7 @@ void FPendingPackage::OnBulkDataListResults(FSharedBuffer Buffer)
 		// currently in. Direct the owner to keep requests alive to avoid that deadlock.
 		BulkDataListCacheRequest.KeepAlive();
 		ThisPointer.Reset();
-		// *this has been deleted and can no longer be accessed
+		// *this may have been deleted and can no longer be accessed
 	}
 }
 
