Avoid FAssetData reporting that a UBlueprintGeneratedClass is actually a UBlueprint causing the wrong ActorFactory to be selected

diff --git a/Engine/Source/Editor/UnrealEd/Private/AssetSelection.cpp b/Engine/Source/Editor/UnrealEd/Private/AssetSelection.cpp
--- a/Engine/Source/Editor/UnrealEd/Private/AssetSelection.cpp
+++ b/Engine/Source/Editor/UnrealEd/Private/AssetSelection.cpp
@@ -837,7 +837,7 @@ AActor* FActorFactoryAssetProxy::AddActorForAsset( UObject* AssetObj, bool Selec
 {
 	AActor* Result = NULL;
 
-	const FAssetData AssetData( AssetObj );
+	const FAssetData AssetData(AssetObj, FAssetData::ECreationFlags::AllowBlueprintClass);
 	FText UnusedErrorMessage;
 	if ( AssetObj != NULL )
 	{
