Only mark new blueprint variables as MultiLine if they are a type that supports MultiLine

diff --git a/Engine/Source/Editor/UnrealEd/Private/Kismet2/BlueprintEditorUtils.cpp b/Engine/Source/Editor/UnrealEd/Private/Kismet2/BlueprintEditorUtils.cpp
--- a/Engine/Source/Editor/UnrealEd/Private/Kismet2/BlueprintEditorUtils.cpp
+++ b/Engine/Source/Editor/UnrealEd/Private/Kismet2/BlueprintEditorUtils.cpp
@@ -4800,7 +4800,10 @@ bool FBlueprintEditorUtils::AddMemberVariable(UBlueprint* Blueprint, const FName
 	NewVar.VarType.bIsReference   = false;
 
 	// Text variables, etc. should default to multiline
-	NewVar.SetMetaData(TEXT("MultiLine"), TEXT("true"));
+	if (NewVarType.PinCategory == UEdGraphSchema_K2::PC_String || NewVarType.PinCategory == UEdGraphSchema_K2::PC_Text)
+	{
+		NewVar.SetMetaData(TEXT("MultiLine"), TEXT("true"));
+	}
 
 	Blueprint->NewVariables.Add(NewVar);
 
