PackageUtilities::NormalizePathNames: If a member of System.Core:Paths is not mounted, log a warning about it and do not include packages under that directory in the results.
#rnx

diff --git a/Engine/Source/Editor/UnrealEd/Private/Commandlets/PackageUtilities.cpp b/Engine/Source/Editor/UnrealEd/Private/Commandlets/PackageUtilities.cpp
--- a/Engine/Source/Editor/UnrealEd/Private/Commandlets/PackageUtilities.cpp
+++ b/Engine/Source/Editor/UnrealEd/Private/Commandlets/PackageUtilities.cpp
@@ -111,11 +111,20 @@ bool NormalizePackageNames( TArray<FString> PackageNames, TArray<FString>& Packa
 		TArray<FString> Paths;
 		if ( GConfig->GetArray( TEXT("Core.System"), TEXT("Paths"), Paths, GEngineIni ) > 0 )
 		{
-			for ( int32 i = 0; i < Paths.Num(); i++ )
+			TStringBuilder<256> UnusedPackagePath;
+			TStringBuilder<256> UnusedFilePath;
+			TStringBuilder<256> UnusedRelPath;
+			for ( const FString& Path : Paths)
 			{
-				FString SearchWildcard = Paths[i] / PackageWildcard;
+				if (!FPackageName::TryGetMountPointForPath(Path, UnusedPackagePath, UnusedFilePath, UnusedRelPath))
+				{
+					UE_LOG(LogPackageUtilities, Warning,
+						TEXT("Engine.ini:[Core.System]:Paths entry '%s' is not mounted. Skipping it."), *Path);
+					continue;
+				}
+				FString SearchWildcard = Path / PackageWildcard;
 				UE_LOG(LogPackageUtilities, Log, TEXT("Searching using wildcard: '%s'"), *SearchWildcard);
-				SearchDirectoryRecursive( SearchWildcard, PackageNames, PackagePathNames );
+				SearchDirectoryRecursive(SearchWildcard, PackageNames, PackagePathNames);
 			}
 		}
 
