Fixed a crash due to the selection containing an actor that was unloaded when exiting a pie session.

diff --git a/Engine/Source/Editor/UnrealEd/Private/PlayLevel.cpp b/Engine/Source/Editor/UnrealEd/Private/PlayLevel.cpp
--- a/Engine/Source/Editor/UnrealEd/Private/PlayLevel.cpp
+++ b/Engine/Source/Editor/UnrealEd/Private/PlayLevel.cpp
@@ -200,7 +200,7 @@ void UEditorEngine::EndPlayMap()
 	EnableScreenSaver( true );
 
 	// Make a list of all the actors that should be selected
-	TArray<UObject *> SelectedActors;
+	TArray<TWeakObjectPtr<AActor>> SelectedActors;
 	if ( ActorsThatWereSelected.Num() > 0 )
 	{
 		for ( int32 ActorIndex = 0; ActorIndex < ActorsThatWereSelected.Num(); ++ActorIndex )
@@ -208,7 +208,7 @@ void UEditorEngine::EndPlayMap()
 			TWeakObjectPtr<AActor> Actor = ActorsThatWereSelected[ ActorIndex ].Get();
 			if (Actor.IsValid())
 			{
-				SelectedActors.Add( Actor.Get() );
+				SelectedActors.Add( Actor );
 			}
 		}
 		ActorsThatWereSelected.Empty();
@@ -238,19 +238,6 @@ void UEditorEngine::EndPlayMap()
 	GetSelectedObjects()->DeselectAll();
 	GetSelectedComponents()->DeselectAll();
 
-	// For every actor that was selected previously, make sure it's editor equivalent is selected
-	GEditor->GetSelectedActors()->BeginBatchSelectOperation();
-	for ( int32 ActorIndex = 0; ActorIndex < SelectedActors.Num(); ++ActorIndex )
-	{
-		AActor* Actor = Cast<AActor>( SelectedActors[ ActorIndex ] );
-		if (Actor)
-		{
-			// We need to notify or else the manipulation transform widget won't appear, but only notify once at the end because OnEditorSelectionChanged is expensive for large groups. 
-			SelectActor( Actor, true, false );
-		}
-	}	
-	GEditor->GetSelectedActors()->EndBatchSelectOperation(true);
-
 	// let the editor know
 	FEditorDelegates::EndPIE.Broadcast(bIsSimulatingInEditor);
 
@@ -477,6 +464,19 @@ void UEditorEngine::EndPlayMap()
 	FWorldContext& EditorWorldContext = GEditor->GetEditorWorldContext();
 	UPackage* Package = EditorWorldContext.World()->GetOutermost();
 
+	// For every actor that was selected previously, make sure it's editor equivalent is selected. We do that after the cleanup in case some actor where removed.
+	GEditor->GetSelectedActors()->BeginBatchSelectOperation();
+	for (int32 ActorIndex = 0; ActorIndex < SelectedActors.Num(); ++ActorIndex)
+	{
+		AActor* Actor = SelectedActors[ActorIndex].Get();
+		if (Actor)
+		{
+			// We need to notify or else the manipulation transform widget won't appear, but only notify once at the end because OnEditorSelectionChanged is expensive for large groups. 
+			SelectActor(Actor, true, false);
+		}
+	}
+	GEditor->GetSelectedActors()->EndBatchSelectOperation(true);
+
 	// Spawn note actors dropped in PIE.
 	if(GEngine->PendingDroppedNotes.Num() > 0)
 	{
