Fixes for CIS localization warnings.
	- Introduced the FMetasoundMacroDescriptor to UGatherTextFromSourceCommandlet to support and parse METASOUND_PARAM() macros and fix related warnings.
	- Fixed a bug in UGatherTextFromSourceCommandlet where macros that look like the below example are not parsed properly. The ',' on the second line
	does not cause the previous parameter to be registered and concatenantes Param2 and Param3.
	MYMACRO(Param1, "Param2"
		, "Param3")
	-Added LOC_DEFINE_REGION guards to MetasoundParamHelper to prevent the macro from being gathered and generate warnings.
	- Fixed support for C++11 raw strings to have delimiters that do not start with an alpha numberic character e.g _JSON as a delimiter.
	#rnx

diff --git a/Engine/Plugins/Runtime/Metasound/Source/MetasoundGraphCore/Public/MetasoundParamHelper.h b/Engine/Plugins/Runtime/Metasound/Source/MetasoundGraphCore/Public/MetasoundParamHelper.h
--- a/Engine/Plugins/Runtime/Metasound/Source/MetasoundGraphCore/Public/MetasoundParamHelper.h
+++ b/Engine/Plugins/Runtime/Metasound/Source/MetasoundGraphCore/Public/MetasoundParamHelper.h
@@ -46,10 +46,12 @@
 */
 
 #if WITH_EDITOR
+#define LOC_DEFINE_REGION
 #define METASOUND_PARAM(NAME, NAME_TEXT, TOOLTIP_TEXT) \
 	static const TCHAR* NAME##Name = TEXT(NAME_TEXT); \
 	static const FText NAME##Tooltip = LOCTEXT(#NAME "Tooltip", TOOLTIP_TEXT); \
 	static const FText NAME##DisplayName = LOCTEXT(#NAME "DisplayName", NAME_TEXT); 
+#undef LOC_DEFINE_REGION
 #else 
 #define METASOUND_PARAM(NAME, NAME_TEXT, TOOLTIP_TEXT) \
 	static const TCHAR* NAME##Name = TEXT(NAME_TEXT); \
diff --git a/Engine/Source/Editor/UnrealEd/Classes/Commandlets/GatherTextFromSourceCommandlet.h b/Engine/Source/Editor/UnrealEd/Classes/Commandlets/GatherTextFromSourceCommandlet.h
--- a/Engine/Source/Editor/UnrealEd/Classes/Commandlets/GatherTextFromSourceCommandlet.h
+++ b/Engine/Source/Editor/UnrealEd/Classes/Commandlets/GatherTextFromSourceCommandlet.h
@@ -314,7 +314,7 @@ private:
 		}
 
 		virtual const FString& GetToken() const override
-		{
+			{
 			return Name;
 		}
 
@@ -366,6 +366,25 @@ private:
 
 		virtual void TryParse(const FString& Text, FSourceFileParseContext& Context) const override;
 	};
+/** Macro descriptor to parse METASOUND_PARAM(NAME, NAME_TEXT, TOOLTIP_TEXT) macros. */
+	class FMetasoundParamMacroDescriptor : public FMacroDescriptor
+	{
+	public:
+		FMetasoundParamMacroDescriptor()
+			: FMacroDescriptor(TEXT("METASOUND_PARAM"), 3)
+		{
+		}
+
+		virtual void TryParse(const FString& Text, FSourceFileParseContext& Context) const override;
+
+	protected:
+		FMetasoundParamMacroDescriptor(FString InName, int32 InMinNumberOfArgument)
+			: FMacroDescriptor(MoveTemp(InName), InMinNumberOfArgument)
+		{
+		}
+
+		void TryParseArgs(const FString& Text, FSourceFileParseContext& Context, const TArray<FString>& Arguments, const int32 ArgIndexOffset) const;
+	};
 
 	class FStringMacroDescriptor : public FMacroDescriptor
 	{
diff --git a/Engine/Source/Editor/UnrealEd/Private/Commandlets/GatherTextFromSourceCommandlet.cpp b/Engine/Source/Editor/UnrealEd/Private/Commandlets/GatherTextFromSourceCommandlet.cpp
--- a/Engine/Source/Editor/UnrealEd/Private/Commandlets/GatherTextFromSourceCommandlet.cpp
+++ b/Engine/Source/Editor/UnrealEd/Private/Commandlets/GatherTextFromSourceCommandlet.cpp
@@ -243,6 +243,8 @@ int32 UGatherTextFromSourceCommandlet::Main( const FString& Params )
 
 	Parsables.Add(new FUICommandExtMacroDescriptor());
 
+	Parsables.Add(new FMetasoundParamMacroDescriptor());
+	
 	// New Localization System with Namespace as literal argument.
 	Parsables.Add(new FStringMacroDescriptor( FString(TEXT("NSLOCTEXT")),
 		FStringMacroDescriptor::FMacroArg(FStringMacroDescriptor::MAS_Namespace, true),
@@ -828,7 +830,8 @@ bool UGatherTextFromSourceCommandlet::ParseSourceText(const FString& Text, const
 									{
 										break;
 									}
-									if (DelimChar == 0 || !FChar::IsAlnum(DelimChar))
+									// We also permit '_' to support the use of _JSON as a delimiter for the raw strings 
+									if (DelimChar == 0 || !(FChar::IsAlnum(DelimChar) || DelimChar == TCHAR('_')))
 									{
 										bIsValid = false;
 										break;
@@ -1071,7 +1074,7 @@ int32 UGatherTextFromSourceCommandlet::FMacroArgumentGatherer::GetNumberOfArgume
 
 
 bool UGatherTextFromSourceCommandlet::FMacroArgumentGatherer::EndArgument()
-{	
+{
 	if (CurrentArgument.IsEmpty() || IsInDoubleQuotes())
 	{
 		return false;
@@ -1614,6 +1617,18 @@ bool UGatherTextFromSourceCommandlet::FMacroDescriptor::ParseArgumentString(cons
 
 					ArgStart = Cursor + 1;
 				}
+				else
+				{
+					// The ',' character is the first thing in the line. We need to close out the previous argument 
+					// E.g MYMACRO(Param1, "Param2"
+					//	    , "Param3")
+					if (!ArgsGatherer.EndArgument())
+					{
+						UE_LOG(LogGatherTextFromSourceCommandlet, Warning, TEXT("Parsing Arguments failed in %s macro while parsing %s:%d. %s"), *GetToken(), *Context.Filename, Context.LineNumber, *FLocTextHelper::SanitizeLogOutput(Context.LineText.TrimStartAndEnd()));
+						return false;
+					}
+					ArgStart = Cursor + 1;
+				}
 			}
 		}
 		else if ('\\' == *Cursor)
@@ -1869,6 +1884,115 @@ void UGatherTextFromSourceCommandlet::FUICommandExtMacroDescriptor::TryParse(con
 	}
 }
 
+void UGatherTextFromSourceCommandlet::FMetasoundParamMacroDescriptor::TryParseArgs(const FString& Text, FSourceFileParseContext& Context, const TArray<FString>& Arguments, const int32 ArgIndexOffset) const
+{
+	FString NameArgument = Arguments[ArgIndexOffset];
+	// Remove whitespace at the start of the line.
+	NameArgument.TrimStartInline();
+
+	// NAME should neverbe be in quotes 
+	bool HasQuotes = false;
+	FString MacroDesc = FString::Printf(TEXT("\"InDescription\" argument in %s macro at %s:%d"), *GetToken(), *Context.Filename, Context.LineNumber);
+	FString SourceLocation = FSourceLocation(Context.Filename, Context.LineNumber).ToString();
+	if (PrepareArgument(NameArgument, true, MacroDesc, HasQuotes))
+	{
+		if (HasQuotes)
+		{
+			UE_LOG(LogGatherTextFromSourceCommandlet, Warning, TEXT("%s macro at \'%s\' has NAME argument enclosed in double quotes and cannot be gathered. Please make sure the argument is not enclosed in double quotes."), *GetToken(), *SourceLocation);
+			return;
+		}
+		else if (NameArgument.IsEmpty())
+		{
+			//The metasound param does not have a NAME so we cannot gather it 
+			UE_LOG(LogGatherTextFromSourceCommandlet, Warning, TEXT("%s macro at \'%s\' has an empty NAME argument and cannot be gathered."), *GetToken(), *SourceLocation);
+			return;
+		}
+	}
+
+	FString NameText= Arguments[ArgIndexOffset + 1];
+	// Remove whitespace at the start of the line.
+	NameText.TrimStartInline();
+
+	// NAME_TEXT should always be in quotes 
+	HasQuotes = false;
+	if (PrepareArgument(NameText, true, MacroDesc, HasQuotes))
+	{
+		if (!HasQuotes)
+		{
+			UE_LOG(LogGatherTextFromSourceCommandlet, Warning, TEXT("%s macro at \'%s\' has NAME_TEXT argument not enclosed in double quotes and cannot be gathered. Please make sure the argument is enclosed in double quotes."), *GetToken(), *SourceLocation);
+			return;
+		}
+		else if (NameText.IsEmpty())
+		{
+			//The metasound param does not have a NAME_TEXT so we cannot gather it 
+			UE_LOG(LogGatherTextFromSourceCommandlet, Warning, TEXT("%s macro at \'%s\' has an empty NAME_TEXT argumentand cannot be gathered."), *GetToken(), *SourceLocation);
+			return;
+		}
+	}
+
+	// parse TOOLTIP_TEXT argument- this arg will be in quotes without TEXT macro
+	FString TooltipSourceText = Arguments[ArgIndexOffset + 2];
+	TooltipSourceText.TrimStartInline();
+	MacroDesc = FString::Printf(TEXT("\"InDescription\" argument in %s macro at %s:%d"), *GetToken(), *Context.Filename, Context.LineNumber);
+	if (PrepareArgument(TooltipSourceText, true, MacroDesc, HasQuotes))
+	{
+		if (!HasQuotes)
+		{
+			UE_LOG(LogGatherTextFromSourceCommandlet, Warning, TEXT("%s macro at \'%s\' has a TOOLTIP_TEXT argument not enclosed in double quotes and cannot be gathered. Please ensure the argument is enclosed in double quotes."), *GetToken(), *SourceLocation);
+			return;
+		}
+		else if (TooltipSourceText.IsEmpty())
+		{
+			UE_LOG(LogGatherTextFromSourceCommandlet, Warning, TEXT("%s macro at \'%s\' has an empty TOOLTIP_TEXT and cannot be gathered."), *GetToken(), *SourceLocation);
+			return;
+		}
+		else
+		{
+			// All Metasound loc data is editor only 
+			static const FString WithEditorString = TEXT("WITH_EDITOR");
+
+			// Create the tooltip entry first 
+			FManifestContext MetasoundArgumentContext;
+			MetasoundArgumentContext.Key = NameArgument + TEXT("ToolTip");
+			MetasoundArgumentContext.SourceLocation = SourceLocation;
+			MetasoundArgumentContext.PlatformName = Context.FilePlatformName;
+
+			// We manually append to the macro stack  for editor only due to the nature of METASOUND_PARAM
+			Context.PushMacroBlock(WithEditorString);
+			Context.AddManifestText(GetToken(), Context.Namespace, TooltipSourceText, MetasoundArgumentContext);
+
+			// now add the display name entry. We can recycle a lot of the values used for tooltips  
+			MetasoundArgumentContext.Key = NameArgument + TEXT("DisplayName");
+			Context.AddManifestText(GetToken(), Context.Namespace, NameText, MetasoundArgumentContext);
+
+			// Now we pop the editor only macro 
+			Context.PopMacroBlock();
+		}
+	}
+}
+
+void UGatherTextFromSourceCommandlet::FMetasoundParamMacroDescriptor::TryParse(const FString& Text, FSourceFileParseContext& Context) const
+{
+	// Attempt to parse something of the format
+	// METASOUND_PARAM(NAME, NAME_TEXT, TOOLTIP_TEXT)
+
+	if (!Context.ExcludedRegion && !Context.WithinBlockComment && !Context.WithinLineComment && !Context.WithinStringLiteral)
+	{
+		TArray<FString> Arguments;
+		if (ParseArgsFromMacro(StripCommentsFromToken(Text, Context), Arguments, Context))
+		{
+			// Validate that we got the rightnumber of Arguments
+			if (Arguments.Num() < GetMinNumberOfArgument())
+			{
+				UE_LOG(LogGatherTextFromSourceCommandlet, Warning, TEXT("Expected at least %d arguments for %s macro, but got %d while parsing %s:%d. %s"), GetMinNumberOfArgument(), *GetToken(), Arguments.Num(), *Context.Filename, Context.LineNumber, *FLocTextHelper::SanitizeLogOutput(Context.LineText.TrimStartAndEnd()));
+				return;
+			}
+
+			TryParseArgs(Text, Context, Arguments, 0);
+		}
+	}
+}
+
 void UGatherTextFromSourceCommandlet::FStringMacroDescriptor::TryParse(const FString& Text, FSourceFileParseContext& Context) const
 {
 	// Attempt to parse something of the format
