WorldPartitionConvertCommandlet: better logging for bad commandlet parameters.

diff --git a/Engine/Source/Editor/UnrealEd/Private/Commandlets/WorldPartitionConvertCommandlet.cpp b/Engine/Source/Editor/UnrealEd/Private/Commandlets/WorldPartitionConvertCommandlet.cpp
--- a/Engine/Source/Editor/UnrealEd/Private/Commandlets/WorldPartitionConvertCommandlet.cpp
+++ b/Engine/Source/Editor/UnrealEd/Private/Commandlets/WorldPartitionConvertCommandlet.cpp
@@ -5,7 +5,7 @@
 =============================================================================*/
 
 #include "Commandlets/WorldPartitionConvertCommandlet.h"
-#include "Algo/Transform.h"
+#include "Algo/ForEach.h"
 #include "HAL/PlatformFileManager.h"
 #include "Misc/Paths.h"
 #include "Misc/ConfigCacheIni.h"
@@ -692,9 +692,16 @@ int32 UWorldPartitionConvertCommandlet::Main(const FString& Params)
 	TMap<FString, FString> Arguments;
 	ParseCommandLine(*Params, Tokens, Switches, Arguments);
 
-	if (Tokens.Num() != 1)
+	if (!Tokens.Num())
 	{
-		UE_LOG(LogWorldPartitionConvertCommandlet, Error, TEXT("ConvertToPartitionedLevel bad parameters"));
+		UE_LOG(LogWorldPartitionConvertCommandlet, Error, TEXT("missing map name"));
+		return 1;
+	}
+	else if (Tokens.Num() > 1)
+	{
+		FString BadParams;
+		Algo::ForEach(Tokens, [&BadParams](const FString& Token) { BadParams += Token; BadParams += TEXT(" "); });
+		UE_LOG(LogWorldPartitionConvertCommandlet, Error, TEXT("extra parameters %s"), *BadParams);
 		return 1;
 	}
 
