Various improvements to the DumpMaterialShaderTypes commandlet to make it easier to use with automation.

- Copy the output file to the Engine/Programs/AutomationTool/Saved/MaterialStats directory so if the results can be uploaded as a build artifact.
- Store the output filename and output file path to make copying to the other directory easier.
- Add better headers to the output:
    - "MaterialSummary"
    - "MaterialInstancesSummary"
    -
- Include the project name in the output file to make it easier to distinguish the file names if they are all in the same directory.
- Remove commandlet timing from output.  This makes subsequent runs have different output and the plain text diff fails.
- Sort the material and material instances lists alphabetically so the order is stable between runs.

diff --git a/Engine/Source/Editor/UnrealEd/Private/Commandlets/DumpMaterialShaderTypes.cpp b/Engine/Source/Editor/UnrealEd/Private/Commandlets/DumpMaterialShaderTypes.cpp
--- a/Engine/Source/Editor/UnrealEd/Private/Commandlets/DumpMaterialShaderTypes.cpp
+++ b/Engine/Source/Editor/UnrealEd/Private/Commandlets/DumpMaterialShaderTypes.cpp
@@ -20,14 +20,19 @@ class FShaderStatsGatheringContext
 {
 public:
 	FShaderStatsGatheringContext() = delete;
-	FShaderStatsGatheringContext(const FString& FileName)
+	FShaderStatsGatheringContext(const FString& InFileName) : FileName(InFileName)
 	{
-		DebugWriter = IFileManager::Get().CreateFileWriter(*FileName);
+		OutputFileName = FPaths::Combine(*FPaths::ProjectSavedDir(), TEXT("MaterialStats"), FileName);
+		DebugWriter = IFileManager::Get().CreateFileWriter(*OutputFileName);
 	}
 	~FShaderStatsGatheringContext()
 	{
 		DebugWriter->Close();
 		delete DebugWriter;
+
+		// Copy to the automation directory.
+		const FString AutomationFilePath = FPaths::Combine(*FPaths::EngineDir(), TEXT("Programs"), TEXT("AutomationTool"), TEXT("Saved"), TEXT("MaterialStats"), FileName);
+		IFileManager::Get().Copy(*AutomationFilePath, *OutputFileName);
 	}
 
 	void AddToGlobalShaderTypeHistogram(const TCHAR* GlobalShaderName)
@@ -215,6 +220,12 @@ private:
 
 	/** Map of global shader type display names to their counts. */
 	TMap<FString, int32>	GlobalShaderTypeHistogram;
+
+	/** Store a copy of the the filename. */
+	FString FileName;
+
+	/** Store the full path to the output file. */
+	FString OutputFileName;
 };
 
 UDumpMaterialShaderTypesCommandlet::UDumpMaterialShaderTypesCommandlet(const FObjectInitializer& ObjectInitializer)
@@ -332,9 +343,9 @@ int ProcessMaterials(const ITargetPlatform* TargetPlatform, const EShaderPlatfor
 	}
 
 	Output.Log(TEXT(""));
-	Output.Log(TEXT("Summary"));
+	Output.Log(TEXT("MaterialSummary"));
 	Output.Log(FString::Printf(TEXT("Total Materials: %d"), MaterialList.Num()));
-	Output.Log(FString::Printf(TEXT("Total Shaders: %d"), TotalShaders));
+	Output.Log(FString::Printf(TEXT("Total Material Shaders: %d"), TotalShaders));
 
 	return TotalShaders;
 }
@@ -387,10 +398,10 @@ int ProcessMaterialInstances(const ITargetPlatform* TargetPlatform, const EShade
 	}
 
 	Output.Log(TEXT(""));
-	Output.Log(TEXT("Summary"));
+	Output.Log(TEXT("Material Instances Summary"));
 	Output.Log(FString::Printf(TEXT("Total Material Instances: %d"), MaterialInstanceList.Num()));
 	Output.Log(FString::Printf(TEXT("Material Instances w/ Static Permutations: %d"), StaticPermutations));
-	Output.Log(FString::Printf(TEXT("Total Shaders: %d"), TotalShaders));
+	Output.Log(FString::Printf(TEXT("Total Material Instances Shaders: %d"), TotalShaders));
 
 	return TotalShaders;
 }
@@ -442,7 +453,7 @@ int ProcessGlobalShaders(const ITargetPlatform* TargetPlatform, const EShaderPla
 
 	Output.Log(TEXT(""));
 	Output.Log(TEXT("Global Shaders Summary"));
-	Output.Log(FString::Printf(TEXT("Total Shaders: %d"), TotalShaders));
+	Output.Log(FString::Printf(TEXT("Total Global Shaders: %d"), TotalShaders));
 
 	return TotalShaders;
 }
@@ -451,10 +462,8 @@ void ProcessForTargetAndShaderPlatform(const ITargetPlatform* TargetPlatform, co
 {
 	TRACE_CPUPROFILER_EVENT_SCOPE(ProcessForTargetAndShaderPlatform);
 
-	const double StartTime = FPlatformTime::Seconds();
-
 	const FString TimeNow = FDateTime::Now().ToString();
-	FString FileName = FPaths::Combine(*FPaths::ProjectSavedDir(), FString::Printf(TEXT("MaterialStats/ShaderTypes-%s-%s-%s.txt"), *TargetPlatform->PlatformName(), *LexToString(ShaderPlatform), *TimeNow));
+	FString FileName = FString::Printf(TEXT("%s-ShaderTypes-%s-%s-%s.txt"), FApp::GetProjectName(), *TargetPlatform->PlatformName(), *LexToString(ShaderPlatform), *TimeNow);
 
 	FShaderStatsGatheringContext Output(FileName);
 
@@ -478,11 +487,6 @@ void ProcessForTargetAndShaderPlatform(const ITargetPlatform* TargetPlatform, co
 	Output.PrintHistogram(TotalShaders);
 	Output.Log(FString::Printf(TEXT("\nAlphabetic list of types:")));
 	Output.PrintAlphabeticList();
-
-	const double EndTime = FPlatformTime::Seconds() - StartTime;
-	Output.Log(TEXT(""));
-	Output.Log(FString::Printf(TEXT("Commandlet Took: %lf"), EndTime));
-
 }
 
 int32 UDumpMaterialShaderTypesCommandlet::Main(const FString& Params)
@@ -552,6 +556,10 @@ int32 UDumpMaterialShaderTypesCommandlet::Main(const FString& Params)
 	const double AssetRegistryEnd = FPlatformTime::Seconds();
 	UE_LOG(LogDumpMaterialShaderTypesCommandlet, Display, TEXT("Asset scan took: %.3f"), AssetRegistryEnd - AssetRegistryStart);
 
+	// Sort the material lists by name so the order is stable.
+	Algo::SortBy(MaterialList, [](const FAssetData& AssetData) { return AssetData.ObjectPath; }, FNameLexicalLess());
+	Algo::SortBy(MaterialInstanceList, [](const FAssetData& AssetData) { return AssetData.ObjectPath; }, FNameLexicalLess());
+
 	// For all active platforms
 	ITargetPlatformManagerModule* TPM = GetTargetPlatformManager();
 	const TArray<ITargetPlatform*>& Platforms = TPM->GetActiveTargetPlatforms();
