Cooker: Change IniSettingRecurse into a threadlocal, so that FConfigFiles that are deleted on other threads will still be correctly processed and their pointers removed.
#rnx

diff --git a/Engine/Source/Editor/UnrealEd/Classes/CookOnTheSide/CookOnTheFlyServer.h b/Engine/Source/Editor/UnrealEd/Classes/CookOnTheSide/CookOnTheFlyServer.h
--- a/Engine/Source/Editor/UnrealEd/Classes/CookOnTheSide/CookOnTheFlyServer.h
+++ b/Engine/Source/Editor/UnrealEd/Classes/CookOnTheSide/CookOnTheFlyServer.h
@@ -286,7 +286,6 @@ private:
 	typedef TMap<FName, TMap<FName, TMap<FName, TArray<FString>>>> FIniSettingContainer;
 
 	mutable FCriticalSection ConfigFileCS;
-	mutable bool IniSettingRecurse = false;
 	mutable FIniSettingContainer AccessedIniStrings;
 	TArray<const FConfigFile*> OpenConfigFiles;
 	TArray<FString> ConfigSettingDenyList;
diff --git a/Engine/Source/Editor/UnrealEd/Private/CookOnTheFlyServer.cpp b/Engine/Source/Editor/UnrealEd/Private/CookOnTheFlyServer.cpp
--- a/Engine/Source/Editor/UnrealEd/Private/CookOnTheFlyServer.cpp
+++ b/Engine/Source/Editor/UnrealEd/Private/CookOnTheFlyServer.cpp
@@ -5161,7 +5161,10 @@ void GetAdditionalCurrentIniVersionStrings( const UCookOnTheFlyServer* CookOnThe
 
 bool UCookOnTheFlyServer::GetCurrentIniVersionStrings( const ITargetPlatform* TargetPlatform, FIniSettingContainer& IniVersionStrings ) const
 {
-	IniVersionStrings = AccessedIniStrings;
+	{
+		FScopeLock Lock(&ConfigFileCS);
+		IniVersionStrings = AccessedIniStrings;
+	}
 
 	// this should be called after the cook is finished
 	TArray<FString> IniFiles;
@@ -5181,9 +5184,12 @@ bool UCookOnTheFlyServer::GetCurrentIniVersionStrings( const ITargetPlatform* Ta
 		
 	}
 
-	for (const FConfigFile* ConfigFile : OpenConfigFiles)
 	{
-		ProcessAccessedIniSettings(ConfigFile, IniVersionStrings);
+		FScopeLock Lock(&ConfigFileCS);
+		for (const FConfigFile* ConfigFile : OpenConfigFiles)
+		{
+			ProcessAccessedIniSettings(ConfigFile, IniVersionStrings);
+		}
 	}
 
 
@@ -5329,29 +5335,28 @@ bool UCookOnTheFlyServer::GetCookedIniVersionStrings(const ITargetPlatform* Targ
 	return true;
 }
 
-
+static thread_local bool GSuppressProcessConfigSettings = false;
 
 void UCookOnTheFlyServer::OnFConfigCreated(const FConfigFile* Config)
 {
-	FScopeLock Lock(&ConfigFileCS);
-	if (IniSettingRecurse)
+	if (GSuppressProcessConfigSettings)
 	{
 		return;
 	}
 
+	FScopeLock Lock(&ConfigFileCS);
 	OpenConfigFiles.Add(Config);
 }
 
 void UCookOnTheFlyServer::OnFConfigDeleted(const FConfigFile* Config)
 {
-	FScopeLock Lock(&ConfigFileCS);
-	if (IniSettingRecurse)
+	if (GSuppressProcessConfigSettings)
 	{
 		return;
 	}
 
+	FScopeLock Lock(&ConfigFileCS);
 	ProcessAccessedIniSettings(Config, AccessedIniStrings);
-
 	OpenConfigFiles.Remove(Config);
 }
 
@@ -5489,7 +5494,7 @@ void UCookOnTheFlyServer::ProcessAccessedIniSettings(const FConfigFile* Config,
 
 bool UCookOnTheFlyServer::IniSettingsOutOfDate(const ITargetPlatform* TargetPlatform) const
 {
-	TGuardValue<bool> A(IniSettingRecurse, true);
+	TGuardValue<bool> A(GSuppressProcessConfigSettings, true);
 
 	FIniSettingContainer OldIniSettings;
 	TMap<FString, FString> OldAdditionalSettings;
@@ -5630,7 +5635,7 @@ bool UCookOnTheFlyServer::IniSettingsOutOfDate(const ITargetPlatform* TargetPlat
 
 bool UCookOnTheFlyServer::SaveCurrentIniSettings(const ITargetPlatform* TargetPlatform) const
 {
-	TGuardValue<bool> S(IniSettingRecurse, true);
+	TGuardValue<bool> S(GSuppressProcessConfigSettings, true);
 
 	TMap<FString, FString> AdditionalIniSettings;
 	GetAdditionalCurrentIniVersionStrings(this, TargetPlatform, AdditionalIniSettings);
