EditorDomain: Fix TrySavePackage feedback when the package digest is unavailable because the package is marked is not loadable in EditorDomain.
#rnx

diff --git a/Engine/Source/Editor/UnrealEd/Private/EditorDomain/EditorDomain.cpp b/Engine/Source/Editor/UnrealEd/Private/EditorDomain/EditorDomain.cpp
--- a/Engine/Source/Editor/UnrealEd/Private/EditorDomain/EditorDomain.cpp
+++ b/Engine/Source/Editor/UnrealEd/Private/EditorDomain/EditorDomain.cpp
@@ -209,6 +209,7 @@ bool FEditorDomain::TryFindOrAddPackageSource(FName PackageName,
 			*WriteToString<256>(PackageName), *PackageDigest.GetStatusString());
 		PackageSource = new FPackageSource();
 		PackageSource->Source = EPackageSource::Workspace;
+		PackageSource->Digest = MoveTemp(PackageDigest);
 		OutSource = PackageSource;
 		return true;
 	}
diff --git a/Engine/Source/Editor/UnrealEd/Private/EditorDomain/EditorDomainUtils.cpp b/Engine/Source/Editor/UnrealEd/Private/EditorDomain/EditorDomainUtils.cpp
--- a/Engine/Source/Editor/UnrealEd/Private/EditorDomain/EditorDomainUtils.cpp
+++ b/Engine/Source/Editor/UnrealEd/Private/EditorDomain/EditorDomainUtils.cpp
@@ -1564,8 +1564,8 @@ bool TrySavePackage(UPackage* Package)
 	FPackageDigest PackageDigest = EditorDomain->GetPackageDigest(PackageName);
 	if (!PackageDigest.IsSuccessful())
 	{
-		UE_LOG(LogEditorDomain, Warning, TEXT("Could not save package to EditorDomain: %s."),
-			*PackageDigest.GetStatusString());
+		UE_LOG(LogEditorDomain, Warning, TEXT("Could not save package to EditorDomain: %s. Reason: %s"),
+			*WriteToString<256>(PackageName), *PackageDigest.GetStatusString());
 		return false;
 	}
 	if (!EnumHasAnyFlags(PackageDigest.DomainUse, EDomainUse::SaveEnabled))
