Fixed crash when renaming an actor folder with same name but different case.

diff --git a/Engine/Source/Editor/UnrealEd/Private/WorldPersistentFolders.cpp b/Engine/Source/Editor/UnrealEd/Private/WorldPersistentFolders.cpp
--- a/Engine/Source/Editor/UnrealEd/Private/WorldPersistentFolders.cpp
+++ b/Engine/Source/Editor/UnrealEd/Private/WorldPersistentFolders.cpp
@@ -71,7 +71,9 @@ bool FWorldPersistentFolders::RenameFolder(const FFolder& InOldFolder, const FFo
 
 	UActorFolder* ActorFolder = GetActorFolder(InOldFolder);
 	check(IsValid(ActorFolder));
-	check(!IsValid(GetActorFolder(InNewFolder)));
+	UActorFolder* FoundFolder = GetActorFolder(InNewFolder);
+	check(!IsValid(FoundFolder) || !FoundFolder->GetPath().IsEqual(InNewFolder.GetPath(), ENameCase::CaseSensitive));
+
 
 	ULevel* Level = GetRootObjectContainer(InOldFolder, GetWorld());
 	check(Level);
@@ -83,7 +85,7 @@ bool FWorldPersistentFolders::RenameFolder(const FFolder& InOldFolder, const FFo
 		ActorFolder->SetParent(ParentActorFolder);
 		const FString FolderLabel = InNewFolder.GetLeafName().ToString();
 		ActorFolder->SetLabel(FolderLabel);
-		check(ActorFolder->GetPath() == InNewFolder.GetPath());
+		check(ActorFolder->GetPath().IsEqual(InNewFolder.GetPath(), ENameCase::CaseSensitive));
 	});
 
 	return true;
@@ -105,7 +107,7 @@ void FWorldPersistentFolders::ModifyFolderAndDetectChanges(ULevel* InLevel, cons
 	{
 		FName NewPath = FolderIt->GetPath();
 		FName* OldPath = OldFolderToPath.Find(FolderIt);
-		if (OldPath && (*OldPath != NewPath))
+		if (OldPath && !OldPath->IsEqual(NewPath, ENameCase::CaseSensitive))
 		{
 			ChangedFolders.Emplace(*OldPath, NewPath);
 		}
@@ -119,8 +121,9 @@ void FWorldPersistentFolders::ModifyFolderAndDetectChanges(ULevel* InLevel, cons
 		FFolder OldFolder(ChangedFolder.Key, InRootObject);
 		FFolder NewFolder(ChangedFolder.Value, InRootObject);
 		FActorFolderProps* OldFolderProperties = Owner.FoldersProperties.Find(OldFolder);
-		Owner.FoldersProperties.Add(NewFolder, OldFolderProperties ? *OldFolderProperties : DefaultFolderProperties);
+		FActorFolderProps FolderProperties = OldFolderProperties ? *OldFolderProperties : DefaultFolderProperties;
 		Owner.FoldersProperties.Remove(OldFolder);
+		Owner.FoldersProperties.Add(NewFolder, FolderProperties);
 			
 		Owner.BroadcastOnActorFolderMoved(OldFolder, NewFolder);
 	}
diff --git a/Engine/Source/Runtime/Engine/Private/ActorFolder.cpp b/Engine/Source/Runtime/Engine/Private/ActorFolder.cpp
--- a/Engine/Source/Runtime/Engine/Private/ActorFolder.cpp
+++ b/Engine/Source/Runtime/Engine/Private/ActorFolder.cpp
@@ -117,7 +117,7 @@ FActorFolderDesc UActorFolder::GetAssetRegistryInfoFromPackage(FName ActorFolder
 void UActorFolder::SetLabel(const FString& InFolderLabel)
 {
 	check(IsValid());
-	if (FolderLabel != InFolderLabel)
+	if (!FolderLabel.Equals(InFolderLabel, ESearchCase::CaseSensitive))
 	{
 		Modify();
 		FolderLabel = InFolderLabel;
