Reduce overhead from high frequency redundant access tracking caused by resolving the same TObjectPtr repeatedly.  Saves 3 minutes from large project full cook.

diff --git a/Engine/Plugins/FX/Niagara/Source/Niagara/Public/NiagaraTypes.h b/Engine/Plugins/FX/Niagara/Source/Niagara/Public/NiagaraTypes.h
--- a/Engine/Plugins/FX/Niagara/Source/Niagara/Public/NiagaraTypes.h
+++ b/Engine/Plugins/FX/Niagara/Source/Niagara/Public/NiagaraTypes.h
@@ -1540,9 +1540,9 @@ public:
 	{
 		FNiagaraTypeRegistry& Registry = Get();
 
+		const uint32 TypeHash = GetTypeHash(NewType);
 		{
 			FReadScopeLock Lock(Registry.RegisteredTypesLock);
-			const uint32 TypeHash = GetTypeHash(NewType);
 			if (const int32* ExistingIndex = Registry.RegisteredTypeIndexMap.Find(TypeHash))
 			{
 				return *ExistingIndex;
@@ -1551,7 +1551,7 @@ public:
 
 		FRWScopeLock Lock(Registry.RegisteredTypesLock, SLT_Write);
 		const int32 Index = Registry.RegisteredTypes.AddUnique(NewType);
-		Registry.RegisteredTypeIndexMap.Add(GetTypeHash(NewType), Index);
+		Registry.RegisteredTypeIndexMap.Add(TypeHash, Index);
 		return Index;
 	}
 
diff --git a/Engine/Plugins/FX/Niagara/Source/NiagaraEditor/Private/NiagaraEditorUtilities.cpp b/Engine/Plugins/FX/Niagara/Source/NiagaraEditor/Private/NiagaraEditorUtilities.cpp
--- a/Engine/Plugins/FX/Niagara/Source/NiagaraEditor/Private/NiagaraEditorUtilities.cpp
+++ b/Engine/Plugins/FX/Niagara/Source/NiagaraEditor/Private/NiagaraEditorUtilities.cpp
@@ -1087,7 +1087,7 @@ void FNiagaraEditorUtilities::SetStaticSwitchConstants(UNiagaraGraph* Graph, TAr
 {
 	const UEdGraphSchema_Niagara* Schema = GetDefault<UEdGraphSchema_Niagara>();
 
-	for (UEdGraphNode* Node : Graph->Nodes)
+	for (TObjectPtr<UEdGraphNode> Node : Graph->Nodes)
 	{
 		// if there is a static switch node its value must be set by the caller
 		UNiagaraNodeStaticSwitch* SwitchNode = Cast<UNiagaraNodeStaticSwitch>(Node);
diff --git a/Engine/Source/Editor/UnrealEd/Private/Cooker/PackageBuildDependencyTracker.cpp b/Engine/Source/Editor/UnrealEd/Private/Cooker/PackageBuildDependencyTracker.cpp
--- a/Engine/Source/Editor/UnrealEd/Private/Cooker/PackageBuildDependencyTracker.cpp
+++ b/Engine/Source/Editor/UnrealEd/Private/Cooker/PackageBuildDependencyTracker.cpp
@@ -60,49 +60,68 @@ FPackageBuildDependencyTracker::FPackageBuildDependencyTracker()
 
 void FPackageBuildDependencyTracker::StaticOnObjectHandleRead(UObject* ReadObject)
 {
-	if (ReadObject)
+	StaticOnObjectHandleReadInternal(ReadObject);
+
+	if (Singleton.PreviousObjectHandleReadFunction)
+	{
+		Singleton.PreviousObjectHandleReadFunction(ReadObject);
+	}
+}
+
+void FPackageBuildDependencyTracker::StaticOnObjectHandleReadInternal(UObject* ReadObject)
+{
+	if (!ReadObject)
+	{
+		return;
+	}
+
+	PackageAccessTracking_Private::FTrackedData* AccumulatedScopeData = PackageAccessTracking_Private::FPackageAccessRefScope::GetCurrentThreadAccumulatedData();
+	if (!AccumulatedScopeData)
+	{
+		return;
+	}
+
+	if (AccumulatedScopeData->BuildOpName.IsNone())
+	{
+		return;
+	}
+
+	if (!ReadObject->HasAnyFlags(RF_Public))
 	{
-		if (PackageAccessTracking_Private::FTrackedData* AccumulatedScopeData = PackageAccessTracking_Private::FPackageAccessRefScope::GetCurrentThreadAccumulatedData())
+		return;
+	}
+
+	FName Referencer = AccumulatedScopeData->PackageName;
+	FName Referenced = ReadObject->GetOutermost()->GetFName();
+	if (Referencer == Referenced)
+	{
+		return;
+	}
+
+	if (AccumulatedScopeData->OpName == PackageAccessTrackingOps::NAME_NoAccessExpected)
+	{
+		UE_LOG(LogPackageBuildDependencyTracker, Warning, TEXT("Object %s is referencing object %s inside of a NAME_NoAccessExpected scope. Programmer should narrow the scope or debug the reference."),
+			*Referencer.ToString(), *Referenced.ToString());
+	}
+
+	LLM_SCOPE_BYNAME(TEXTVIEW("TObjectPtr"));
+
+	FBuildDependencyAccessData AccessData{Referenced, AccumulatedScopeData->TargetPlatform};
+	FScopeLock RecordsScopeLock(&Singleton.RecordsLock);
+	if (Referencer == Singleton.LastReferencer)
+	{
+		if (AccessData != Singleton.LastAccessData)
 		{
-			if (!AccumulatedScopeData->BuildOpName.IsNone())
-			{
-				FName Referencer = AccumulatedScopeData->PackageName;
-				FName Referenced = ReadObject->GetOutermost()->GetFName();
-				if (Referencer != Referenced)
-				{
-					if (AccumulatedScopeData->OpName == PackageAccessTrackingOps::NAME_NoAccessExpected)
-					{
-						UE_LOG(LogPackageBuildDependencyTracker, Warning, TEXT("Object %s is referencing object %s inside of a NAME_NoAccessExpected scope. Programmer should narrow the scope or debug the reference."),
-							*Referencer.ToString(), *Referenced.ToString());
-					}
-
-					LLM_SCOPE_BYNAME(TEXTVIEW("TObjectPtr"));
-
-					FBuildDependencyAccessData AccessData{ Referenced, AccumulatedScopeData->TargetPlatform };
-					FScopeLock RecordsScopeLock(&Singleton.RecordsLock);
-					if (Referencer == Singleton.LastReferencer)
-					{
-						if (AccessData != Singleton.LastAccessData)
-						{
-							Singleton.LastAccessData = AccessData;
-							Singleton.LastReferencerSet->Add(AccessData);
-						}
-					}
-					else
-					{
-						Singleton.LastAccessData = AccessData;
-						Singleton.LastReferencer = Referencer;
-						Singleton.LastReferencerSet = &Singleton.Records.FindOrAdd(Referencer);
-						Singleton.LastReferencerSet->Add(AccessData);
-					}
-				}
-			}
+			Singleton.LastAccessData = AccessData;
+			Singleton.LastReferencerSet->Add(AccessData);
 		}
 	}
-
-	if (Singleton.PreviousObjectHandleReadFunction)
+	else
 	{
-		Singleton.PreviousObjectHandleReadFunction(ReadObject);
+		Singleton.LastAccessData = AccessData;
+		Singleton.LastReferencer = Referencer;
+		Singleton.LastReferencerSet = &Singleton.Records.FindOrAdd(Referencer);
+		Singleton.LastReferencerSet->Add(AccessData);
 	}
 }
 
diff --git a/Engine/Source/Editor/UnrealEd/Private/Cooker/PackageBuildDependencyTracker.h b/Engine/Source/Editor/UnrealEd/Private/Cooker/PackageBuildDependencyTracker.h
--- a/Engine/Source/Editor/UnrealEd/Private/Cooker/PackageBuildDependencyTracker.h
+++ b/Engine/Source/Editor/UnrealEd/Private/Cooker/PackageBuildDependencyTracker.h
@@ -40,6 +40,7 @@ private:
 
 	/** Track object reference reads */
 	static void StaticOnObjectHandleRead(UObject* ReadObject);
+	static void StaticOnObjectHandleReadInternal(UObject* ReadObject);
 	ObjectHandleReadFunction* PreviousObjectHandleReadFunction = nullptr;
 
 	mutable FCriticalSection RecordsLock;
diff --git a/Engine/Source/Runtime/CoreUObject/Private/Misc/PackageAccessTracking.cpp b/Engine/Source/Runtime/CoreUObject/Private/Misc/PackageAccessTracking.cpp
--- a/Engine/Source/Runtime/CoreUObject/Private/Misc/PackageAccessTracking.cpp
+++ b/Engine/Source/Runtime/CoreUObject/Private/Misc/PackageAccessTracking.cpp
@@ -100,7 +100,7 @@ FPackageAccessRefScope::~FPackageAccessRefScope()
 
 void FPackageAccessRefScope::SetPackageName(FName InPackageName)
 {
-	checkf(FPackageName::IsValidLongPackageName(InPackageName.ToString(), true), TEXT("Invalid package name: %s"), *InPackageName.ToString());
+	checkfSlow(FPackageName::IsValidLongPackageName(InPackageName.ToString(), true), TEXT("Invalid package name: %s"), *InPackageName.ToString());
 	checkf(CurrentThreadScope == this, TEXT("Invalid SetPackageName while a child scope is already on the stack. We have not yet implemented propagating a change in PackageName to child scopes' AccumulatedData."));
 
 	DirectData.PackageName = InPackageName;
diff --git a/Engine/Source/Runtime/Engine/Classes/Materials/MaterialFunctionInterface.h b/Engine/Source/Runtime/Engine/Classes/Materials/MaterialFunctionInterface.h
--- a/Engine/Source/Runtime/Engine/Classes/Materials/MaterialFunctionInterface.h
+++ b/Engine/Source/Runtime/Engine/Classes/Materials/MaterialFunctionInterface.h
@@ -243,7 +243,7 @@ public:
 
 			for (UMaterialFunctionInterface* Function : Functions)
 			{
-				for (UMaterialExpression* FunctionExpression : *Function->GetFunctionExpressions())
+				for (TObjectPtr<UMaterialExpression> FunctionExpression : *Function->GetFunctionExpressions())
 				{
 					if (ExpressionType* ParameterExpression = Cast<ExpressionType>(FunctionExpression))
 					{
diff --git a/Engine/Source/Runtime/Engine/Private/Materials/MaterialExpressions.cpp b/Engine/Source/Runtime/Engine/Private/Materials/MaterialExpressions.cpp
--- a/Engine/Source/Runtime/Engine/Private/Materials/MaterialExpressions.cpp
+++ b/Engine/Source/Runtime/Engine/Private/Materials/MaterialExpressions.cpp
@@ -14270,13 +14270,13 @@ void UMaterialExpressionMaterialFunctionCall::PostLoad()
 #if WITH_EDITORONLY_DATA
 bool UMaterialExpressionMaterialFunctionCall::IterateDependentFunctions(TFunctionRef<bool(UMaterialFunctionInterface*)> Predicate) const
 {
-	if (MaterialFunction)
+	if (UMaterialFunctionInterface* MaterialFunctionInterface = MaterialFunction)
 	{
-		if (!MaterialFunction->IterateDependentFunctions(Predicate))
+		if (!MaterialFunctionInterface->IterateDependentFunctions(Predicate))
 		{
 			return false;
 		}
-		if (!Predicate(MaterialFunction))
+		if (!Predicate(MaterialFunctionInterface))
 		{
 			return false;
 		}
