DumpMaterialShaderTypes should output its results to the Log directory.

- This way the file can be uploaded automatically as a build artifact.

diff --git a/Engine/Source/Editor/UnrealEd/Private/Commandlets/DumpMaterialShaderTypes.cpp b/Engine/Source/Editor/UnrealEd/Private/Commandlets/DumpMaterialShaderTypes.cpp
--- a/Engine/Source/Editor/UnrealEd/Private/Commandlets/DumpMaterialShaderTypes.cpp
+++ b/Engine/Source/Editor/UnrealEd/Private/Commandlets/DumpMaterialShaderTypes.cpp
@@ -31,7 +31,7 @@ public:
 		delete DebugWriter;
 
 		// Copy to the automation directory.
-		const FString AutomationFilePath = FPaths::Combine(*FPaths::EngineDir(), TEXT("Programs"), TEXT("AutomationTool"), TEXT("Saved"), TEXT("MaterialStats"), FileName);
+		const FString AutomationFilePath = FPaths::Combine(*FPaths::EngineDir(), TEXT("Programs"), TEXT("AutomationTool"), TEXT("Saved"), TEXT("Logs"), FileName);
 		IFileManager::Get().Copy(*AutomationFilePath, *OutputFileName);
 	}
 
