Material from Texture - masked if alpha present

diff --git a/Engine/Source/Editor/UnrealEd/Private/LevelEditorViewport.cpp b/Engine/Source/Editor/UnrealEd/Private/LevelEditorViewport.cpp
--- a/Engine/Source/Editor/UnrealEd/Private/LevelEditorViewport.cpp
+++ b/Engine/Source/Editor/UnrealEd/Private/LevelEditorViewport.cpp
@@ -364,7 +364,19 @@ static bool TryAndCreateMaterialInput( UMaterial* UnrealMaterial, EMaterialKind:
 		return false;
 	}
 
-	bool bSetupAsNormalMap = UnrealTexture->IsNormalMap();
+	bool bTextureHasAlpha = !UnrealTexture->CompressionNoAlpha;
+	if ( bTextureHasAlpha )
+	{
+		if ( const UTexture2D* Texture2D = Cast<UTexture2D>(UnrealTexture) )
+		{
+			const EPixelFormatChannelFlags ValidTextureChannels = GetPixelFormatValidChannels(Texture2D->GetPixelFormat());
+			bTextureHasAlpha = EnumHasAnyFlags(ValidTextureChannels, EPixelFormatChannelFlags::A);
+		}
+	}
+
+	const bool bSetupAsNormalMap = UnrealTexture->IsNormalMap();
+	
+	UnrealMaterial->BlendMode = bTextureHasAlpha ? EBlendMode::BLEND_Masked : EBlendMode::BLEND_Opaque;
 
 	// Create a new texture sample expression, this is our texture input node into the material output.
 	UMaterialExpressionTextureSample* UnrealTextureExpression = NewObject<UMaterialExpressionTextureSample>(UnrealMaterial);
@@ -456,6 +468,11 @@ static bool TryAndCreateMaterialInput( UMaterial* UnrealMaterial, EMaterialKind:
 		if ( TextureKind == EMaterialKind::Base )
 		{
 			UnrealMaterial->BaseColor.Expression = UnrealTextureExpression;
+			if ( bTextureHasAlpha && UnrealTextureExpression->Outputs.IsValidIndex(3) )
+			{
+				UnrealMaterial->Opacity.Connect(3, UnrealTextureExpression);
+				UnrealMaterial->OpacityMask.Connect(3, UnrealTextureExpression);
+			}
 		}
 		else if ( TextureKind == EMaterialKind::Specular )
 		{
