Material from Texture index fix

diff --git a/Engine/Source/Editor/UnrealEd/Private/LevelEditorViewport.cpp b/Engine/Source/Editor/UnrealEd/Private/LevelEditorViewport.cpp
--- a/Engine/Source/Editor/UnrealEd/Private/LevelEditorViewport.cpp
+++ b/Engine/Source/Editor/UnrealEd/Private/LevelEditorViewport.cpp
@@ -468,10 +468,10 @@ static bool TryAndCreateMaterialInput( UMaterial* UnrealMaterial, EMaterialKind:
 		if ( TextureKind == EMaterialKind::Base )
 		{
 			UnrealMaterial->BaseColor.Expression = UnrealTextureExpression;
-			if ( bTextureHasAlpha && UnrealTextureExpression->Outputs.IsValidIndex(3) )
+			if ( bTextureHasAlpha && UnrealTextureExpression->Outputs.IsValidIndex(4) )
 			{
-				UnrealMaterial->Opacity.Connect(3, UnrealTextureExpression);
-				UnrealMaterial->OpacityMask.Connect(3, UnrealTextureExpression);
+				UnrealMaterial->Opacity.Connect(4, UnrealTextureExpression);
+				UnrealMaterial->OpacityMask.Connect(4, UnrealTextureExpression);
 			}
 		}
 		else if ( TextureKind == EMaterialKind::Specular )
