Don't force-flush render resources from UEditorEngine::Tick when PIE is running.

If realtime was disabled, the force-flush was happening even when PIE was active, which cost a lot of performance.

#rnx

diff --git a/Engine/Source/Editor/UnrealEd/Private/EditorEngine.cpp b/Engine/Source/Editor/UnrealEd/Private/EditorEngine.cpp
--- a/Engine/Source/Editor/UnrealEd/Private/EditorEngine.cpp
+++ b/Engine/Source/Editor/UnrealEd/Private/EditorEngine.cpp
@@ -1685,6 +1685,7 @@ void UEditorEngine::Tick( float DeltaSeconds, bool bIdleMode )
 	EmitDynamicResolutionEvent(EDynamicResolutionStateEvent::BeginFrame);
 
 	// if we have the side-by-side world for "Play From Here", tick it unless we are ensuring slate is responsive
+	bool bHasPIEViewport = false;
 	if( FSlateThrottleManager::Get().IsAllowingExpensiveTasks() )
 	{
 		// Determine number of PIE worlds that should tick and if they feature an active movie sequence tick.
@@ -1797,6 +1798,7 @@ void UEditorEngine::Tick( float DeltaSeconds, bool bIdleMode )
 				if ( GameViewport != NULL )
 				{
 					GameViewport->Tick(TickDeltaSeconds);
+					bHasPIEViewport = true;
 				}
 			}
 
@@ -1982,7 +1984,7 @@ void UEditorEngine::Tick( float DeltaSeconds, bool bIdleMode )
 	}
 
 	// If we're not Realtime and no NonRealtime viewports are drawn, or if everything is hidden, make sure RHI gets its flush to prevent memory from accumulating
-	if ((bAllWindowsHidden || !IsRealtime) && !bEditorFrameNonRealtimeViewportDrawn && IsRunningRHIInSeparateThread())
+	if ((bAllWindowsHidden || !IsRealtime) && !bEditorFrameNonRealtimeViewportDrawn && !bHasPIEViewport && IsRunningRHIInSeparateThread())
 	{
 		ENQUEUE_RENDER_COMMAND(FlushPendingDeleteRHIResources_NonRealtime)(
 			[](FRHICommandListImmediate& RHICmdList)
