Fixed crash baking levels with volumetric lightmap density volumes due missing LWC handling

diff --git a/Engine/Source/Editor/UnrealEd/Private/Lightmass/Lightmass.cpp b/Engine/Source/Editor/UnrealEd/Private/Lightmass/Lightmass.cpp
--- a/Engine/Source/Editor/UnrealEd/Private/Lightmass/Lightmass.cpp
+++ b/Engine/Source/Editor/UnrealEd/Private/Lightmass/Lightmass.cpp
@@ -931,7 +931,7 @@ void FLightmassExporter::WriteVisibilityData( int32 Channel )
 
 	// This used to walk along Matinee sequences and write the positions for pre-computed visibility.
 	// Leaving the swarm channel code here for fear of breaking swarm.
-	TArray<FVector4> CameraTrackPositions;
+	TArray<FVector4f> CameraTrackPositions;
 	int32 NumCameraPositions = CameraTrackPositions.Num();
 	Swarm.WriteChannel( Channel, &NumCameraPositions, sizeof(NumCameraPositions) );
 	Swarm.WriteChannel( Channel, CameraTrackPositions.GetData(), CameraTrackPositions.Num() * CameraTrackPositions.GetTypeSize() );
@@ -947,10 +947,15 @@ void FLightmassExporter::WriteVolumetricLightmapData( int32 Channel )
 		VolumeData.Bounds = FBox3f(DetailVolume->GetComponentsBoundingBox(true));
 		VolumeData.AllowedMipLevelRange = FIntPoint(DetailVolume->AllowedMipLevelRange.Min, DetailVolume->AllowedMipLevelRange.Max);
 
-		TArray<FPlane> Planes;
-		DetailVolume->Brush->GetSurfacePlanes(DetailVolume, Planes);
-		const int32 NumPlanes = Planes.Num();
-		VolumeData.NumPlanes = NumPlanes;
+		TArray<FPlane> Planes4d;
+		DetailVolume->Brush->GetSurfacePlanes(DetailVolume, Planes4d);
+		TArray<FPlane4f> Planes;
+		for (FPlane Plane4d : Planes4d)
+		{
+			Planes.Add(FPlane4f(Plane4d));
+		}
+
+		VolumeData.NumPlanes = Planes.Num();
 
 		Swarm.WriteChannel( Channel, &VolumeData, sizeof(VolumeData) );
 		Swarm.WriteChannel( Channel, Planes.GetData(), Planes.Num() * Planes.GetTypeSize() );
