Add error logging to UResavePackagesCommandlet::CheckoutAndSavePackage

diff --git a/Engine/Source/Editor/UnrealEd/Private/Commandlets/ContentCommandlets.cpp b/Engine/Source/Editor/UnrealEd/Private/Commandlets/ContentCommandlets.cpp
--- a/Engine/Source/Editor/UnrealEd/Private/Commandlets/ContentCommandlets.cpp
+++ b/Engine/Source/Editor/UnrealEd/Private/Commandlets/ContentCommandlets.cpp
@@ -1616,7 +1616,14 @@ void UResavePackagesCommandlet::CheckoutAndSavePackage(UPackage* Package, TArray
 			if (CheckoutFile(PackageFilename, true, bIgnoreAlreadyCheckedOut))
 			{
 				SublevelFilenames.Add(PackageFilename);
-				SavePackageHelper(Package, PackageFilename);
+				if (!SavePackageHelper(Package, PackageFilename))
+				{
+					UE_LOG(LogContentCommandlet, Error, TEXT("Failed to save existing package %s"), *PackageFilename);
+				}
+			}
+			else
+			{
+				UE_LOG(LogContentCommandlet, Error, TEXT("Failed to check out existing package %s"), *PackageFilename);
 			}
 		}
 		else
@@ -1627,6 +1634,14 @@ void UResavePackagesCommandlet::CheckoutAndSavePackage(UPackage* Package, TArray
 				{
 					SublevelFilenames.Add(PackageFilename);
 				}
+				else
+				{
+					UE_LOG(LogContentCommandlet, Error, TEXT("Failed to check out new package %s"), *PackageFilename);
+				}
+			}
+			else
+			{
+				UE_LOG(LogContentCommandlet, Error, TEXT("Failed to save new package %s"), *PackageFilename);
 			}
 		}
 	}
