Fix potential assert/crash when renaming packages

#rnx

### Problem
- This problem was introduced in CL 19224316 which removed the direct call to ::ResetLoaders as forcing a package that was being unloaded to load all of it's bulkdata payloads into memory did not make much sense.
- Since the fix in that CL was about flushing all async compilation work before calling ::ResetLoaders, if we could avoid calling it unless absolutely needed (objects that were not gced) then we could reduce the chance of us actually needing to flush the async compilation work, which can be slow.

### Fix
- However if the package is being renamed, then the uobjects contained within the package will have been moved to the new package, in which case we want any bulkdata objects that they own to remain valid so we cannot rely on the call to ::ResetLoaders that occurs when a package is garbage collected as that variant leaves the newly detached bulkdata objects in an invalid state that can no longer access their payloads.
- Rolled the fix back to a 'dumber' version. We once again call ::ResetLoaders on each package but before we start doing that we check all packages that we are trying to unload to see if any have asnyc compilation work running and only if we detect this do we flush the system.
-- Renamed ::ResetPackageLoaders to ::FlushAsyncCompilation and moved it to be called much earlier.
-- Removed the array of weak ptrs added in 19224316 as they are no longer used.
- I also changed the code documentation on ::ResetLoaders to make it clear why it is being called there.

diff --git a/Engine/Source/Editor/UnrealEd/Private/PackageTools.cpp b/Engine/Source/Editor/UnrealEd/Private/PackageTools.cpp
--- a/Engine/Source/Editor/UnrealEd/Private/PackageTools.cpp
+++ b/Engine/Source/Editor/UnrealEd/Private/PackageTools.cpp
@@ -79,7 +79,7 @@ namespace
  * Utility function that checks each UObject inside of the given UPackage to see if it is waiting 
  * on async compilation.
  * 
- * @return true is the package contains at least one UObject that has compilation work running, otherwise false.
+ * @return true if the package contains at least one UObject that has compilation work running, otherwise false.
  */
 static bool IsPackageCompiling(const UPackage* Package)
 {
@@ -102,24 +102,20 @@ static bool IsPackageCompiling(const UPackage* Package)
 }
 
 /**
- * Utility function to reset the loaders for an array of packages safely.
- * 
- * If any of the packages contains an object with running async compilation
- * tasks then we need to wait for all existing async compilation tasks to
- * be flushed before we can call ResetLoaders, due to thread safety issues.
+ * Utility function that checks all of the provided packages to see if any
+ * of them contain assets that currently have async compilation work running.
+ * If there are assets that are waiting on async compilation work then we 
+ * wait on all currently outstanding work to finish before returning.
  */
-static void ResetPackageLoaders(const TArray<TWeakObjectPtr<UPackage>>& PackagesToUnload)
+static void FlushAsyncCompilation(const TSet<UPackage*>& PackagesToUnload)
 {
 	bool bHasAsyncCompilationWork = false;
-	for (const TWeakObjectPtr<UPackage>& Package : PackagesToUnload)
+	for (const UPackage* Package : PackagesToUnload)
 	{
-		if (Package.IsValid())
+		if (IsPackageCompiling(Package))
 		{
-			if (IsPackageCompiling(Package.Get()))
-			{
-				bHasAsyncCompilationWork = true;
-				break;
-			}
+			bHasAsyncCompilationWork = true;
+			break;
 		}
 	}
 
@@ -127,14 +123,6 @@ static void ResetPackageLoaders(const TArray<TWeakObjectPtr<UPackage>>& Packages
 	{
 		FAssetCompilingManager::Get().FinishAllCompilation();
 	}
-
-	for (const TWeakObjectPtr<UPackage>& Package : PackagesToUnload)
-	{
-		if (Package.IsValid())
-		{
-			ResetLoaders(Package.Get());
-		}
-	}
 }
 
 }
@@ -434,14 +422,6 @@ UPackageTools::UPackageTools(const FObjectInitializer& ObjectInitializer)
 			FlushAsyncLoading();
 			(*GFlushStreamingFunc)();
 
-			// Store the packages as weak pointers so that we can identify which ones 
-			// survive garbage collection later on
-			TArray<TWeakObjectPtr<UPackage>> PackagesToReset;
-			for (UPackage* Package : PackagesToUnload)
-			{
-				PackagesToReset.Add(Package);
-			}
-
 			// Remove potential references to to-be deleted objects from the GB selection set.
 			GEditor->GetSelectedObjects()->DeselectAll();
 
@@ -460,6 +440,10 @@ UPackageTools::UPackageTools(const FObjectInitializer& ObjectInitializer)
 				}
 			}
 
+			// We need to make sure that there is no async compilation work running for the packages that we are about to unload
+			// so that it is safe to call ::ResetLoaders
+			FlushAsyncCompilation(PackagesToUnload);
+
 			// Now try to clean up assets in all packages to unload.
 			int32 PackageIndex = 0;
 			for (UPackage* PackageBeingUnloaded : PackagesToUnload)
@@ -522,6 +506,15 @@ UPackageTools::UPackageTools(const FObjectInitializer& ObjectInitializer)
 					ObjectsInPackage.Reset();
 				}
 
+				// Calling ::ResetLoaders now will force any bulkdata objects still attached to the FLinkerLoad to load
+				// their payloads into memory. If we don't call this now, then the version that will be called during
+				// garbage collection will cause the bulkdata objects to be invalidated rather than loading the payloads 
+				// into memory.
+				// This might seem odd, but if the package we are unloading is being renamed, then the inner UObjects will
+				// be moved to the newly named package rather than being garbage collected and so we need to make sure that
+				// their bulkdata objects remain valid, otherwise renamed packages will not save correctly and cease to function.
+				ResetLoaders(PackageBeingUnloaded);
+
 				if( PackageBeingUnloaded->IsDirty() )
 				{
 					// The package was marked dirty as a result of something that happened above (e.g callbacks in CollectGarbage).  
@@ -573,10 +566,6 @@ UPackageTools::UPackageTools(const FObjectInitializer& ObjectInitializer)
 
 			CollectGarbage( GARBAGE_COLLECTION_KEEPFLAGS );
 
-			// We need to reset the loader of any of the package that survived garbage collection to  
-			// make sure that they are detached from their file on disk.
-			ResetPackageLoaders(PackagesToReset);
-			
 			// Restore the standalone flag on any metadata objects that survived the GC
 			for ( const TWeakObjectPtr<UMetaData>& WeakPackageMetaData : PackageMetaDataWithClearedStandaloneFlag )
 			{
