Fix payload loading errors seen on nightly HLOD jobs

#rnx

- When running with -BuildHLOD we use -AllowCommandletRendering so textures being loaded will create async compilation tasks
- Due to an unrelated problem the async compilation was seeing very poor access to the DDC often timing out and running the compilation locally
- This increase in latency was causing the async work to be running after the package containing the asset was gced and detached from the file on disk meaning we could no longer load the bulkdata payload.
- This fix will make sure that we don't run gc until the async compilation work has finished, but only is we are allowing rendering.
-- This is not a great fix and needs to be improved on, the associated jira will track this investigation work. However it will solve any potential problem for now.

diff --git a/Engine/Source/Editor/UnrealEd/Private/Commandlets/ContentCommandlets.cpp b/Engine/Source/Editor/UnrealEd/Private/Commandlets/ContentCommandlets.cpp
--- a/Engine/Source/Editor/UnrealEd/Private/Commandlets/ContentCommandlets.cpp
+++ b/Engine/Source/Editor/UnrealEd/Private/Commandlets/ContentCommandlets.cpp
@@ -51,7 +51,7 @@
 #include "WorldPartition/WorldPartition.h"
 #include "WorldPartition/WorldPartitionHelpers.h"
 #include "LevelInstance/LevelInstanceActor.h"
-
+#include "AssetCompilingManager.h"
 #include "PackageHelperFunctions.h"
 #include "PackageTools.h"
 #include "UObject/PackageTrailer.h"
@@ -890,6 +890,11 @@ void UResavePackagesCommandlet::LoadAndSaveOnePackage(const FString& Filename)
 
 		if (!GarbageCollectionFrequency || Counter++ % GarbageCollectionFrequency == 0)
 		{
+			if (FApp::CanEverRender())
+			{
+				FAssetCompilingManager::Get().FinishAllCompilation();
+			}
+
 			if (GarbageCollectionFrequency > 1)
 			{
 				UE_LOG(LogContentCommandlet, Display, TEXT("GC"));
