Cloth - Fix crash when attempting to rebuild fixed vert data on a missing cloth section.

diff --git a/Engine/Source/Editor/ClothPainter/Private/ClothPainter.cpp b/Engine/Source/Editor/ClothPainter/Private/ClothPainter.cpp
--- a/Engine/Source/Editor/ClothPainter/Private/ClothPainter.cpp
+++ b/Engine/Source/Editor/ClothPainter/Private/ClothPainter.cpp
@@ -274,11 +274,12 @@ void FClothPainter::Tick(FEditorViewportClient* ViewportClient, float DeltaTime)
 			{
 				if(Asset)
 				{
-					Asset->ApplyParameterMasks();
+					constexpr bool bUpdateFixedVertData = true;  // There's Currently no way of telling whether the MaxDistance mask has been edited
+					constexpr bool bInvalidateDerivedDataCache = false;  // No need to rebuild the DDC while previewing
+
+					Asset->ApplyParameterMasks(bUpdateFixedVertData, bInvalidateDerivedDataCache);
 				}
 			}
-			static const bool bInvalidateDerivedDataCache = false;  // No need to rebuild the DDC while previewing
-			SkeletalMeshComponent->RebuildClothingSectionsFixedVerts(bInvalidateDerivedDataCache);
 		}
 
 		FComponentReregisterContext ReregisterContext(SkeletalMeshComponent);
diff --git a/Engine/Source/Editor/ClothPainter/Private/ClothingPaintEditMode.cpp b/Engine/Source/Editor/ClothPainter/Private/ClothingPaintEditMode.cpp
--- a/Engine/Source/Editor/ClothPainter/Private/ClothingPaintEditMode.cpp
+++ b/Engine/Source/Editor/ClothPainter/Private/ClothingPaintEditMode.cpp
@@ -93,11 +93,12 @@ void FClothingPaintEditMode::Exit()
 				for(UClothingAssetBase* AssetBase : SkelMesh->GetMeshClothingAssets())
 				{
 					UClothingAssetCommon* ConcreteAsset = CastChecked<UClothingAssetCommon>(AssetBase);
-					ConcreteAsset->ApplyParameterMasks();
+					constexpr bool bUpdateFixedVertData = true;
+					constexpr bool bInvalidateDerivedDataCache = true;
+					ConcreteAsset->ApplyParameterMasks(bUpdateFixedVertData, bInvalidateDerivedDataCache);
 				}
 			}
 
-			MeshComponent->RebuildClothingSectionsFixedVerts();
 			MeshComponent->ResetMeshSectionVisibility();
 			MeshComponent->SelectedClothingGuidForPainting = FGuid();
 			MeshComponent->SelectedClothingLodForPainting = INDEX_NONE;
diff --git a/Engine/Source/Editor/UnrealEd/Classes/Animation/DebugSkelMeshComponent.h b/Engine/Source/Editor/UnrealEd/Classes/Animation/DebugSkelMeshComponent.h
--- a/Engine/Source/Editor/UnrealEd/Classes/Animation/DebugSkelMeshComponent.h
+++ b/Engine/Source/Editor/UnrealEd/Classes/Animation/DebugSkelMeshComponent.h
@@ -532,6 +532,7 @@ public:
 	// Rebuilds the fixed parameter on the mesh to mesh data, to be used if the editor has
 	// changed a vert to be fixed or unfixed otherwise the simulation will not work
 	// bInvalidateDerivedDataCache can only be false during previewing as otherwise the changes won't be correctly saved
+	UE_DEPRECATED(5.0, "This function is redundant, since it is always called after ApplyParameterMasks and therefore will be removed.")
 	void RebuildClothingSectionsFixedVerts(bool bInvalidateDerivedDataCache = true);
 
 	TArray<FVector3f> SkinnedSelectedClothingPositions;
@@ -540,6 +541,7 @@ public:
 private:
 	// Rebuilds the fixed vertex attribute on any cloth deformer mappings,
 	// including LOD bias mappings, that reference the specified LOD section.
+	UE_DEPRECATED(5.0, "This function is redundant, since it is always called after ApplyParameterMasks and therefore will be removed.")
 	void RebuildClothingSectionFixedVerts(int32 LODIndex, int32 SectionIndex);
 
 	// Helper function to generate space bases for current frame
diff --git a/Engine/Source/Editor/UnrealEd/Private/Animation/DebugSkelMeshComponent.cpp b/Engine/Source/Editor/UnrealEd/Private/Animation/DebugSkelMeshComponent.cpp
--- a/Engine/Source/Editor/UnrealEd/Private/Animation/DebugSkelMeshComponent.cpp
+++ b/Engine/Source/Editor/UnrealEd/Private/Animation/DebugSkelMeshComponent.cpp
@@ -958,7 +958,9 @@ void UDebugSkelMeshComponent::RebuildClothingSectionsFixedVerts(bool bInvalidate
 	{
 		for (int32 SectionIndex = 0; SectionIndex < LODModels[LODIndex].Sections.Num(); ++SectionIndex)
 		{
+			PRAGMA_DISABLE_DEPRECATION_WARNINGS
 			RebuildClothingSectionFixedVerts(LODIndex, SectionIndex);
+			PRAGMA_ENABLE_DEPRECATION_WARNINGS
 		}
 	}
 
diff --git a/Engine/Source/Runtime/ClothingSystemRuntimeCommon/Private/ClothingAsset.cpp b/Engine/Source/Runtime/ClothingSystemRuntimeCommon/Private/ClothingAsset.cpp
--- a/Engine/Source/Runtime/ClothingSystemRuntimeCommon/Private/ClothingAsset.cpp
+++ b/Engine/Source/Runtime/ClothingSystemRuntimeCommon/Private/ClothingAsset.cpp
@@ -863,7 +863,7 @@ void UClothingAssetCommon::ForEachInteractorUsingClothing(TFunction<void(UClothi
 	}
 }
 
-void UClothingAssetCommon::ApplyParameterMasks(bool bUpdateFixedVertData)
+void UClothingAssetCommon::ApplyParameterMasks(bool bUpdateFixedVertData, bool bInvalidateDerivedDataCache)
 {
 	for(FClothLODDataCommon& Lod : LodData)
 	{
@@ -899,8 +899,10 @@ void UClothingAssetCommon::ApplyParameterMasks(bool bUpdateFixedVertData)
 				}
 			}
 		}
-		// We must always dirty the DDC key unless previewing
-		SkeletalMesh->InvalidateDeriveDataCacheGUID();
+		if (bInvalidateDerivedDataCache)  // We must always dirty the DDC key unless previewing
+		{
+			SkeletalMesh->InvalidateDeriveDataCacheGUID();
+		}
 	}
 }
 
diff --git a/Engine/Source/Runtime/ClothingSystemRuntimeCommon/Public/ClothingAsset.h b/Engine/Source/Runtime/ClothingSystemRuntimeCommon/Public/ClothingAsset.h
--- a/Engine/Source/Runtime/ClothingSystemRuntimeCommon/Public/ClothingAsset.h
+++ b/Engine/Source/Runtime/ClothingSystemRuntimeCommon/Public/ClothingAsset.h
@@ -146,9 +146,9 @@ public:
 	/** 
 	 * Callback envoked after weights have been edited.
 	 * Calls \c PushWeightsToMesh() on each \c ClothLodData, and invalidates cached data. 
-	 * Optionaly recalculate the owner's sections fixed vertex data based on this asset masks.
+	 * Optionaly recalculate the owner's sections fixed vertex data based on this asset masks, and invalidate the DDC for this asset.
 	 */
-	void ApplyParameterMasks(bool bUpdateFixedVertData = false);
+	void ApplyParameterMasks(bool bUpdateFixedVertData = false, bool bInvalidateDerivedDataCache = true);
 
 	/**
 	 * Builds the LOD transition data.
