Add cvar to control PIE features being shown or not

diff --git a/Engine/Source/Editor/UnrealEd/Private/Kismet2/DebuggerCommands.cpp b/Engine/Source/Editor/UnrealEd/Private/Kismet2/DebuggerCommands.cpp
--- a/Engine/Source/Editor/UnrealEd/Private/Kismet2/DebuggerCommands.cpp
+++ b/Engine/Source/Editor/UnrealEd/Private/Kismet2/DebuggerCommands.cpp
@@ -87,6 +87,16 @@
 
 DEFINE_LOG_CATEGORY_STATIC(LogDebuggerCommands, Log, All);
 
+namespace DebuggerCommands
+{
+	static bool bAllowPlayWorldFeature = true;
+	static FAutoConsoleVariableRef AllowPlayWorldFeatureCVar(TEXT("Editor.AllowPlayWorldFeature"), bAllowPlayWorldFeature, TEXT("When true play world is allowed."));
+	static bool AllowPlayWorldFeature()
+	{
+		return bAllowPlayWorldFeature;
+	}
+}
+
 void SGlobalPlayWorldActions::Construct(const FArguments& InArgs)
 {
 	// Always keep track of the current active play world actions widget so we later set user focus on it
@@ -468,7 +478,7 @@ void FPlayWorldCommands::BindGlobalPlayWorldCommands()
 		FExecuteAction::CreateStatic(&FInternalPlayWorldCommandCallbacks::StopPlaySession_Clicked),
 		FCanExecuteAction::CreateStatic(&FInternalPlayWorldCommandCallbacks::HasPlayWorld),
 		FIsActionChecked(),
-		FIsActionButtonVisible()
+		FIsActionButtonVisible::CreateStatic(&DebuggerCommands::AllowPlayWorldFeature)
 	);
 
 	// Late join session
@@ -498,7 +508,7 @@ void FPlayWorldCommands::BindGlobalPlayWorldCommands()
 		FExecuteAction::CreateStatic(&FInternalPlayWorldCommandCallbacks::SingleFrameAdvance_Clicked),
 		FCanExecuteAction::CreateStatic(&FPlayWorldCommandCallbacks::HasPlayWorldAndPaused),
 		FIsActionChecked(),
-		FIsActionButtonVisible()
+		FIsActionButtonVisible::CreateStatic(&DebuggerCommands::AllowPlayWorldFeature)
 	);
 
 	ActionList.MapAction(Commands.TogglePlayPauseOfPlaySession,
@@ -521,7 +531,7 @@ void FPlayWorldCommands::BindGlobalPlayWorldCommands()
 		FExecuteAction::CreateStatic(&FInternalPlayWorldCommandCallbacks::PossessEjectPlayer_Clicked),
 		FCanExecuteAction::CreateStatic(&FInternalPlayWorldCommandCallbacks::CanPossessEjectPlayer),
 		FIsActionChecked(),
-		FIsActionButtonVisible()
+		FIsActionButtonVisible::CreateStatic(&DebuggerCommands::AllowPlayWorldFeature)
 	);
 
 	// Breakpoint-only commands
@@ -637,6 +647,7 @@ void FPlayWorldCommands::BuildToolbar(FToolMenuSection& InSection, bool bInclude
 
 	FUIAction SpecialPIEOptionsMenuAction;
 	SpecialPIEOptionsMenuAction.CanExecuteAction = FCanExecuteAction::CreateStatic(&FInternalPlayWorldCommandCallbacks::CanShowNonPlayWorldOnlyActions);
+	SpecialPIEOptionsMenuAction.IsActionVisibleDelegate = FIsActionButtonVisible::CreateStatic(&DebuggerCommands::AllowPlayWorldFeature);
 
 	FToolMenuEntry PIEComboEntry = FToolMenuEntry::InitComboButton("PIECombo", SpecialPIEOptionsMenuAction, FOnGetContent::CreateStatic(&GeneratePlayMenuContent, GlobalPlayWorldActions.ToSharedRef()), FText(), LOCTEXT("PIEComboToolTip", "Change Play Mode and Play Settings"));
 	PIEComboEntry.StyleNameOverride = FName("Toolbar.BackplateRightCombo");
@@ -1281,6 +1292,11 @@ void SetLastExecutedLaunchMode(ELaunchModeType LaunchMode)
 
 void FInternalPlayWorldCommandCallbacks::RepeatLastPlay_Clicked()
 {
+	if (!DebuggerCommands::AllowPlayWorldFeature())
+	{
+		return;
+	}
+
 	ULevelEditorPlaySettings* PlaySettings = GetMutableDefault<ULevelEditorPlaySettings>();
 	PlaySettings->PostEditChange();
 
@@ -1302,6 +1318,11 @@ void FInternalPlayWorldCommandCallbacks::RepeatLastPlay_Clicked()
 
 bool FInternalPlayWorldCommandCallbacks::RepeatLastPlay_CanExecute()
 {
+	if (!DebuggerCommands::AllowPlayWorldFeature())
+	{
+		return false;
+	}
+
 	const ULevelEditorPlaySettings* PlaySettings = GetDefault<ULevelEditorPlaySettings>();
 	if (PlaySettings->LastExecutedPlayModeType == EPlayModeType::PlayMode_QuickLaunch)
 	{
@@ -1775,17 +1796,17 @@ void FInternalPlayWorldCommandCallbacks::TogglePlayPause_Clicked()
 
 bool FInternalPlayWorldCommandCallbacks::CanShowNonPlayWorldOnlyActions()
 {
-	return !HasPlayWorld();
+	return !HasPlayWorld() && DebuggerCommands::AllowPlayWorldFeature();
 }
 
 bool FInternalPlayWorldCommandCallbacks::CanShowVulkanNonPlayWorldOnlyActions()
 {
-	return !HasPlayWorld() && GetDefault<UEditorExperimentalSettings>()->bAllowVulkanPreview && FModuleManager::Get().ModuleExists(TEXT("VulkanRHI"));
+	return !HasPlayWorld() && GetDefault<UEditorExperimentalSettings>()->bAllowVulkanPreview && FModuleManager::Get().ModuleExists(TEXT("VulkanRHI")) && DebuggerCommands::AllowPlayWorldFeature();
 }
 
 bool FInternalPlayWorldCommandCallbacks::CanShowVROnlyActions()
 {
-	return !HasPlayWorld();
+	return !HasPlayWorld() && DebuggerCommands::AllowPlayWorldFeature();
 }
 
 int32 FInternalPlayWorldCommandCallbacks::GetNumberOfClients()
