Add insight tag to editor feedback context update

diff --git a/Engine/Source/Editor/UnrealEd/Private/FeedbackContextEditor.cpp b/Engine/Source/Editor/UnrealEd/Private/FeedbackContextEditor.cpp
--- a/Engine/Source/Editor/UnrealEd/Private/FeedbackContextEditor.cpp
+++ b/Engine/Source/Editor/UnrealEd/Private/FeedbackContextEditor.cpp
@@ -512,6 +512,8 @@ void FFeedbackContextEditor::FinalizeSlowTask()
 
 void FFeedbackContextEditor::ProgressReported( const float TotalProgressInterp, FText DisplayMessage )
 {
+	TRACE_CPUPROFILER_EVENT_SCOPE(FFeedbackContextEditor::ProgressReported);
+
 	if (!(FPlatformSplash::IsShown() || BuildProgressWidget.IsValid() || SlowTaskWindow.IsValid()))
 	{
 		return;
