Add analytics provider to DumpMaterialShaderTypes to track nightly statistics across various projects.

- Calculate the total shaders in the DefaultMaterials as a good metric.  This impacts editor startup time and is the default materials are the worst case scenario for usage flags.
- Only gather analytics if passed in on the commandline.
- Fix a copy by value on a large array.
- Const-ify some args.
- Add the ability to pass -analytics through the automated test.

# Analytics
We report the following analytics:

	* ProjectName
	* BuildVersion
	* Platform
	* ShaderPlatform
	* TotalShaders
	* TotalMaterials
	* TotalMaterialInstances
	* TotalGlobalShaders
	* TotalDefaultMaterialShaders

diff --git a/Engine/Build/ShaderTypeRegressionTest.xml b/Engine/Build/ShaderTypeRegressionTest.xml
--- a/Engine/Build/ShaderTypeRegressionTest.xml
+++ b/Engine/Build/ShaderTypeRegressionTest.xml
@@ -14,6 +14,14 @@
 	<Option Name="CompileCategory" DefaultValue="Editors" Description="Category string for Compile Label"/>
 	<Option Name="RunCategory" DefaultValue="Test" Description="Category string for Run Label"/>
 	
+	<Property Name="AnalyticsArg" Value=""/>
+	<Option Name="SendAnalytics" DefaultValue="false" Description="Whether or not to shader stat analytics."/>
+	<Do If="$(SendAnalytics)">
+		<Property Name="AnalyticsArg" Value="-analytics"/>
+	</Do>
+	
+	<Property Name="CommandletArguments" Value="-targetplatform=Windows $(AnalyticsArg)"/>
+	
 	<!-- Projects setup-->
 	<Option Name="Projects" DefaultValue="false" Restrict=".+" Description="The projects to test"/>
 	<Error Message="No project names specified" If="'$(Projects)' == 'false'"/>
@@ -55,8 +63,7 @@
 			<!-- Run the commandlet itself. -->
 			<Property Name="RunNodeName" Value="$(Project) Dump Material Shader Types"/>
 			<Node Name="$(RunNodeName)" Requires="$(RunNodeRequires);$(RunNodeList)">
-				<Property Name="Arguments" Value="-targetplatform=Windows"/>
-				<Commandlet Name="DumpMaterialShaderTypes" Project="$(Project)" Arguments="$(Arguments)"/>
+				<Commandlet Name="DumpMaterialShaderTypes" Project="$(Project)" Arguments="$(CommandletArguments)"/>
 			</Node>
 			<Property Name="RunNodeList" Value="$(RunNodeList);$(RunNodeName)"/>
 			
diff --git a/Engine/Source/Editor/UnrealEd/Private/Commandlets/DumpMaterialShaderTypes.cpp b/Engine/Source/Editor/UnrealEd/Private/Commandlets/DumpMaterialShaderTypes.cpp
--- a/Engine/Source/Editor/UnrealEd/Private/Commandlets/DumpMaterialShaderTypes.cpp
+++ b/Engine/Source/Editor/UnrealEd/Private/Commandlets/DumpMaterialShaderTypes.cpp
@@ -1,10 +1,12 @@
 // Copyright Epic Games, Inc. All Rights Reserved.
 
 #include "Commandlets/DumpMaterialShaderTypes.h"
+#include "AnalyticsET.h"
 #include "AssetData.h"
 #include "AssetRegistryModule.h"
 #include "GlobalShader.h"
 #include "HAL/FileManager.h"
+#include "IAnalyticsProviderET.h"
 #include "Interfaces/ITargetPlatform.h"
 #include "Interfaces/ITargetPlatformManagerModule.h"
 #include "Materials/Material.h"
@@ -319,7 +321,7 @@ void PrintDebugShaderInfo(FShaderStatsGatheringContext& Output, const TArray<FDe
 	}
 }
 
-int ProcessMaterials(const ITargetPlatform* TargetPlatform, const EShaderPlatform ShaderPlatform, FShaderStatsGatheringContext& Output, TArray<FAssetData>& MaterialList)
+int ProcessMaterials(const ITargetPlatform* TargetPlatform, const EShaderPlatform ShaderPlatform, FShaderStatsGatheringContext& Output, const TArray<FAssetData>& MaterialList)
 {
 	TRACE_CPUPROFILER_EVENT_SCOPE(ProcessMaterials);
 
@@ -350,7 +352,7 @@ int ProcessMaterials(const ITargetPlatform* TargetPlatform, const EShaderPlatfor
 	return TotalShaders;
 }
 
-int ProcessMaterialInstances(const ITargetPlatform* TargetPlatform, const EShaderPlatform ShaderPlatform, FShaderStatsGatheringContext& Output, TArray<FAssetData>& MaterialInstanceList)
+int ProcessMaterialInstances(const ITargetPlatform* TargetPlatform, const EShaderPlatform ShaderPlatform, FShaderStatsGatheringContext& Output, const TArray<FAssetData>& MaterialInstanceList)
 {
 	TRACE_CPUPROFILER_EVENT_SCOPE(ProcessMaterialInstances);
 
@@ -458,7 +460,7 @@ int ProcessGlobalShaders(const ITargetPlatform* TargetPlatform, const EShaderPla
 	return TotalShaders;
 }
 
-void ProcessForTargetAndShaderPlatform(const ITargetPlatform* TargetPlatform, const EShaderPlatform ShaderPlatform, const FString& Params, TArray<FAssetData>& MaterialList, TArray<FAssetData> MaterialInstanceList)
+void ProcessForTargetAndShaderPlatform(const ITargetPlatform* TargetPlatform, const EShaderPlatform ShaderPlatform, const FString& Params, const TArray<FAssetData>& MaterialList, const TArray<FAssetData>& MaterialInstanceList, TSharedPtr<IAnalyticsProviderET> Provider)
 {
 	TRACE_CPUPROFILER_EVENT_SCOPE(ProcessForTargetAndShaderPlatform);
 
@@ -477,16 +479,47 @@ void ProcessForTargetAndShaderPlatform(const ITargetPlatform* TargetPlatform, co
 	TotalShaders += ProcessMaterialInstances(TargetPlatform, ShaderPlatform, Output, MaterialInstanceList);
 	TotalAssets += MaterialInstanceList.Num();
 
-	TotalShaders += ProcessGlobalShaders(TargetPlatform, ShaderPlatform, Output);
+	const int TotalGlobalShaders = ProcessGlobalShaders(TargetPlatform, ShaderPlatform, Output);
+	TotalShaders += TotalGlobalShaders;
+
+	int TotalDefaultMaterialShaders = 0;
+	{
+		for (int32 Domain = 0; Domain < MD_MAX; ++Domain)
+		{
+			UMaterial* Material = UMaterial::GetDefaultMaterial((EMaterialDomain)Domain);
+			if (Material)
+			{
+				TArray<FDebugShaderTypeInfo> OutShaderInfo;
+				Material->GetShaderTypes(ShaderPlatform, TargetPlatform, OutShaderInfo);
+				TotalDefaultMaterialShaders += GetTotalShaders(OutShaderInfo);
+			}
+		}
+	}
 
 	Output.Log(TEXT(""));
 	Output.Log(TEXT("Summary"));
 	Output.Log(FString::Printf(TEXT("Total Assets: %d"), TotalAssets));
 	Output.Log(FString::Printf(TEXT("Total Shaders: %d"), TotalShaders));
+	Output.Log(FString::Printf(TEXT("Total Default Material Shaders: %d"), TotalDefaultMaterialShaders));
 	Output.Log(FString::Printf(TEXT("Histogram:")));
 	Output.PrintHistogram(TotalShaders);
 	Output.Log(FString::Printf(TEXT("\nAlphabetic list of types:")));
 	Output.PrintAlphabeticList();
+
+	if (Provider.IsValid())
+	{
+		Provider->RecordEvent(TEXT("DumpMaterialShaderTypes"), MakeAnalyticsEventAttributeArray(
+			TEXT("ProjectName"), FApp::GetProjectName(),
+			TEXT("BuildVersion"), FApp::GetBuildVersion(),
+			TEXT("Platform"), TargetPlatform->PlatformName(),
+			TEXT("ShaderPlatform"), LexToString(ShaderPlatform),
+			TEXT("TotalShaders"), TotalShaders,
+			TEXT("TotalMaterials"), MaterialList.Num(),
+			TEXT("TotalMaterialInstances"), MaterialInstanceList.Num(),
+			TEXT("TotalGlobalShaders"), TotalGlobalShaders,
+			TEXT("TotalDefaultMaterialShaders"), TotalDefaultMaterialShaders
+			));
+	}
 }
 
 int32 UDumpMaterialShaderTypesCommandlet::Main(const FString& Params)
@@ -504,9 +537,12 @@ int32 UDumpMaterialShaderTypesCommandlet::Main(const FString& Params)
 		UE_LOG(LogDumpMaterialShaderTypesCommandlet, Log, TEXT("Options:"));
 		UE_LOG(LogDumpMaterialShaderTypesCommandlet, Log, TEXT(" Required: -targetplatform=<platform(s)>     (Which target platform do you want results, e.g. WindowsClient, WindowsEditor. Multiple shader platforms are allowed)."));
 		UE_LOG(LogDumpMaterialShaderTypesCommandlet, Log, TEXT(" Optional: -collection=<name>                (You can also specify a collection of assets to narrow down the results e.g. if you maintain a collection that represents the actually used in-game assets)."));
+		UE_LOG(LogDumpMaterialShaderTypesCommandlet, Log, TEXT(" Optional: -analytics                        (Whether or not to send analytics data for tracking purposes)."));
 		return 0;
 	}
 
+	const bool bSendAnalytics = FParse::Param(FCommandLine::Get(), TEXT("analytics"));
+
 	const double AssetRegistryStart = FPlatformTime::Seconds();
 
 	TArray<FAssetData> MaterialList;
@@ -564,6 +600,27 @@ int32 UDumpMaterialShaderTypesCommandlet::Main(const FString& Params)
 	ITargetPlatformManagerModule* TPM = GetTargetPlatformManager();
 	const TArray<ITargetPlatform*>& Platforms = TPM->GetActiveTargetPlatforms();
 
+	TSharedPtr<IAnalyticsProviderET> Provider;
+	if (bSendAnalytics)
+	{
+		FAnalyticsET::Config Config;
+		Config.APIKeyET = FString(TEXT("StudioAnalytics.Dev"));
+		Config.APIServerET = FString(TEXT("https://datarouter.ol.epicgames.com/"));
+		Config.AppVersionET = FEngineVersion::Current().ToString();
+
+		// There are other things to configure, but the default are usually fine.
+		Provider = FAnalyticsET::Get().CreateAnalyticsProvider(Config);
+		if (Provider.IsValid())
+		{
+			const FString UserID = FPlatformProcess::UserName(false);
+			Provider->SetUserID(UserID);
+			Provider->StartSession(MakeAnalyticsEventAttributeArray(
+				TEXT("ProjectName"), FApp::GetProjectName(),
+				TEXT("Version"), FApp::GetBuildVersion()
+			));
+		}
+	}
+
 	for (int32 Index = 0; Index < Platforms.Num(); Index++)
 	{
 		TArray<FName> DesiredShaderFormats;
@@ -574,7 +631,7 @@ int32 UDumpMaterialShaderTypesCommandlet::Main(const FString& Params)
 			const EShaderPlatform ShaderPlatform = ShaderFormatToLegacyShaderPlatform(DesiredShaderFormats[FormatIndex]);
 
 			UE_LOG(LogDumpMaterialShaderTypesCommandlet, Display, TEXT("Dumping material shader types for '%s' - '%s'..."), *Platforms[Index]->PlatformName(), *LexToString(ShaderPlatform));
-			ProcessForTargetAndShaderPlatform(Platforms[Index], ShaderPlatform, Params, MaterialList, MaterialInstanceList);
+			ProcessForTargetAndShaderPlatform(Platforms[Index], ShaderPlatform, Params, MaterialList, MaterialInstanceList, Provider);
 		}
 	}
 
