Minimap for Open World Template maps is not displayed in the World Partition Editor outside of QAGame

* Avoid UDIM setup for the minimap as there are no fallback possible to normal texture if Virtual Texture support is disabled in the project settings
* Rebuilt minimap for the default OpenWorld template to ensure it is not UDIM based and wil work outside of QAGame
* Added warning if a minimap exceeds the maximum texture resolution & VT support is disabled.
* Fix WorldPartitionBuilder cell iteration sometimes performed on extra cells outside of the world bound
* When running World Partition commandlets from the editor, provide the -unattended flag (to avoid blocking popups) and hide the subprocess window.

diff --git a/Engine/Source/Editor/UnrealEd/Private/WorldPartition/WorldPartitionBuilder.cpp b/Engine/Source/Editor/UnrealEd/Private/WorldPartition/WorldPartitionBuilder.cpp
--- a/Engine/Source/Editor/UnrealEd/Private/WorldPartition/WorldPartitionBuilder.cpp
+++ b/Engine/Source/Editor/UnrealEd/Private/WorldPartition/WorldPartitionBuilder.cpp
@@ -141,6 +141,15 @@ static FIntVector GetCellCoord(const FVector& InPos, const int32 InCellSize)
 	);
 }
 
+static FIntVector GetCellCount(const FBox& InBounds, const int32 InCellSize)
+{
+	const FIntVector MinCellCoords = GetCellCoord(InBounds.Min, InCellSize);
+	const FIntVector MaxCellCoords = FIntVector(FMath::CeilToInt(InBounds.Max.X / InCellSize),
+												FMath::CeilToInt(InBounds.Max.Y / InCellSize),
+												FMath::CeilToInt(InBounds.Max.Z / InCellSize));
+	return MaxCellCoords - MinCellCoords;
+}
+
 bool UWorldPartitionBuilder::Run(UWorld* World, FPackageSourceControlHelper& PackageHelper)
 {
 	// Notify derived classes that partition building process starts
@@ -192,20 +201,22 @@ bool UWorldPartitionBuilder::Run(UWorld* World, FPackageSourceControlHelper& Pac
 	if ((LoadingMode == ELoadingMode::IterativeCells) || (LoadingMode == ELoadingMode::IterativeCells2D))
 	{
 		// do partial loading loop that calls RunInternal
-		auto CanIterateZ = [](const bool bInResult, const ELoadingMode InLoadingMode, const int32 InZ, const int32 InMinZ, const int32 InMaxZ) -> bool
+		const FIntVector MinCellCoords = GetCellCoord(CellInfo.EditorBounds.Min, IterativeCellSize);
+		const FIntVector NumCellsIterations = GetCellCount(CellInfo.EditorBounds, IterativeCellSize);
+		const FIntVector BeginCellCoords = MinCellCoords;
+		const FIntVector EndCellCoords = BeginCellCoords + NumCellsIterations;
+
+		auto CanIterateZ = [&BeginCellCoords, &EndCellCoords, LoadingMode](const bool bInResult, const int32 InZ) -> bool
 		{
-			if (InLoadingMode == ELoadingMode::IterativeCells2D)
+			if (LoadingMode == ELoadingMode::IterativeCells2D)
 			{
-				return bInResult && (InZ == InMinZ);
+				return bInResult && (InZ == BeginCellCoords.Z);
 			}
 
-			return bInResult && (InZ <= InMaxZ);
+			return bInResult && (InZ < EndCellCoords.Z);
 		};
 
-		const FIntVector MinCellCoords = GetCellCoord(CellInfo.EditorBounds.Min, IterativeCellSize);
-		const FIntVector MaxCellCoords = GetCellCoord(CellInfo.EditorBounds.Max, IterativeCellSize);
-
-		const int32 IterationCount = ((LoadingMode == ELoadingMode::IterativeCells2D) ? 1 : (MaxCellCoords.Z - MinCellCoords.Z + 1)) * (MaxCellCoords.Y - MinCellCoords.Y + 1) * (MaxCellCoords.X - MinCellCoords.X + 1);
+		const int32 IterationCount = ((LoadingMode == ELoadingMode::IterativeCells2D) ? 1 : NumCellsIterations.Z) * NumCellsIterations.Y * NumCellsIterations.X;
 		int32 IterationIndex = 0;
 
 		UE_LOG(LogWorldPartitionBuilder, Display, TEXT("Iterative Cell Mode"));
@@ -216,11 +227,11 @@ bool UWorldPartitionBuilder::Run(UWorld* World, FPackageSourceControlHelper& Pac
 		
 		FBox LoadedBounds(ForceInit);
 
-		for (int32 z = MinCellCoords.Z; CanIterateZ(bResult, LoadingMode, z, MinCellCoords.Z, MaxCellCoords.Z); z++)
+		for (int32 z = BeginCellCoords.Z; CanIterateZ(bResult, z); z++)
 		{
-			for (int32 y = MinCellCoords.Y; bResult && (y <= MaxCellCoords.Y); y++)
+			for (int32 y = BeginCellCoords.Y; bResult && (y < EndCellCoords.Y); y++)
 			{
-				for (int32 x = MinCellCoords.X; bResult && (x <= MaxCellCoords.X); x++)
+				for (int32 x = BeginCellCoords.X; bResult && (x < EndCellCoords.X); x++)
 				{
 					IterationIndex++;
 					UE_LOG(LogWorldPartitionBuilder, Display, TEXT("[%d / %d] Processing cells..."), IterationIndex, IterationCount);
diff --git a/Engine/Source/Editor/UnrealEd/Private/WorldPartition/WorldPartitionMiniMapBuilder.cpp b/Engine/Source/Editor/UnrealEd/Private/WorldPartition/WorldPartitionMiniMapBuilder.cpp
--- a/Engine/Source/Editor/UnrealEd/Private/WorldPartition/WorldPartitionMiniMapBuilder.cpp
+++ b/Engine/Source/Editor/UnrealEd/Private/WorldPartition/WorldPartitionMiniMapBuilder.cpp
@@ -78,18 +78,18 @@ bool UWorldPartitionMiniMapBuilder::PreRun(UWorld* World, FPackageSourceControlH
 
 bool UWorldPartitionMiniMapBuilder::RunInternal(UWorld* World, const FCellInfo& InCellInfo, FPackageSourceControlHelper& PackageHelper)
 {
-	FMinimapTile MinimapTile;
-	
-	MinimapTile.Coordinates.X = InCellInfo.Location.X;
-	MinimapTile.Coordinates.Y = InCellInfo.Location.Y;
-
 	EditorBounds = InCellInfo.EditorBounds;
 	IterativeCellSize = InCellInfo.IterativeCellSize;
 
 	FString TextureName = FString::Format(TEXT("MinimapTile_{0}_{1}_{2}"), { InCellInfo.Location.X, InCellInfo.Location.Y, InCellInfo.Location.Z });
 
-	FWorldPartitionMiniMapHelper::CaptureBoundsMiniMapToTexture(World, WorldMiniMap, WorldMiniMap->MiniMapTileSize, static_cast<UTexture2D*&>(MinimapTile.Texture), TextureName, InCellInfo.Bounds);
+	UTexture2D* TileTexture = nullptr;
+	FWorldPartitionMiniMapHelper::CaptureBoundsMiniMapToTexture(World, WorldMiniMap, WorldMiniMap->MiniMapTileSize, TileTexture, TextureName, InCellInfo.Bounds);
 
+	FMinimapTile MinimapTile;
+	MinimapTile.Texture = TStrongObjectPtr<UTexture2D>(TileTexture);
+	MinimapTile.Coordinates.X = InCellInfo.Location.X;
+	MinimapTile.Coordinates.Y = InCellInfo.Location.Y;
 	MiniMapTiles.Add(MoveTemp(MinimapTile));
 
 	return true;
@@ -97,9 +97,6 @@ bool UWorldPartitionMiniMapBuilder::RunInternal(UWorld* World, const FCellInfo&
 
 bool UWorldPartitionMiniMapBuilder::PostRun(UWorld* World, FPackageSourceControlHelper& PackageHelper, const bool bInRunSuccess)
 {
-	TArray<FTextureSourceBlock> SourceBlocks;
-	TArray<const uint8*>		SourceImageData;
-
 	if (!bInRunSuccess)
 	{
 		return false;
@@ -116,63 +113,74 @@ bool UWorldPartitionMiniMapBuilder::PostRun(UWorld* World, FPackageSourceControl
 		return false;
 	}
 
-	SourceBlocks.Reserve(MiniMapTiles.Num());
-	SourceImageData.Reserve(MiniMapTiles.Num());
-
-	FTileCoordinates Min = { FMath::FloorToInt32(EditorBounds.Min.X / IterativeCellSize), FMath::FloorToInt32(EditorBounds.Min.Y / IterativeCellSize) };
-	FTileCoordinates Max = { FMath::FloorToInt32(EditorBounds.Max.X / IterativeCellSize), FMath::FloorToInt32(EditorBounds.Max.Y / IterativeCellSize) };
-
-	FVector2D				Offset;
-	ETextureSourceFormat	Format = TSF_BGRA8;
-
-	for (FMinimapTile& Tile : MiniMapTiles)
+	FIntVector2 MinTileCoords(INT32_MAX);
+	FIntVector2 MaxTileCoords(-INT32_MAX);
+	for (const FMinimapTile& Tile : MiniMapTiles)
 	{
-		FTextureSourceBlock* Block = new(SourceBlocks) FTextureSourceBlock();
-		bool bIsOrigin = (Tile.Coordinates.X == 0) && (Tile.Coordinates.Y == 0);
-
-		Format = (Tile.Texture != nullptr) ? Tile.Texture->Source.GetFormat() : Format;
-
-		// Offset X coordinates and reverse Y
-		Tile.Coordinates.X = Tile.Coordinates.X - Min.X;
-		Tile.Coordinates.Y = Max.Y - Tile.Coordinates.Y;
-
-		if (bIsOrigin)
-		{
-			Offset.X = Tile.Coordinates.X;
-			Offset.Y = -Tile.Coordinates.Y;
-		}
-
-		Block->BlockX = Tile.Coordinates.X;
-		Block->BlockY = Tile.Coordinates.Y;
-		Block->SizeX = Tile.Texture->GetSizeX();
-		Block->SizeY = Tile.Texture->GetSizeY();
-		Block->NumSlices = 1;
-		Block->NumMips = 1;
-
-		const uint8* DataPtr = Tile.Texture->Source.LockMip(0);
-		SourceImageData.Add(DataPtr);
+		MinTileCoords = FIntVector2(FMath::Min(MinTileCoords.X, Tile.Coordinates.X), FMath::Min(MinTileCoords.Y, Tile.Coordinates.Y));
+		MaxTileCoords = FIntVector2(FMath::Max(MaxTileCoords.X, Tile.Coordinates.X), FMath::Max(MaxTileCoords.Y, Tile.Coordinates.Y));
 	}
 
-	TStrongObjectPtr<UTextureFactory> Factory(NewObject<UTextureFactory>());
-
-	WorldMiniMap->MiniMapTexture = Factory->CreateTexture2D(WorldMiniMap, TEXT("MinimapTexture"), RF_NoFlags);
-	WorldMiniMap->MiniMapTexture->Source.InitBlocked(&Format, SourceBlocks.GetData(), 1, SourceBlocks.Num(), SourceImageData.GetData());
-
-	for (const FMinimapTile& Tile : MiniMapTiles)
+	const FIntVector2 NumTilesCoords = FIntVector2(MaxTileCoords.X - MinTileCoords.X + 1, MaxTileCoords.Y - MinTileCoords.Y + 1);
+	const int32 MinimapTileSize = WorldMiniMap->MiniMapTileSize;
+	const int32 MinimapImageSizeX = NumTilesCoords.X * MinimapTileSize;
+	const int32 MinimapImageSizeY = NumTilesCoords.Y * MinimapTileSize;
+	
+	// Copy all tiles to the final minimap texture
 	{
-		Tile.Texture->Source.UnlockMip(0);
+		TStrongObjectPtr<UTextureFactory> Factory(NewObject<UTextureFactory>());
+		UTexture2D* MiniMapTexture = Factory->CreateTexture2D(WorldMiniMap, TEXT("MinimapTexture"), RF_NoFlags);
+		MiniMapTexture->Source.Init(MinimapImageSizeX, MinimapImageSizeY, 1, 1, TSF_BGRA8);
+
+	    const uint32 DstDataStride = MiniMapTexture->Source.GetSizeX() * MiniMapTexture->Source.GetBytesPerPixel();
+	    uint8* DstDataBasePtr = MiniMapTexture->Source.LockMip(0);
+    
+	    for (const FMinimapTile& Tile : MiniMapTiles)
+	    {
+		    const uint32 SrcDataStride = MinimapTileSize * Tile.Texture->Source.GetBytesPerPixel();
+		    const uint8* SrcDataPtr = Tile.Texture->Source.LockMipReadOnly(0);
+    
+		    const int32 TileOffsetX = Tile.Coordinates.X - MinTileCoords.X;
+		    const int32 TileOffsetY = Tile.Coordinates.Y - MinTileCoords.Y;
+    
+		    uint8* DstDataPtr = DstDataBasePtr + (TileOffsetY * MinimapTileSize * DstDataStride) + (TileOffsetX * SrcDataStride);
+    
+		    for (int RowIdx = 0; RowIdx < MinimapTileSize; ++RowIdx)
+		    {
+			    FMemory::Memcpy(DstDataPtr + DstDataStride * RowIdx,
+							    SrcDataPtr + SrcDataStride * RowIdx,
+							    SrcDataStride);
+		    }
+    
+		    Tile.Texture->Source.UnlockMip(0);
+	    }
+
+		MiniMapTexture->Source.UnlockMip(0);
+		MiniMapTexture->PowerOfTwoMode = ETexturePowerOfTwoSetting::PadToPowerOfTwo;	// Required for VTs
+		MiniMapTexture->AdjustMinAlpha = 1.f;
+		MiniMapTexture->LODGroup = TEXTUREGROUP_UI;
+		MiniMapTexture->VirtualTextureStreaming = true;
+		MiniMapTexture->UpdateResource();
+
+		WorldMiniMap->MiniMapTexture = MiniMapTexture;
 	}
 
-	WorldMiniMap->MiniMapTexture->AdjustMinAlpha = 1.f;
-	WorldMiniMap->MiniMapTexture->LODGroup = TEXTUREGROUP_UI;
-	WorldMiniMap->MiniMapTexture->VirtualTextureStreaming = true;
-	WorldMiniMap->MiniMapTexture->UpdateResource();
+	// Compute relevant UV space for the minimap
+	{
+		FBox2D MinimapWorldBounds = FBox2D({ EditorBounds.Min.X, EditorBounds.Min.Y }, { EditorBounds.Max.X, EditorBounds.Max.Y });
+		FBox2D ImageWorldBounds = FBox2D(FVector2D(MinTileCoords.X, MinTileCoords.Y) * IterativeCellSize, FVector2D(MaxTileCoords.X + 1, MaxTileCoords.Y + 1) * IterativeCellSize);
+
+		FVector2D TexturePOW2ScaleFactor = FVector2D((float)MinimapImageSizeX / FMath::RoundUpToPowerOfTwo(MinimapImageSizeX),
+			(float)MinimapImageSizeY / FMath::RoundUpToPowerOfTwo(MinimapImageSizeY));
 
-	FBox2D UVOffset(FVector2D(EditorBounds.Min.X / IterativeCellSize + Offset.X, EditorBounds.Min.Y / IterativeCellSize + Offset.Y),
-					FVector2D(EditorBounds.Max.X / IterativeCellSize + Offset.X, EditorBounds.Max.Y / IterativeCellSize + Offset.Y));
+		FBox2D UVOffset(FVector2D((MinimapWorldBounds.Min - ImageWorldBounds.Min) / ImageWorldBounds.GetSize()),
+			FVector2D((MinimapWorldBounds.Max - ImageWorldBounds.Min) / ImageWorldBounds.GetSize()));
 
-	UVOffset.bIsValid = true;
-	WorldMiniMap->UVOffset = UVOffset;
+		UVOffset.Min *= TexturePOW2ScaleFactor;
+		UVOffset.Max *= TexturePOW2ScaleFactor;
+
+		WorldMiniMap->UVOffset = UVOffset;
+	}
 
 	// Make sure the minimap texture is ready before saving
 	FAssetCompilingManager::Get().FinishAllCompilation();
@@ -222,4 +230,4 @@ bool UWorldPartitionMiniMapBuilder::PostRun(UWorld* World, FPackageSourceControl
 	}
 
 	return true;
-}
+}
\ No newline at end of file
diff --git a/Engine/Source/Editor/UnrealEd/Public/WorldPartition/WorldPartitionMiniMapBuilder.h b/Engine/Source/Editor/UnrealEd/Public/WorldPartition/WorldPartitionMiniMapBuilder.h
--- a/Engine/Source/Editor/UnrealEd/Public/WorldPartition/WorldPartitionMiniMapBuilder.h
+++ b/Engine/Source/Editor/UnrealEd/Public/WorldPartition/WorldPartitionMiniMapBuilder.h
@@ -2,11 +2,24 @@
 
 #pragma once
 
+#include "CoreMinimal.h"
+#include "UObject/StrongObjectPtr.h"
 #include "WorldPartition/WorldPartitionBuilder.h"
 #include "WorldPartition/WorldPartitionMiniMap.h"
 
 #include "WorldPartitionMiniMapBuilder.generated.h"
 
+
+/** Structure holding Minimap's tile resources */
+struct FMinimapTile
+{
+	/** Tile's texture */
+	TStrongObjectPtr<UTexture2D> Texture;
+
+	/** Coordinates of the tile in the Minimap's Virtual Texture */
+	FIntVector2 Coordinates;
+};
+
 UCLASS(config = Engine, defaultconfig)
 class UNREALED_API UWorldPartitionMiniMapBuilder : public UWorldPartitionBuilder
 {
@@ -26,11 +39,9 @@ protected:
 
 private:
 	/*MiniMap Texture Tiles for displaying on world partition window*/
-	UPROPERTY(Transient, VisibleAnywhere, Category = WorldPartitionMiniMap)
 	TArray<FMinimapTile> MiniMapTiles;
+	TObjectPtr<AWorldPartitionMiniMap> WorldMiniMap;
 
-	UPROPERTY(Transient, VisibleAnywhere, Category = WorldPartitionMiniMap)
-	TObjectPtr<AWorldPartitionMiniMap> WorldMiniMap = nullptr;
 	FBox EditorBounds;
 	bool bAutoSubmit;
 };
\ No newline at end of file
diff --git a/Engine/Source/Editor/WorldPartitionEditor/Private/WorldPartition/SWorldPartitionEditorGridSpatialHash.cpp b/Engine/Source/Editor/WorldPartitionEditor/Private/WorldPartition/SWorldPartitionEditorGridSpatialHash.cpp
--- a/Engine/Source/Editor/WorldPartitionEditor/Private/WorldPartition/SWorldPartitionEditorGridSpatialHash.cpp
+++ b/Engine/Source/Editor/WorldPartitionEditor/Private/WorldPartition/SWorldPartitionEditorGridSpatialHash.cpp
@@ -34,7 +34,6 @@
 
 SWorldPartitionEditorGridSpatialHash::SWorldPartitionEditorGridSpatialHash()
 	: WorldMiniMapBounds(ForceInit)
-	, WorldMinimapUsesVirtualTexture(false)
 {
 }
 
@@ -49,9 +48,6 @@ void SWorldPartitionEditorGridSpatialHash::Construct(const FArguments& InArgs)
 
 	if (WorldPartition)
 	{
-		UWorldPartitionEditorSpatialHash* EditorSpatialHash = (UWorldPartitionEditorSpatialHash*)WorldPartition->EditorHash;
-
-		//Update MiniMap data for drawing  
 		UpdateWorldMiniMapDetails();
 
 		bShowActors = !WorldMiniMapBrush.HasUObject();
@@ -60,38 +56,16 @@ void SWorldPartitionEditorGridSpatialHash::Construct(const FArguments& InArgs)
 	SWorldPartitionEditorGrid2D::Construct(SWorldPartitionEditorGrid::FArguments().InWorld(InArgs._InWorld));
 }
 
-void SWorldPartitionEditorGridSpatialHash::Tick(const FGeometry& AllottedGeometry, const double InCurrentTime, const float InDeltaTime)
-{
-	SWorldPartitionEditorGrid2D::Tick(AllottedGeometry, InCurrentTime, InDeltaTime);
-
-	UTexture2D* Texture2D = Cast<UTexture2D>(WorldMiniMapBrush.GetResourceObject());
-
-	if ((Texture2D != nullptr) && (Texture2D->IsCurrentlyVirtualTextured() != WorldMinimapUsesVirtualTexture))
-	{
-		UpdateWorldMiniMapDetails();
-	}
-}
-
 void SWorldPartitionEditorGridSpatialHash::UpdateWorldMiniMapDetails()
 {
-	auto WorldMiniMap = FWorldPartitionMiniMapHelper::GetWorldPartitionMiniMap(World);
+	AWorldPartitionMiniMap* WorldMiniMap = FWorldPartitionMiniMapHelper::GetWorldPartitionMiniMap(World);
 	if (WorldMiniMap)
 	{
 		WorldMiniMapBounds = FBox2D(FVector2D(WorldMiniMap->MiniMapWorldBounds.Min), FVector2D(WorldMiniMap->MiniMapWorldBounds.Max));
 		if (UTexture2D* MiniMapTexture = WorldMiniMap->MiniMapTexture)
 		{
-			WorldMinimapUsesVirtualTexture = MiniMapTexture->IsCurrentlyVirtualTextured();
-
-			if (WorldMinimapUsesVirtualTexture)
-			{
-				WorldMiniMapBrush.SetUVRegion(WorldMiniMap->UVOffset);
-			}
-			else
-			{
-				WorldMiniMapBrush.SetUVRegion(FBox2D({ 0.f, 0.f }, { 1.f, 1.f }));
-			}
-
-			WorldMiniMapBrush.SetImageSize(FVector2D(MiniMapTexture->GetSizeX(), MiniMapTexture->GetSizeY()));
+			WorldMiniMapBrush.SetUVRegion(WorldMiniMap->UVOffset);
+			WorldMiniMapBrush.SetImageSize(MiniMapTexture->GetImportedSize());
 			WorldMiniMapBrush.SetResourceObject(MiniMapTexture);
 		}
 	}
diff --git a/Engine/Source/Editor/WorldPartitionEditor/Private/WorldPartitionEditorModule.cpp b/Engine/Source/Editor/WorldPartitionEditor/Private/WorldPartitionEditorModule.cpp
--- a/Engine/Source/Editor/WorldPartitionEditor/Private/WorldPartitionEditorModule.cpp
+++ b/Engine/Source/Editor/WorldPartitionEditor/Private/WorldPartitionEditorModule.cpp
@@ -278,10 +278,13 @@ static void RunCommandletAsExternalProcess(const FString& InCommandletArgs, cons
 	FString ProjectPath = FPaths::IsProjectFilePathSet() ? FPaths::GetProjectFilePath() : FApp::GetProjectName();
 
 	uint32 ProcessID;
-	FString Arguments = FString::Printf(TEXT("\"%s\" %s"), *ProjectPath, *InCommandletArgs);
+	FString Arguments = FString::Printf(TEXT("\"%s\" %s -unattended"), *ProjectPath, *InCommandletArgs);
 	
 	UE_LOG(LogWorldPartitionEditor, Display, TEXT("Running commandlet: %s %s"), *CurrentExecutableName, *Arguments);
-	ProcessHandle = FPlatformProcess::CreateProc(*CurrentExecutableName, *Arguments, true, false, false, &ProcessID, 0, nullptr, WritePipe, ReadPipe);
+	const bool bLaunchDetached = true;
+	const bool bLaunchHidden = true;
+	const bool bLaunchReallyHidden = true;
+	ProcessHandle = FPlatformProcess::CreateProc(*CurrentExecutableName, *Arguments, bLaunchDetached, bLaunchHidden, bLaunchReallyHidden, &ProcessID, 0, nullptr, WritePipe, ReadPipe);
 
 	while (FPlatformProcess::IsProcRunning(ProcessHandle))
 	{
diff --git a/Engine/Source/Editor/WorldPartitionEditor/Public/WorldPartition/SWorldPartitionEditorGridSpatialHash.h b/Engine/Source/Editor/WorldPartitionEditor/Public/WorldPartition/SWorldPartitionEditorGridSpatialHash.h
--- a/Engine/Source/Editor/WorldPartitionEditor/Public/WorldPartition/SWorldPartitionEditorGridSpatialHash.h
+++ b/Engine/Source/Editor/WorldPartitionEditor/Public/WorldPartition/SWorldPartitionEditorGridSpatialHash.h
@@ -16,13 +16,10 @@ public:
 
 	void Construct(const FArguments& InArgs);
 
-	virtual void Tick(const FGeometry& AllottedGeometry, const double InCurrentTime, const float InDeltaTime) override;
-
 	int32 PaintGrid(const FGeometry& AllottedGeometry, const FSlateRect& MyCullingRect, FSlateWindowElementList& OutDrawElements, int32 LayerId) const override;
 
 private:
 	void UpdateWorldMiniMapDetails();
 	FBox2D          WorldMiniMapBounds;
 	FSlateBrush		WorldMiniMapBrush;
-	bool			WorldMinimapUsesVirtualTexture;
 };
diff --git a/Engine/Source/Runtime/Engine/Private/WorldPartition/WorldPartitionMiniMap.cpp b/Engine/Source/Runtime/Engine/Private/WorldPartition/WorldPartitionMiniMap.cpp
--- a/Engine/Source/Runtime/Engine/Private/WorldPartition/WorldPartitionMiniMap.cpp
+++ b/Engine/Source/Runtime/Engine/Private/WorldPartition/WorldPartitionMiniMap.cpp
@@ -1,6 +1,12 @@
 // Copyright Epic Games, Inc. All Rights Reserved.
 
 #include "WorldPartition/WorldPartitionMiniMap.h"
+#include "Engine/Texture2D.h"
+#include "Logging/MessageLog.h"
+#include "Misc/MapErrors.h"
+#include "Misc/UObjectToken.h"
+#include "RenderUtils.h"
+#include "RHI.h"
 
 #define LOCTEXT_NAMESPACE "WorldPartitionEditor"
 
@@ -12,6 +18,26 @@ AWorldPartitionMiniMap::AWorldPartitionMiniMap(const FObjectInitializer& ObjectI
 }
 
 #if WITH_EDITOR
+void AWorldPartitionMiniMap::CheckForErrors()
+{
+	Super::CheckForErrors();
+
+	if (MiniMapTexture != nullptr)
+	{
+		const int32 MaxTextureDimension = GetMax2DTextureDimension();
+		const bool bExceedMaxDimension = MiniMapTexture->GetImportedSize().GetMax() > MaxTextureDimension;
+		if (!UseVirtualTexturing(GMaxRHIFeatureLevel) && bExceedMaxDimension)
+		{
+			FFormatNamedArguments Arguments;
+			Arguments.Add(TEXT("ActorName"), FText::FromString(GetName()));
+			FMessageLog("MapCheck").Warning()
+				->AddToken(FUObjectToken::Create(this))
+				->AddToken(FTextToken::Create(FText::Format(LOCTEXT("MapCheck_Message_MinimapTextureSize", "{ActorName} : Texture size is too big, minimap won't be rendered. Reduce the MiniMapTileSize property or enable Virtual Texture Support for your project."), Arguments)))
+				->AddToken(FMapErrorToken::Create("MinimapTextureSize"));
+		}
+	}
+}
+
 void AWorldPartitionMiniMap::PostEditChangeProperty(FPropertyChangedEvent& PropertyChangedEvent)
 {
 	MiniMapTileSize = FMath::Clamp<uint32>(FMath::RoundUpToPowerOfTwo(MiniMapTileSize), 256, 8192);
diff --git a/Engine/Source/Runtime/Engine/Public/WorldPartition/WorldPartitionMiniMap.h b/Engine/Source/Runtime/Engine/Public/WorldPartition/WorldPartitionMiniMap.h
--- a/Engine/Source/Runtime/Engine/Public/WorldPartition/WorldPartitionMiniMap.h
+++ b/Engine/Source/Runtime/Engine/Public/WorldPartition/WorldPartitionMiniMap.h
@@ -6,36 +6,6 @@
 #include "WorldPartition/DataLayer/ActorDataLayer.h"
 #include "WorldPartitionMiniMap.generated.h"
 
-/** Coordinates of the tile in Minimap texture */
-USTRUCT()
-struct FTileCoordinates
-{
-	GENERATED_BODY()
-
-	UPROPERTY(VisibleAnywhere, Category = WorldPartitionMiniMap)
-	int32 X = 0;
-
-	UPROPERTY(VisibleAnywhere, Category = WorldPartitionMiniMap)
-	int32 Y = 0;
-};
-
-/** Structure holding Minimap's tile resources */
-USTRUCT()
-struct FMinimapTile
-{
-	GENERATED_BODY()
-
-	FMinimapTile() : Texture(nullptr) {}
-
-	/** Tile's texture */
-	UPROPERTY(VisibleAnywhere, Category = WorldPartitionMiniMap)
-	TObjectPtr<UTexture2D> Texture;
-
-	/** Coordinates of the tile in the Minimap's Virtual Texture */
-	UPROPERTY(VisibleAnywhere, Category = WorldPartitionMiniMap)
-	FTileCoordinates Coordinates;
-};
-
 /**
  * A mini map to preview the world in world partition window. (editor-only)
  */
@@ -54,6 +24,7 @@ public:
 	virtual bool IsLockLocation() const { return true; }
 	virtual bool IsUserManaged() const final { return false; }
 	virtual bool SupportsDataLayer() const final { return false; }
+	virtual void CheckForErrors() override;
 #endif
 
 	/*WorldBounds for MinMapTexture*/
