Fixed warning "DoesPackageExist called on PackageName that will always return false" in PIE

diff --git a/Engine/Source/Editor/UnrealEd/Private/AssetThumbnail.cpp b/Engine/Source/Editor/UnrealEd/Private/AssetThumbnail.cpp
--- a/Engine/Source/Editor/UnrealEd/Private/AssetThumbnail.cpp
+++ b/Engine/Source/Editor/UnrealEd/Private/AssetThumbnail.cpp
@@ -1555,7 +1555,7 @@ void FAssetThumbnailPool::DirtyThumbnailForObject(UObject* ObjectBeingModified)
 
 	if (ObjectBeingModified->HasAnyFlags(RF_ClassDefaultObject))
 	{
-		if (ObjectBeingModified->GetClass()->ClassGeneratedBy != NULL)
+		if (ObjectBeingModified->GetClass()->ClassGeneratedBy != nullptr)
 		{
 			// This is a blueprint modification. Check to see if this thumbnail is the blueprint of the modified CDO
 			ObjectBeingModified = ObjectBeingModified->GetClass()->ClassGeneratedBy;
@@ -1573,27 +1573,34 @@ void FAssetThumbnailPool::DirtyThumbnailForObject(UObject* ObjectBeingModified)
 		// regenerated on demand later. (Before being displayed in the browser, or package saves, etc.)
 		FObjectThumbnail* Thumbnail = ThumbnailTools::GetThumbnailForObject(ObjectBeingModified);
 
-		// Don't try loading thumbnails for package that have never been saved
-		if (Thumbnail == NULL && !IsGarbageCollecting() && !ObjectBeingModified->GetOutermost()->HasAnyPackageFlags(PKG_NewlyCreated))
+		// If we don't yet have a thumbnail map, load one from disk if possible
+		if (Thumbnail == nullptr)
 		{
-			// If we don't yet have a thumbnail map, load one from disk if possible
-			// Don't attempt to do this while garbage collecting since loading or finding objects during GC is illegal
-			FName ObjectFullName = FName(*ObjectBeingModified->GetFullName());
-			TArray<FName> ObjectFullNames;
-			FThumbnailMap LoadedThumbnails;
-			ObjectFullNames.Add(ObjectFullName);
-			if (ThumbnailTools::ConditionallyLoadThumbnailsForObjects(ObjectFullNames, LoadedThumbnails))
+			UPackage* ObjectPackage = ObjectBeingModified->GetOutermost();
+
+			const bool bMemoryPackage = FPackageName::IsMemoryPackage(ObjectBeingModified->GetPathName());	// Don't try to load from disk if the package is a memory package
+			const bool bUnsavedPackage = ObjectPackage->HasAnyPackageFlags(PKG_NewlyCreated);				// Don't try loading thumbnails for package that have never been saved
+			const bool bIsGarbageCollecting = IsGarbageCollecting();										// Don't attempt to do this while garbage collecting since loading or finding objects during GC is illegal
+
+			const bool bTryLoadThumbnailFromDisk = !bIsGarbageCollecting && !bMemoryPackage && !bUnsavedPackage;
+			if (bTryLoadThumbnailFromDisk)
 			{
-				Thumbnail = LoadedThumbnails.Find(ObjectFullName);
+				FName ObjectFullName = FName(*ObjectBeingModified->GetFullName());
 
-				if (Thumbnail != NULL)
+				FThumbnailMap LoadedThumbnails;
+				if (ThumbnailTools::ConditionallyLoadThumbnailsForObjects({ ObjectFullName }, LoadedThumbnails))
 				{
-					Thumbnail = ThumbnailTools::CacheThumbnail(ObjectBeingModified->GetFullName(), Thumbnail, ObjectBeingModified->GetOutermost());
+					Thumbnail = LoadedThumbnails.Find(ObjectFullName);
+
+					if (Thumbnail != nullptr)
+					{
+						Thumbnail = ThumbnailTools::CacheThumbnail(ObjectBeingModified->GetFullName(), Thumbnail, ObjectPackage);
+					}
 				}
 			}
 		}
 
-		if (Thumbnail != NULL)
+		if (Thumbnail != nullptr)
 		{
 			// Mark the thumbnail as dirty
 			Thumbnail->MarkAsDirty();
