HLSL Syntax Hightlighter: added more type keywords. Merged with the old list of keywords and sorted based on string comparison in descending order. The keywords are from CrossCompiler::EHlslToken.

I feel like exposing an API to query those keywords on CrossCompiler might not be worth the effort. We will need to introduce an extra storage there to stores the types as they are added to their internal map -> filter out type keywords based on a range in the enum class ->  merge it with existing list -> sort the list -> remove duplicates. Could get difficult to debug if one of the steps goes wrong.

diff --git a/Engine/Source/Editor/UnrealEd/Private/Text/HLSLSyntaxHighlighterMarshaller.cpp b/Engine/Source/Editor/UnrealEd/Private/Text/HLSLSyntaxHighlighterMarshaller.cpp
--- a/Engine/Source/Editor/UnrealEd/Private/Text/HLSLSyntaxHighlighterMarshaller.cpp
+++ b/Engine/Source/Editor/UnrealEd/Private/Text/HLSLSyntaxHighlighterMarshaller.cpp
@@ -11,67 +11,154 @@
 // NOTE: Since FSyntaxTokenizer matches on a first-token-encountered basis, it's important that
 // tokens with the same prefix are ordered by longest-prefix-first. Ideally FSyntaxTokenizer 
 // should be using a prefix tree structure for longest prefix matching.
+
+// Type Keywords are copied from CrossCompiler::EHlslToken
 const TCHAR* HlslKeywords[] =
 {
-    TEXT("bool2"),
-    TEXT("bool3"),
-    TEXT("bool4"),
-	TEXT("bool"),
-	TEXT("break"),
-	TEXT("Buffer"),
-	TEXT("case"),
-	TEXT("column_major"),
-	TEXT("const"),
-	TEXT("continue"),
-	TEXT("default"),
-	TEXT("do"),
-	TEXT("double"),
-	TEXT("dword"),
-	TEXT("else"),
-	TEXT("enum"),
-	TEXT("export"),
-	TEXT("extern"),
-	TEXT("extern"),
-	TEXT("false"),
-    TEXT("float2x2"),
-    TEXT("float3x3"),
-    TEXT("float4x4"),
-	TEXT("float2"),
-	TEXT("float3"),
+	TEXT("while"),
+	TEXT("volatile"),
+	TEXT("void"),
+	TEXT("vector"),
+	TEXT("unorm"),
+	TEXT("uniform"),
+	TEXT("uint4x4"),
+	TEXT("uint4x3"),
+	TEXT("uint4x2"),
+	TEXT("uint4x1"),
+	TEXT("uint4"),
+	TEXT("uint3x4"),
+	TEXT("uint3x3"),
+	TEXT("uint3x2"),
+	TEXT("uint3x1"),
+	TEXT("uint3"),
+	TEXT("uint2x4"),
+	TEXT("uint2x3"),
+	TEXT("uint2x2"),
+	TEXT("uint2x1"),
+	TEXT("uint2"),
+	TEXT("uint1x4"),
+	TEXT("uint1x3"),
+	TEXT("uint1x2"),
+	TEXT("uint1x1"),
+	TEXT("uint1"),
+	TEXT("uint"),
+	TEXT("true"),
+	TEXT("switch"),
+	TEXT("struct"),
+	TEXT("static"),
+	TEXT("snorm"),
+	TEXT("shared"),
+	TEXT("row_major"),
+	TEXT("return"),
+	TEXT("register"),
+	TEXT("precise"),
+	TEXT("packoffset"),
+	TEXT("numthreads"),
+	TEXT("nointerpolation"),
+	TEXT("matrix"),
+	TEXT("int4x4"),
+	TEXT("int4x3"),
+	TEXT("int4x2"),
+	TEXT("int4x1"),
+	TEXT("int4"),
+	TEXT("int3x4"),
+	TEXT("int3x3"),
+	TEXT("int3x2"),
+	TEXT("int3x1"),
+	TEXT("int3"),
+	TEXT("int2x4"),
+	TEXT("int2x3"),
+	TEXT("int2x2"),
+	TEXT("int2x1"),
+	TEXT("int2"),
+	TEXT("int1x4"),
+	TEXT("int1x3"),
+	TEXT("int1x2"),
+	TEXT("int1x1"),
+	TEXT("int1"),
+	TEXT("int"),
+	TEXT("if"),
+	TEXT("half4x4"),
+	TEXT("half4x3"),
+	TEXT("half4x2"),
+	TEXT("half4x1"),
+	TEXT("half4"),
+	TEXT("half3x4"),
+	TEXT("half3x3"),
+	TEXT("half3x2"),
+	TEXT("half3x1"),
+	TEXT("half3"),
+	TEXT("half2x4"),
+	TEXT("half2x3"),
+	TEXT("half2x2"),
+	TEXT("half2x1"),
+	TEXT("half2"),
+	TEXT("half1x4"),
+	TEXT("half1x3"),
+	TEXT("half1x2"),
+	TEXT("half1x1"),
+	TEXT("half1"),
+	TEXT("half"),
+	TEXT("groupshared"),
+	TEXT("goto"),
+	TEXT("for"),
+	TEXT("float4x4"),
+	TEXT("float4x3"),
+	TEXT("float4x2"),
+	TEXT("float4x1"),
 	TEXT("float4"),
+	TEXT("float3x4"),
+	TEXT("float3x3"),
+	TEXT("float3x2"),
+	TEXT("float3x1"),
+	TEXT("float3"),
+	TEXT("float2x4"),
+	TEXT("float2x3"),
+	TEXT("float2x2"),
+	TEXT("float2x1"),
+	TEXT("float2"),
+	TEXT("float1x4"),
+	TEXT("float1x3"),
+	TEXT("float1x2"),
+	TEXT("float1x1"),
+	TEXT("float1"),
 	TEXT("float"),
-	TEXT("for"),
-	TEXT("goto"),
-	TEXT("groupshared"),
-	TEXT("half"),
-	TEXT("if"),
-    TEXT("int2"),
-    TEXT("int3"),
-    TEXT("int4"),
-	TEXT("int"),
-	TEXT("matrix"),
-	TEXT("nointerpolation"),
-    TEXT("numthreads"),	
-	TEXT("packoffset"),
-	TEXT("precise"),
-	TEXT("register"),
-	TEXT("return"),
-	TEXT("row_major"),
-	TEXT("shared"),
-	TEXT("snorm"),
-	TEXT("static"),
-	TEXT("struct"),
-	TEXT("switch"),
-	TEXT("true"),
-    TEXT("uint2"),
-    TEXT("uint3"),
-    TEXT("uint4"),
-	TEXT("uint"),
-	TEXT("uniform"),
-	TEXT("unorm"),
-	TEXT("vector"),
-	TEXT("volatile"),
-	TEXT("while"),
+	TEXT("false"),
+	TEXT("extern"),
+	TEXT("export"),
+	TEXT("enum"),
+	TEXT("else"),
+	TEXT("dword"),
+	TEXT("double"),
+	TEXT("do"),
+	TEXT("default"),
+	TEXT("continue"),
+	TEXT("const"),
+	TEXT("column_major"),
+	TEXT("case"),
+	TEXT("break"),
+	TEXT("bool4x4"),
+	TEXT("bool4x3"),
+	TEXT("bool4x2"),
+	TEXT("bool4x1"),
+	TEXT("bool4"),
+	TEXT("bool3x4"),
+	TEXT("bool3x3"),
+	TEXT("bool3x2"),
+	TEXT("bool3x1"),
+	TEXT("bool3"),
+	TEXT("bool2x4"),
+	TEXT("bool2x3"),
+	TEXT("bool2x2"),
+	TEXT("bool2x1"),
+	TEXT("bool2"),
+	TEXT("bool1x4"),
+	TEXT("bool1x3"),
+	TEXT("bool1x2"),
+	TEXT("bool1x1"),
+	TEXT("bool1"),
+	TEXT("bool"),
+	TEXT("Buffer"),
 };
 
 const TCHAR* Operators[] =
