Restore UE4 default behavior of bCookAll when no packages are specified in the cook.
#rnx

diff --git a/Engine/Source/Editor/UnrealEd/Private/CookOnTheFlyServer.cpp b/Engine/Source/Editor/UnrealEd/Private/CookOnTheFlyServer.cpp
--- a/Engine/Source/Editor/UnrealEd/Private/CookOnTheFlyServer.cpp
+++ b/Engine/Source/Editor/UnrealEd/Private/CookOnTheFlyServer.cpp
@@ -6380,8 +6380,11 @@ void UCookOnTheFlyServer::CollectFilesToCook(TArray<FName>& FilesInPath, TMap<FN
 			}
 		}
 
+		// Keep the old behavior of cooking all by default until we implement good feedback in the editor about the missing setting
+		constexpr bool bCookAllByDefault = true;
+
 		// If no packages were explicitly added by command line or game callback, add all maps
-		if (bCookAll)
+		if (bCookAll || (bCookAllByDefault && FilesInPath.Num() == InitialNum))
 		{
 			TArray<FString> Tokens;
 			Tokens.Empty(2);
@@ -6416,10 +6419,11 @@ void UCookOnTheFlyServer::CollectFilesToCook(TArray<FName>& FilesInPath, TMap<FN
 				}
 			}
 		}
-		if (FilesInPath.Num() == InitialNum && !bCookAll)
+		else if (FilesInPath.Num() == InitialNum)
 		{
-			LogCookerMessage(TEXT("Cooker found no package requests specified on commandline or ini. Run cookcommandlet with -helpcookusage to see request options."),
-				EMessageSeverity::Info);
+			LogCookerMessage(TEXT("No package requests specified on -run=Cook commandline or ini. ")
+				TEXT("Set the flag 'Edit->Project Settings->Project/Packaging->Packaging/Advanced->Cook Everything in the Project Content Directory'. ")
+				TEXT("Or launch 'UnrealEditor -run=cook -helpcookusage' to see all package request options."), EMessageSeverity::Warning);
 		}
 	}
 
