Build navigation for runtime data layers.

diff --git a/Engine/Source/Editor/UnrealEd/Private/WorldPartition/WorldPartitionNavigationDataBuilder.cpp b/Engine/Source/Editor/UnrealEd/Private/WorldPartition/WorldPartitionNavigationDataBuilder.cpp
--- a/Engine/Source/Editor/UnrealEd/Private/WorldPartition/WorldPartitionNavigationDataBuilder.cpp
+++ b/Engine/Source/Editor/UnrealEd/Private/WorldPartition/WorldPartitionNavigationDataBuilder.cpp
@@ -14,6 +14,7 @@
 
 #include "WorldPartition/WorldPartition.h"
 #include "WorldPartition/WorldPartitionSubsystem.h"
+#include "WorldPartition/DataLayer/WorldDataLayers.h"
 #include "WorldPartition/NavigationData/NavigationDataChunkActor.h"
 
 DEFINE_LOG_CATEGORY_STATIC(LogWorldPartitionNavigationDataBuilder, Log, All);
@@ -25,6 +26,20 @@ UWorldPartitionNavigationDataBuilder::UWorldPartitionNavigationDataBuilder(const
 
 bool UWorldPartitionNavigationDataBuilder::PreRun(UWorld* World, FPackageSourceControlHelper& PackageHelper)
 {
+	// Set runtime data layer to be included in the generation.
+	if (const AWorldDataLayers* WorldDataLayers = World->GetWorldDataLayers())
+	{
+		WorldDataLayers->ForEachDataLayer([this](const UDataLayer* DataLayer)
+		{
+			if (DataLayer->IsRuntime())
+			{
+				DataLayerLabels.Add(DataLayer->GetDataLayerLabel());
+			}
+
+			return true;
+		});
+	}
+	
 	const TSubclassOf<APartitionActor>& NavigationDataActorClass = ANavigationDataChunkActor::StaticClass();
 	uint32 GridSize = NavigationDataActorClass->GetDefaultObject<APartitionActor>()->GetDefaultGridSize(World);
 	GridSize = FMath::Max(GridSize, 1u);
diff --git a/Engine/Source/Runtime/Engine/Public/WorldPartition/DataLayer/DataLayer.h b/Engine/Source/Runtime/Engine/Public/WorldPartition/DataLayer/DataLayer.h
--- a/Engine/Source/Runtime/Engine/Public/WorldPartition/DataLayer/DataLayer.h
+++ b/Engine/Source/Runtime/Engine/Public/WorldPartition/DataLayer/DataLayer.h
@@ -23,8 +23,8 @@ UENUM(BlueprintType)
 enum class EDataLayerRuntimeState : uint8
 {
 	Unloaded,
-	Loaded,
-	Activated
+	Loaded,		// Loaded but not in the world
+	Activated	// Loaded and in the world
 };
 
 const inline TCHAR* GetDataLayerRuntimeStateName(EDataLayerRuntimeState State)
