Fixing feature level check to include feature levels over SM5.

diff --git a/Engine/Source/Editor/UnrealEd/Private/SEditorViewportViewMenu.cpp b/Engine/Source/Editor/UnrealEd/Private/SEditorViewportViewMenu.cpp
--- a/Engine/Source/Editor/UnrealEd/Private/SEditorViewportViewMenu.cpp
+++ b/Engine/Source/Editor/UnrealEd/Private/SEditorViewportViewMenu.cpp
@@ -158,7 +158,7 @@ void SEditorViewportViewMenu::FillViewMenu(UToolMenu* Menu) const
 
 						{
 							FToolMenuSection& Section = Menu->AddSection("OptimizationViewmodes", LOCTEXT("OptimizationSubMenuHeader", "Optimization Viewmodes"));
-							if (FeatureLevel == ERHIFeatureLevel::SM5)
+							if (FeatureLevel >= ERHIFeatureLevel::SM5)
 							{
 								Section.AddMenuEntry(BaseViewportCommands.LightComplexityMode, UViewModeUtils::GetViewModeDisplayName(VMI_LightComplexity));
 								if (IConsoleManager::Get().FindTConsoleVariableDataInt(TEXT("r.AllowStaticLighting"))->GetValueOnAnyThread() != 0)
