Fix for PIE starting a non single process game without saving dirty blueprints.  Update Save Warning so user cannot play without saving since there can be serialization errors between a pie instance and a separate processes.

diff --git a/Engine/Source/Editor/UnrealEd/Private/PlayLevel.cpp b/Engine/Source/Editor/UnrealEd/Private/PlayLevel.cpp
--- a/Engine/Source/Editor/UnrealEd/Private/PlayLevel.cpp
+++ b/Engine/Source/Editor/UnrealEd/Private/PlayLevel.cpp
@@ -998,10 +998,13 @@ bool UEditorEngine::SaveMapsForPlaySession()
 	}
 
 	// Also save dirty packages, this is required because we're going to be launching a session outside of our normal process
-	const bool bPromptUserToSave = true;
-	const bool bSaveMapPackages = true;
-	const bool bSaveContentPackages = true;
-	if (!FEditorFileUtils::SaveDirtyPackages(bPromptUserToSave, bSaveMapPackages, bSaveContentPackages))
+	const bool bPromptUserToSave      = true;
+	const bool bSaveMapPackages       = true;
+	const bool bSaveContentPackages   = true;
+	const bool bFastSave              = false;
+	const bool bNotifyNoPackagesSaved = false;
+	const bool bCanBeDeclined         = false;
+	if (!FEditorFileUtils::SaveDirtyPackages(bPromptUserToSave, bSaveMapPackages, bSaveContentPackages, bFastSave, bNotifyNoPackagesSaved, bCanBeDeclined))
 	{
 		return false;
 	}
