ImageWrapper and import/export refactor
FImage is now the standard preferred type for a bag of pixels
FImageView can point at pixels without owning an allocation
ERawImageFormat (FImage) converts to ETextureSourceFormat
FImageUtils provides generic load/save and get/set from FImage
major cleanup in the ImageWrappers
new preferred API is through ImageWrapperModule Compress/Decompress
SetRaw/GetRaw functions cleaned up to not have undefined behavior on unexpected formats
ImageWrapper output added for HDR,BMP,TGA
RGBA32F format added and supported throughout import/export
EditorFactories import/export made more generic, most image types handled the same way using FImage now
Deprecate old TSF RGBA order pixel formats
Fix many crashes or bad handling of unusual pixel formats
Pixel access functions should be used instead of switches on pixel type

diff --git a/Engine/Plugins/Editor/BlueprintMaterialTextureNodes/Source/BlueprintMaterialTextureNodes/BlueprintMaterialTextureNodes.Build.cs b/Engine/Plugins/Editor/BlueprintMaterialTextureNodes/Source/BlueprintMaterialTextureNodes/BlueprintMaterialTextureNodes.Build.cs
--- a/Engine/Plugins/Editor/BlueprintMaterialTextureNodes/Source/BlueprintMaterialTextureNodes/BlueprintMaterialTextureNodes.Build.cs
+++ b/Engine/Plugins/Editor/BlueprintMaterialTextureNodes/Source/BlueprintMaterialTextureNodes/BlueprintMaterialTextureNodes.Build.cs
@@ -30,7 +30,8 @@ public class BlueprintMaterialTextureNodes : ModuleRules
 			new string[]
 			{
 				"Core",
-                "AssetRegistry",
+				"ImageCore",
+				"AssetRegistry",
                 "AssetTools",
                 "RHI",
                 "StaticMeshEditor",
diff --git a/Engine/Plugins/Editor/BlueprintMaterialTextureNodes/Source/BlueprintMaterialTextureNodes/Private/BlueprintMaterialTextureNodesBPLibrary.cpp b/Engine/Plugins/Editor/BlueprintMaterialTextureNodes/Source/BlueprintMaterialTextureNodes/Private/BlueprintMaterialTextureNodesBPLibrary.cpp
--- a/Engine/Plugins/Editor/BlueprintMaterialTextureNodes/Source/BlueprintMaterialTextureNodes/Private/BlueprintMaterialTextureNodesBPLibrary.cpp
+++ b/Engine/Plugins/Editor/BlueprintMaterialTextureNodes/Source/BlueprintMaterialTextureNodes/Private/BlueprintMaterialTextureNodesBPLibrary.cpp
@@ -24,6 +24,8 @@
 #include "Materials/MaterialInstance.h"
 #include "Materials/MaterialInstanceConstant.h"
 #include "MaterialShared.h"
+#include "ImageCore.h"
+#include "ImageCoreUtils.h"
 
 #define LOCTEXT_NAMESPACE "BlueprintMaterialTextureLibrary"
 
@@ -80,47 +82,16 @@ FLinearColor UBlueprintMaterialTextureNodesBPLibrary::Texture2D_SampleUV_EditorO
 		{
 			FTextureSource& TextureSource = Texture->Source;
 
+			// gets a copy of the whole mip to sample one pixel :
 			TArray64<uint8> SourceData;
 			Texture->Source.GetMipData(SourceData, Mip);
 			ETextureSourceFormat SourceFormat = TextureSource.GetFormat();
 			int32 Index = ((Y * MipWidth) + X) * TextureSource.GetBytesPerPixel();
-
-			if ((SourceFormat == TSF_BGRA8 || SourceFormat == TSF_BGRE8))
-			{
-				FColor OutSourceColor;
-				const uint8* PixelPtr = SourceData.GetData() + Index;
-				OutSourceColor = *((FColor*)PixelPtr);
-				if (Texture->SRGB)
-				{
-					return	FLinearColor::FromSRGBColor(OutSourceColor);
-				}
-
-				return FLinearColor(float(OutSourceColor.R), float(OutSourceColor.G), float(OutSourceColor.B), float(OutSourceColor.A)) / 255.0f;
-
-			}
-			else if ((SourceFormat == TSF_RGBA16 || SourceFormat == TSF_RGBA16F))
-			{
-				FFloat16Color OutSourceColor;
-				const uint8* PixelPtr = SourceData.GetData() + Index;
-				OutSourceColor = *((FFloat16Color*)PixelPtr);
-
-				return FLinearColor(float(OutSourceColor.R), float(OutSourceColor.G), float(OutSourceColor.B), float(OutSourceColor.A));
-			}
-			else if (SourceFormat == TSF_G8)
-			{
-				FColor OutSourceColor;
-				const uint8* PixelPtr = SourceData.GetData() + Index;
-				const uint8 value = *PixelPtr;
-				if (Texture->SRGB)
-				{
-					return FLinearColor::FromSRGBColor(FColor(float(value), 0, 0, 0));
-				}
-
-				return FLinearColor(float(value), 0, 0, 0);
-
-			}
-
-			FMessageLog("Blueprint").Warning(LOCTEXT("Texture2D_SampleUV_InvalidFormat", "Texture2D_SampleUV_EditorOnly: Source was unavailable or of unsupported format."));
+			const uint8* PixelPtr = SourceData.GetData() + Index;
+			
+			ERawImageFormat::Type RawFormat = FImageCoreUtils::ConvertToRawImageFormat(SourceFormat);
+			FLinearColor Color = ERawImageFormat::GetOnePixelLinear(PixelPtr,RawFormat,Texture->SRGB);
+			return Color;
 		}
 
 	}
diff --git a/Engine/Plugins/Enterprise/MDLImporter/Source/MDLImporter/Private/common/ImageUtils.cpp b/Engine/Plugins/Enterprise/MDLImporter/Source/MDLImporter/Private/common/ImageUtils.cpp
--- a/Engine/Plugins/Enterprise/MDLImporter/Source/MDLImporter/Private/common/ImageUtils.cpp
+++ b/Engine/Plugins/Enterprise/MDLImporter/Source/MDLImporter/Private/common/ImageUtils.cpp
@@ -1,7 +1,8 @@
 // Copyright Epic Games, Inc. All Rights Reserved.
 
 /*=============================================================================
-ImageUtils.cpp: Image utility functions.
+this is MDLImporter ImageUtils
+NOT the ImageUtils in Engine
 =============================================================================*/
 
 #include "ImageUtils.h"
diff --git a/Engine/Plugins/Experimental/ImagePlate/Source/ImagePlate/ImagePlate.Build.cs b/Engine/Plugins/Experimental/ImagePlate/Source/ImagePlate/ImagePlate.Build.cs
--- a/Engine/Plugins/Experimental/ImagePlate/Source/ImagePlate/ImagePlate.Build.cs
+++ b/Engine/Plugins/Experimental/ImagePlate/Source/ImagePlate/ImagePlate.Build.cs
@@ -19,6 +19,7 @@ namespace UnrealBuildTool.Rules
 				new string[]
 				{
 					"Core",
+					"ImageCore",
 					"CoreUObject",
 					"Engine",
 					"ImageWrapper",
diff --git a/Engine/Plugins/Experimental/ImagePlate/Source/ImagePlate/Private/ImagePlateFileSequence.cpp b/Engine/Plugins/Experimental/ImagePlate/Source/ImagePlate/Private/ImagePlateFileSequence.cpp
--- a/Engine/Plugins/Experimental/ImagePlate/Source/ImagePlate/Private/ImagePlateFileSequence.cpp
+++ b/Engine/Plugins/Experimental/ImagePlate/Source/ImagePlate/Private/ImagePlateFileSequence.cpp
@@ -17,7 +17,7 @@
 #include "Engine/Texture2DDynamic.h"
 #include "Engine/TextureRenderTarget2D.h"
 #include "Textures/SlateTextureData.h"
-
+#include "ImageUtils.h"
 
 DEFINE_LOG_CATEGORY_STATIC(LogImagePlateFileSequence, Log, Warning);
 
@@ -551,47 +551,38 @@ namespace ImagePlateFrameCache
 	/** Load an image file into CPU memory */
 	FImagePlateSourceFrame LoadFileData(FString FilenameToLoad, int32 FrameNumber)
 	{
-		// Start at 100k
-		TArray64<uint8> SourceFileData;
-		SourceFileData.Reserve(1024*100);
-		if (!FFileHelper::LoadFileToArray(SourceFileData, *FilenameToLoad))
+		FImage Image;
+		if ( ! FImageUtils::LoadImage(*FilenameToLoad,Image) )
 		{
 			UE_LOG(LogImagePlateFileSequence, Warning, TEXT("Failed to load file data from '%s'"), *FilenameToLoad);
 			return FImagePlateSourceFrame();
 		}
+		
+		// this module only handles these formats :
+		//case 16: Texture2DDynamic->Format = PF_FloatRGBA; break;
+		//case 8: Texture2DDynamic->Format = PF_R8G8B8A8; break;
+		int BitDepth;
 
-		IImageWrapperModule& ImageWrapperModule = FModuleManager::GetModuleChecked<IImageWrapperModule>("ImageWrapper");
-
-		EImageFormat ImageType = ImageWrapperModule.DetectImageFormat(SourceFileData.GetData(), SourceFileData.Num());
-		TSharedPtr<IImageWrapper> ImageWrapper = ImageWrapperModule.CreateImageWrapper(ImageType);
-
-		if (!ImageWrapper.IsValid())
-		{
-			UE_LOG(LogImagePlateFileSequence, Warning, TEXT("File '%s' is not a supported image type."), *FilenameToLoad);
-			return FImagePlateSourceFrame();
-		}
-
-		ImageWrapper->SetCompressed(SourceFileData.GetData(), SourceFileData.Num());
-
-		int32 SourceBitDepth = ImageWrapper->GetBitDepth();
-		TArray64<uint8> RawImageData;
-		if (ImageWrapper->GetRaw(ERGBFormat::RGBA, SourceBitDepth, RawImageData))
+		if ( ERawImageFormat::IsHDR(Image.Format) )
 		{
-			// BMP image wrappers supply the bitdepth per pixel, rather than per-channel
-			SourceBitDepth = ImageType == EImageFormat::BMP ? ImageWrapper->GetBitDepth() / 4 : ImageWrapper->GetBitDepth();
-
-			return FImagePlateSourceFrame(
-				RawImageData,
-				ImageWrapper->GetWidth(),
-				ImageWrapper->GetHeight(),
-				SourceBitDepth
-			);
+			Image.ChangeFormat(ERawImageFormat::RGBA16F,EGammaSpace::Linear);
+			BitDepth = 16;
 		}
 		else
 		{
-			UE_LOG(LogImagePlateFileSequence, Warning, TEXT("Failed to get raw rgba data from image file '%s'."), *FilenameToLoad);
-			return FImagePlateSourceFrame();
+			Image.ChangeFormat(ERawImageFormat::BGRA8,EGammaSpace::sRGB);
+
+			// this module wants RGBA not BGRA :
+			FImageCore::TransposeImageRGBABGRA(Image);
+
+			BitDepth = 8;
 		}
+
+		return FImagePlateSourceFrame(
+			MoveTemp(Image.RawData),
+			Image.SizeX,
+			Image.SizeY,
+			BitDepth);
 	}
 
 	class FFrameLoadingThread : FRunnable
diff --git a/Engine/Plugins/Experimental/MeshLODToolset/Source/MeshLODToolset/Private/Graphs/GenerateStaticMeshLODProcess.cpp b/Engine/Plugins/Experimental/MeshLODToolset/Source/MeshLODToolset/Private/Graphs/GenerateStaticMeshLODProcess.cpp
--- a/Engine/Plugins/Experimental/MeshLODToolset/Source/MeshLODToolset/Private/Graphs/GenerateStaticMeshLODProcess.cpp
+++ b/Engine/Plugins/Experimental/MeshLODToolset/Source/MeshLODToolset/Private/Graphs/GenerateStaticMeshLODProcess.cpp
@@ -513,7 +513,7 @@ int UGenerateStaticMeshLODProcess::SelectTextureToBake(const TArray<FTextureInfo
 #if WITH_EDITORONLY_DATA
 			// Texture has multiple channels
 			ETextureSourceFormat Format = Tex2D->Source.GetFormat();
-			if (Format == TSF_BGRA8 || Format == TSF_BGRE8 || Format == TSF_RGBA16 || Format == TSF_RGBA16F)
+			if (Format == TSF_BGRA8 || Format == TSF_BGRE8 || Format == TSF_RGBA16 || Format == TSF_RGBA16F || Format == TSF_RGBA32F)
 			{
 				++TextureVotes[TextureIndex];
 			}
diff --git a/Engine/Plugins/Experimental/MeshModelingToolsetExp/Source/MeshModelingToolsExp/Private/BakeMeshAttributeTool.cpp b/Engine/Plugins/Experimental/MeshModelingToolsetExp/Source/MeshModelingToolsExp/Private/BakeMeshAttributeTool.cpp
--- a/Engine/Plugins/Experimental/MeshModelingToolsetExp/Source/MeshModelingToolsExp/Private/BakeMeshAttributeTool.cpp
+++ b/Engine/Plugins/Experimental/MeshModelingToolsetExp/Source/MeshModelingToolsExp/Private/BakeMeshAttributeTool.cpp
@@ -282,7 +282,7 @@ int UBakeMeshAttributeTool::SelectColorTextureToBake(const TArray<UTexture*>& Te
 #if WITH_EDITORONLY_DATA
 			// Texture has multiple channels
 			ETextureSourceFormat Format = Tex->Source.GetFormat();
-			if (Format == TSF_BGRA8 || Format == TSF_BGRE8 || Format == TSF_RGBA16 || Format == TSF_RGBA16F)
+			if (Format == TSF_BGRA8 || Format == TSF_BGRE8 || Format == TSF_RGBA16 || Format == TSF_RGBA16F || Format == TSF_RGBA32F)
 			{
 				++TextureVotes[TextureIndex];
 			}
diff --git a/Engine/Plugins/Runtime/AR/Google/GoogleARCore/Source/GoogleARCoreBase/Private/GoogleARCoreCookSupport.cpp b/Engine/Plugins/Runtime/AR/Google/GoogleARCore/Source/GoogleARCoreBase/Private/GoogleARCoreCookSupport.cpp
--- a/Engine/Plugins/Runtime/AR/Google/GoogleARCore/Source/GoogleARCoreBase/Private/GoogleARCoreCookSupport.cpp
+++ b/Engine/Plugins/Runtime/AR/Google/GoogleARCore/Source/GoogleARCoreBase/Private/GoogleARCoreCookSupport.cpp
@@ -15,7 +15,7 @@
 #include "GenericPlatform/GenericPlatformFile.h"
 #include "RenderUtils.h"
 
-#include "png.h"
+#include "png.h" // todo: remove png.h
 
 #if WITH_EDITORONLY_DATA
 
@@ -29,12 +29,23 @@
 
 bool FGoogleARCoreSessionConfigCookSupport::SaveTextureToPNG(UTexture2D *Tex, const FString &Filename)
 {
+	// todo: use FImageUtils::SaveTexture2D instead
+	//  
+	#if 0
+
+	FImage Image;
+	if ( ! FImageUtils::GetTexture2DSourceImage(Tex,Image) )
+		return false;
+	return FImageUtils::SaveImageByExtension(*FileName,Image);
+
+	#endif
+
 	TArray64<uint8> MipData;
 	bool Ret = true;
 
 	if (Tex->Source.GetMipData(MipData, 0))
 	{
-		if (Tex->Source.GetFormat() != TSF_BGRA8 && Tex->Source.GetFormat() != TSF_RGBA8)
+		if (Tex->Source.GetFormat() != TSF_BGRA8 )
 		{
 			UE_LOG(
 				LogGoogleARCore, Error,
diff --git a/Engine/Source/Developer/SkeletalMeshUtilitiesCommon/Private/LODUtilities.cpp b/Engine/Source/Developer/SkeletalMeshUtilitiesCommon/Private/LODUtilities.cpp
--- a/Engine/Source/Developer/SkeletalMeshUtilitiesCommon/Private/LODUtilities.cpp
+++ b/Engine/Source/Developer/SkeletalMeshUtilitiesCommon/Private/LODUtilities.cpp
@@ -20,6 +20,8 @@
 #include "OverlappingCorners.h"
 #include "Framework/Commands/UIAction.h"
 #include "HAL/ThreadSafeBool.h"
+#include "ImageCore.h"
+#include "ImageCoreUtils.h"
 
 #include "ObjectTools.h"
 
@@ -3204,11 +3206,10 @@ bool FLODUtilities::StripLODGeometry(USkeletalMesh* SkeletalMesh, const int32 LO
 		}
 		SkeletalMesh->Modify();
 		
-		const float ThresholdClamp = FMath::Clamp(Threshold, 0.0f, 1.0f);
-		const uint8 Threshold8 = (uint8)FMath::FloorToInt(ThresholdClamp * (float)(0xFF));
-		const uint16 Threshold16 = (uint16)FMath::FloorToInt(ThresholdClamp * (float)(0xFFFF));
+		ERawImageFormat::Type RawFormat = FImageCoreUtils::ConvertToRawImageFormat(SourceFormat);
+		bool bSRGB = TextureMask->SRGB;
 
-		auto ShouldStripTriangle = [&InitialSource, &ResX, &ResY, &FormatDataSize, &Ref2DData, &SourceFormat, &Threshold, &Threshold8, &Threshold16](const FVector2D& UvA, const FVector2D& UvB, const FVector2D& UvC)->bool
+		auto ShouldStripTriangle = [&](const FVector2D& UvA, const FVector2D& UvB, const FVector2D& UvC)->bool
 		{
 			FVector2D PixelUvA = FVector2D(FMath::FloorToInt(UvA.X * (float)ResX) % (ResX + 1), FMath::FloorToInt(UvA.Y * (float)ResY) % (ResY + 1));
 			FVector2D PixelUvB = FVector2D(FMath::FloorToInt(UvB.X * (float)ResX) % (ResX + 1), FMath::FloorToInt(UvB.Y * (float)ResY) % (ResY + 1));
@@ -3220,54 +3221,18 @@ bool FLODUtilities::StripLODGeometry(USkeletalMesh* SkeletalMesh, const int32 LO
 			int32 MaxV = FMath::Clamp(FMath::Max3<int32>(PixelUvA.Y, PixelUvB.Y, PixelUvC.Y), 0, ResY);
 
 			//Do not read the alpha value when testing the texture value
-			auto IsPixelZero = [&Ref2DData, &InitialSource, &SourceFormat, &FormatDataSize, &Threshold, &Threshold8, &Threshold16](int32 PosX, int32 PosY) -> bool
+			auto IsPixelZero = [&](int32 PosX, int32 PosY) -> bool
 			{
-				uint8 CurPos[16];
 				const int32 RefPos = PosX + (PosY * InitialSource.GetSizeX());
-				FMemory::Memcpy(&(CurPos[0]), Ref2DData.GetData() + (RefPos * FormatDataSize), FormatDataSize);
-				bool bPixelIsZero = true;
-				switch (SourceFormat)
-				{
-				case TSF_BGRA8:
-				case TSF_BGRE8:
-				case TSF_RGBA8:
-				case TSF_RGBE8:
-				{
-					
-					if (CurPos[0] > Threshold8 || CurPos[1] > Threshold8 || CurPos[2] > Threshold8)
-					{
-						bPixelIsZero = false;
-					}
-				}
-				break;
-				case TSF_G8:
-				{
-					bPixelIsZero = !(CurPos[0] > Threshold8);
-				}
-				break;
-				case TSF_G16:
-				{
-					bPixelIsZero = !(((uint16*)(&CurPos[0]))[0] > Threshold16);
-				}
-				break;
+				const void * PixelPtr = Ref2DData.GetData() + RefPos * FormatDataSize;
 				
-				break;
-				case TSF_RGBA16:
-				case TSF_RGBA16F:
-				{
-					FFloat16 HalfValueR = *(FFloat16*)(&CurPos[0]);
-					FFloat16 HalfValueG = *(FFloat16*)(&CurPos[2]);
-					FFloat16 HalfValueB = *(FFloat16*)(&CurPos[4]);
-					if ( !FMath::IsNearlyZero(float(HalfValueR), Threshold) || !FMath::IsNearlyZero(float(HalfValueG), Threshold) || !FMath::IsNearlyZero(float(HalfValueB), Threshold) )
-					{
-						bPixelIsZero = false;
-					}
-				}
-				break;
-				default:
-					//Unknown format
-					return false;
-				}
+				FLinearColor Color = ERawImageFormat::GetOnePixelLinear(PixelPtr,RawFormat,bSRGB);	
+
+				bool bPixelIsZero = 
+					FMath::IsNearlyZero(Color.R,Threshold) &&
+					FMath::IsNearlyZero(Color.G,Threshold) &&
+					FMath::IsNearlyZero(Color.B,Threshold);
+
 				return bPixelIsZero;
 			};
 
diff --git a/Engine/Source/Developer/SkeletalMeshUtilitiesCommon/SkeletalMeshUtilitiesCommon.Build.cs b/Engine/Source/Developer/SkeletalMeshUtilitiesCommon/SkeletalMeshUtilitiesCommon.Build.cs
--- a/Engine/Source/Developer/SkeletalMeshUtilitiesCommon/SkeletalMeshUtilitiesCommon.Build.cs
+++ b/Engine/Source/Developer/SkeletalMeshUtilitiesCommon/SkeletalMeshUtilitiesCommon.Build.cs
@@ -17,8 +17,9 @@ namespace UnrealBuildTool.Rules
 					"CoreUObject",
 					"RenderCore",
 					"MeshUtilitiesCommon",
+					"ImageCore",
 					"Engine",
-                    "ClothingSystemRuntimeCommon",
+					"ClothingSystemRuntimeCommon",
 					"Slate",
                 }
             );
diff --git a/Engine/Source/Developer/TextureBuild/Private/TextureBuildFunction.cpp b/Engine/Source/Developer/TextureBuild/Private/TextureBuildFunction.cpp
--- a/Engine/Source/Developer/TextureBuild/Private/TextureBuildFunction.cpp
+++ b/Engine/Source/Developer/TextureBuild/Private/TextureBuildFunction.cpp
@@ -8,6 +8,7 @@
 #include "IImageWrapper.h"
 #include "IImageWrapperModule.h"
 #include "ImageCore.h"
+#include "ImageCoreUtils.h"
 #include "Interfaces/ITextureFormat.h"
 #include "Interfaces/ITextureFormatModule.h"
 #include "Memory/CompositeBuffer.h"
@@ -176,18 +177,7 @@ static void ReadOutputSettingsFromCompactBinary(const FCbObjectView& Object, int
 
 static ERawImageFormat::Type ComputeRawImageFormat(ETextureSourceFormat SourceFormat)
 {
-	switch (SourceFormat)
-	{
-	case TSF_G8:		return ERawImageFormat::G8;
-	case TSF_G16:		return ERawImageFormat::G16;
-	case TSF_BGRA8:		return ERawImageFormat::BGRA8;
-	case TSF_BGRE8:		return ERawImageFormat::BGRE8;
-	case TSF_RGBA16:	return ERawImageFormat::RGBA16;
-	case TSF_RGBA16F:	return ERawImageFormat::RGBA16F;
-	default:
-		checkf(false, TEXT("Invalid source texture format encountered: %d."), SourceFormat);
-		return ERawImageFormat::G8;
-	}
+	return FImageCoreUtils::ConvertToRawImageFormat(SourceFormat);
 }
 
 static bool TryReadTextureSourceFromCompactBinary(FCbFieldView Source, UE::DerivedData::FBuildContext& Context, TArray<FImage>& OutMips)
diff --git a/Engine/Source/Editor/UMGEditor/Private/WidgetBlueprintEditor.cpp b/Engine/Source/Editor/UMGEditor/Private/WidgetBlueprintEditor.cpp
--- a/Engine/Source/Editor/UMGEditor/Private/WidgetBlueprintEditor.cpp
+++ b/Engine/Source/Editor/UMGEditor/Private/WidgetBlueprintEditor.cpp
@@ -425,14 +425,13 @@ void FWidgetBlueprintEditor::CaptureThumbnail()
 		return;
 	}
 
-	TArray64<uint8> RawData;
-	FImageUtils::GetRawData(RenderTarget2D, RawData);
-	IImageWrapperModule& ImageWrapperModule = FModuleManager::Get().LoadModuleChecked<IImageWrapperModule>(TEXT("ImageWrapper"));
-	TSharedPtr<IImageWrapper> PNGImageWrapper = ImageWrapperModule.CreateImageWrapper(EImageFormat::PNG);
-	PNGImageWrapper->SetRaw(RawData.GetData(), RawData.GetAllocatedSize(), RenderTarget2D->SizeX, RenderTarget2D->SizeY, ERGBFormat::BGRA, 8);
-	const TArray64<uint8> PNGData = PNGImageWrapper->GetCompressed(100);
-
-	UTexture2D* ThumbnailTexture = FImageUtils::ImportBufferAsTexture2D(PNGData);
+	FImage Image;
+	if ( !FImageUtils::GetRenderTargetImage(RenderTarget2D,Image) )
+	{
+		return;
+	}
+
+	UTexture2D* ThumbnailTexture = FImageUtils::CreateTexture2DFromImage(Image);
 	FWidgetBlueprintEditorUtils::SetTextureAsAssetThumbnail(GetWidgetBlueprintObj(), ThumbnailTexture);
 }
 
diff --git a/Engine/Source/Editor/UnrealEd/Classes/Exporters/RenderTargetExporterHDR.h b/Engine/Source/Editor/UnrealEd/Classes/Exporters/RenderTargetExporterHDR.h
--- a/Engine/Source/Editor/UnrealEd/Classes/Exporters/RenderTargetExporterHDR.h
+++ b/Engine/Source/Editor/UnrealEd/Classes/Exporters/RenderTargetExporterHDR.h
@@ -22,4 +22,27 @@ class UNREALED_API URenderTargetExporterHDR : public UExporter
 };
 
 
+UCLASS()
+class UNREALED_API URenderTargetExporterPNG : public UExporter
+{
+	GENERATED_UCLASS_BODY()
+	
+	//~ Begin UExporter Interface
+	virtual bool SupportsObject(UObject* Object) const override;
+	virtual bool ExportBinary( UObject* Object, const TCHAR* Type, FArchive& Ar, FFeedbackContext* Warn, int32 FileIndex = 0, uint32 PortFlags=0 ) override;
+	//~ End UExporter Interface
+};
+
+
+
+UCLASS()
+class UNREALED_API URenderTargetExporterEXR : public UExporter
+{
+	GENERATED_UCLASS_BODY()
+	
+	//~ Begin UExporter Interface
+	virtual bool SupportsObject(UObject* Object) const override;
+	virtual bool ExportBinary( UObject* Object, const TCHAR* Type, FArchive& Ar, FFeedbackContext* Warn, int32 FileIndex = 0, uint32 PortFlags=0 ) override;
+	//~ End UExporter Interface
+};
 
diff --git a/Engine/Source/Editor/UnrealEd/Classes/Exporters/TextureCubeExporterHDR.h b/Engine/Source/Editor/UnrealEd/Classes/Exporters/TextureCubeExporterHDR.h
--- a/Engine/Source/Editor/UnrealEd/Classes/Exporters/TextureCubeExporterHDR.h
+++ b/Engine/Source/Editor/UnrealEd/Classes/Exporters/TextureCubeExporterHDR.h
@@ -21,5 +21,13 @@ class UNREALED_API UTextureCubeExporterHDR : public UExporter
 	//~ End UExporter Interface
 };
 
+UCLASS()
+class UNREALED_API URenderTargetCubeExporterHDR : public UTextureCubeExporterHDR
+{
+	GENERATED_UCLASS_BODY()
+};
+
+
+
 
 
diff --git a/Engine/Source/Editor/UnrealEd/Classes/Exporters/TextureExporterEXR.h b/Engine/Source/Editor/UnrealEd/Classes/Exporters/TextureExporterEXR.h
new file mode 100644
--- /dev/null
+++ b/Engine/Source/Editor/UnrealEd/Classes/Exporters/TextureExporterEXR.h
@@ -0,0 +1,27 @@
+// Copyright Epic Games, Inc. All Rights Reserved.
+
+//~=============================================================================
+// TextureExporterEXR
+//~=============================================================================
+
+#pragma once
+
+#include "CoreMinimal.h"
+#include "UObject/ObjectMacros.h"
+#include "Exporters/Exporter.h"
+#include "TextureExporterEXR.generated.h"
+
+UCLASS()
+class UNREALED_API UTextureExporterEXR : public UExporter
+{
+	GENERATED_UCLASS_BODY()
+
+
+	//~ Begin UExporter Interface
+	virtual bool SupportsObject(UObject* Object) const override;
+	virtual bool ExportBinary( UObject* Object, const TCHAR* Type, FArchive& Ar, FFeedbackContext* Warn, int32 FileIndex = 0, uint32 PortFlags=0 ) override;
+	//~ End UExporter Interface
+};
+
+
+
diff --git a/Engine/Source/Editor/UnrealEd/Classes/Exporters/TextureExporterPNG.h b/Engine/Source/Editor/UnrealEd/Classes/Exporters/TextureExporterPNG.h
new file mode 100644
--- /dev/null
+++ b/Engine/Source/Editor/UnrealEd/Classes/Exporters/TextureExporterPNG.h
@@ -0,0 +1,27 @@
+// Copyright Epic Games, Inc. All Rights Reserved.
+
+//~=============================================================================
+// TextureExporterPNG
+//~=============================================================================
+
+#pragma once
+
+#include "CoreMinimal.h"
+#include "UObject/ObjectMacros.h"
+#include "Exporters/Exporter.h"
+#include "TextureExporterPNG.generated.h"
+
+UCLASS()
+class UNREALED_API UTextureExporterPNG : public UExporter
+{
+	GENERATED_UCLASS_BODY()
+
+
+	//~ Begin UExporter Interface
+	virtual bool SupportsObject(UObject* Object) const override;
+	virtual bool ExportBinary( UObject* Object, const TCHAR* Type, FArchive& Ar, FFeedbackContext* Warn, int32 FileIndex = 0, uint32 PortFlags=0 ) override;
+	//~ End UExporter Interface
+};
+
+
+
diff --git a/Engine/Source/Editor/UnrealEd/Classes/Factories/TextureFactory.h b/Engine/Source/Editor/UnrealEd/Classes/Factories/TextureFactory.h
--- a/Engine/Source/Editor/UnrealEd/Classes/Factories/TextureFactory.h
+++ b/Engine/Source/Editor/UnrealEd/Classes/Factories/TextureFactory.h
@@ -15,7 +15,7 @@ struct FImportImage
 	TArray64<uint8> RawData;
 	ETextureSourceFormat Format = TSF_Invalid;
 	TextureCompressionSettings CompressionSettings = TC_Default;
-	int32 NumMips;
+	int32 NumMips = 0;
 	int32 SizeX = 0;
 	int32 SizeY = 0;
 	bool SRGB = true;
@@ -202,6 +202,9 @@ protected:
 	/** Keep track of if we are doing a reimport */
 	bool bIsDoingAReimport = false;
 
+	/* Store YesAll/NoAll responses: */
+	EAppReturnType::Type HDRImportShouldBeLongLatCubeMap = EAppReturnType::Retry;
+
 private:
 	/** This variable is static because in StaticImportObject() the type of the factory is not known. */
 	static bool bSuppressImportOverwriteDialog;
@@ -219,7 +222,7 @@ private:
 	*
 	*	@return	bool					true if the given height/width represent a supported texture resolution, false if not
 	*/
-	static bool IsImportResolutionValid(int32 Width, int32 Height, bool bAllowNonPowerOfTwo, FFeedbackContext* Warn);
+	static bool IsImportResolutionValid(int64 Width, int64 Height, bool bAllowNonPowerOfTwo, FFeedbackContext* Warn);
 
 	/** Flags to be used when calling ImportImage */
 	enum class EImageImportFlags
@@ -234,7 +237,7 @@ private:
 	FRIEND_ENUM_CLASS_FLAGS(EImageImportFlags);
 
 	/** Import image file into generic image struct, may be easily copied to FTextureSource */
-	bool ImportImage(const uint8* Buffer, uint32 Length, FFeedbackContext* Warn, EImageImportFlags Flags, FImportImage& OutImage);
+	bool ImportImage(const uint8* Buffer, int64 Length, FFeedbackContext* Warn, EImageImportFlags Flags, FImportImage& OutImage);
 
 	/** used by CreateTexture() */
 	UTexture* ImportTexture(UClass* Class, UObject* InParent, FName Name, EObjectFlags Flags, const TCHAR* Type, const uint8*& Buffer, const uint8* BufferEnd, FFeedbackContext* Warn);
diff --git a/Engine/Source/Editor/UnrealEd/Private/Factories/EditorFactories.cpp b/Engine/Source/Editor/UnrealEd/Private/Factories/EditorFactories.cpp
--- a/Engine/Source/Editor/UnrealEd/Private/Factories/EditorFactories.cpp
+++ b/Engine/Source/Editor/UnrealEd/Private/Factories/EditorFactories.cpp
@@ -134,6 +134,8 @@
 #include "Exporters/RenderTargetExporterHDR.h"
 #include "Exporters/TextureExporterPCX.h"
 #include "Exporters/TextureExporterTGA.h"
+#include "Exporters/TextureExporterPNG.h"
+#include "Exporters/TextureExporterEXR.h"
 #include "EngineGlobals.h"
 #include "GameFramework/ForceFeedbackEffect.h"
 #include "Engine/StaticMesh.h"
@@ -185,7 +187,6 @@
 #include "NormalMapIdentification.h"
 #include "AudioDeviceManager.h"
 #include "Kismet2/KismetEditorUtilities.h"
-#include "BmpImageSupport.h"
 #include "ScopedTransaction.h"
 #include "BSPOps.h"
 #include "LevelUtils.h"
@@ -201,6 +202,8 @@
 #include "IImageWrapper.h"
 #include "IImageWrapperModule.h"
 #include "TgaImageSupport.h"
+#include "ImageCore.h"
+#include "ImageCoreUtils.h"
 
 #include "FbxImporter.h"
 #include "Misc/FbxErrors.h"
@@ -419,12 +422,15 @@ UObject* UTexture2DFactoryNew::FactoryCreateNew( UClass* InClass, UObject* InPar
 	//Set the source art to be white as default.
 	if( Object->Source.IsValid() )
 	{
-		TArray64<uint8> TexturePixels;
-		Object->Source.GetMipData( TexturePixels, 0 );
+		// scope MipLock
+		{
+			FTextureSource::FMipLock MipLock(FTextureSource::ELockState::ReadWrite,&Object->Source,0);
 
-		uint8* DestData = Object->Source.LockMip(0);
-		FMemory::Memset(DestData, 255, TexturePixels.Num() * sizeof(uint8));
-		Object->Source.UnlockMip(0);
+			if ( MipLock.IsValid() )
+			{
+				FMemory::Memset(MipLock.Image.RawData, 255, MipLock.Image.GetImageSizeBytes());
+			}
+		}
 
 		Object->PostEditChange();
 	}
@@ -2212,12 +2218,13 @@ UObject* UCurveLinearColorAtlasFactory::FactoryCreateNew(UClass* Class, UObject*
 	const int32 TextureDataSize = Object->Source.CalcMipSize(0);
 	Object->SrcData.AddUninitialized(TextureDataSize);
 	uint32* TextureData = (uint32*)Object->Source.LockMip(0);
+	FFloat16Color InitColor(FLinearColor::White);
 	for (uint32 y = 0; y < Object->TextureSize; y++)
 	{
 		// Create base mip for the texture we created.
 		for (uint32 x = 0; x < Object->TextureSize; x++)
 		{
-			Object->SrcData[x*Object->TextureSize + y] = FLinearColor::White;
+			Object->SrcData[x*Object->TextureSize + y] = InitColor;
 		}
 	}
 	FMemory::Memcpy(TextureData, Object->SrcData.GetData(), TextureDataSize);
@@ -2936,6 +2943,12 @@ void FillZeroAlphaPNGData( int32 SizeX, int32 SizeY, ETextureSourceFormat Source
 			PNGFill.ProcessData();
 			break;
 		}
+
+		default:
+		{
+			// G8, G16, no alpha to fill
+			break;
+		}
 	}
 }
 
@@ -3002,295 +3015,179 @@ void* FImportImage::GetMipData(int32 InMipIndex)
 	return &RawData[Offset];
 }
 
-bool UTextureFactory::ImportImage(const uint8* Buffer, uint32 Length, FFeedbackContext* Warn, EImageImportFlags Flags, FImportImage& OutImage)
+static void AutoDetectAndChangeGrayScale(FImage & Image)
 {
-	TRACE_CPUPROFILER_EVENT_SCOPE(UTextureFactory::ImportImage)
-
-	const bool bAllowNonPowerOfTwo = EnumHasAllFlags(Flags, EImageImportFlags::AllowNonPowerOfTwo);
+	if ( Image.Format != ERawImageFormat::BGRA8 )
+	{
+		return;
+	}
 
-	IImageWrapperModule& ImageWrapperModule = FModuleManager::LoadModuleChecked<IImageWrapperModule>(FName("ImageWrapper"));
+	// auto-detect gray BGRA8 and change to G8
 
-	// Use the magic bytes when possible to avoid calling inefficient code to check if the image is of the right format
-	EImageFormat ImageFormat = ImageWrapperModule.DetectImageFormat(Buffer, int64(Length));
+	const FColor * Colors = (const FColor *)Image.RawData.GetData();
+	int64 NumPixels = Image.GetNumPixels();
 
-	//
-	// PNG
-	//
-	if (ImageFormat == EImageFormat::PNG)
+	for(int64 i=0;i<NumPixels;i++)
 	{
-		TSharedPtr<IImageWrapper> PngImageWrapper = ImageWrapperModule.CreateImageWrapper(EImageFormat::PNG);
-		if (PngImageWrapper.IsValid() && PngImageWrapper->SetCompressed(Buffer, Length))
+		if ( Colors[i].A != 255 ||
+			Colors[i].R != Colors[i].B ||
+			Colors[i].G != Colors[i].B )
 		{
-			if (!IsImportResolutionValid(PngImageWrapper->GetWidth(), PngImageWrapper->GetHeight(), bAllowNonPowerOfTwo, Warn))
-			{
-				return false;
-			}
-
-			// Select the texture's source format
-			ETextureSourceFormat TextureFormat = TSF_Invalid;
-			int32 BitDepth = PngImageWrapper->GetBitDepth();
-			ERGBFormat Format = PngImageWrapper->GetFormat();
-
-			if (Format == ERGBFormat::Gray)
-			{
-				if (BitDepth <= 8)
-				{
-					TextureFormat = TSF_G8;
-					Format = ERGBFormat::Gray;
-					BitDepth = 8;
-				}
-				else if (BitDepth == 16)
-				{
-					TextureFormat = TSF_G16;
-					Format = ERGBFormat::Gray;
-					BitDepth = 16;
-				}
-			}
-			else if (Format == ERGBFormat::RGBA || Format == ERGBFormat::BGRA)
-			{
-				if (BitDepth <= 8)
-				{
-					TextureFormat = TSF_BGRA8;
-					Format = ERGBFormat::BGRA;
-					BitDepth = 8;
-				}
-				else if (BitDepth == 16)
-				{
-					TextureFormat = TSF_RGBA16;
-					Format = ERGBFormat::RGBA;
-					BitDepth = 16;
-				}
-			}
+			return;
+		}
+	}
 
-			if (TextureFormat == TSF_Invalid)
-			{
-				Warn->Logf(ELogVerbosity::Error, TEXT("PNG file contains data in an unsupported format."));
-				return false;
-			}
+	// yes, it's gray, do it :
+	Image.ChangeFormat(ERawImageFormat::G8,Image.GammaSpace);
+}
 
-			OutImage.Init2DWithParams(
-				PngImageWrapper->GetWidth(),
-				PngImageWrapper->GetHeight(),
-				TextureFormat,
-				BitDepth < 16
-			);
+bool UTextureFactory::ImportImage(const uint8* Buffer, int64 Length, FFeedbackContext* Warn, EImageImportFlags Flags, FImportImage& OutImage)
+{
+	TRACE_CPUPROFILER_EVENT_SCOPE(UTextureFactory::ImportImage)
 
-			if (PngImageWrapper->GetRaw(Format, BitDepth, OutImage.RawData))
-			{
-				bool bFillPNGZeroAlpha = true;
-				GConfig->GetBool(TEXT("TextureImporter"), TEXT("FillPNGZeroAlpha"), bFillPNGZeroAlpha, GEditorIni);
+	const bool bAllowNonPowerOfTwo = EnumHasAllFlags(Flags, EImageImportFlags::AllowNonPowerOfTwo);
 
-				if (bFillPNGZeroAlpha)
-				{
-					// Replace the pixels with 0.0 alpha with a color value from the nearest neighboring color which has a non-zero alpha
-					FillZeroAlphaPNGData(OutImage.SizeX, OutImage.SizeY, OutImage.Format, OutImage.RawData.GetData());
-				}
-			}
-			else
-			{
-				Warn->Logf(ELogVerbosity::Error, TEXT("Failed to decode PNG."));
-				return false;
-			}
+	IImageWrapperModule& ImageWrapperModule = FModuleManager::LoadModuleChecked<IImageWrapperModule>(FName("ImageWrapper"));
 
-			return true;
-		}
-	}
+	// Use the magic bytes when possible to avoid calling inefficient code to check if the image is of the right format
+	EImageFormat ImageFormat = ImageWrapperModule.DetectImageFormat(Buffer, Length);
 
 	//
 	// JPEG
 	//
 	if (ImageFormat == EImageFormat::JPEG)
 	{
-		TSharedPtr<IImageWrapper> JpegImageWrapper = ImageWrapperModule.CreateImageWrapper(EImageFormat::JPEG);
-		if (JpegImageWrapper.IsValid() && JpegImageWrapper->SetCompressed(Buffer, Length))
+		// unusual loader, retains jpeg
+		bool bRetainJpegFormat = false;
+		if (EnumHasAllFlags(Flags, EImageImportFlags::AllowReturnOfCompressedData))
 		{
-			if (!IsImportResolutionValid(JpegImageWrapper->GetWidth(), JpegImageWrapper->GetHeight(), bAllowNonPowerOfTwo, Warn))
-			{
-				return false;
-			}
-
-			// Select the texture's source format
-			ETextureSourceFormat TextureFormat = TSF_Invalid;
-			int32 BitDepth = JpegImageWrapper->GetBitDepth();
-			ERGBFormat Format = JpegImageWrapper->GetFormat();
-
-			if (Format == ERGBFormat::Gray)
-			{
-				if (BitDepth <= 8)
-				{
-					TextureFormat = TSF_G8;
-					Format = ERGBFormat::Gray;
-					BitDepth = 8;
-				}
-			}
-			else if (Format == ERGBFormat::RGBA)
+			// For now this option is opt in via the config files once there is no technical risk this will become the default path.
+			GConfig->GetBool(TEXT("TextureImporter"), TEXT("RetainJpegFormat"), bRetainJpegFormat, GEditorIni);
+		}
+		if ( bRetainJpegFormat)
+		{
+			// does not decode jpeg, just to get width/height :
+			TSharedPtr<IImageWrapper> JpegImageWrapper = ImageWrapperModule.CreateImageWrapper(EImageFormat::JPEG);
+			if (JpegImageWrapper.IsValid() && JpegImageWrapper->SetCompressed(Buffer, Length))
 			{
-				if (BitDepth <= 8)
+				if (!IsImportResolutionValid(JpegImageWrapper->GetWidth(), JpegImageWrapper->GetHeight(), bAllowNonPowerOfTwo, Warn))
 				{
-					TextureFormat = TSF_BGRA8;
-					Format = ERGBFormat::BGRA;
-					BitDepth = 8;
+					return false;
 				}
-			}
-
-			if (TextureFormat == TSF_Invalid)
-			{
-				Warn->Logf(ELogVerbosity::Error, TEXT("JPEG file contains data in an unsupported format."));
-				return false;
-			}
-
-			OutImage.Init2DWithParams(
-				JpegImageWrapper->GetWidth(),
-				JpegImageWrapper->GetHeight(),
-				TextureFormat,
-				BitDepth < 16
-			);
+			
+				ERawImageFormat::Type RawFormat = JpegImageWrapper->GetClosestRawImageFormat();
+				check( RawFormat != ERawImageFormat::Invalid );
+				ETextureSourceFormat TextureFormat = FImageCoreUtils::ConvertToTextureSourceFormat(RawFormat);
+				bool bSRGB = JpegImageWrapper->GetSRGB();
 
-			bool bRetainJpegFormat = false;
-			if (EnumHasAllFlags(Flags, EImageImportFlags::AllowReturnOfCompressedData))
-			{
-				// For now this option is opt in via the config files once there is no technical risk this will become the default path.
-				GConfig->GetBool(TEXT("TextureImporter"), TEXT("RetainJpegFormat"), bRetainJpegFormat, GEditorIni);
-			}
+				OutImage.Init2DWithParams(
+					JpegImageWrapper->GetWidth(),
+					JpegImageWrapper->GetHeight(),
+					TextureFormat,
+					bSRGB
+				);
 
-			if (bRetainJpegFormat)
-			{
 				OutImage.RawData.Append(Buffer, Length);
 				OutImage.RawDataCompressionFormat = ETextureSourceCompressionFormat::TSCF_JPEG;
-			}
-			else if (!JpegImageWrapper->GetRaw(Format, BitDepth, OutImage.RawData))
-			{
-				Warn->Logf(ELogVerbosity::Error, TEXT("Failed to decode JPEG."));
-				return false;
-			}
 
-			return true;
+				UE_LOG(LogEditorFactories,Display,TEXT("JPEG imported and retained as JPEG in uasset."));
+
+				return true;
+			}
 		}
+		// if not bRetainJpegFormat , jpeg will use generic loader
 	}
 
-	//
-	// EXR
-	//
-	if (ImageFormat == EImageFormat::EXR)
+	if ( ImageFormat != EImageFormat::Invalid )
 	{
-		TSharedPtr<IImageWrapper> ExrImageWrapper = ImageWrapperModule.CreateImageWrapper(EImageFormat::EXR);
-		if (ExrImageWrapper.IsValid() && ExrImageWrapper->SetCompressed(Buffer, Length))
+		// Generic ImageWrapper loader :
+		// for PNG,EXR,BMP,TGA :
+		FImage LoadedImage;
+		if ( ImageWrapperModule.DecompressImage(Buffer, Length,LoadedImage) )
 		{
-			int32 Width = ExrImageWrapper->GetWidth();
-			int32 Height = ExrImageWrapper->GetHeight();
-
-			if (!IsImportResolutionValid(Width, Height, bAllowNonPowerOfTwo, Warn))
+			if (!IsImportResolutionValid(LoadedImage.SizeX, LoadedImage.SizeY, bAllowNonPowerOfTwo, Warn))
 			{
 				return false;
 			}
+		
+			AutoDetectAndChangeGrayScale(LoadedImage);
 
-			// Currently the only texture source format compatible with EXR image formats is TSF_RGBA16F.
-			// EXR decoder automatically converts imported image channels into the requested float format.
-			// In case if the imported image is a gray scale image, its content will be stored in the green channel of the created texture.
-			ETextureSourceFormat TextureFormat = TSF_RGBA16F;
-			ERGBFormat Format = ERGBFormat::RGBAF;
-			int32 BitDepth = 16;
+			ETextureSourceFormat TextureFormat = FImageCoreUtils::ConvertToTextureSourceFormat(LoadedImage.Format);
+			bool bSRGB = LoadedImage.GammaSpace != EGammaSpace::Linear;
 
 			OutImage.Init2DWithParams(
-				Width,
-				Height,
+				LoadedImage.SizeX,
+				LoadedImage.SizeY,
 				TextureFormat,
-				false
+				bSRGB
 			);
-			OutImage.CompressionSettings = TC_HDR;
 
-			if (!ExrImageWrapper->GetRaw(Format, BitDepth, OutImage.RawData))
+			// move LoadedImage bits to OutImage :
+			Swap(OutImage.RawData,LoadedImage.RawData);
+		
+			if ( ERawImageFormat::IsHDR(LoadedImage.Format) )
 			{
-				Warn->Logf(ELogVerbosity::Error, TEXT("Failed to decode EXR."));
-				return false;
+				OutImage.CompressionSettings = TC_HDR;
+				check( bSRGB == false );
 			}
 
-			return true;
-		}
-	}
+			// do per-format processing to match legacy behavior :
 
-	//
-	// BMP
-	//
-	if (ImageFormat == EImageFormat::BMP)
-	{
-		TSharedPtr<IImageWrapper> BmpImageWrapper = ImageWrapperModule.CreateImageWrapper(EImageFormat::BMP);
-		if (BmpImageWrapper.IsValid() && BmpImageWrapper->SetCompressed(Buffer, Length))
-		{
-			// Check the resolution of the imported texture to ensure validity
-			if (!IsImportResolutionValid(BmpImageWrapper->GetWidth(), BmpImageWrapper->GetHeight(), bAllowNonPowerOfTwo, Warn))
+			if ( ImageFormat == EImageFormat::PNG )
 			{
-				return false;
-			}
-
-			OutImage.Init2DWithParams(
-				BmpImageWrapper->GetWidth(),
-				BmpImageWrapper->GetHeight(),
-				TSF_BGRA8,
-				false
-			);
-
-			return BmpImageWrapper->GetRaw(BmpImageWrapper->GetFormat(), BmpImageWrapper->GetBitDepth(), OutImage.RawData);
-		}
-	}
+				bool bFillPNGZeroAlpha = true;
+				GConfig->GetBool(TEXT("TextureImporter"), TEXT("FillPNGZeroAlpha"), bFillPNGZeroAlpha, GEditorIni);
 
-	//
-	// TIFF
-	//
-	if (ImageFormat == EImageFormat::TIFF)
-	{
-		TSharedPtr<IImageWrapper> TiffImageWrapper = ImageWrapperModule.CreateImageWrapper(EImageFormat::TIFF);
-		if (TiffImageWrapper.IsValid())
-		{
-			if (TiffImageWrapper->SetCompressed(Buffer, Length))
-			{ 
-				// Check the resolution of the imported texture to ensure validity
-				if (!IsImportResolutionValid(TiffImageWrapper->GetWidth(), TiffImageWrapper->GetHeight(), bAllowNonPowerOfTwo, Warn))
+				if (bFillPNGZeroAlpha)
 				{
-					return false;
+					// Replace the pixels with 0.0 alpha with a color value from the nearest neighboring color which has a non-zero alpha
+					FillZeroAlphaPNGData(OutImage.SizeX, OutImage.SizeY, OutImage.Format, OutImage.RawData.GetData());
 				}
+			}
+			else if ( ImageFormat == EImageFormat::TGA )
+			{
+				const FTGAFileHeader* TGA = (FTGAFileHeader*)Buffer;
 
-				ETextureSourceFormat SourceFormat = TSF_Invalid;
-				ERGBFormat TiffFormat = TiffImageWrapper->GetFormat();
-				const int32 BitDepth = TiffImageWrapper->GetBitDepth();
-				bool bIsSRGB = false;
-
-				if (TiffFormat == ERGBFormat::BGRA)
-				{
-					SourceFormat = TSF_BGRA8;
-					bIsSRGB = true;
-				}
-				else if (TiffFormat == ERGBFormat::RGBA)
+				if (TGA->ColorMapType == 1 && TGA->ImageTypeCode == 1 && TGA->BitsPerPixel == 8)
 				{
-					SourceFormat = TSF_RGBA16;
+					// Notes: The Scaleform GFx exporter (dll) strips all font glyphs into a single 8-bit texture.
+					// The targa format uses this for a palette index; GFx uses a palette of (i,i,i,i) so the index
+					// is also the alpha value.
+					//
+					// We store the image as PF_G8, where it will be used as alpha in the Glyph shader.
+				
+					// ?? check or convert? or neither?
+					//check( TextureFormat == TSF_G8 );
+
+					OutImage.CompressionSettings = TC_Grayscale;
 				}
-				else if (TiffFormat == ERGBFormat::RGBAF)
+				else if(TGA->ColorMapType == 0 && TGA->ImageTypeCode == 3 && TGA->BitsPerPixel == 8)
 				{
-					SourceFormat = TSF_RGBA16F;
+					// standard grayscale images
+				
+					// ?? check or convert? or neither?
+					//check( TextureFormat == TSF_G8 );
+
+					OutImage.CompressionSettings = TC_Grayscale;
 				}
-				else if (TiffFormat == ERGBFormat::Gray)
+
+				if ( OutImage.CompressionSettings == TC_Grayscale && TGA->ImageTypeCode == 3)
 				{
-					SourceFormat = TSF_G8;
-					if (BitDepth == 16)
-					{
-						SourceFormat = TSF_G16;
-					}
+					// default grayscales to linear as they wont get compression otherwise and are commonly used as masks
+					OutImage.SRGB = false;
 				}
-
-				OutImage.Init2DWithParams(
-					TiffImageWrapper->GetWidth(),
-					TiffImageWrapper->GetHeight(),
-					SourceFormat,
-					bIsSRGB
-				);
-
-				return TiffImageWrapper->GetRaw(TiffFormat, BitDepth, OutImage.RawData);
 			}
 
-			return false;
+			return true;
+		}
+		else
+		{
+			Warn->Logf(ELogVerbosity::Error, TEXT("Image Import failed."));
+			// do not return false, let loaders below try
 		}
 	}
 
+
 	//
 	// PCX
 	//
@@ -3407,82 +3304,6 @@ bool UTextureFactory::ImportImage(const uint8* Buffer, uint32 Length, FFeedbackC
 		return true;
 	}
 
-	//
-	// TGA
-	//
-	TSharedPtr<IImageWrapper> TgaImageWrapper = ImageWrapperModule.CreateImageWrapper(EImageFormat::TGA);
-	if (TgaImageWrapper.IsValid() && TgaImageWrapper->SetCompressed(Buffer, Length))
-	{
-		// Check the resolution of the imported texture to ensure validity
-		if (!IsImportResolutionValid(TgaImageWrapper->GetWidth(), TgaImageWrapper->GetHeight(), bAllowNonPowerOfTwo, Warn))
-		{
-			return false;
-		}
-
-		const FTGAFileHeader* TGA = (FTGAFileHeader*)Buffer;
-
-		if (TGA->ColorMapType == 1 && TGA->ImageTypeCode == 1 && TGA->BitsPerPixel == 8)
-		{
-			// Notes: The Scaleform GFx exporter (dll) strips all font glyphs into a single 8-bit texture.
-			// The targa format uses this for a palette index; GFx uses a palette of (i,i,i,i) so the index
-			// is also the alpha value.
-			//
-			// We store the image as PF_G8, where it will be used as alpha in the Glyph shader.
-			OutImage.Init2DWithOneMip(
-				TGA->Width,
-				TGA->Height,
-				TSF_G8);
-			OutImage.CompressionSettings = TC_Grayscale;
-		}
-		else if(TGA->ColorMapType == 0 && TGA->ImageTypeCode == 3 && TGA->BitsPerPixel == 8)
-		{
-			// standard grayscale images
-			OutImage.Init2DWithOneMip(
-				TGA->Width,
-				TGA->Height,
-				TSF_G8);
-			OutImage.CompressionSettings = TC_Grayscale;
-		}
-		else
-		{
-			if(TGA->ImageTypeCode == 10) // 10 = RLE compressed 
-			{
-				if( TGA->BitsPerPixel != 32 &&
-					TGA->BitsPerPixel != 24 &&
-					TGA->BitsPerPixel != 16 )
-				{
-					Warn->Logf(ELogVerbosity::Error, TEXT("TGA uses an unsupported rle-compressed bit-depth: %u"),TGA->BitsPerPixel);
-					return false;
-				}
-			}
-			else
-			{
-				if( TGA->BitsPerPixel != 32 &&
-					TGA->BitsPerPixel != 16 &&
-					TGA->BitsPerPixel != 24)
-				{
-					Warn->Logf(ELogVerbosity::Error, TEXT("TGA uses an unsupported bit-depth: %u"),TGA->BitsPerPixel);
-					return false;
-				}
-			}
-
-			OutImage.Init2DWithOneMip(
-				TGA->Width,
-				TGA->Height,
-				TSF_BGRA8);
-		}
-
-		const bool bResult = TgaImageWrapper->GetRaw( TgaImageWrapper->GetFormat(), TgaImageWrapper->GetBitDepth(), OutImage.RawData );
-
-		if (bResult && OutImage.CompressionSettings == TC_Grayscale && TGA->ImageTypeCode == 3)
-		{
-			// default grayscales to linear as they wont get compression otherwise and are commonly used as masks
-			OutImage.SRGB = false;
-		}
-
-		return bResult;
-	}
-
 	//
 	// PSD File
 	//
@@ -3540,7 +3361,7 @@ bool UTextureFactory::ImportImage(const uint8* Buffer, uint32 Length, FFeedbackC
 	
 	//
 	// DDS Texture
-	//
+	// @@!! move this out with the other DDS cube/aray loader
 	FDDSLoadHelper  DDSLoadHelper(Buffer, Length);
 	if (DDSLoadHelper.IsValid2DTexture())
 	{
@@ -3711,11 +3532,13 @@ UTexture* UTextureFactory::ImportTextureUDIM(UClass* Class, UObject* InParent, F
 
 UTexture* UTextureFactory::ImportTexture(UClass* Class, UObject* InParent, FName Name, EObjectFlags Flags, const TCHAR* Type, const uint8*& Buffer, const uint8* BufferEnd, FFeedbackContext* Warn)
 {
+	TRACE_CPUPROFILER_EVENT_SCOPE(ImportTexture);
+
 	bool bAllowNonPowerOfTwo = false;
 	GConfig->GetBool( TEXT("TextureImporter"), TEXT("AllowNonPowerOfTwoTextures"), bAllowNonPowerOfTwo, GEditorIni );
 
 	// Validate it.
-	const int32 Length = BufferEnd - Buffer;
+	const int64 Length = BufferEnd - Buffer;
 
 	// We want to allow certain texture types to potentially return the compressed data rather than raw uncompressed
 	EImageImportFlags ImportFlags = EImageImportFlags::AllowReturnOfCompressedData;
@@ -3724,55 +3547,7 @@ UTexture* UTextureFactory::ImportTexture(UClass* Class, UObject* InParent, FName
 		ImportFlags |= EImageImportFlags::AllowNonPowerOfTwo;
 	}
 
-	//
-	// Generic 2D Image
-	//
-	FImportImage Image;
-	if (ImportImage(Buffer, Length, Warn, ImportFlags, Image))
-	{
-		UTexture2D* Texture = CreateTexture2D(InParent, Name, Flags);
-		if (Texture)
-		{
-			if (Image.RawDataCompressionFormat == ETextureSourceCompressionFormat::TSCF_None)
-			{
-				Texture->Source.Init(
-					Image.SizeX,
-					Image.SizeY,
-					/*NumSlices=*/ 1,
-					Image.NumMips,
-					Image.Format,
-					Image.RawData.GetData()
-				);
-			}
-			else
-			{
-				Texture->Source.InitWithCompressedSourceData(
-					Image.SizeX,
-					Image.SizeY,
-					Image.NumMips,
-					Image.Format,
-					Image.RawData,
-					Image.RawDataCompressionFormat
-				);
-			}
-
-			Texture->CompressionSettings = Image.CompressionSettings;
-
-			if (ColorSpaceMode == ETextureSourceColorSpace::Auto)
-			{
-				Texture->SRGB = Image.SRGB;
-			}
-			else if (ColorSpaceMode == ETextureSourceColorSpace::Linear)
-			{
-				Texture->SRGB = false;
-			}
-			else if (ColorSpaceMode == ETextureSourceColorSpace::SRGB)
-			{
-				Texture->SRGB = true;
-			}
-		}
-		return Texture;
-	}
+	// Do unusual loaders first before Generic
 
 	//
 	// DDS Cubemap
@@ -3812,7 +3587,7 @@ UTexture* UTextureFactory::ImportTexture(UClass* Class, UObject* InParent, FName
 				NumMips,
 				Format
 				);
-			if(Format == TSF_RGBA16F)
+			if(Format == TSF_RGBA16F || Format == TSF_RGBA32F)
 			{
 				TextureCube->CompressionSettings = TC_HDR;
 				TextureCube->SRGB = false;
@@ -3846,6 +3621,8 @@ UTexture* UTextureFactory::ImportTexture(UClass* Class, UObject* InParent, FName
 				TextureCube->MipGenSettings = TMGS_LeaveExistingMips;
 			}
 		}
+		
+		UE_LOG(LogEditorFactories,Display,TEXT("DDS imported as TextureCube.") );
 
 		return TextureCube;
 	}
@@ -3879,7 +3656,7 @@ UTexture* UTextureFactory::ImportTexture(UClass* Class, UObject* InParent, FName
 				NumMips,
 				Format
 			);
-			if (Format == TSF_RGBA16F)
+			if (Format == TSF_RGBA16F || Format == TSF_RGBA32F)
 			{
 				TextureArray->CompressionSettings = TC_HDR;
 				TextureArray->SRGB = false;
@@ -3914,49 +3691,101 @@ UTexture* UTextureFactory::ImportTexture(UClass* Class, UObject* InParent, FName
 				MipGenSettings = TMGS_LeaveExistingMips;
 			}
 		}
+		
+		UE_LOG(LogEditorFactories,Display,TEXT("DDS imported as TextureArray.") );
 
 		return TextureArray;
 	}
 
 	//
-	// HDR File
+	// HDR File -> longlat cube map
 	//
 	IImageWrapperModule& ImageWrapperModule = FModuleManager::LoadModuleChecked<IImageWrapperModule>(FName("ImageWrapper"));
 	if (ImageWrapperModule.DetectImageFormat(Buffer, Length) == EImageFormat::HDR)
 	{
-		FHdrImageWrapper HdrImageWrapper;
-		if (HdrImageWrapper.SetCompressedFromView(TArrayView64<const uint8>(Buffer, Length)))
-		{
-			TArray64<uint8> UnCompressedData;
-			if (HdrImageWrapper.GetRaw(ERGBFormat::BGRE, 8, UnCompressedData))
+		// do longlat or 2d or ask ?
+		
+		// note: HDRImportShouldBeLongLatCubeMap member variable is not preserved across imports
+		// because the TextureFactory is reconstructed each time
+		// it is saved across multiple files in a single import
+		
+		// if we don't have an All answer, pop the question :
+		if ( HDRImportShouldBeLongLatCubeMap != EAppReturnType::YesAll &&
+			HDRImportShouldBeLongLatCubeMap != EAppReturnType::NoAll )
+		{
+			// trinary bool
+			// 0 = no
+			// 1 = yes
+			// other = ask
+			int LoadHdrAsLongLatCubemap = 2;
+			GConfig->GetInt(TEXT("TextureImporter"), TEXT("LoadHdrAsLongLatCubemap"), LoadHdrAsLongLatCubemap, GEditorIni);
+			if ( LoadHdrAsLongLatCubemap == 0 )
+			{
+				HDRImportShouldBeLongLatCubeMap = EAppReturnType::NoAll;
+			}
+			else if ( LoadHdrAsLongLatCubemap == 1 )
 			{
-				// create the cube texture
-				UTextureCube* TextureCube = CreateTextureCube(InParent, Name, Flags);
-				if ( TextureCube )
+				HDRImportShouldBeLongLatCubeMap = EAppReturnType::YesAll;
+			}
+			else
+			{
+				// ask
+
+				if ( GIsRunningUnattendedScript )
 				{
-					TextureCube->Source.Init(
-						HdrImageWrapper.GetWidth(),
-						HdrImageWrapper.GetHeight(),
-						/*NumSlices=*/ 1,
-						/*NumMips=*/ 1,
-						TSF_BGRE8,
-						UnCompressedData.GetData()
-						);
-					// the loader can suggest a compression setting
-					TextureCube->CompressionSettings = TC_HDR;
-					TextureCube->SRGB = false;
-
-					return TextureCube;
+					// not allowed to messagebox
+					// just default to legacy behavior = yes latlong cubemap
+					HDRImportShouldBeLongLatCubeMap = EAppReturnType::Yes;
+				}
+				else
+				{
+					HDRImportShouldBeLongLatCubeMap = FMessageDialog::Open(EAppMsgType::YesNoYesAllNoAll, NSLOCTEXT("UnrealEd", "Import_HDRLongLat", "Should HDR texture be imported as longlat cubemap instead of 2d texture?"));
 				}
 			}
 		}
 
-		Warn->Log(ELogVerbosity::Error, HdrImageWrapper.GetErrorMessage().ToString());
-		
-	}
-
-	//
-	// IES File (usually measured real world light profiles)
+		// branch if they want longlat cubemap :
+		if (HDRImportShouldBeLongLatCubeMap == EAppReturnType::Yes ||
+			HDRImportShouldBeLongLatCubeMap == EAppReturnType::YesAll)
+		{
+			FHdrImageWrapper HdrImageWrapper;
+			if (HdrImageWrapper.SetCompressedFromView(TArrayView64<const uint8>(Buffer, Length)))
+			{
+				TArray64<uint8> UnCompressedData;
+				if (HdrImageWrapper.GetRaw(ERGBFormat::BGRE, 8, UnCompressedData))
+				{
+					// create the cube texture
+					UTextureCube* TextureCube = CreateTextureCube(InParent, Name, Flags);
+					if ( TextureCube )
+					{
+						TextureCube->Source.Init(
+							HdrImageWrapper.GetWidth(),
+							HdrImageWrapper.GetHeight(),
+							/*NumSlices=*/ 1,
+							/*NumMips=*/ 1,
+							TSF_BGRE8,
+							UnCompressedData.GetData()
+							);
+						// the loader can suggest a compression setting
+						TextureCube->CompressionSettings = TC_HDR;
+						TextureCube->SRGB = false;
+						
+						UE_LOG(LogEditorFactories,Display,TEXT("HDR Image imported as LongLat cube map.") );
+
+						return TextureCube;
+					}
+				}
+			}
+
+			Warn->Log(ELogVerbosity::Error, HdrImageWrapper.GetErrorMessage().ToString());
+			return nullptr;
+		}
+		// else didn't want cubemap
+		// will fall down to generic importer and make a 2d texture
+	}
+
+	//
+	// IES File (usually measured real world light profiles)
 	//
 	if(!FCString::Stricmp(Type, TEXT("ies")))
 	{
@@ -3990,6 +3819,70 @@ UTexture* UTextureFactory::ImportTexture(UClass* Class, UObject* InParent, FName
 		}
 	}
 
+	//
+	// Generic 2D Image
+	//
+	FImportImage Image;
+	if (ImportImage(Buffer, Length, Warn, ImportFlags, Image))
+	{
+		UTexture2D* Texture = CreateTexture2D(InParent, Name, Flags);
+		if (Texture)
+		{
+			UE_LOG(LogEditorFactories,Display,TEXT("Image imported as : %s"), *(StaticEnum<ETextureSourceFormat>()->GetDisplayNameTextByValue(Image.Format).ToString()) );
+
+			if (Image.RawDataCompressionFormat == ETextureSourceCompressionFormat::TSCF_None)
+			{
+				Texture->Source.Init(
+					Image.SizeX,
+					Image.SizeY,
+					/*NumSlices=*/ 1,
+					Image.NumMips,
+					Image.Format,
+					Image.RawData.GetData()
+				);
+			}
+			else
+			{
+				Texture->Source.InitWithCompressedSourceData(
+					Image.SizeX,
+					Image.SizeY,
+					Image.NumMips,
+					Image.Format,
+					Image.RawData,
+					Image.RawDataCompressionFormat
+				);
+			}
+
+			Texture->CompressionSettings = Image.CompressionSettings;
+
+			// check Texture Format before setting SRGB
+			if ( ERawImageFormat::GetFormatNeedsGammaSpace( FImageCoreUtils::ConvertToRawImageFormat(Image.Format) ) )
+			{
+				if (ColorSpaceMode == ETextureSourceColorSpace::Auto)
+				{
+					Texture->SRGB = Image.SRGB;
+				}
+				else if (ColorSpaceMode == ETextureSourceColorSpace::Linear)
+				{
+					Texture->SRGB = false;
+				}
+				else if (ColorSpaceMode == ETextureSourceColorSpace::SRGB)
+				{
+					Texture->SRGB = true;
+				}
+				else
+				{
+					check(0);
+				}
+			}
+			else
+			{
+				Texture->SRGB = false;
+			}
+		}
+		return Texture;
+	}
+
 	return nullptr;
 }
 
@@ -4011,6 +3904,7 @@ UObject* UTextureFactory::FactoryCreateBinary
 	FFeedbackContext*	Warn
 )
 {
+	TRACE_CPUPROFILER_EVENT_SCOPE(FactoryCreateBinary);
 	check(Type);
 
 	FName TextureName = Name;
@@ -4615,31 +4509,50 @@ void UTextureFactory::ApplyAutoImportSettings(UTexture* Texture)
 	}
 }
 
-bool UTextureFactory::IsImportResolutionValid(int32 Width, int32 Height, bool bAllowNonPowerOfTwo, FFeedbackContext* Warn)
+bool UTextureFactory::IsImportResolutionValid(int64 Width, int64 Height, bool bAllowNonPowerOfTwo, FFeedbackContext* Warn)
 {
 	static const auto CVarVirtualTexturesEnabled = IConsoleManager::Get().FindTConsoleVariableDataInt(TEXT("r.VirtualTextures")); check(CVarVirtualTexturesEnabled);
 
 	// In theory this value could be much higher, but various UE image code currently uses 32bit size/offset values
-	const int32 MaximumSupportedVirtualTextureResolution = 16 * 1024;
+	// TODO: in theory we can try making this bigger and fix any int32 overflow issues
+	const int64 MaximumSupportedVirtualTextureResolution = 16 * 1024;
 
 	// Calculate the maximum supported resolution utilizing the global max texture mip count
 	// (Note, have to subtract 1 because 1x1 is a valid mip-size; this means a GMaxTextureMipCount of 4 means a max resolution of 8x8, not 2^4 = 16x16)
-	const int32 MaximumSupportedResolution = CVarVirtualTexturesEnabled->GetValueOnAnyThread() ? MaximumSupportedVirtualTextureResolution : (1 << (GMaxTextureMipCount - 1));
+	int64 MaximumSupportedResolution = CVarVirtualTexturesEnabled->GetValueOnAnyThread() ? MaximumSupportedVirtualTextureResolution : ((int64)1 << (GMaxTextureMipCount - 1));
 
-	bool bValid = true;
+	// MaximumSupportedResolution is also limited by Oodle Texture maximum of 16k
+	MaximumSupportedResolution = FMath::Min(MaximumSupportedResolution,(int64)16384);
 
+	bool bValid = true;
+	
 	// Check if the texture is above the supported resolution and prompt the user if they wish to continue if it is
-	if ( Width > MaximumSupportedResolution || Height > MaximumSupportedResolution )
+
+	if (Width < 0 || Height < 0 || Width > MAX_int32 || Height > MAX_int32)
 	{
-		if (Width * Height < 0)
-		{
-			Warn->Log(ELogVerbosity::Error, *FText::Format(
-				NSLOCTEXT("UnrealEd", "Warning_TextureSizeTooLargeOrInvalid", "Texture is too large to import or it has an invalid resolution. The current maximun is {0} pixels"),
-				FText::AsNumber(FMath::Square(MaximumSupportedVirtualTextureResolution))
-				).ToString());
+		Warn->Log(ELogVerbosity::Error, *FText::Format(
+			NSLOCTEXT("UnrealEd", "Warning_TextureSizeTooLargeOrInvalid", "Texture is too large to import or it has an invalid resolution. The current maximun is {0} pixels"),
+			FText::AsNumber(FMath::Square(MaximumSupportedVirtualTextureResolution))
+			).ToString());
 
-			bValid = false;
-		}
+		bValid = false;
+	}
+
+	// pixel count must fit in int32 :
+	//  mip surface could still be larger than 2 GB, that's allowed
+	//	eg. 16k RGBA float = 4 GB
+	if ( bValid && Width * Height > MAX_int32)
+	{
+		Warn->Log(ELogVerbosity::Error, *FText::Format(
+			NSLOCTEXT("UnrealEd", "Warning_TextureSizeTooLargeOrInvalid", "Texture is too large to import or it has an invalid resolution. The current maximun is {0} pixels"),
+			FText::AsNumber(FMath::Square(MaximumSupportedVirtualTextureResolution))
+			).ToString());
+
+		bValid = false;
+	}
+
+	if ( Width > MaximumSupportedResolution || Height > MaximumSupportedResolution )
+	{
 
 		if ( bValid && (Width * Height) > FMath::Square(MaximumSupportedVirtualTextureResolution))
 		{
@@ -4767,6 +4680,40 @@ bool UTextureExporterPCX::ExportBinary( UObject* Object, const TCHAR* Type, FArc
 	return true;
 }
 
+
+static bool ExportTexture2DGeneric( UObject* Object,FArchive& Ar, const TCHAR * ImageFormat, FFeedbackContext* Warn, int LayerIndex=0)
+{
+	UTexture2D* Texture = CastChecked<UTexture2D>( Object );
+	if ( Texture == nullptr )
+	{
+		return false;
+	}
+
+	if ( !Texture->Source.IsValid() )
+	{
+		Warn->Logf(ELogVerbosity::Error, TEXT("ExportTexture failed : source not valid"));
+		return false;
+	}
+	
+	FImage Image;
+	if ( ! Texture->Source.GetMipImage(Image,0,LayerIndex,0) )
+	{
+		Warn->Logf(ELogVerbosity::Error, TEXT("ExportTexture failed : could not GetMipImage"));
+		return false;
+	}
+	
+	TArray64<uint8> OutArray;
+	if ( ! FImageUtils::CompressImage(OutArray,ImageFormat,Image) )
+	{
+		Warn->Logf(ELogVerbosity::Error, TEXT("ExportTexture failed : could not CompressImage to format %s"),ImageFormat);
+		return false;
+	}
+
+	Ar.Serialize(OutArray.GetData(),OutArray.Num());
+
+	return true;
+}
+
 /*------------------------------------------------------------------------------
 	UTextureExporterBMP implementation.
 ------------------------------------------------------------------------------*/
@@ -4782,6 +4729,7 @@ UTextureExporterBMP::UTextureExporterBMP(const FObjectInitializer& ObjectInitial
 
 UTexture2D* UTextureExporterBMP::GetExportTexture(UObject* Object) const
 {
+	// for UVirtualTextureBuilderExporterBMP
 	UVirtualTextureBuilder* VirtualTextureBuilder = Cast<UVirtualTextureBuilder>(Object);
 	if (VirtualTextureBuilder != nullptr)
 	{
@@ -4796,19 +4744,24 @@ bool UTextureExporterBMP::SupportsObject(UObject* Object) const
 	if (Super::SupportsObject(Object))
 	{
 		UTexture2D* Texture = GetExportTexture(Object);
-		bSupportsObject = Texture != nullptr && (Texture->Source.GetFormat() == TSF_BGRA8 || Texture->Source.GetFormat() == TSF_RGBA16 || Texture->Source.GetFormat() == TSF_G8 || Texture->Source.GetFormat() == TSF_G16);
+		if ( Texture )
+		{
+			bSupportsObject = Texture->Source.GetFormat() == TSF_BGRA8 || Texture->Source.GetFormat() == TSF_G8;
+		}
 	}
 	return bSupportsObject;
 }
 
 int32 UTextureExporterBMP::GetFileCount(UObject* Object) const
 {
+	// for UVirtualTextureBuilderExporterBMP
 	UTexture2D* Texture = GetExportTexture(Object);
 	return (Texture != nullptr) ? Texture->Source.GetNumLayers() : 1;
 }
 
 FString UTextureExporterBMP::GetUniqueFilename(const TCHAR* Filename, int32 FileIndex, int32 FileCount)
 {
+	// for UVirtualTextureBuilderExporterBMP
 	return (FileCount == 1) ? Filename : FString::Printf(TEXT("%s%d%s"), *FPaths::GetBaseFilename(Filename, false), FileIndex, *FPaths::GetExtension(Filename, true));
 }
 
@@ -4819,148 +4772,137 @@ bool UTextureExporterBMP::ExportBinary( UObject* Object, const TCHAR* Type, FArc
 	{
 		return false;
 	}
+	
+	// FileIndex for layers for VT :
+
+	return ExportTexture2DGeneric(Texture,Ar,TEXT("BMP"),Warn,FileIndex);
+}
 
-	const int32 LayerIndex = FileIndex;
-	ETextureSourceFormat Format = Texture->Source.IsValid() ? Texture->Source.GetFormat(LayerIndex) : TSF_Invalid;
-	if (Format != TSF_BGRA8 && Format != TSF_RGBA16 && Format != TSF_G16 && Format != TSF_G8)
+UVirtualTextureBuilderExporterBMP::UVirtualTextureBuilderExporterBMP(const FObjectInitializer& ObjectInitializer)
+	: UTextureExporterBMP(ObjectInitializer)
+{
+	SupportedClass = UVirtualTextureBuilder::StaticClass();
+	// already done by UTextureExporterBMP
+	//   I just change the SupportedClass
+	//PreferredFormatIndex = 0;
+	//FormatExtension.Add(TEXT("BMP"));
+	//FormatDescription.Add(TEXT("BMP Windows Bitmap"));
+}
+
+/*------------------------------------------------------------------------------
+	URenderTargetExporterHDR implementation.
+	Exports render targets.
+------------------------------------------------------------------------------*/
+// This is not for screenshots it's for "rendertarget" assets in levels
+//	when you do asset action export on them
+
+static bool ExportRenderTarget2DGeneric(UObject* Object, FArchive& Ar, const TCHAR * ImageExtension)
+{
+	UTextureRenderTarget2D* TexRT2D = Cast<UTextureRenderTarget2D>(Object);
+
+	if (TexRT2D == nullptr)
 	{
 		return false;
 	}
-
-	const bool bIsRGBA16 = Format == TSF_RGBA16;
-	const bool bIsG16 = Format == TSF_G16;
-	const bool bIsG8 = Format == TSF_G8;
-	const int32 SourceBytesPerPixel = bIsRGBA16 ? 8 : (bIsG16 ? 2 : (bIsG8 ? 1 : 4));
-
-	if (bIsRGBA16)
+	
+	FImage Image;
+	if ( ! FImageUtils::GetRenderTargetImage(TexRT2D,Image) )
 	{
-		FMessageLog ExportWarning("EditorErrors");
-		FFormatNamedArguments Arguments;
-		Arguments.Add(TEXT("Name"), FText::FromString(Texture->GetName()));
-		ExportWarning.Warning(FText::Format(LOCTEXT("BitDepthBMPWarning", "{Name}: Texture is RGBA16 and cannot be represented at such high bit depth in .bmp. Color will be scaled to RGBA8."), Arguments));
-		ExportWarning.Open(EMessageSeverity::Warning);
+		return false;
 	}
 
-	if (bIsG16)
+	TArray64<uint8> CompressedData;
+	if ( ! FImageUtils::CompressImage(CompressedData,ImageExtension,Image) )
 	{
-		FMessageLog ExportWarning("EditorErrors");
-		FFormatNamedArguments Arguments;
-		Arguments.Add(TEXT("Name"), FText::FromString(Texture->GetName()));
-		ExportWarning.Warning(FText::Format(LOCTEXT("BitDepthBMPWarningG16", "{Name}: Texture is G16 and cannot be represented at such high bit depth in .bmp. Color will be split across two channels of RGBA8."), Arguments));
-		ExportWarning.Open(EMessageSeverity::Warning);
+		return false;
 	}
 
-	int32 SizeX = Texture->Source.GetSizeX();
-	int32 SizeY = Texture->Source.GetSizeY();
-	TArray64<uint8> RawData;
-	Texture->Source.GetMipData(RawData, 0, LayerIndex, 0);
-
-	FBitmapFileHeader bmf;
-	FBitmapInfoHeader bmhdr;
-
-	// File header.
-	bmf.bfType      = 'B' + (256*(int32)'M');
-	bmf.bfReserved1 = 0;
-	bmf.bfReserved2 = 0;
-	int32 biSizeImage = SizeX * SizeY * 3;
-	bmf.bfOffBits   = sizeof(FBitmapFileHeader) + sizeof(FBitmapInfoHeader);
-	bmhdr.biBitCount = 24;
-
-	bmf.bfSize		= bmf.bfOffBits + biSizeImage;
-	Ar << bmf;
-
-	// Info header.
-	bmhdr.biSize          = sizeof(FBitmapInfoHeader);
-	bmhdr.biWidth         = SizeX;
-	bmhdr.biHeight        = SizeY;
-	bmhdr.biPlanes        = 1;
-	bmhdr.biCompression   = BCBI_RGB;
-	bmhdr.biSizeImage     = biSizeImage;
-	bmhdr.biXPelsPerMeter = 0;
-	bmhdr.biYPelsPerMeter = 0;
-	bmhdr.biClrUsed       = 0;
-	bmhdr.biClrImportant  = 0;
-	Ar << bmhdr;
-
-
-	// Upside-down scanlines.
-	for( int32 i=SizeY-1; i>=0; i-- )
-	{
-		uint8* ScreenPtr = &RawData[i*SizeX*SourceBytesPerPixel];
-		for( int32 j=SizeX; j>0; j-- )
-		{
-			if (bIsRGBA16)
-			{
-				Ar << ScreenPtr[1];
-				Ar << ScreenPtr[3];
-				Ar << ScreenPtr[5];
-				ScreenPtr += 8;
-			}
-			if (bIsG16)
-			{
-				uint8 DummyZero = 0;
-				Ar << DummyZero;
-				Ar << ScreenPtr[0];
-				Ar << ScreenPtr[1];
-				ScreenPtr += 2;
-			}
-			else if (bIsG8)
-			{
-				Ar << ScreenPtr[0];
-				Ar << ScreenPtr[0];
-				Ar << ScreenPtr[0];
-				ScreenPtr += 1;
-			}
-			else
-			{
-				Ar << ScreenPtr[0];
-				Ar << ScreenPtr[1];
-				Ar << ScreenPtr[2];
-				ScreenPtr += 4;
-			}
-		}
-	}
+	Ar.Serialize((void*)CompressedData.GetData(), CompressedData.GetAllocatedSize());
+
 	return true;
 }
 
-UVirtualTextureBuilderExporterBMP::UVirtualTextureBuilderExporterBMP(const FObjectInitializer& ObjectInitializer)
+// HDR is terrible, use EXR or PNG instead
+//  left for now for backwards compatibility
+//  EXR and PNG exporters are prioritized in AssetTools
+URenderTargetExporterHDR::URenderTargetExporterHDR(const FObjectInitializer& ObjectInitializer)
 	: Super(ObjectInitializer)
 {
-	SupportedClass = UVirtualTextureBuilder::StaticClass();
+	SupportedClass = UTextureRenderTarget2D::StaticClass();
 	PreferredFormatIndex = 0;
-	FormatExtension.Add(TEXT("BMP"));
-	FormatDescription.Add(TEXT("Windows Bitmap"));
+	FormatExtension.Add(TEXT("HDR"));
+	FormatDescription.Add(TEXT("HDR Radiance RGBE lossy rendertarget"));
+
+	// HDR can export both RGBA8 and 16F render targets
+	//  so it has no SupportsObject() check
+	//  though it's poor for both
 }
 
-/*------------------------------------------------------------------------------
-	URenderTargetExporterHDR implementation.
-	Exports render targets.
-------------------------------------------------------------------------------*/
-URenderTargetExporterHDR::URenderTargetExporterHDR(const FObjectInitializer& ObjectInitializer)
+bool URenderTargetExporterHDR::ExportBinary(UObject* Object, const TCHAR* Type, FArchive& Ar, FFeedbackContext* Warn, int32 FileIndex, uint32 PortFlags)
+{
+	return ExportRenderTarget2DGeneric(Object,Ar,TEXT("HDR"));
+}
+
+URenderTargetExporterPNG::URenderTargetExporterPNG(const FObjectInitializer& ObjectInitializer)
 	: Super(ObjectInitializer)
 {
-	SupportedClass = UTextureRenderTarget::StaticClass();
+	SupportedClass = UTextureRenderTarget2D::StaticClass();
 	PreferredFormatIndex = 0;
-	FormatExtension.Add(TEXT("HDR"));
-	FormatDescription.Add(TEXT("HDR"));
+	FormatExtension.Add(TEXT("PNG"));
+	FormatDescription.Add(TEXT("PNG"));
 }
 
-bool URenderTargetExporterHDR::ExportBinary(UObject* Object, const TCHAR* Type, FArchive& Ar, FFeedbackContext* Warn, int32 FileIndex, uint32 PortFlags)
+bool URenderTargetExporterPNG::SupportsObject(UObject* Object) const
 {
-	UTextureRenderTarget2D* TexRT2D = Cast<UTextureRenderTarget2D>(Object);
-	UTextureRenderTargetCube* TexRTCube = Cast<UTextureRenderTargetCube>(Object);
-
-	if (TexRT2D != nullptr)
+	if (Super::SupportsObject(Object))
 	{
-		return FImageUtils::ExportRenderTarget2DAsHDR(TexRT2D, Ar);
+		UTextureRenderTarget2D* TexRT2D = Cast<UTextureRenderTarget2D>(Object);
+
+		if (TexRT2D)
+		{
+			EPixelFormat PixelFormat = TexRT2D->GetFormat();
+
+			return PixelFormat == PF_B8G8R8A8;
+		}
 	}
-	else if (TexRTCube != nullptr)
+	return false;
+}
+
+bool URenderTargetExporterPNG::ExportBinary(UObject* Object, const TCHAR* Type, FArchive& Ar, FFeedbackContext* Warn, int32 FileIndex, uint32 PortFlags)
+{
+	return ExportRenderTarget2DGeneric(Object,Ar,TEXT("PNG"));
+}
+
+URenderTargetExporterEXR::URenderTargetExporterEXR(const FObjectInitializer& ObjectInitializer)
+	: Super(ObjectInitializer)
+{
+	SupportedClass = UTextureRenderTarget2D::StaticClass();
+	PreferredFormatIndex = 0;
+	FormatExtension.Add(TEXT("EXR"));
+	FormatDescription.Add(TEXT("EXR HDR float rendertarget"));
+}
+
+bool URenderTargetExporterEXR::SupportsObject(UObject* Object) const
+{
+	if (Super::SupportsObject(Object))
 	{
-		return FImageUtils::ExportRenderTargetCubeAsHDR(TexRTCube, Ar);
+		UTextureRenderTarget2D* TexRT2D = Cast<UTextureRenderTarget2D>(Object);
+
+		if (TexRT2D)
+		{
+			EPixelFormat PixelFormat = TexRT2D->GetFormat();
+
+			return PixelFormat == PF_FloatRGBA;
+		}
 	}
 	return false;
 }
 
+bool URenderTargetExporterEXR::ExportBinary(UObject* Object, const TCHAR* Type, FArchive& Ar, FFeedbackContext* Warn, int32 FileIndex, uint32 PortFlags)
+{
+	return ExportRenderTarget2DGeneric(Object,Ar,TEXT("EXR"));
+}
+
 /*------------------------------------------------------------------------------
 	UTextureCubeExporterHDR implementation.
 	Export UTextureCubes as .HDR
@@ -4971,20 +4913,42 @@ UTextureCubeExporterHDR::UTextureCubeExporterHDR(const FObjectInitializer& Objec
 	SupportedClass = UTextureCube::StaticClass();
 	PreferredFormatIndex = 0;
 	FormatExtension.Add(TEXT("HDR"));
-	FormatDescription.Add(TEXT("HDR"));
+	FormatDescription.Add(TEXT("HDR cube LongLatUnwrap"));
+	
+	// this is the only exporter registered for UTextureCube
+	// this does GenerateLongLatUnwrap
+
+	// @@!! todo : register a DDS cube exporter that can write the faces
+	//		how can user choose between the two? -> by extension
 }
 
 bool UTextureCubeExporterHDR::ExportBinary(UObject* Object, const TCHAR* Type, FArchive& Ar, FFeedbackContext* Warn, int32 FileIndex, uint32 PortFlags)
 {
 	UTextureCube* TexCube = Cast<UTextureCube>(Object);
+	// for URenderTargetCubeExporterHDR
+	UTextureRenderTargetCube* TexRTCube = Cast<UTextureRenderTargetCube>(Object);
 
 	if (TexCube != nullptr)
 	{
+		// this does GenerateLongLatUnwrap
 		return FImageUtils::ExportTextureCubeAsHDR(TexCube, Ar);
 	}
+	else if (TexRTCube != nullptr)
+	{
+		// -> this can continue to use HDR for the longlat unwraps
+		// this does GenerateLongLatUnwrap
+		return FImageUtils::ExportRenderTargetCubeAsHDR(TexRTCube, Ar);
+	}
+
 	return false;
 }
 
+URenderTargetCubeExporterHDR::URenderTargetCubeExporterHDR(const FObjectInitializer& ObjectInitializer)
+	: UTextureCubeExporterHDR(ObjectInitializer)
+{
+	SupportedClass = UTextureRenderTargetCube::StaticClass();
+}
+
 /*------------------------------------------------------------------------------
 	UTextureExporterHDR implementation.
 	Export UTexture2D as .HDR
@@ -4995,7 +4959,7 @@ UTextureExporterHDR::UTextureExporterHDR(const FObjectInitializer& ObjectInitial
 	SupportedClass = UTexture2D::StaticClass();
 	PreferredFormatIndex = 0;
 	FormatExtension.Add(TEXT("HDR"));
-	FormatDescription.Add(TEXT("HDR"));
+	FormatDescription.Add(TEXT("HDR Radiance RGBE texture"));
 }
 
 bool UTextureExporterHDR::SupportsObject(UObject* Object) const
@@ -5007,7 +4971,8 @@ bool UTextureExporterHDR::SupportsObject(UObject* Object) const
 
 		if (Texture)
 		{
-			bSupportsObject = Texture->Source.GetFormat() == TSF_BGRA8 || Texture->Source.GetFormat() == TSF_RGBA16F;
+			// only BGRE is allowed to go out to HDR now (lossless)
+			bSupportsObject = Texture->Source.GetFormat() == TSF_BGRE8;
 		}
 	}
 	return bSupportsObject;
@@ -5015,22 +4980,83 @@ bool UTextureExporterHDR::SupportsObject(UObject* Object) const
 
 bool UTextureExporterHDR::ExportBinary(UObject* Object, const TCHAR* Type, FArchive& Ar, FFeedbackContext* Warn, int32 FileIndex, uint32 PortFlags)
 {
-	UTexture2D* Texture = Cast<UTexture2D>(Object);
+	return ExportTexture2DGeneric(Object,Ar,TEXT("HDR"),Warn);
+}
 
-	if (Texture != nullptr)
+/*------------------------------------------------------------------------------
+	UTextureExporterPNG implementation.
+------------------------------------------------------------------------------*/
+UTextureExporterPNG::UTextureExporterPNG(const FObjectInitializer& ObjectInitializer)
+	: Super(ObjectInitializer)
+{
+	SupportedClass = UTexture2D::StaticClass();
+	PreferredFormatIndex = 0;
+	FormatExtension.Add(TEXT("PNG"));
+	FormatDescription.Add(TEXT("PNG"));
+}
+
+bool UTextureExporterPNG::SupportsObject(UObject* Object) const
+{
+	if (Super::SupportsObject(Object))
 	{
-		return FImageUtils::ExportTexture2DAsHDR(Texture, Ar);
+		UTexture2D* Texture = Cast<UTexture2D>(Object);
+
+		if (Texture)
+		{
+			ETextureSourceFormat TSF = Texture->Source.GetFormat();
+			ERawImageFormat::Type RawFormat = FImageCoreUtils::ConvertToRawImageFormat(TSF);
+			// supports all non-HDR formats :
+			return ! ERawImageFormat::IsHDR(RawFormat);
+		}
 	}
 	return false;
 }
 
+bool UTextureExporterPNG::ExportBinary( UObject* Object, const TCHAR* Type, FArchive& Ar, FFeedbackContext* Warn, int32 FileIndex, uint32 PortFlags )
+{
+	return ExportTexture2DGeneric(Object,Ar,TEXT("PNG"),Warn);
+}
+
+/*------------------------------------------------------------------------------
+	UTextureExporterEXR implementation.
+------------------------------------------------------------------------------*/
+UTextureExporterEXR::UTextureExporterEXR(const FObjectInitializer& ObjectInitializer)
+	: Super(ObjectInitializer)
+{
+	SupportedClass = UTexture2D::StaticClass();
+	PreferredFormatIndex = 0;
+	FormatExtension.Add(TEXT("EXR"));
+	FormatDescription.Add(TEXT("EXR HDR float texture"));
+}
+
+bool UTextureExporterEXR::SupportsObject(UObject* Object) const
+{
+	if (Super::SupportsObject(Object))
+	{
+		UTexture2D* Texture = Cast<UTexture2D>(Object);
+
+		if (Texture)
+		{
+			ETextureSourceFormat TSF = Texture->Source.GetFormat();
+			ERawImageFormat::Type RawFormat = FImageCoreUtils::ConvertToRawImageFormat(TSF);
+			// supports all HDR formats :
+			return ERawImageFormat::IsHDR(RawFormat);
+		}
+	}
+	return false;
+}
+
+bool UTextureExporterEXR::ExportBinary( UObject* Object, const TCHAR* Type, FArchive& Ar, FFeedbackContext* Warn, int32 FileIndex, uint32 PortFlags )
+{
+	return ExportTexture2DGeneric(Object,Ar,TEXT("EXR"),Warn);
+}
+
 /*------------------------------------------------------------------------------
 	UTextureExporterTGA implementation.
 ------------------------------------------------------------------------------*/
 UTextureExporterTGA::UTextureExporterTGA(const FObjectInitializer& ObjectInitializer)
 	: Super(ObjectInitializer)
 {
-
 	SupportedClass = UTexture2D::StaticClass();
 	PreferredFormatIndex = 0;
 	FormatExtension.Add(TEXT("TGA"));
@@ -5046,7 +5072,7 @@ bool UTextureExporterTGA::SupportsObject(UObject* Object) const
 
 		if (Texture)
 		{
-			bSupportsObject = Texture->Source.GetFormat() == TSF_BGRA8 || Texture->Source.GetFormat() == TSF_RGBA16;
+			bSupportsObject = Texture->Source.GetFormat() == TSF_BGRA8;
 		}
 	}
 	return bSupportsObject;
@@ -5056,6 +5082,9 @@ bool UTextureExporterTGA::ExportBinary( UObject* Object, const TCHAR* Type, FArc
 {
 	UTexture2D* Texture = CastChecked<UTexture2D>( Object );
 
+	// could move to TGAImageWrapper and use generic exporter instead
+	//	maybe just leave this as-is
+
 	if (!Texture->Source.IsValid() || (Texture->Source.GetFormat() != TSF_BGRA8 && Texture->Source.GetFormat() != TSF_RGBA16))
 	{
 		return false;
@@ -5074,6 +5103,9 @@ bool UTextureExporterTGA::ExportBinary( UObject* Object, const TCHAR* Type, FArc
 
 	const int32 BytesPerPixel = bIsRGBA16 ? 8 : 4;
 
+	// Detects alpha channel
+	// writes either 24 or 32 bit
+
 	int32 SizeX = Texture->Source.GetSizeX();
 	int32 SizeY = Texture->Source.GetSizeY();
 	TArray64<uint8> RawData;
@@ -5081,7 +5113,9 @@ bool UTextureExporterTGA::ExportBinary( UObject* Object, const TCHAR* Type, FArc
 
 	// If we should export the file with no alpha info.  
 	// If the texture is compressed with no alpha we should definitely not export an alpha channel
-	bool bExportWithAlpha = !Texture->CompressionNoAlpha;
+	//  Umm what
+	//bool bExportWithAlpha = !Texture->CompressionNoAlpha;
+	bool bExportWithAlpha = true;
 	if( bExportWithAlpha )
 	{
 		// If the texture isn't compressed with no alpha scan the texture to see if the alpha values are all 255 which means we can skip exporting it.
diff --git a/Engine/Source/Editor/UnrealEd/Private/Lightmass/Lightmass.cpp b/Engine/Source/Editor/UnrealEd/Private/Lightmass/Lightmass.cpp
--- a/Engine/Source/Editor/UnrealEd/Private/Lightmass/Lightmass.cpp
+++ b/Engine/Source/Editor/UnrealEd/Private/Lightmass/Lightmass.cpp
@@ -50,6 +50,7 @@
 #include "EditorLevelUtils.h"
 #include "Misc/MessageDialog.h"
 #include "Modules/ModuleManager.h"
+#include "ImageCoreUtils.h"
 
 extern FSwarmDebugOptions GSwarmDebugOptions;
 
@@ -1139,10 +1140,10 @@ void FLightmassExporter::WriteLights( int32 Channel )
 				RectData.SourceTextureSizeX = SizeX;
 				RectData.SourceTextureSizeY = SizeY;
 
-				while( SizeX > 1 || SizeY > 1 )
+				ERawImageFormat::Type RawFormat = FImageCoreUtils::ConvertToRawImageFormat(Format);
+
+				if ( SizeX > 1 || SizeY > 1 )
 				{
-					//int32 i = SourceTexture.AddUninitialized( SizeX * SizeY );
-					
 					TArray64< uint8 > MipData;
 					Source.GetMipData( MipData, MipLevel );
 
@@ -1150,45 +1151,13 @@ void FLightmassExporter::WriteLights( int32 Channel )
 					uint8* PixelEnd = Pixel + MipData.Num();
 					while( Pixel < PixelEnd )
 					{
-						if( Format == TSF_RGBA16 || Format == TSF_RGBA16F )
-						{
-							//SourceTexture[ i++ ] = *(FFloat16Color*)Pixel;
-							RectData.SourceTextureAvgColor += FLinearColor( *(FFloat16Color*)Pixel );
-						}
-						else
-						{
-							FColor Color = FColor( Pixel[0], Pixel[0], Pixel[0] );
-							if( Format == TSF_BGRA8 || Format == TSF_BGRE8 )
-							{
-								Color = *(FColor*)Pixel;
-							}
-
-							FLinearColor LinearColor;
-							if( SRGB )
-							{
-								LinearColor = FLinearColor( Color );
-							}
-							else
-							{
-								LinearColor.R = float( Color.R ) / 255.0f;
-								LinearColor.G = float( Color.G ) / 255.0f;
-								LinearColor.B = float( Color.B ) / 255.0f;
-							}
-							
-							//SourceTexture[ i++ ] = FFloat16Color( LinearColor );
-							RectData.SourceTextureAvgColor += LinearColor;
-						}
+						FLinearColor Color = ERawImageFormat::GetOnePixelLinear(Pixel,RawFormat,SRGB);
+						RectData.SourceTextureAvgColor += Color;
 
 						Pixel += BytesPerPixel;
 					}
 
-					// Source doesn't often have mips just average color for now.
 					RectData.SourceTextureAvgColor *= 1.0f / ( SizeX * SizeY );
-					break;
-					
-					SizeX = SizeX > 1 ? SizeX >> 1 : 1;
-					SizeY = SizeY > 1 ? SizeY >> 1 : 1;
-					MipLevel++;
 				}
 			}
 		}
diff --git a/Engine/Source/Editor/UnrealEd/Private/PhysicalMaterialMaskImport.cpp b/Engine/Source/Editor/UnrealEd/Private/PhysicalMaterialMaskImport.cpp
--- a/Engine/Source/Editor/UnrealEd/Private/PhysicalMaterialMaskImport.cpp
+++ b/Engine/Source/Editor/UnrealEd/Private/PhysicalMaterialMaskImport.cpp
@@ -202,7 +202,6 @@ void FPhysicalMaterialMaskImport::GetSupportedTextureFileTypes(TArray<FString>&
 void FPhysicalMaterialMaskImport::GetSupportedTextureSourceFormats(TArray<ETextureSourceFormat>& OutSourceFormats)
 {
 	OutSourceFormats.Empty();
-	OutSourceFormats.Emplace(ETextureSourceFormat::TSF_RGBA8);
 	OutSourceFormats.Emplace(ETextureSourceFormat::TSF_BGRA8);
 	OutSourceFormats.Emplace(ETextureSourceFormat::TSF_RGBA16);
 }
diff --git a/Engine/Source/Editor/UnrealEd/UnrealEd.Build.cs b/Engine/Source/Editor/UnrealEd/UnrealEd.Build.cs
--- a/Engine/Source/Editor/UnrealEd/UnrealEd.Build.cs
+++ b/Engine/Source/Editor/UnrealEd/UnrealEd.Build.cs
@@ -125,6 +125,7 @@ public class UnrealEd : ModuleRules
 		PrivateDependencyModuleNames.AddRange(
 			new string[]
 			{
+				"ImageCore",
 				"AssetRegistry",
 				"AssetTagsEditor",
 				"BSPUtils",
diff --git a/Engine/Source/Runtime/Core/Private/Misc/FileHelper.cpp b/Engine/Source/Runtime/Core/Private/Misc/FileHelper.cpp
--- a/Engine/Source/Runtime/Core/Private/Misc/FileHelper.cpp
+++ b/Engine/Source/Runtime/Core/Private/Misc/FileHelper.cpp
@@ -893,6 +893,8 @@ bool FFileHelper::CreateBitmap( const TCHAR* Pattern, int32 SourceWidth, int32 S
 		}
 
 		// Colors.
+		// @todo fix me : calling Serialize per byte = insanely slow
+		//	BmpImageWrapper now has a good writer, prefer that ; use FImageUtils::SaveImage
 		for( int32 i = SubRectangle->Max.Y - 1; i >= SubRectangle->Min.Y; i-- )
 		{
 			for( int32 j = SubRectangle->Min.X; j < SubRectangle->Max.X; j++ )
diff --git a/Engine/Source/Runtime/Core/Public/Misc/FileHelper.h b/Engine/Source/Runtime/Core/Public/Misc/FileHelper.h
--- a/Engine/Source/Runtime/Core/Public/Misc/FileHelper.h
+++ b/Engine/Source/Runtime/Core/Public/Misc/FileHelper.h
@@ -162,6 +162,10 @@ struct CORE_API FFileHelper
 	 *
 	 * @return true if success
 	 */
+	 // DEPRECATED
+	//	BmpImageWrapper now has a good writer, prefer that for general image saving ; use FImageUtils::SaveImage
+	// note this also calls SendDataToPCViaUnrealConsole
+	//   and uses GenerateNextBitmapFilename
 	static bool CreateBitmap( const TCHAR* Pattern, int32 DataWidth, int32 DataHeight, const struct FColor* Data, struct FIntRect* SubRectangle = NULL, IFileManager* FileManager = &IFileManager::Get(), FString* OutFilename = NULL, bool bInWriteAlpha = false, EColorChannel ColorChannel = EColorChannel::All);
 
 	/**
diff --git a/Engine/Source/Runtime/Core/Public/PixelFormat.h b/Engine/Source/Runtime/Core/Public/PixelFormat.h
--- a/Engine/Source/Runtime/Core/Public/PixelFormat.h
+++ b/Engine/Source/Runtime/Core/Public/PixelFormat.h
@@ -21,8 +21,8 @@ enum EPixelFormat : uint8
 	PF_DXT3                 =6,
 	PF_DXT5                 =7,
 	PF_UYVY                 =8,
-	PF_FloatRGB             =9,
-	PF_FloatRGBA            =10,
+	PF_FloatRGB             =9, // 16F
+	PF_FloatRGBA            =10, // 16F
 	PF_DepthStencil         =11,
 	PF_ShadowDepth          =12,
 	PF_R32_FLOAT            =13,
diff --git a/Engine/Source/Runtime/Engine/Classes/Engine/Texture.h b/Engine/Source/Runtime/Engine/Classes/Engine/Texture.h
--- a/Engine/Source/Runtime/Engine/Classes/Engine/Texture.h
+++ b/Engine/Source/Runtime/Engine/Classes/Engine/Texture.h
@@ -19,6 +19,7 @@
 #include "Engine/StreamableRenderAsset.h"
 #include "PerPlatformProperties.h"
 #include "Misc/FieldAccessor.h"
+#include "ImageCore.h"
 #if WITH_EDITORONLY_DATA
 #include "Misc/TVariant.h"
 #include "DerivedDataCacheKeyProxy.h"
@@ -141,7 +142,14 @@ struct FTextureSource
 	ENGINE_API FTextureSource();
 
 	ENGINE_API static int32 GetBytesPerPixel(ETextureSourceFormat Format);
-	FORCEINLINE static bool IsHDR(ETextureSourceFormat Format) { return (Format == TSF_BGRE8 || Format == TSF_RGBA16F); }
+	FORCEINLINE static bool IsHDR(ETextureSourceFormat Format) { return (Format == TSF_BGRE8 || Format == TSF_RGBA16F || Format == TSF_RGBA32F); }
+	
+	enum class ELockState : uint8
+	{
+		None,
+		ReadOnly,
+		ReadWrite
+	};
 
 #if WITH_EDITOR
 	// FwdDecl for member structs
@@ -273,6 +281,8 @@ struct FTextureSource
 	/** Force the GUID to change even if mip data has not been modified. */
 	ENGINE_API void ForceGenerateGuid();
 
+	/** Use FMipLock to encapsulate a locked mip */
+
 	/** Lock a mip for reading. */
 	ENGINE_API const uint8* LockMipReadOnly(int32 BlockIndex, int32 LayerIndex, int32 MipIndex);
 
@@ -284,6 +294,21 @@ struct FTextureSource
 
 	/** Retrieve a copy of the data for a particular mip. */
 	ENGINE_API bool GetMipData(TArray64<uint8>& OutMipData, int32 BlockIndex, int32 LayerIndex, int32 MipIndex, class IImageWrapperModule* ImageWrapperModule = nullptr);
+	
+	/** Legacy API that defaults to LayerIndex 0 */
+	FORCEINLINE bool GetMipData(TArray64<uint8>& OutMipData, int32 MipIndex, class IImageWrapperModule* ImageWrapperModule = nullptr)
+	{
+		return GetMipData(OutMipData, 0, 0, MipIndex, ImageWrapperModule);
+	}
+
+	/** Retrieve a copy of the MipData as an FImage */
+	ENGINE_API bool GetMipImage(FImage & OutImage, int32 BlockIndex, int32 LayerIndex, int32 MipIndex);
+	
+	/** Retrieve a copy of the MipData as an FImage */
+	FORCEINLINE bool GetMipImage(FImage & OutImage, int32 MipIndex)
+	{
+		return GetMipImage(OutImage, 0, 0, MipIndex);
+	}
 
 	/** Returns a FMipData structure that wraps around the entire mip chain for read only operations. This is more efficient than calling the above method once per mip. */
 	ENGINE_API FMipData GetMipData(class IImageWrapperModule* ImageWrapperModule);
@@ -349,12 +374,6 @@ struct FTextureSource
 	/** Sets the GUID to use, and whether that GUID is actually a hash of some data. */
 	ENGINE_API void SetId(const FGuid& InId, bool bInGuidIsHash);
 
-	/** Legacy API that defaults to LayerIndex 0 */
-	FORCEINLINE bool GetMipData(TArray64<uint8>& OutMipData, int32 MipIndex, class IImageWrapperModule* ImageWrapperModule = nullptr)
-	{
-		return GetMipData(OutMipData, 0, 0, MipIndex, ImageWrapperModule);
-	}
-
 	FORCEINLINE int64 CalcMipSize(int32 MipIndex) const { return CalcMipSize(0, 0, MipIndex); }
 	/** Lock a mip for reading. */
 	FORCEINLINE const uint8* LockMipReadOnly(int32 MipIndex) { return LockMipReadOnly(0, 0, MipIndex); }
@@ -450,6 +469,59 @@ struct FTextureSource
 		const FTextureSource& TextureSource;
 		FSharedBuffer MipData;
 	};
+	
+	/**
+	* FMipLock to encapsulate a locked mip
+	*  acquires the lock in construct and unlocks in destruct
+	*/
+	struct ENGINE_API FMipLock
+	{
+		// constructor locks the mip (can fail, pointer will be null)
+		FMipLock(ELockState InLockState,FTextureSource * InTextureSource,int32 InBlockIndex, int32 InLayerIndex, int32 InMipIndex);
+		FMipLock(ELockState InLockState,FTextureSource * InTextureSource,int32 InMipIndex);
+		
+		// move constructor
+		FMipLock(FMipLock&&);
+
+		// destructor unlocks
+		~FMipLock();
+		
+		ELockState LockState;
+		FTextureSource * TextureSource;
+		int32 BlockIndex;
+		int32 LayerIndex;
+		int32 MipIndex;
+		FImageView Image;
+
+		const void * GetRawData() const
+		{
+			check( LockState != ELockState::None );
+			return Image.RawData;
+		}
+		void * GetMutableData() const
+		{
+			check( LockState == ELockState::ReadWrite );
+			return Image.RawData;
+		}
+		int64 GetDataSize() const { return Image.GetImageSizeBytes(); }
+
+		bool IsValid() const
+		{
+			// two criteria should be the same :
+			bool bNotNull = Image.RawData != nullptr;
+			bool bNotNone = LockState != ELockState::None;
+			check( bNotNull == bNotNone );
+			return bNotNull;
+		}
+
+	private:
+		
+		// no copying or default construct
+		FMipLock() = delete;
+		FMipLock(const FMipLock&) = delete;
+		FMipLock& operator=(const FMipLock&) = delete;
+		FMipLock& operator=(FMipLock&& Other) = delete;
+	};
 #endif // WITH_EDITOR
 
 private:
@@ -473,12 +545,6 @@ private:
 	/** Number of mips that are locked. */
 	uint32 NumLockedMips;
 
-	enum class ELockState : uint8
-	{
-		None,
-		ReadOnly,
-		ReadWrite
-	};
 	/** The state of any lock being held on the mip data */
 	ELockState LockState;
 
@@ -490,11 +556,10 @@ private:
 	uint8* LockMipInternal(int32 BlockIndex, int32 LayerIndex, int32 MipIndex, ELockState RequestedLockState);
 	
 	/** Returns the source data fully decompressed */
-	FSharedBuffer Decompress(class IImageWrapperModule* ImageWrapperModule) const;
-	/** Attempt to decompress the source data from a compressed png format. All failures will be logged and result in the method returning false */
-	FSharedBuffer TryDecompressPngData(IImageWrapperModule* ImageWrapperModule) const;
-	/** Attempt to decompress the source data from Jpeg format. All failures will be logged and result in the method returning false */
-	FSharedBuffer TryDecompressJpegData(IImageWrapperModule* ImageWrapperModule) const;
+	// ImageWrapperModule is not used
+	FSharedBuffer Decompress(class IImageWrapperModule* ImageWrapperModule = nullptr) const;
+	/** Attempt to decompress the source data from a compressed format. All failures will be logged and result in the method returning false */
+	FSharedBuffer TryDecompressData() const;
 
 	/** Return true if the source art is not png compressed but could be. */
 	bool CanPNGCompress() const;
@@ -1439,6 +1504,11 @@ public:
 	ENGINE_API void GetDefaultFormatSettings(FTextureFormatSettings& OutSettings) const;
 	
 	ENGINE_API ETextureEncodeSpeed GetDesiredEncodeSpeed() const;
+
+	ENGINE_API EGammaSpace GetGammaSpace() const
+	{
+		return SRGB ? ( bUseLegacyGamma ? EGammaSpace::Pow22 : EGammaSpace::sRGB ) : EGammaSpace::Linear;
+	}
 #endif
 
 	/** @return the width of the surface represented by the texture. */
diff --git a/Engine/Source/Runtime/Engine/Classes/Engine/TextureDefines.h b/Engine/Source/Runtime/Engine/Classes/Engine/TextureDefines.h
--- a/Engine/Source/Runtime/Engine/Classes/Engine/TextureDefines.h
+++ b/Engine/Source/Runtime/Engine/Classes/Engine/TextureDefines.h
@@ -266,13 +266,14 @@ enum ECompositeTextureMode
 UENUM()
 enum ETextureSourceCompressionFormat
 {
-	TSCF_None	UMETA(DisplayName = "ZLib"),
+	TSCF_None	UMETA(DisplayName = "None"),
 	TSCF_PNG	UMETA(DisplayName = "PNG"),
 	TSCF_JPEG	UMETA(DisplayName = "JPEG"),
 
 	TSCF_MAX
 };
 
+// ETextureSourceFormat should map one-to-one to ImageCore ERawImageFormat::Type
 UENUM()
 enum ETextureSourceFormat
 {
@@ -283,14 +284,21 @@ enum ETextureSourceFormat
 	TSF_RGBA16,
 	TSF_RGBA16F,
 
-	//@todo: Deprecated!
-	TSF_RGBA8,
-	//@todo: Deprecated!
-	TSF_RGBE8,
+	// these are mapped to TSF_BGRA8/TSF_BGRE8 on load, so the runtime will never see them after loading :
+	// keep them here to preserve enum values
+	TSF_RGBA8_DEPRECATED,
+	TSF_RGBE8_DEPRECATED,
 
 	TSF_G16,
+	TSF_RGBA32F,
+	TSF_R16F,
 
-	TSF_MAX
+	TSF_MAX,
+
+	// error in UHT, enable when fixed
+	// provide aliases to the old names with deprecation warnings
+	//TSF_RGBA8 UE_DEPRECATED(5.1,"Legacy ETextureSourceFormat not supported, use BGRA8") = TSF_RGBA8_DEPRECATED,
+	//TSF_RGBE8 UE_DEPRECATED(5.1,"Legacy ETextureSourceFormat not supported, use BGRE8") = TSF_RGBE8_DEPRECATED
 };
 
 // This needs to be mirrored in EditorFactories.cpp.
diff --git a/Engine/Source/Runtime/Engine/Engine.Build.cs b/Engine/Source/Runtime/Engine/Engine.Build.cs
--- a/Engine/Source/Runtime/Engine/Engine.Build.cs
+++ b/Engine/Source/Runtime/Engine/Engine.Build.cs
@@ -13,7 +13,15 @@ public class Engine : ModuleRules
 
 		SharedPCHHeaderFile = "Public/EngineSharedPCH.h";
 
-		PublicIncludePathModuleNames.AddRange(new string[] { "Renderer", "PacketHandler", "AudioMixer", "AudioMixerCore", "AnimationCore" });
+		PublicIncludePathModuleNames.AddRange(new string[] {
+			"Renderer", 
+			"PacketHandler", 
+			"AudioMixer", 
+			"AudioMixerCore",
+			"AnimationCore",
+			"ImageCore"
+			}
+		);
 
 		PrivateIncludePaths.AddRange(
 			new string[] {
diff --git a/Engine/Source/Runtime/Engine/Private/GameViewportClient.cpp b/Engine/Source/Runtime/Engine/Private/GameViewportClient.cpp
--- a/Engine/Source/Runtime/Engine/Private/GameViewportClient.cpp
+++ b/Engine/Source/Runtime/Engine/Private/GameViewportClient.cpp
@@ -1979,17 +1979,15 @@ bool UGameViewportClient::ProcessScreenShots(FViewport* InViewport)
 					Size = FIntVector(NewWidth, NewHeight, 0);
 				}
 
-				if (!FPaths::GetExtension(ScreenShotName).IsEmpty())
+				if ( FPaths::GetExtension(ScreenShotName).IsEmpty() )
 				{
-					ScreenShotName = FPaths::GetBaseFilename(ScreenShotName, false);
 					ScreenShotName += TEXT(".png");
 				}
 
 				// Save the contents of the array to a png file.
-				TArray64<uint8> CompressedBitmap;
-				FImageUtils::PNGCompressImageArray(Size.X, Size.Y, Bitmap, CompressedBitmap);
-				bIsScreenshotSaved = FFileHelper::SaveArrayToFile(CompressedBitmap, *ScreenShotName);
-
+				
+				FImageView Image((const FColor *)Bitmap.GetData(),Size.X,Size.Y);
+				bIsScreenshotSaved = FImageUtils::SaveImageByExtension(*ScreenShotName,Image);
 			}
 		}
 
diff --git a/Engine/Source/Runtime/Engine/Private/ImageUtils.cpp b/Engine/Source/Runtime/Engine/Private/ImageUtils.cpp
--- a/Engine/Source/Runtime/Engine/Private/ImageUtils.cpp
+++ b/Engine/Source/Runtime/Engine/Private/ImageUtils.cpp
@@ -18,11 +18,140 @@ ImageUtils.cpp: Image utility functions.
 #include "Misc/ObjectThumbnail.h"
 #include "Misc/Paths.h"
 #include "Modules/ModuleManager.h"
+#include "ImageCoreUtils.h"
 
 DEFINE_LOG_CATEGORY_STATIC(LogImageUtils, Log, All);
 
 #define LOCTEXT_NAMESPACE "ImageUtils"
 
+
+static IImageWrapperModule * GetOrLoadImageWrapperModule()
+{
+	static FName ImageWrapperName("ImageWrapper");
+
+	if ( IsInGameThread() )
+	{
+		IImageWrapperModule& ImageWrapperModule = FModuleManager::LoadModuleChecked<IImageWrapperModule>(ImageWrapperName);
+		return &ImageWrapperModule;
+	}
+	else
+	{
+		IImageWrapperModule * ImageWrapperModule = FModuleManager::GetModulePtr<IImageWrapperModule>(ImageWrapperName);
+
+		if ( ImageWrapperModule == nullptr )
+		{
+			UE_LOG(LogImageUtils,Warning,TEXT("Not on GameThread, cannot load ImageWrapper.  Do on main thread in startup."));
+
+			// LoadModule needs to be done on the Game thread as part of your initialization
+			//   before any thread tries to run this code
+			// Engine does this.  If you are making a stand-alone app and hit this error,
+			//   add this to your initialization on the Game thread :
+			// FModuleManager::LoadModuleChecked<IImageWrapperModule>("ImageWrapper");
+		}
+
+		return ImageWrapperModule;
+	}
+}
+
+
+bool FImageUtils::SaveImageAutoFormat(const TCHAR * FileName, const FImageView & InImage, int32 Quality)
+{
+	IImageWrapperModule * ImageWrapperModule = GetOrLoadImageWrapperModule();
+	if ( ImageWrapperModule == nullptr )
+	{
+		return false;
+	}
+
+	EImageFormat ToFormat = ImageWrapperModule->GetDefaultOutputFormat(InImage.Format);
+	if ( ToFormat == EImageFormat::Invalid )
+	{
+		return false;
+	}
+	
+	FString Scratch;
+	if ( ImageWrapperModule->GetImageFormatFromExtension(FileName) != ToFormat )
+	{
+		const TCHAR * Extension = ImageWrapperModule->GetExtension(ToFormat);
+		
+		Scratch = FPaths::GetBaseFilename(FileName, false);
+		Scratch += TEXT('.');
+		Scratch += Extension;
+		FileName = *Scratch;
+	}
+
+	TArray64<uint8> Buffer;
+	if ( ! ImageWrapperModule->CompressImage(Buffer,ToFormat,InImage,Quality) )
+	{
+		return false;
+	}
+
+	return FFileHelper::SaveArrayToFile(Buffer,FileName);
+}
+
+bool FImageUtils::LoadImage(const TCHAR * Filename, FImage & OutImage)
+{
+	TArray64<uint8> Buffer;
+	if ( ! FFileHelper::LoadFileToArray(Buffer, Filename) )
+	{
+		return false;
+	}
+
+	return DecompressImage(Buffer.GetData(),Buffer.Num(),OutImage);
+}
+
+bool FImageUtils::SaveImageByExtension(const TCHAR * Filename, const FImageView & InImage, int32 Quality)
+{
+	TArray64<uint8> Buffer;
+	if ( !  CompressImage(Buffer, Filename, InImage,Quality) )
+	{
+		return false;
+	}
+
+	return FFileHelper::SaveArrayToFile(Buffer,Filename);
+}
+
+// OutImage is allocated and filled, any existing contents are discarded
+bool FImageUtils::DecompressImage(const void* InCompressedData, int64 InCompressedSize, FImage & OutImage)
+{
+	IImageWrapperModule * ImageWrapperModule = GetOrLoadImageWrapperModule();
+	if ( ImageWrapperModule == nullptr )
+	{
+		return false;
+	}
+
+	return ImageWrapperModule->DecompressImage(InCompressedData,InCompressedSize,OutImage);
+}
+
+// OutData is filled with the file format encoded data
+//	 lossy conversion may be done if necessary
+bool FImageUtils::CompressImage(TArray64<uint8> & OutData, const TCHAR * ToFormatExtension, const FImageView & InImage, int32 Quality)
+{
+	IImageWrapperModule * ImageWrapperModule = GetOrLoadImageWrapperModule();
+	if ( ImageWrapperModule == nullptr )
+	{
+		return false;
+	}
+
+	
+	EImageFormat ToFormat;
+	
+	if ( ToFormatExtension == nullptr )
+	{
+		ToFormat = ImageWrapperModule->GetDefaultOutputFormat(InImage.Format);
+	}
+	else
+	{
+		ToFormat = ImageWrapperModule->GetImageFormatFromExtension(ToFormatExtension);
+	}
+
+	if ( ToFormat == EImageFormat::Invalid )
+	{
+		return false;
+	}
+
+	return ImageWrapperModule->CompressImage(OutData,ToFormat,InImage,Quality);
+}
+
 /**
  * Returns data containing the pixmap of the passed in rendertarget.
  * @param TexRT - The 2D rendertarget from which to read pixmap data.
@@ -40,15 +169,19 @@ bool FImageUtils::GetRawData(UTextureRenderTarget2D* TexRT, TArray64<uint8>& Raw
 	switch (Format)
 	{
 	case PF_FloatRGBA:
-	{
+		{
 		TArray<FFloat16Color> FloatColors;
 		bReadSuccess = RenderTarget->ReadFloat16Pixels(FloatColors);
 		FMemory::Memcpy(RawData.GetData(), FloatColors.GetData(), ImageBytes);
-	}
-	break;
+		}
+		break;
 	case PF_B8G8R8A8:
 		bReadSuccess = RenderTarget->ReadPixelsPtr((FColor*)RawData.GetData());
 		break;
+	default:
+		// bReadSuccess == false
+		UE_LOG(LogImageUtils,Warning,TEXT("RenderTarget GetRawData PixelFormat no supported : %s") , *(StaticEnum<EPixelFormat>()->GetDisplayNameTextByValue(Format).ToString()) );
+		break;
 	}
 	if (bReadSuccess == false)
 	{
@@ -57,6 +190,49 @@ bool FImageUtils::GetRawData(UTextureRenderTarget2D* TexRT, TArray64<uint8>& Raw
 	return bReadSuccess;
 }
 
+bool FImageUtils::GetRenderTargetImage(UTextureRenderTarget2D* TexRT, FImage & Image)
+{
+	Image = FImage();
+	
+	// see ETextureRenderTargetFormat
+	// for allowed formats
+
+	TArray64<uint8> RawData;
+	if ( ! GetRawData(TexRT,RawData) )
+	{
+		return false;
+	}
+	
+	// @@!! in theory the RT knows its gamma
+	//   but from what I've seen it's usually wrong?
+	// TexRT->GetDisplayGamma() or TexRT->SRGB or TexRT->IsSRGB() , OMG
+
+	Image.RawData = MoveTemp(RawData);
+	Image.SizeX = TexRT->SizeX;
+	Image.SizeY = TexRT->SizeY;
+	Image.NumSlices = 1;
+	
+	EPixelFormat PixelFormat = TexRT->GetFormat();
+	switch(PixelFormat)
+	{
+	case PF_FloatRGBA:
+		Image.Format = ERawImageFormat::RGBA16F;
+		Image.GammaSpace = EGammaSpace::Linear;
+		break;	
+	case PF_B8G8R8A8:
+		Image.Format = ERawImageFormat::BGRA8;
+		Image.GammaSpace = EGammaSpace::sRGB;
+		break;
+	default:
+		// should not get here because GetRawData would have returned false
+		check(0);
+		return false;
+	}
+
+	return true;
+}
+
+
 /**
  * Resizes the given image using a simple average filter and stores it in the destination array.
  *
@@ -424,22 +600,13 @@ void FImageUtils::ThumbnailCompressImageArray( int32 ImageWidth, int32 ImageHeig
 
 void FImageUtils::PNGCompressImageArray(int32 ImageWidth, int32 ImageHeight, const TArrayView64<const FColor>& SrcData, TArray64<uint8>& DstData)
 {
-	const int64 PixelsNum = ImageWidth * ImageHeight;
-	check(SrcData.Num() == PixelsNum);
-
-	const uint8* SrcFirstByte = static_cast<const uint8*>(static_cast<const void*>(SrcData.GetData()));
-	const int64 MemorySize = PixelsNum * sizeof(FColor);
-
 	DstData.Reset();
 
 	if (SrcData.Num() > 0 && ImageWidth > 0 && ImageHeight > 0)
 	{
-		IImageWrapperModule& ImageWrapperModule = FModuleManager::LoadModuleChecked<IImageWrapperModule>(FName("ImageWrapper"));
-		TSharedPtr<IImageWrapper> ImageWrapper = ImageWrapperModule.CreateImageWrapper(EImageFormat::PNG);
-		if (ImageWrapper.IsValid() && ImageWrapper->SetRaw(SrcFirstByte, MemorySize, ImageWidth, ImageHeight, ERGBFormat::BGRA, 8))
-		{
-			DstData = ImageWrapper->GetCompressed();
-		}
+		FImageView Image((void *)SrcData.GetData(),ImageWidth,ImageHeight, ERawImageFormat::BGRA8);
+
+		CompressImage(DstData,TEXT("png"),Image);
 	}
 }
 
@@ -522,8 +689,10 @@ UTextureCube* FImageUtils::CreateCheckerboardCubeTexture(FColor ColorOne, FColor
 /*------------------------------------------------------------------------------
 HDR file format helper.
 ------------------------------------------------------------------------------*/
-// @todo Oodle : can we get rid of this HDR writer?
-//  try to remove uses of HDR as an export format
+// DEPRECATED
+// do not use HDR for exports, it is very lossy, use EXR instead for float output
+// if you need HDR, do not use this export code
+// instead use HdrImageWrapper via FImageUtils::CompressImage
 class FHDRExportHelper
 {
 public:
@@ -564,9 +733,14 @@ public:
 
 #if WITH_EDITORONLY_DATA
 		Size = FIntPoint(Texture->Source.GetSizeX(), Texture->Source.GetSizeY());
+
 		bReadSuccess = Texture->Source.GetMipData(RawData, 0);
 		const ETextureSourceFormat NewFormat = Texture->Source.GetFormat();
 
+		// DEPRECATED
+		// not a general purpose HDR exporter
+		// can't write F32 or BGRE8
+		// do not use this except for longlat cubemaps
 		if (NewFormat == TSF_BGRA8)
 		{
 			Format = PF_B8G8R8A8;
@@ -639,6 +813,7 @@ public:
 	* @param Ar - Archive object to write HDR data to.
 	* @return true on successful export.
 	*/
+	// this does GenerateLongLatUnwrap
 	bool ExportHDR(UTextureCube* TexCube, FArchive& Ar)
 	{
 		check(TexCube != nullptr);
@@ -664,6 +839,7 @@ public:
 	* @param Ar - Archive object to write HDR data to.
 	* @return true on successful export.
 	*/
+	// this does GenerateLongLatUnwrap
 	bool ExportHDR(UTextureRenderTargetCube* TexCube, FArchive& Ar)
 	{
 		check(TexCube != nullptr);
@@ -798,8 +974,11 @@ private:
 		}
 		else
 		{
+			check( Format == PF_B8G8R8A8 );
 			WriteHDRBits(Ar, (FColor*)RawData.GetData());
 		}
+		// can't actually write BGRE8 pixels
+		// use HdrImageWrapper instead
 	}
 
 	FIntPoint Size;
@@ -808,132 +987,116 @@ private:
 
 bool FImageUtils::ExportRenderTarget2DAsHDR(UTextureRenderTarget2D* TexRT, FArchive& Ar)
 {
-	FHDRExportHelper Exporter;
-	return Exporter.ExportHDR(TexRT, Ar);
-}
+	UE_LOG(LogImageUtils, Warning, TEXT("HDR file format is very lossy, use EXR or PNG instead.") );
 
-bool FImageUtils::ExportRenderTarget2DAsPNG(UTextureRenderTarget2D* TexRT, FArchive& Ar)
-{
-	bool bSuccess = false;
-	if(TexRT->GetFormat() == PF_B8G8R8A8)
+	FImage Image;
+	if ( ! GetRenderTargetImage(TexRT,Image) )
 	{
-		check(TexRT != nullptr);
-		FRenderTarget* RenderTarget = TexRT->GameThread_GetRenderTargetResource();
-		FIntPoint Size = RenderTarget->GetSizeXY();
-
-		TArray64<uint8> RawData;
-		bSuccess = GetRawData(TexRT, RawData);
+		return false;
+	}
 
-		IImageWrapperModule& ImageWrapperModule = FModuleManager::Get().LoadModuleChecked<IImageWrapperModule>(TEXT("ImageWrapper"));
+	TArray64<uint8> CompressedData;
+	if ( ! CompressImage(CompressedData,TEXT("HDR"),Image) )
+	{
+		return false;
+	}
 
-		TSharedPtr<IImageWrapper> PNGImageWrapper = ImageWrapperModule.CreateImageWrapper(EImageFormat::PNG);
+	Ar.Serialize((void*)CompressedData.GetData(), CompressedData.GetAllocatedSize());
 
-		PNGImageWrapper->SetRaw(RawData.GetData(), RawData.GetAllocatedSize(), Size.X, Size.Y, ERGBFormat::BGRA, 8);
+	return true;
+}
 
-		const TArray64<uint8> PNGData = PNGImageWrapper->GetCompressed(100);
+bool FImageUtils::ExportRenderTarget2DAsPNG(UTextureRenderTarget2D* TexRT, FArchive& Ar)
+{
+	FImage Image;
+	if ( ! GetRenderTargetImage(TexRT,Image) )
+	{
+		return false;
+	}
 
-		Ar.Serialize((void*)PNGData.GetData(), PNGData.GetAllocatedSize());
+	TArray64<uint8> CompressedData;
+	if ( ! CompressImage(CompressedData,TEXT("PNG"),Image) )
+	{
+		return false;
 	}
 
-	return bSuccess;
+	Ar.Serialize((void*)CompressedData.GetData(), CompressedData.GetAllocatedSize());
+
+	return true;
 }
 
-ENGINE_API bool FImageUtils::ExportRenderTarget2DAsEXR(UTextureRenderTarget2D* TexRT, FArchive& Ar)
+bool FImageUtils::ExportRenderTarget2DAsEXR(UTextureRenderTarget2D* TexRT, FArchive& Ar)
 {
-	bool bSuccess = false;
-	if(TexRT->GetFormat() == PF_B8G8R8A8 || TexRT->GetFormat() == PF_FloatRGBA)
+	FImage Image;
+	if ( ! GetRenderTargetImage(TexRT,Image) )
 	{
-		check(TexRT != nullptr);
-		FRenderTarget* RenderTarget = TexRT->GameThread_GetRenderTargetResource();
-		FIntPoint Size = RenderTarget->GetSizeXY();
-
-		TArray64<uint8> RawData;
-		bSuccess = GetRawData(TexRT, RawData);
+		return false;
+	}
 
-		int32 BitsPerPixel = TexRT->GetFormat() == PF_B8G8R8A8 ? 8 : (sizeof(FFloat16Color) / 4) * 8;
-		ERGBFormat RGBFormat = TexRT->GetFormat() == PF_B8G8R8A8 ? ERGBFormat::BGRA : ERGBFormat::RGBAF;
+	TArray64<uint8> CompressedData;
+	if ( ! CompressImage(CompressedData,TEXT("EXR"),Image) )
+	{
+		return false;
+	}
 
-		IImageWrapperModule& ImageWrapperModule = FModuleManager::Get().LoadModuleChecked<IImageWrapperModule>(TEXT("ImageWrapper"));
+	Ar.Serialize((void*)CompressedData.GetData(), CompressedData.GetAllocatedSize());
 
-		TSharedPtr<IImageWrapper> EXRImageWrapper = ImageWrapperModule.CreateImageWrapper(EImageFormat::EXR);
+	return true;
+}
 
-		EXRImageWrapper->SetRaw(RawData.GetData(), RawData.GetAllocatedSize(), Size.X, Size.Y, RGBFormat, BitsPerPixel);
+// if Texture source is available, get it as an FImage
+bool FImageUtils::GetTexture2DSourceImage(UTexture2D* Texture, FImage & OutImage)
+{
+#if WITH_EDITORONLY_DATA	
+	return Texture->Source.GetMipImage(OutImage,0);
+#else
 
-		const TArray64<uint8> Data = EXRImageWrapper->GetCompressed();
+	UE_LOG(LogImageUtils,Warning,TEXT("GetTexture2DSourceImage from PlatformData not implemented yet"));
 
-		Ar.Serialize((void*)Data.GetData(), Data.GetAllocatedSize());
+	// @@!! todo?
+	// can try to get texture from platformdata
+	// only a few formats would be possible to grab
+	// eg. see ExportHDR below
 
-		bSuccess = true;
-	}
+	/*
+	EPixelFormat PixelFormat = Texture->GetFormat();
+	Texture->GetMipData(0, (void**)RawData2.GetData());
+	if (Texture->GetPlatformData()->Mips.Num() == 0)
+	*/
+	
+	//uint8* MipData = static_cast<uint8*>(NewTexture->GetPlatformData()->Mips[0].BulkData.Lock(LOCK_READ_WRITE));
+	//int64 MipDataSize = NewTexture->GetPlatformData()->Mips[0].BulkData.GetBulkDataSize();
+	
+	//if (NewFormat == PF_B8G8R8A8)
+	//else if (NewFormat == PF_FloatRGBA)
 
-	return bSuccess;
+	return false;
+#endif
 }
 
 bool FImageUtils::ExportTexture2DAsHDR(UTexture2D* TexRT, FArchive& Ar)
 {
+	UE_LOG(LogImageUtils, Warning, TEXT("HDR file format is very lossy, use EXR or PNG instead.") );
+
+	// could use GetTexture2DSourceImage
+	// then CompressImage
+	// for arbitrary format exports
+	// but just leave it alone for now
+
+	// this is only used by UKismetRenderingLibrary::ExportTexture2D
+	//	so we should take filename as input to auto-detect format
+
 	FHDRExportHelper Exporter;
 	return Exporter.ExportHDR(TexRT, Ar);
 }
 
 UTexture2D* FImageUtils::ImportFileAsTexture2D(const FString& Filename)
 {
-	IImageWrapperModule& ImageWrapperModule = FModuleManager::Get().LoadModuleChecked<IImageWrapperModule>(TEXT("ImageWrapper"));
-
 	UTexture2D* NewTexture = nullptr;
 	TArray64<uint8> Buffer;
 	if (FFileHelper::LoadFileToArray(Buffer, *Filename))
 	{
-		EPixelFormat PixelFormat = PF_Unknown;
-
-		uint8* RawData = nullptr;
-		int32 BitDepth = 0;
-		int32 Width = 0;
-		int32 Height = 0;
-
-		if (FPaths::GetExtension(Filename) == TEXT("HDR"))
-		{
-			TSharedPtr<IImageWrapper> HdrImageWrapper =  ImageWrapperModule.CreateImageWrapper(EImageFormat::HDR);
-	
-			if(HdrImageWrapper->SetCompressed(Buffer.GetData(), Buffer.Num()))
-			{
-				PixelFormat = PF_FloatRGBA;
-				Width = HdrImageWrapper->GetWidth();
-				Height = HdrImageWrapper->GetHeight();
-
-				TArray64<uint8> BGREImage;
-				if (HdrImageWrapper->GetRaw(ERGBFormat::BGRE, 8, BGREImage))
-				{
-					NewTexture = UTexture2D::CreateTransient(Width, Height, PixelFormat);
-					if (NewTexture)
-					{
-						uint8* MipData = static_cast<uint8*>(NewTexture->GetPlatformData()->Mips[0].BulkData.Lock(LOCK_READ_WRITE));
-						int64 MipDataSize = NewTexture->GetPlatformData()->Mips[0].BulkData.GetBulkDataSize();
-
-						TArrayView64<FColor> SourceColors(reinterpret_cast<FColor*>(BGREImage.GetData()), BGREImage.Num() / sizeof(FColor));
-
-						// Bulk data was already allocated for the correct size when we called CreateTransient above
-						TArrayView64<FFloat16Color> DestinationColors(reinterpret_cast<FFloat16Color*>(MipData), MipDataSize / sizeof(FFloat16Color));
-
-						check( SourceColors.Num() == DestinationColors.Num() );
-
-						int64 DestinationIndex = 0;
-						for (const FColor& Color: SourceColors)
-						{
-							FLinearColor LinearColor = Color.FromRGBE();
-							DestinationColors[DestinationIndex++] = FFloat16Color(LinearColor);
-						}
-
-						NewTexture->GetPlatformData()->Mips[0].BulkData.Unlock();
-
-						NewTexture->UpdateResource();
-					}
-				}
-			}
-		}
-		else
-		{
-			NewTexture = FImageUtils::ImportBufferAsTexture2D(Buffer);
-		}
+		NewTexture = FImageUtils::ImportBufferAsTexture2D(Buffer);
 
 		if(!NewTexture)
 		{
@@ -950,72 +1113,44 @@ UTexture2D* FImageUtils::ImportFileAsTexture2D(const FString& Filename)
 
 UTexture2D* FImageUtils::ImportBufferAsTexture2D(TArrayView64<const uint8> Buffer)
 {
-	IImageWrapperModule& ImageWrapperModule = FModuleManager::Get().LoadModuleChecked<IImageWrapperModule>(TEXT("ImageWrapper"));
+	//@@!! test me
 
-	EImageFormat Format = ImageWrapperModule.DetectImageFormat(Buffer.GetData(), Buffer.Num());
+	FImage Image;
+	if ( ! DecompressImage(Buffer.GetData(),Buffer.Num(), Image) )
+	{
+		UE_LOG(LogImageUtils, Warning, TEXT("Error creating texture. Couldn't determine the file format"));
+		return nullptr;
+	}
 
-	UTexture2D* NewTexture = nullptr;
-	EPixelFormat PixelFormat = PF_Unknown;
+	return CreateTexture2DFromImage(Image);
+}
 
-	if (Format != EImageFormat::Invalid)
+UTexture2D* FImageUtils::CreateTexture2DFromImage(const FImageView & Image)
+{
+	ERawImageFormat::Type PixelFormatRawFormat;
+	EPixelFormat PixelFormat = FImageCoreUtils::GetPixelFormatForRawImageFormat(Image.Format,&PixelFormatRawFormat);
+						
+	UTexture2D* NewTexture = UTexture2D::CreateTransient(Image.SizeX, Image.SizeY, PixelFormat);
+	if (NewTexture == nullptr)
 	{
-		TSharedPtr<IImageWrapper> ImageWrapper = ImageWrapperModule.CreateImageWrapper(Format);
-		
-		int32 BitDepth = 0;
-		int32 Width = 0;
-		int32 Height = 0;
+		UE_LOG(LogImageUtils, Warning, TEXT("Error in CreateTransient"));
+		return nullptr;
+	}
 
-		if (ImageWrapper->SetCompressed((void*)Buffer.GetData(), Buffer.Num()))
-		{
-			PixelFormat = PF_Unknown;
-			
-			//@todo Oodle : guessing format from bit depth looks wrong
-			ERGBFormat RGBFormat = ERGBFormat::Invalid;
-			BitDepth = ImageWrapper->GetBitDepth();
-						
-			if (BitDepth == 16)
-			{
-				PixelFormat = PF_FloatRGBA;
-				RGBFormat = ERGBFormat::RGBAF;
-			}
-			else if (BitDepth == 8)
-			{
-				PixelFormat = PF_B8G8R8A8;
-				RGBFormat = ERGBFormat::BGRA;
-			}
-			else
-			{
-				UE_LOG(LogImageUtils, Warning, TEXT("Error creating texture. Bit depth is unsupported. (%d)"), BitDepth);
-				return nullptr;
-			}
-			
-			Width = ImageWrapper->GetWidth();
-			Height = ImageWrapper->GetHeight();
+	NewTexture->bNotOfflineProcessed = true;
 
-			TArray64<uint8> UncompressedData;
-			ImageWrapper->GetRaw(RGBFormat, BitDepth, UncompressedData);
-			
-			NewTexture = UTexture2D::CreateTransient(Width, Height, PixelFormat);
-			if (NewTexture)
-			{
-				NewTexture->bNotOfflineProcessed = true;
-				uint8* MipData = static_cast<uint8*>(NewTexture->GetPlatformData()->Mips[0].BulkData.Lock(LOCK_READ_WRITE));
-				int64 MipDataSize = NewTexture->GetPlatformData()->Mips[0].BulkData.GetBulkDataSize();
+	uint8* MipData = static_cast<uint8*>(NewTexture->GetPlatformData()->Mips[0].BulkData.Lock(LOCK_READ_WRITE));
+	check( MipData != nullptr );
+	int64 MipDataSize = NewTexture->GetPlatformData()->Mips[0].BulkData.GetBulkDataSize();
+
+	FImageView MipImage(MipData,Image.SizeX,Image.SizeY,1,PixelFormatRawFormat,Image.GammaSpace);
 
-				// Bulk data was already allocated for the correct size when we called CreateTransient above
-				check( UncompressedData.Num() >= MipDataSize );
-				FMemory::Memcpy(MipData, UncompressedData.GetData(), MipDataSize);
+	// copy into texture and convert if necessary :
+	FImageCore::CopyImage(Image,MipImage);
 				
-				NewTexture->GetPlatformData()->Mips[0].BulkData.Unlock();
+	NewTexture->GetPlatformData()->Mips[0].BulkData.Unlock();
 
-				NewTexture->UpdateResource();
-			}
-		}
-	}
-	else
-	{
-		UE_LOG(LogImageUtils, Warning, TEXT("Error creating texture. Couldn't determine the file format"));
-	}
+	NewTexture->UpdateResource();
 
 	return NewTexture;
 }
@@ -1027,12 +1162,14 @@ UTexture2D* FImageUtils::ImportBufferAsTexture2D(const TArray<uint8>& Buffer)
 
 bool FImageUtils::ExportRenderTargetCubeAsHDR(UTextureRenderTargetCube* TexRT, FArchive& Ar)
 {
+	// this does GenerateLongLatUnwrap
 	FHDRExportHelper Exporter;
 	return Exporter.ExportHDR(TexRT, Ar);
 }
 
 bool FImageUtils::ExportTextureCubeAsHDR(UTextureCube* TexRT, FArchive& Ar)
 {
+	// this does GenerateLongLatUnwrap
 	FHDRExportHelper Exporter;
 	return Exporter.ExportHDR(TexRT, Ar);
 }
diff --git a/Engine/Source/Runtime/Engine/Private/PhysicsEngine/PhysicalMaterialMask.cpp b/Engine/Source/Runtime/Engine/Private/PhysicsEngine/PhysicalMaterialMask.cpp
--- a/Engine/Source/Runtime/Engine/Private/PhysicsEngine/PhysicalMaterialMask.cpp
+++ b/Engine/Source/Runtime/Engine/Private/PhysicsEngine/PhysicalMaterialMask.cpp
@@ -375,19 +375,16 @@ void UPhysicalMaterialMask::GenerateMaskData(TArray<uint32>& OutMaskData, int32&
 					break;
 				}
 
-				case TSF_RGBA8:
-				{
-					MaskDataGenerator<uint8, 0, 1, 2, 3> MaskDataGen(OutSizeX, OutSizeY, TextureData);
-					MaskDataGen.GenerateMask(OutMaskData);
-					break;
-				}
-
 				case TSF_RGBA16:
 				{
 					MaskDataGenerator<uint16, 0, 1, 2, 3> MaskDataGen(OutSizeX, OutSizeY, TextureData);
 					MaskDataGen.GenerateMask(OutMaskData);
 					break;
 				}
+
+				default:
+					check(0);
+					break;
 				}
 			}
 		}
diff --git a/Engine/Source/Runtime/Engine/Private/Texture.cpp b/Engine/Source/Runtime/Engine/Private/Texture.cpp
--- a/Engine/Source/Runtime/Engine/Private/Texture.cpp
+++ b/Engine/Source/Runtime/Engine/Private/Texture.cpp
@@ -36,6 +36,8 @@
 #include "Engine/TextureCube.h"
 #include "Engine/RendererSettings.h"
 #include "ColorSpace.h"
+#include "ImageCoreUtils.h"
+#include "ImageUtils.h"
 
 #if WITH_EDITOR
 #include "TextureCompiler.h"
@@ -692,6 +694,34 @@ void UTexture::PostLoad()
 		Info.Insert(FAssetImportInfo::FSourceFile(SourceFilePath_DEPRECATED));
 		AssetImportData->SourceData = MoveTemp(Info);
 	}
+
+	if ( Source.GetFormat() == TSF_RGBA8_DEPRECATED
+		|| Source.GetFormat() == TSF_RGBE8_DEPRECATED )
+	{
+		// ensure that later code doesn't ever see the _DEPRECATED formats
+
+		// needs RB swap
+		// force BulkData to become resident
+		// do the swap on the bits
+		// change format to swapped version
+		// these formats are incredibly rare and old
+		// just warn and change the enum but don't swap the bits
+		// will appear RB swapped until reimported
+		UE_LOG(LogTexture, Warning, TEXT("TextureSource is a deprecated RB swapped format, needs reimport!: %s"), *GetPathName());
+		
+		for(int i=0;i<Source.LayerFormat.Num();i++)
+		{
+			if ( Source.LayerFormat[i] == TSF_RGBA8_DEPRECATED )
+			{
+				Source.LayerFormat[i] = TSF_BGRA8;
+			}
+			else if ( Source.LayerFormat[i] == TSF_RGBE8_DEPRECATED )
+			{
+				Source.LayerFormat[i] = TSF_BGRE8;
+			}
+		}
+	}
+
 #endif
 
 	if( !IsTemplate() )
@@ -1226,18 +1256,8 @@ FTextureSourceBlock::FTextureSourceBlock()
 
 int32 FTextureSource::GetBytesPerPixel(ETextureSourceFormat Format)
 {
-	int32 BytesPerPixel = 0;
-	switch (Format)
-	{
-	case TSF_G8:		BytesPerPixel = 1; break;
-	case TSF_G16:		BytesPerPixel = 2; break;
-	case TSF_BGRA8:		BytesPerPixel = 4; break;
-	case TSF_BGRE8:		BytesPerPixel = 4; break;
-	case TSF_RGBA16:	BytesPerPixel = 8; break;
-	case TSF_RGBA16F:	BytesPerPixel = 8; break;
-	default:			BytesPerPixel = 0; break;
-	}
-	return BytesPerPixel;
+	ERawImageFormat::Type RawFormat = FImageCoreUtils::ConvertToRawImageFormat(Format);
+	return ERawImageFormat::GetBytesPerPixel(RawFormat);
 }
 
 #if WITH_EDITOR
@@ -1475,6 +1495,8 @@ void FTextureSource::Compress()
 
 		IImageWrapperModule& ImageWrapperModule = FModuleManager::LoadModuleChecked<IImageWrapperModule>( FName("ImageWrapper") );
 		TSharedPtr<IImageWrapper> ImageWrapper = ImageWrapperModule.CreateImageWrapper( EImageFormat::PNG );
+		
+		// Legacy bug, must be matched in Compress & Decompress
 		// TODO: TSF_BGRA8 is stored as RGBA, so the R and B channels are swapped in the internal png. Should we fix this?
 		ERGBFormat RawFormat = (Format == TSF_G8 || Format == TSF_G16) ? ERGBFormat::Gray : ERGBFormat::RGBA;
 		int RawBitsPerChannel = (Format == TSF_G16 || Format == TSF_RGBA16) ? 16 : 8;
@@ -1513,21 +1535,92 @@ void FTextureSource::Compress()
 	}
 }
 
-FSharedBuffer FTextureSource::Decompress(IImageWrapperModule* ImageWrapperModule) const
+FSharedBuffer FTextureSource::Decompress(IImageWrapperModule* ) const
 {
 	TRACE_CPUPROFILER_EVENT_SCOPE(FTextureSource::Decompress);
+	
+	// ImageWrapperModule argument ignored, not drilled through DecompressImage
 
-	if (CompressionFormat == TSCF_JPEG)
+	if (CompressionFormat != TSCF_None || bPNGCompressed)
 	{
-		return TryDecompressJpegData(ImageWrapperModule);
+		return TryDecompressData();
 	}
-	else if (bPNGCompressed) // Change to CompressionFormat == TSCF_PNG once bPNGCompressed is deprecated
+	else
 	{
-		return TryDecompressPngData(ImageWrapperModule);
+		return BulkData.GetPayload().Get();
+	}
+}
+
+// constructor locks the mip (can fail, pointer will be null)
+FTextureSource::FMipLock::FMipLock(ELockState InLockState,FTextureSource * InTextureSource,int32 InBlockIndex, int32 InLayerIndex, int32 InMipIndex)
+:
+	LockState(InLockState),
+	TextureSource(InTextureSource),
+	BlockIndex(InBlockIndex),
+	LayerIndex(InLayerIndex),
+	MipIndex(InMipIndex)
+{
+	void * Locked = TextureSource->LockMipInternal(BlockIndex, LayerIndex, MipIndex, LockState);
+	if ( Locked )
+	{	
+		FTextureSourceBlock Block;
+		TextureSource->GetBlock(BlockIndex, Block);
+		check(MipIndex < Block.NumMips);
+
+		Image.RawData = Locked;
+		Image.SizeX = FMath::Max(Block.SizeX >> MipIndex, 1);
+		Image.SizeY = FMath::Max(Block.SizeY >> MipIndex, 1);
+		Image.NumSlices = Block.NumSlices;
+		Image.Format = FImageCoreUtils::ConvertToRawImageFormat(TextureSource->Format);
+
+		if ( TextureSource->Owner != nullptr )
+		{
+			Image.GammaSpace = TextureSource->Owner->GetGammaSpace();
+		}
+		else
+		{
+			// TextureSource does not have gamma information, just guess it ?
+			Image.GammaSpace = ERawImageFormat::GetDefaultGammaSpace(Image.Format);
+		}
+		
+		const int64 MipSizeBytes = TextureSource->CalcMipSize(BlockIndex, LayerIndex, MipIndex);
+
+		check( Image.GetImageSizeBytes() == MipSizeBytes );		
+		check( IsValid() );
 	}
 	else
 	{
-		return BulkData.GetPayload().Get();
+		LockState = ELockState::None;
+	}
+}
+
+FTextureSource::FMipLock::FMipLock(ELockState InLockState,FTextureSource * InTextureSource,int32 InMipIndex) :
+	FMipLock(InLockState,InTextureSource,0,0,InMipIndex)
+{
+}
+
+// move constructor :
+FTextureSource::FMipLock::FMipLock(FMipLock&& RHS) :
+	LockState(RHS.LockState),
+	TextureSource(RHS.TextureSource),
+	BlockIndex(RHS.BlockIndex),
+	LayerIndex(RHS.LayerIndex),
+	MipIndex(RHS.MipIndex),
+	Image(RHS.Image)
+{
+	// blank out RHS so he doesn't try to unlock :
+	RHS.LockState = ELockState::None;
+	RHS.Image = FImageView();
+	check( ! RHS.IsValid() );
+}
+
+FTextureSource::FMipLock::~FMipLock()
+{
+	if ( IsValid() )
+	{
+		TextureSource->UnlockMip(BlockIndex,LayerIndex,MipIndex);
+		LockState = ELockState::None;
+		Image.RawData = nullptr;
 	}
 }
 
@@ -1619,7 +1712,42 @@ void FTextureSource::UnlockMip(int32 BlockIndex, int32 LayerIndex, int32 MipInde
 	}
 }
 
-bool FTextureSource::GetMipData(TArray64<uint8>& OutMipData, int32 BlockIndex, int32 LayerIndex, int32 MipIndex, IImageWrapperModule* ImageWrapperModule)
+bool FTextureSource::GetMipImage(FImage & OutImage, int32 BlockIndex, int32 LayerIndex, int32 MipIndex)
+{
+	TArray64<uint8> MipData;
+	if ( ! GetMipData(MipData,BlockIndex,LayerIndex,MipIndex) )
+	{
+		return false;
+	}
+	
+	const int64 MipSizeBytes = CalcMipSize(BlockIndex, LayerIndex, MipIndex);
+	check( MipData.Num() == MipSizeBytes );
+	
+	FTextureSourceBlock Block;
+	GetBlock(BlockIndex, Block);
+	check(MipIndex < Block.NumMips);
+
+	OutImage.RawData = MoveTemp(MipData);
+	OutImage.SizeX = FMath::Max(Block.SizeX >> MipIndex, 1);
+	OutImage.SizeY = FMath::Max(Block.SizeY >> MipIndex, 1);
+	OutImage.NumSlices = Block.NumSlices;
+	OutImage.Format = FImageCoreUtils::ConvertToRawImageFormat(Format);
+
+	if ( Owner != nullptr )
+	{
+		OutImage.GammaSpace = Owner->GetGammaSpace();
+	}
+	else
+	{
+		// TextureSource does not have gamma information, just guess it ?
+		OutImage.GammaSpace = ERawImageFormat::GetDefaultGammaSpace(OutImage.Format);
+	}
+
+	check( OutImage.GetImageSizeBytes() == MipSizeBytes );
+	return true;
+}
+
+bool FTextureSource::GetMipData(TArray64<uint8>& OutMipData, int32 BlockIndex, int32 LayerIndex, int32 MipIndex, IImageWrapperModule* )
 {
 	TRACE_CPUPROFILER_EVENT_SCOPE(FTextureSource::GetMipData (TArray64));
 
@@ -1638,7 +1766,7 @@ bool FTextureSource::GetMipData(TArray64<uint8>& OutMipData, int32 BlockIndex, i
 		// by the one above. If it fires then it indicates that there is a lock/unlock mismatch as well as invalid access!
 		checkf(LockedMipData.IsNull(), TEXT("Attempting to access mip data while locked mip data is still allocated"));
 
-		FSharedBuffer DecompressedData = Decompress(ImageWrapperModule);
+		FSharedBuffer DecompressedData = Decompress();
 
 		if (!DecompressedData.IsNull())
 		{
@@ -1663,7 +1791,7 @@ bool FTextureSource::GetMipData(TArray64<uint8>& OutMipData, int32 BlockIndex, i
 	return bSuccess;
 }
 
-FTextureSource::FMipData FTextureSource::GetMipData(IImageWrapperModule* ImageWrapperModule)
+FTextureSource::FMipData FTextureSource::GetMipData(IImageWrapperModule* )
 {
 	TRACE_CPUPROFILER_EVENT_SCOPE(FTextureSource::GetMipData (FMipData));
 
@@ -1676,7 +1804,7 @@ FTextureSource::FMipData FTextureSource::GetMipData(IImageWrapperModule* ImageWr
 	FReadScopeLock _(BulkDataLock.Get());
 #endif //WITH_EDITOR
 
-	FSharedBuffer DecompressedData = Decompress(ImageWrapperModule);
+	FSharedBuffer DecompressedData = Decompress();
 	return FMipData(*this, DecompressedData);
 }
 
@@ -1776,6 +1904,7 @@ ETextureSourceCompressionFormat FTextureSource::GetSourceCompression() const
 {
 	// Until we deprecate bPNGCompressed it might not be 100% in sync with CompressionFormat
 	// so if it is set we should use that rather than the enum.
+	// @todo fix this, remove bPNGCompressed
 	if (bPNGCompressed)
 	{
 		return ETextureSourceCompressionFormat::TSCF_PNG;
@@ -1789,84 +1918,46 @@ FString FTextureSource::GetSourceCompressionAsString() const
 	return StaticEnum<ETextureSourceCompressionFormat>()->GetDisplayNameTextByValue(GetSourceCompression()).ToString();
 }
 
-FSharedBuffer FTextureSource::TryDecompressPngData(IImageWrapperModule* ImageWrapperModule) const
+FSharedBuffer FTextureSource::TryDecompressData() const
 {
-	bool bCanPngCompressFormat = (Format == TSF_G8 || Format == TSF_G16 || Format == TSF_RGBA8 || Format == TSF_BGRA8 || Format == TSF_RGBA16);
-	check(Blocks.Num() == 0 && NumLayers == 1 && NumSlices == 1 && bCanPngCompressFormat);
-
-	FSharedBuffer Payload = BulkData.GetPayload().Get();
-
-	if (ImageWrapperModule == nullptr) // Optional if called from the gamethread, see FModuleManager::WarnIfItWasntSafeToLoadHere()
-	{
-		ImageWrapperModule = &FModuleManager::LoadModuleChecked<IImageWrapperModule>(FName("ImageWrapper"));
-	}
-
-	TSharedPtr<IImageWrapper> ImageWrapper = ImageWrapperModule->CreateImageWrapper(EImageFormat::PNG);
-	if (ImageWrapper.IsValid() && ImageWrapper->SetCompressed(Payload.GetData(), Payload.GetSize()))
+	if (NumLayers == 1 && NumSlices == 1 && Blocks.Num() == 0)
 	{
-		Payload.Reset(); // Image wrapper takes a copy so we can discard the payload now
-
-		check(ImageWrapper->GetWidth() == SizeX);
-		check(ImageWrapper->GetHeight() == SizeY);
+		FSharedBuffer Payload = BulkData.GetPayload().Get();
 
-		TArray64<uint8> RawData;
-		// TODO: TSF_BGRA8 is stored as RGBA, so the R and B channels are swapped in the internal png. Should we fix this?
-		ERGBFormat RawFormat = (Format == TSF_G8 || Format == TSF_G16) ? ERGBFormat::Gray : ERGBFormat::RGBA;
-		if (ImageWrapper->GetRaw(RawFormat, (Format == TSF_G16 || Format == TSF_RGBA16) ? 16 : 8, RawData) && RawData.Num() > 0)
-		{		
-			return MakeSharedBufferFromArray(MoveTemp(RawData));
-		}
-		else
+		FImage Image;
+		if ( ! FImageUtils::DecompressImage(Payload.GetData(), Payload.GetSize(), Image) )
 		{
-			UE_LOG(LogTexture, Warning, TEXT("PNG decompression of source art failed"));
+			UE_LOG(LogTexture, Warning, TEXT("TryDecompressData failed to return uncompressed data"));
 			return FSharedBuffer();
 		}
-	}
-	else
-	{
-		UE_LOG(LogTexture, Log, TEXT("Only pngs are supported"));
-		return FSharedBuffer();
-	}
-}
 
-FSharedBuffer FTextureSource::TryDecompressJpegData(IImageWrapperModule* ImageWrapperModule) const
-{
-	if (NumLayers == 1 && NumSlices == 1 && Blocks.Num() == 0)
-	{
-		if (!ImageWrapperModule)
-		{
-			ImageWrapperModule = &FModuleManager::LoadModuleChecked<IImageWrapperModule>(FName("ImageWrapper"));
-		}
-
-		FSharedBuffer Payload = BulkData.GetPayload().Get();
+		// we got data in Image.Format
+		// we expect data in TSF "Format"
+		ERawImageFormat::Type RawFormat = FImageCoreUtils::ConvertToRawImageFormat(Format);
 
-		TSharedPtr<IImageWrapper> ImageWrapper = ImageWrapperModule->CreateImageWrapper(EImageFormat::JPEG);
-		if (ImageWrapper.IsValid() && ImageWrapper->SetCompressed(Payload.GetData(), Payload.GetSize()))
+		if ( Image.Format != RawFormat )
 		{
-			Payload.Reset(); // Image wrapper takes a copy so we can discard the payload now
+			// this shouldn't ever happen currently
+			UE_LOG(LogTexture, Warning, TEXT("TryDecompressData unexpected format conversion?"));
 
-			TArray64<uint8> RawData;
-			// The two formats we support for JPEG imports, see UTextureFactory::ImportImage
-			const ERGBFormat JpegFormat = Format == TSF_G8 ? ERGBFormat::Gray : ERGBFormat::BGRA;
-			if (ImageWrapper->GetRaw(JpegFormat, 8, RawData))
-			{
-				return MakeSharedBufferFromArray(MoveTemp(RawData));
-			}
-			else
-			{
-				UE_LOG(LogTexture, Warning, TEXT("JPEG decompression of source art failed to return uncompressed data"));
-				return FSharedBuffer();
-			}
+			// TextureSource does not have gamma information, so just guess:
+			//  (gamma conversion should never actually happen here so this should be unused)
+			Image.ChangeFormat(RawFormat, ERawImageFormat::GetDefaultGammaSpace(RawFormat));
 		}
-		else
+			
+		if ( (CompressionFormat == TSCF_PNG || bPNGCompressed) && Image.Format == ERawImageFormat::BGRA8 )
 		{
-			UE_LOG(LogTexture, Warning, TEXT("JPEG decompression of source art failed initialization"));
-			return FSharedBuffer();
+			// Legacy bug, must be matched in Compress & Decompress
+			// see FTextureSource::Compress
+			// TODO: TSF_BGRA8 is stored as RGBA, so the R and B channels are swapped in the internal png. Should we fix this?
+			FImageCore::TransposeImageRGBABGRA(Image);
 		}
+
+		return MakeSharedBufferFromArray(MoveTemp(Image.RawData));
 	}
 	else
 	{
-		UE_LOG(LogTexture, Warning, TEXT("JPEG compressed source art is in an invalid format NumLayers:(%d) NumSlices:(%d) NumBlocks:(%d)"),
+		UE_LOG(LogTexture, Warning, TEXT("Compressed source art is in an invalid format NumLayers:(%d) NumSlices:(%d) NumBlocks:(%d)"),
 			NumLayers, NumSlices, Blocks.Num());
 		return FSharedBuffer();
 	}	
@@ -1965,7 +2056,7 @@ void FTextureSource::ImportCustomProperties(const TCHAR* SourceText, FFeedbackCo
 
 bool FTextureSource::CanPNGCompress() const
 {
-	bool bCanPngCompressFormat = (Format == TSF_G8 || Format == TSF_G16 || Format == TSF_RGBA8 || Format == TSF_BGRA8 || Format == TSF_RGBA16);
+	bool bCanPngCompressFormat = (Format == TSF_G8 || Format == TSF_G16 || Format == TSF_BGRA8 || Format == TSF_RGBA16 );
 
 	if (!bPNGCompressed &&
 		NumLayers == 1 &&
diff --git a/Engine/Source/Runtime/Engine/Private/TextureDerivedDataTask.cpp b/Engine/Source/Runtime/Engine/Private/TextureDerivedDataTask.cpp
--- a/Engine/Source/Runtime/Engine/Private/TextureDerivedDataTask.cpp
+++ b/Engine/Source/Runtime/Engine/Private/TextureDerivedDataTask.cpp
@@ -37,6 +37,7 @@
 #include "DerivedDataValue.h"
 #include "Engine/TextureCube.h"
 #include "Engine/VolumeTexture.h"
+#include "ImageCoreUtils.h"
 #include "GenericPlatform/GenericPlatformMath.h"
 #include "HAL/CriticalSection.h"
 #include "HAL/PlatformProcess.h"
@@ -170,21 +171,12 @@ void FTextureSourceData::Init(UTexture& InTexture, TextureMipGenSettings InMipGe
 	for (int LayerIndex = 0; LayerIndex < NumLayers; ++LayerIndex)
 	{
 		FTextureSourceLayerData* LayerData = new(Layers) FTextureSourceLayerData();
-		switch (InTexture.Source.GetFormat(LayerIndex))
-		{
-		case TSF_G8:		LayerData->ImageFormat = ERawImageFormat::G8;		break;
-		case TSF_G16:		LayerData->ImageFormat = ERawImageFormat::G16;		break;
-		case TSF_BGRA8:		LayerData->ImageFormat = ERawImageFormat::BGRA8;	break;
-		case TSF_BGRE8:		LayerData->ImageFormat = ERawImageFormat::BGRE8;	break;
-		case TSF_RGBA16:	LayerData->ImageFormat = ERawImageFormat::RGBA16;	break;
-		case TSF_RGBA16F:	LayerData->ImageFormat = ERawImageFormat::RGBA16F;  break;
-		default:
-			UE_LOG(LogTexture, Fatal, TEXT("Texture %s has source art in an invalid format."), *InTexture.GetName());
-			return;
-		}
+
+		LayerData->ImageFormat = FImageCoreUtils::ConvertToRawImageFormat( InTexture.Source.GetFormat(LayerIndex) );
 
 		FTextureFormatSettings FormatSettings;
 		InTexture.GetLayerFormatSettings(LayerIndex, FormatSettings);
+
 		LayerData->GammaSpace = FormatSettings.SRGB ? (InTexture.bUseLegacyGamma ? EGammaSpace::Pow22 : EGammaSpace::sRGB) : EGammaSpace::Linear;
 	}
 
@@ -1214,19 +1206,17 @@ public:
 
 		for (int LayerIndex = 0; LayerIndex < NumLayers; ++LayerIndex)
 		{
-			switch (Texture.Source.GetFormat(LayerIndex))
+			ETextureSourceFormat TSF = Texture.Source.GetFormat(LayerIndex);
+			ERawImageFormat::Type RawFormat = FImageCoreUtils::ConvertToRawImageFormat(TSF);
+
+			if ( RawFormat == ERawImageFormat::Invalid )
 			{
-			case TSF_G8:
-			case TSF_G16:
-			case TSF_BGRA8:
-			case TSF_BGRE8:
-			case TSF_RGBA16:
-			case TSF_RGBA16F:
-				break;
-			default:
-				UE_LOG(LogTexture, Fatal, TEXT("Texture %s has source art in an invalid format."), *Texture.GetPathName());
+				UE_LOG(LogTexture, Error, TEXT("Texture %s has source art in an invalid format."), *Texture.GetPathName());
 				return false;
 			}
+
+			// valid TSF should round-trip :
+			check( FImageCoreUtils::ConvertToTextureSourceFormat( RawFormat ) == TSF );
 		}
 
 		const bool bCompositeTextureViable = Texture.CompositeTexture && Texture.CompositeTextureMode != CTM_Disabled;
diff --git a/Engine/Source/Runtime/Engine/Private/TextureDerivedDataTask.h b/Engine/Source/Runtime/Engine/Private/TextureDerivedDataTask.h
--- a/Engine/Source/Runtime/Engine/Private/TextureDerivedDataTask.h
+++ b/Engine/Source/Runtime/Engine/Private/TextureDerivedDataTask.h
@@ -83,9 +83,10 @@ struct FTextureSourceData
 
 	void Init(UTexture& InTexture, TextureMipGenSettings InMipGenSettings, bool bInCubeMap, bool bInTextureArray, bool bInVolumeTexture, bool bAllowAsyncLoading);
 	bool IsValid() const { return bValid; }
-
-	void GetSourceMips(FTextureSource& Source, IImageWrapperModule* InImageWrapper);
-	void GetAsyncSourceMips(IImageWrapperModule* InImageWrapper);
+	
+	// ImageWrapperModule is not used
+	void GetSourceMips(FTextureSource& Source, IImageWrapperModule* InImageWrapper = nullptr);
+	void GetAsyncSourceMips(IImageWrapperModule* InImageWrapper = nullptr);
 
 	void ReleaseMemory()
 	{
diff --git a/Engine/Source/Runtime/Engine/Private/TextureRenderTarget2D.cpp b/Engine/Source/Runtime/Engine/Private/TextureRenderTarget2D.cpp
--- a/Engine/Source/Runtime/Engine/Private/TextureRenderTarget2D.cpp
+++ b/Engine/Source/Runtime/Engine/Private/TextureRenderTarget2D.cpp
@@ -305,12 +305,19 @@ UTexture2D* UTextureRenderTarget2D::ConstructTexture2D(UObject* Outer, const FSt
 	UpdateTexture2D(Result, TextureFormat, Flags, AlphaOverride);
 
 	// if render target gamma used was 1.0 then disable SRGB for the static texture
-	if (FMath::Abs(RenderTarget->GetDisplayGamma() - 1.0f) < KINDA_SMALL_NUMBER)
+	// @@!! UTextureRenderTarget2D also has an explicit SRGB flag in the UTexture parent class
+	//	  these are NOT correctly kept in sync
+	//		I see SRGB = 1 but Gamma = 1.0
+	//	 see also IsSRGB() which is yet another query that has different ideas
+	//	furthermore, float formats do not support anything but Linear gamma
+	float Gamma = RenderTarget->GetDisplayGamma();
+	if (FMath::Abs(Gamma - 1.0f) < KINDA_SMALL_NUMBER)
 	{
 		Flags &= ~CTF_SRGB;
 	}
 
 	Result->SRGB = (Flags & CTF_SRGB) != 0;
+
 	Result->MipGenSettings = TMGS_FromTextureGroup;
 
 	if ((Flags & CTF_AllowMips) == 0)
diff --git a/Engine/Source/Runtime/Engine/Private/TextureRenderTargetCube.cpp b/Engine/Source/Runtime/Engine/Private/TextureRenderTargetCube.cpp
--- a/Engine/Source/Runtime/Engine/Private/TextureRenderTargetCube.cpp
+++ b/Engine/Source/Runtime/Engine/Private/TextureRenderTargetCube.cpp
@@ -150,6 +150,8 @@ UTextureCube* UTextureRenderTargetCube::ConstructTextureCube(
 			case PF_B8G8R8A8:
 				TextureFormat = TSF_BGRA8;
 				break;
+			default:
+				return nullptr;
 		}
 
 		// The r2t resource will be needed to read its surface contents
diff --git a/Engine/Source/Runtime/Engine/Private/VT/VirtualTextureDataBuilder.cpp b/Engine/Source/Runtime/Engine/Private/VT/VirtualTextureDataBuilder.cpp
--- a/Engine/Source/Runtime/Engine/Private/VT/VirtualTextureDataBuilder.cpp
+++ b/Engine/Source/Runtime/Engine/Private/VT/VirtualTextureDataBuilder.cpp
@@ -14,11 +14,9 @@
 
 // Debugging aid to dump tiles to disc as png files
 #define SAVE_TILES 0
+
 #if SAVE_TILES
-#include "IImageWrapper.h"
-#include "IImageWrapperModule.h"
-#include "HAL/FileManager.h"
-#include "Misc/FileHelper.h"
+#include "ImageUtils.h"
 #endif
 
 static TAutoConsoleVariable<int32> CVarVTParallelTileCompression(
@@ -190,7 +188,7 @@ struct FPixelDataRectangle
 		FMemory::Memcpy(DestPixelData, Value, PixelSize);
 	}
 
-	inline void *GetPixel(int32 x, int32 y, size_t PixelSize)
+	inline void *GetPixel(int32 x, int32 y, size_t PixelSize) const
 	{
 		check(x >= 0);
 		check(y >= 0);
@@ -198,57 +196,21 @@ struct FPixelDataRectangle
 		check(y < Height);
 		return Data + (((y * Width) + x) * PixelSize);
 	}
-
+	
 #if SAVE_TILES
-	void Save(FString BaseFileName, IImageWrapperModule* ImageWrapperModule)
+	FImageView GetImageView() const 
 	{
-		IFileManager* FileManager = &IFileManager::Get();
-		auto ImageWrapper = ImageWrapperModule->CreateImageWrapper(EImageFormat::PNG);
-		int BytesPerPixel = FTextureSource::GetBytesPerPixel(Format);
-
-		switch (Format)
-		{
-		case TSF_G8:
-			ImageWrapper->SetRaw(Data, BytesPerPixel * Width * Height, Width, Height, ERGBFormat::Gray, 8);
-			break;
-		case TSF_G16:
-			ImageWrapper->SetRaw(Data, BytesPerPixel * Width * Height, Width, Height, ERGBFormat::Gray, 16);
-			break;
-		case TSF_BGRA8:
-			ImageWrapper->SetRaw(Data, BytesPerPixel * Width * Height, Width, Height, ERGBFormat::BGRA, 8);
-			break;
-		case TSF_BGRE8:
-			// This will probably result in bad image. (Should we convert the data or export it in a format that support float(ERGBFormat::BGRE)?)
-			ImageWrapper->SetRaw(Data, BytesPerPixel * Width * Height, Width, Height, ERGBFormat::BGRA, 8);
-			break;
-		case TSF_RGBA16:
-			ImageWrapper->SetRaw(Data, BytesPerPixel * Width * Height, Width, Height, ERGBFormat::RGBA, 16);
-			break;
-		case TSF_RGBA16F:
-			// This will probably result in bad image. (Should we convert the data or export it in a format that support float(ERGBFormat::RGBAF)?)
-			ImageWrapper->SetRaw(Data, BytesPerPixel * Width * Height, Width, Height, ERGBFormat::RGBA, 16);
-			break;
-		case TSF_RGBA8:
-			ImageWrapper->SetRaw(Data, BytesPerPixel * Width * Height, Width, Height, ERGBFormat::RGBA, 8);
-			break;
-		case TSF_RGBE8:
-			ImageWrapper->SetRaw(Data, BytesPerPixel * Width * Height, Width, Height, ERGBFormat::RGBA, 8);
-			break;
-		default:
-			return;
-		}
+		ERawImageFormat::Type RawFormat = FImageUtils::ConvertToRawImageFormat(Format);
+		return FImageView(Data,Width,Height,RawFormat);
+	}
 
-		// Make sure it has the png extension
-		FString NewExtension = TEXT(".png");
-		FString Filename = FPaths::GetBaseFilename(BaseFileName, false) + NewExtension;
+	void Save(FString BaseFileName, IImageWrapperModule* ImageWrapperModule)
+	{
+		FImageView Image = GetImageView();
 
-		// Compress and write image
-		FArchive* Ar = FileManager->CreateFileWriter(*Filename);
-		if (Ar != nullptr)
+		if ( ! FImageUtils::SaveImageAutoFormat(*BaseFileName,Image) )
 		{
-			TArray64<uint8> CompressedData = ImageWrapper->GetCompressed((int32)EImageCompressionQuality::Default);
-			Ar->Serialize((void *)CompressedData.GetData(), CompressedData.Num());
-			delete Ar;
+			UE_LOG(LogVirtualTexturing,Warning,TEXT("Couldn't save to : %s"),*BaseFileName);
 		}
 	}
 #endif
diff --git a/Engine/Source/Runtime/Engine/Public/ImageUtils.h b/Engine/Source/Runtime/Engine/Public/ImageUtils.h
--- a/Engine/Source/Runtime/Engine/Public/ImageUtils.h
+++ b/Engine/Source/Runtime/Engine/Public/ImageUtils.h
@@ -1,22 +1,25 @@
 // Copyright Epic Games, Inc. All Rights Reserved.
 
 /*=============================================================================
-ImageUtils.h: Image utility functions.
+ImageUtils.h: Image/Textures utility functions 
 =============================================================================*/
 
 #pragma once
 
 #include "CoreMinimal.h"
 #include "Misc/Guid.h"
+#include "ImageCore.h"
 #include "Engine/Texture.h"
 
 class UTexture2D;
 class UTextureCube;
 class UTextureRenderTarget2D;
 class UTextureRenderTargetCube;
+class IImageWrapperModule;
 
 /**
  *	Parameters used for creating a Texture2D from a simple color buffer.
+ * 
  */
 struct FCreateTexture2DParameters
 {
@@ -58,10 +61,75 @@ struct FCreateTexture2DParameters
 
 /**
  * Class of static image utility functions.
+ * 
+ * Provides load/save of FImage
+ * and conversion from Texture2D/RenderTarget to/from FImage
+ * 
+ *  in Engine, cannot be used by standalone texture build workers
  */
 class FImageUtils
 {
 public:
+
+	/**
+	 * Load an image of any type supported by the ImageWrapper module
+	 *
+	 * @param Filename				File name to load
+	 * @param OutImage				Filled with the loaded image, allocated as needed
+	 */
+	ENGINE_API static bool LoadImage(const TCHAR * Filename, FImage & OutImage);
+		
+	/**
+	 * Save an image.  Extension of Filename will be used for output image file format.
+	 * FImageView can be made from any surface pointer.
+	 *
+	 * @param Filename				File name to save to, with extension to identify format
+	 * @param InImage				Image to save
+	 * @param Quality				Mainly for JPEG, but special values have some meaning for other formats
+	 */
+	ENGINE_API static bool SaveImageByExtension(const TCHAR * Filename, const FImageView & InImage, int32 Quality=0);
+
+	/**
+	 * Save an image.  Output format will be chosen automatically based on the Image pixel format
+	 * eg. EXR for float, PNG for BGRA8
+	 * extension will be added to file name
+	 *
+	 * @param Filename				File name to save to, auto extension will be added
+	 * @param InImage				Image to save
+	 * @param Quality				Mainly for JPEG, but special values have some meaning for other formats
+	 */
+	ENGINE_API static bool SaveImageAutoFormat(const TCHAR * Filename, const FImageView & InImage, int32 Quality=0);
+	
+	/**
+	 * "Compress" an image to a file format.  Here "Compress" really means put an image into file format binary.
+	 * Extension is used to identify file format, eg. png/jpeg/etc.
+	 * FImageView can be made from any surface pointer.
+	 *
+	 * @param OutData				Filled with file format Image data
+	 * @param ToFormatExtension		can be a full name like "xxx.png" or just the extension like "png"
+	 *								ToFormatExtension == null is okay & means use default output format
+	 * @param InImage				Image to save
+	 * @param Quality				Mainly for JPEG, but special values have some meaning for other formats
+	 */
+	ENGINE_API static bool CompressImage(TArray64<uint8> & OutData, const TCHAR * ToFormatExtension, const FImageView & InImage, int32 Quality = 0);
+	
+	/**
+	 * "Decompress" an image to a file format.  Here "Compress" really means put an image into file format binary.
+	 *
+	 * @param InCompressedData		File formatted image bits to unpack
+	 * @param InCompressedSize		Size of InCompressedData
+	 * @param OutImage				Filled with the decompressed image and allocated as needed
+	 */
+	ENGINE_API static bool DecompressImage(const void* InCompressedData, int64 InCompressedSize, FImage & OutImage);
+	
+	/**
+	 * if Texture source is available, get it as an FImage
+	 *
+	 * @param Texture		Texture to get gets from
+	 * @param OutImage		Filled with the image copy and allocated as needed
+	 */
+	ENGINE_API static bool GetTexture2DSourceImage(	UTexture2D* Texture, FImage & OutImage);
+
 	/**
 	 * Resizes the given image using a simple average filter and stores it in the destination array.  This version constrains aspect ratio.
 	 *
@@ -130,6 +198,14 @@ public:
 	 *
 	 */
 	ENGINE_API static UTexture2D* CreateTexture2D(int32 SrcWidth, int32 SrcHeight, const TArray<FColor> &SrcData, UObject* Outer, const FString& Name, const EObjectFlags &Flags, const FCreateTexture2DParameters& InParams);
+	
+	/**
+	 * Creates a 2D texture from an FImage
+	 * 
+	 * @param Image			Image that will be copied into a Texture
+	 * @return				Returns a pointer to the constructed 2D texture object.
+	 */	
+	ENGINE_API static UTexture2D* CreateTexture2DFromImage(const FImageView & Image);
 
 	/**
 	 * Crops, and scales an image from a raw image array.
@@ -163,6 +239,7 @@ public:
 
 	/**
 	 * Compress image to PNG format uint8 array.
+	 * deprecated, use CompressImage instead
 	 *
 	 * @param ImageHeight		Source image width.
 	 * @param ImageWidth		Source image height.
@@ -194,6 +271,8 @@ public:
 
 	/**
 	 * Exports a UTextureRenderTarget2D as an HDR image on the disk.
+	 * HDR is very low quality, prefer EXR or PNG instead
+	 * Deprecated.  Prefer GetRenderTargetImage.
 	 *
 	 * @param TexRT		The render target to export
 	 * @param Ar			Archive to fill with data.
@@ -201,21 +280,33 @@ public:
 	 *
 	 */
 	ENGINE_API static bool ExportRenderTarget2DAsHDR(UTextureRenderTarget2D* TexRT, FArchive& Ar);
+	
+	/**
+	 * Exports a UTextureRenderTarget2D as a EXR image on the disk.
+	 * Deprecated.  Prefer GetRenderTargetImage.
+	 *
+	 * @param TexRT		The render target to export
+	 * @param Ar		Archive to fill with data.
+	 * @return			Export operation success or failure.
+	 *
+	 */
+	ENGINE_API static bool ExportRenderTarget2DAsEXR(UTextureRenderTarget2D* TexRT, FArchive& Ar);
 
 	/**
-	 * Exports a UTextureRenderTarget2D as an HDR image on the disk.
+	 * Exports a UTextureRenderTarget2D as a PNG image on the disk.
+	 * Deprecated.  Prefer GetRenderTargetImage.
 	 *
 	 * @param TexRT		The render target to export
-	 * @param Ar			Archive to fill with data.
+	 * @param Ar		Archive to fill with data.
 	 * @return			Export operation success or failure.
 	 *
 	 */
 	ENGINE_API static bool ExportRenderTarget2DAsPNG(UTextureRenderTarget2D* TexRT, FArchive& Ar);
 
-	ENGINE_API static bool ExportRenderTarget2DAsEXR(UTextureRenderTarget2D* TexRT, FArchive& Ar);
-
 	/**
 	* Exports a UTexture2D as an HDR image on the disk.
+	* Deprecated.  Prefer GetTexture2DSourceImage + SaveImage. 
+	* HDR is very low quality, prefer EXR or PNG instead
 	*
 	* @param TexRT		The texture to export
 	* @param Ar			Archive to fill with data.
@@ -233,12 +324,12 @@ public:
 	 * Imports a texture a buffer and creates Texture2D from it
 	 */
 	ENGINE_API static UTexture2D* ImportBufferAsTexture2D(TArrayView64<const uint8> Buffer);
-
 	ENGINE_API static UTexture2D* ImportBufferAsTexture2D(const TArray<uint8>& Buffer);
 	
 	/**
 	* Exports a UTextureRenderTargetCube as an HDR image on the disk.
-	*
+	* this does GenerateLongLatUnwrap
+	* 
 	* @param TexRT		The render target cube to export
 	* @param Ar			Archive to fill with data.
 	* @return			Export operation success or failure.
@@ -248,7 +339,8 @@ public:
 
 	/**
 	* Exports a UTextureCube as an HDR image on the disk.
-	*
+	* this does GenerateLongLatUnwrap
+	* 
 	* @param TexRT		The texture cube to export
 	* @param Ar			Archive to fill with data.
 	* @return			Export operation success or failure.
@@ -256,6 +348,16 @@ public:
 	*/
 	ENGINE_API static bool ExportTextureCubeAsHDR(UTextureCube* TexRT, FArchive& Ar);
 
+	// Should be removed from public API.
+	//  Move out of header.
 	ENGINE_API static bool GetRawData(UTextureRenderTarget2D* TexRT, TArray64<uint8>& RawData);
+	
+	/**
+	* Get the contents of a RenderTarget into an Image
+	* 
+	* @param TexRT		The texture rendertarget to copy from
+	* @param OutImage	Filled with the image, allocated as needed
+	*/
+	ENGINE_API static bool GetRenderTargetImage(UTextureRenderTarget2D* TexRT, FImage & OutImage);
 
 };
diff --git a/Engine/Source/Runtime/ImageCore/ImageCore.Build.cs b/Engine/Source/Runtime/ImageCore/ImageCore.Build.cs
--- a/Engine/Source/Runtime/ImageCore/ImageCore.Build.cs
+++ b/Engine/Source/Runtime/ImageCore/ImageCore.Build.cs
@@ -6,7 +6,15 @@ public class ImageCore : ModuleRules
 {
 	public ImageCore(ReadOnlyTargetRules Target) : base(Target)
 	{
+		// include only, no link :
+		// for TextureDefines.h :
+		PrivateIncludePathModuleNames.AddRange(new string[]
+		{
+			"Engine",
+		});
+
 		PrivateDependencyModuleNames.Add("ColorManagement");
+		PrivateDependencyModuleNames.Add("CoreUObject"); // for TextureDefines.h
 
 		PublicDependencyModuleNames.Add("Core");
 	}
diff --git a/Engine/Source/Runtime/ImageCore/Private/ImageCore.cpp b/Engine/Source/Runtime/ImageCore/Private/ImageCore.cpp
--- a/Engine/Source/Runtime/ImageCore/Private/ImageCore.cpp
+++ b/Engine/Source/Runtime/ImageCore/Private/ImageCore.cpp
@@ -12,6 +12,8 @@ static constexpr float MAX_HALF_FLOAT16 = 65504.0f;
 
 FORCEINLINE static FLinearColor SaturateToHalfFloat(const FLinearColor& LinearCol)
 {
+	// @@!! vector min+max
+	// @@!! why do we need this at all?
 	FLinearColor Result;
 	Result.R = FMath::Clamp(LinearCol.R, -MAX_HALF_FLOAT16, MAX_HALF_FLOAT16);
 	Result.G = FMath::Clamp(LinearCol.G, -MAX_HALF_FLOAT16, MAX_HALF_FLOAT16);
@@ -33,7 +35,7 @@ IMPLEMENT_MODULE(FDefaultModuleImpl, ImageCore);
  */
 static void InitImageStorage(FImage& Image)
 {
-	int64 NumBytes = int64(Image.SizeX) * Image.SizeY * Image.NumSlices * Image.GetBytesPerPixel();
+	int64 NumBytes = Image.GetImageSizeBytes();
 	Image.RawData.Empty(NumBytes);
 	Image.RawData.AddUninitialized(NumBytes);
 }
@@ -130,48 +132,105 @@ static void ParallelLoop(const TCHAR* DebugName, int32 NumJobs, int64 TexelsPerJ
 }
 
 
+// Copy Image, swapping RB, SrcImage must be BGRA8 or RGBA16
+// Src == Dest is okay
+IMAGECORE_API void FImageCore::CopyImageRGBABGRA(const FImageView & SrcImage,const FImageView & DestImage)
+{
+	check( SrcImage.GetNumPixels() == DestImage.GetNumPixels() );
+	check( SrcImage.Format == DestImage.Format );
+	const int64 NumTexels = SrcImage.GetNumPixels();
+
+	if ( SrcImage.Format == ERawImageFormat::BGRA8 )
+	{	
+		const uint8 * SrcColors = (const uint8 *)SrcImage.RawData;
+		const uint8 * SrcColorsEnd = SrcColors + SrcImage.GetImageSizeBytes();
+		uint8 * DestColors = (uint8 *)DestImage.RawData;
+	
+		for(;SrcColors<SrcColorsEnd;SrcColors+=4,DestColors+=4)
+		{
+			// make it work with Src == Dest
+			uint8 R = SrcColors[0];
+			uint8 B = SrcColors[2];
+			DestColors[0] = B;
+			DestColors[1] = SrcColors[1];
+			DestColors[2] = R;
+			DestColors[3] = SrcColors[3];
+		}
+	}
+	else if ( SrcImage.Format == ERawImageFormat::RGBA16 )
+	{
+		const uint16 * SrcColors = (const uint16 *)SrcImage.RawData;
+		const uint16 * SrcColorsEnd = SrcColors + SrcImage.GetImageSizeBytes()/sizeof(uint16);
+		uint16 * DestColors = (uint16 * )DestImage.RawData;
+	
+		for(;SrcColors<SrcColorsEnd;SrcColors+=4,DestColors+=4)
+		{
+			// make it work with Src == Dest
+			uint16 R = SrcColors[0];
+			uint16 B = SrcColors[2];
+			DestColors[0] = B;
+			DestColors[1] = SrcColors[1];
+			DestColors[2] = R;
+			DestColors[3] = SrcColors[3];
+		}
+	}
+	else
+	{
+		check(0);
+	}
+}
+
+IMAGECORE_API void FImageCore::TransposeImageRGBABGRA(const FImageView & Image)
+{
+	// CopyImageRGBABGRA is okay in place
+	CopyImageRGBABGRA(Image,Image);
+}
+
+static uint8 Requantize16to8(const uint16 In)
+{
+	// same as QuantizeRound(In/65535.f);
+    uint32 Ret = ( (uint32)In * 255 + 32768 + 127 )>>16;
+	checkSlow( Ret <= 255 );
+	return (uint8)Ret;
+}
+
 /**
  * Copies an image accounting for format differences. Sizes must match.
  *
  * @param SrcImage - The source image to copy from.
  * @param DestImage - The destination image to copy to.
  */
-static void CopyImage(const FImage& SrcImage, FImage& DestImage)
+IMAGECORE_API void FImageCore::CopyImage(const FImageView & SrcImage,const FImageView & DestImage)
 {
 	TRACE_CPUPROFILER_EVENT_SCOPE(Texture.CopyImage);
 
-	check(SrcImage.SizeX == DestImage.SizeX);
-	check(SrcImage.SizeY == DestImage.SizeY);
-	check(SrcImage.NumSlices == DestImage.NumSlices);
+	check(SrcImage.GetNumPixels() == DestImage.GetNumPixels());
 
 	const bool bDestIsGammaCorrected = DestImage.IsGammaCorrected();
-	const int64 NumTexels = int64(SrcImage.SizeX) * SrcImage.SizeY * SrcImage.NumSlices;
+	const int64 NumTexels = SrcImage.GetNumPixels();
 	int64 TexelsPerJob;
 	int32 NumJobs = ImageParallelForComputeNumJobsForPixels(TexelsPerJob,NumTexels);
 
 	if (SrcImage.Format == DestImage.Format &&
 		SrcImage.GammaSpace == DestImage.GammaSpace)
 	{
-		DestImage.RawData = SrcImage.RawData;
+		int64 Bytes = SrcImage.GetImageSizeBytes();
+		check( DestImage.GetImageSizeBytes() == Bytes );
+		memcpy(DestImage.RawData,SrcImage.RawData,Bytes);
 	}
 	else if (SrcImage.Format == ERawImageFormat::RGBA32F)
 	{
 		// Convert from 32-bit linear floating point.
-		const FLinearColor* SrcColors = SrcImage.AsRGBA32F().GetData();
+		const FLinearColor* SrcColors = (const FLinearColor*) SrcImage.RawData;
 	
-		// if gamma correction is done, it's always to sRGB , not to Pow22
-		// so if Pow22 was requested, change to sRGB
-		// so that Float->int->Float roundtrips correctly
-		if ( DestImage.GammaSpace == EGammaSpace::Pow22 )
-		{
-			DestImage.GammaSpace = EGammaSpace::sRGB;
-		}
+		// if gamma correction is done, it's always *TO* sRGB , not to Pow22
+		check( DestImage.GammaSpace != EGammaSpace::Pow22 );
 
 		switch (DestImage.Format)
 		{
 		case ERawImageFormat::G8:
 			{
-				uint8* DestLum = DestImage.AsG8().GetData();
+				uint8* DestLum = (uint8 *) DestImage.RawData;
 				if (bDestIsGammaCorrected)
 				{
 					ParallelLoop(TEXT("Texture.CopyImage.PF"), NumJobs, TexelsPerJob, NumTexels, [DestLum, SrcColors](int64 TexelIndex)
@@ -191,7 +250,7 @@ static void CopyImage(const FImage& SrcImage, FImage& DestImage)
 
 		case ERawImageFormat::G16:
 			{
-				uint16* DestLum = DestImage.AsG16().GetData();
+				uint16* DestLum = (uint16 *) DestImage.RawData;
 				ParallelLoop(TEXT("Texture.CopyImage.PF"), NumJobs, TexelsPerJob, NumTexels, [DestLum, SrcColors](int64 TexelIndex)
 				{
 					DestLum[TexelIndex] = FColor::QuantizeUNormFloatTo16(SrcColors[TexelIndex].R);
@@ -201,7 +260,7 @@ static void CopyImage(const FImage& SrcImage, FImage& DestImage)
 
 		case ERawImageFormat::BGRA8:
 			{
-				FColor* DestColors = DestImage.AsBGRA8().GetData();
+				FColor* DestColors = (FColor *) DestImage.RawData;
 				if (bDestIsGammaCorrected)
 				{
 					ParallelLoop(TEXT("Texture.CopyImage.PF"), NumJobs, TexelsPerJob, NumTexels,
@@ -265,7 +324,7 @@ static void CopyImage(const FImage& SrcImage, FImage& DestImage)
 		
 		case ERawImageFormat::BGRE8:
 			{
-				FColor* DestColors = DestImage.AsBGRE8().GetData();
+				FColor* DestColors = (FColor*) DestImage.RawData;
 				ParallelLoop(TEXT("Texture.CopyImage.PF"), NumJobs, TexelsPerJob, NumTexels, [DestColors, SrcColors](int64 TexelIndex)
 				{
 					DestColors[TexelIndex] = SrcColors[TexelIndex].ToRGBE();
@@ -275,7 +334,7 @@ static void CopyImage(const FImage& SrcImage, FImage& DestImage)
 		
 		case ERawImageFormat::RGBA16:
 			{
-				uint16* DestColors = DestImage.AsRGBA16().GetData();
+				uint16* DestColors = (uint16 *) DestImage.RawData;
 				ParallelLoop(TEXT("Texture.CopyImage.PF"), NumJobs, TexelsPerJob, NumTexels, [DestColors, SrcColors](int64 TexelIndex)
 				{
 					FLinearColor Src = SrcColors[TexelIndex];
@@ -290,7 +349,7 @@ static void CopyImage(const FImage& SrcImage, FImage& DestImage)
 		
 		case ERawImageFormat::RGBA16F:
 			{
-				FFloat16Color* DestColors = DestImage.AsRGBA16F().GetData();
+				FFloat16Color* DestColors = (FFloat16Color*) DestImage.RawData;
 				ParallelLoop(TEXT("Texture.CopyImage.PF"), NumJobs, TexelsPerJob, NumTexels, [DestColors, SrcColors](int64 TexelIndex)
 				{
 					DestColors[TexelIndex] = FFloat16Color(SrcColors[TexelIndex]);
@@ -300,24 +359,28 @@ static void CopyImage(const FImage& SrcImage, FImage& DestImage)
 
 		case ERawImageFormat::R16F:
 			{
-				FFloat16* DestColors = DestImage.AsR16F().GetData();
+				FFloat16* DestColors = (FFloat16*) DestImage.RawData;
 				ParallelLoop(TEXT("Texture.CopyImage.PF"), NumJobs, TexelsPerJob, NumTexels, [DestColors, SrcColors](int64 TexelIndex)
 				{
 					DestColors[TexelIndex] = FFloat16(SrcColors[TexelIndex].R);
 				});
 			}
 			break;
+			
+		default:
+			check(0);
+			break;
 		}
 	}
 	else if (DestImage.Format == ERawImageFormat::RGBA32F)
 	{
 		// Convert to 32-bit linear floating point.
-		FLinearColor* DestColors = DestImage.AsRGBA32F().GetData();
+		FLinearColor* DestColors = (FLinearColor*) DestImage.RawData;
 		switch (SrcImage.Format)
 		{
 		case ERawImageFormat::G8:
 			{
-				const uint8* SrcLum = SrcImage.AsG8().GetData();
+				const uint8* SrcLum = (const uint8*) SrcImage.RawData;
 				switch (SrcImage.GammaSpace)
 				{
 				case EGammaSpace::Linear:
@@ -350,7 +413,7 @@ static void CopyImage(const FImage& SrcImage, FImage& DestImage)
 
 		case ERawImageFormat::G16:
 			{
-				const uint16* SrcLum = SrcImage.AsG16().GetData();
+				const uint16* SrcLum = (const uint16*) SrcImage.RawData;
 				ParallelLoop(TEXT("Texture.CopyImage.PF"), NumJobs, TexelsPerJob, NumTexels, [DestColors, SrcLum](int64 TexelIndex)
 				{
 					float Value = FColor::DequantizeUNorm16ToFloat(SrcLum[TexelIndex]);
@@ -361,7 +424,7 @@ static void CopyImage(const FImage& SrcImage, FImage& DestImage)
 
 		case ERawImageFormat::BGRA8:
 			{
-				const FColor* SrcColors = SrcImage.AsBGRA8().GetData();
+				const FColor* SrcColors = (const FColor*) SrcImage.RawData;
 				switch (SrcImage.GammaSpace)
 				{
 				case EGammaSpace::Linear:
@@ -427,7 +490,7 @@ static void CopyImage(const FImage& SrcImage, FImage& DestImage)
 
 		case ERawImageFormat::BGRE8:
 			{
-				const FColor* SrcColors = SrcImage.AsBGRE8().GetData();
+				const FColor* SrcColors = (const FColor*) SrcImage.RawData;
 				ParallelLoop(TEXT("Texture.CopyImage.PF"), NumJobs, TexelsPerJob, NumTexels, [DestColors, SrcColors](int64 TexelIndex)
 				{
 					DestColors[TexelIndex] = SrcColors[TexelIndex].FromRGBE();
@@ -437,7 +500,7 @@ static void CopyImage(const FImage& SrcImage, FImage& DestImage)
 
 		case ERawImageFormat::RGBA16:
 			{
-				const uint16* SrcColors = SrcImage.AsRGBA16().GetData();
+				const uint16* SrcColors = (const uint16*) SrcImage.RawData;
 				ParallelLoop(TEXT("Texture.CopyImage.PF"), NumJobs, TexelsPerJob, NumTexels, [DestColors, SrcColors](int64 TexelIndex)
 				{
 					int64 SrcIndex = TexelIndex * 4;
@@ -453,7 +516,7 @@ static void CopyImage(const FImage& SrcImage, FImage& DestImage)
 
 		case ERawImageFormat::RGBA16F:
 			{
-				const FFloat16Color* SrcColors = SrcImage.AsRGBA16F().GetData();
+				const FFloat16Color* SrcColors = (const FFloat16Color*) SrcImage.RawData;
 				ParallelLoop(TEXT("Texture.CopyImage.PF"), NumJobs, TexelsPerJob, NumTexels, [DestColors, SrcColors](int64 TexelIndex)
 				{
 					DestColors[TexelIndex] = SrcColors[TexelIndex].GetFloats();
@@ -463,20 +526,24 @@ static void CopyImage(const FImage& SrcImage, FImage& DestImage)
 
 		case ERawImageFormat::R16F:
 			{
-				const FFloat16* SrcColors = SrcImage.AsR16F().GetData();
+				const FFloat16* SrcColors = (const FFloat16*) SrcImage.RawData;
 				ParallelLoop(TEXT("Texture.CopyImage.PF"), NumJobs, TexelsPerJob, NumTexels, [DestColors, SrcColors](int64 TexelIndex)
 				{
 					DestColors[TexelIndex] = FLinearColor(SrcColors[TexelIndex].GetFloat(), 0, 0, 1);
 				});
 			}
 			break;
+
+		default:
+			check(0);
+			break;
 		}
 	}
-	else if (SrcImage.Format == ERawImageFormat::BGRA8 && SrcImage.GammaSpace == EGammaSpace::Linear &&
-			DestImage.Format == ERawImageFormat::RGBA16 && DestImage.GammaSpace == EGammaSpace::Linear)
+	else if (SrcImage.Format == ERawImageFormat::BGRA8 && DestImage.Format == ERawImageFormat::RGBA16 &&
+		SrcImage.GammaSpace == EGammaSpace::Linear )
 	{
-		const FColor* SrcColors = SrcImage.AsBGRA8().GetData();
-		uint16* DestColors = DestImage.AsRGBA16().GetData();
+		const FColor* SrcColors = (const FColor*) SrcImage.RawData;
+		uint16* DestColors = (uint16*) DestImage.RawData;
 		ParallelLoop<4>(TEXT("Texture.CopyImage.PF"), NumJobs, TexelsPerJob, NumTexels,
 			[DestColors, SrcColors](int64 TexelIndex)
 			{
@@ -512,11 +579,28 @@ static void CopyImage(const FImage& SrcImage, FImage& DestImage)
 			}
 		);
 	}
-	else if (SrcImage.Format == ERawImageFormat::BGRA8 && SrcImage.GammaSpace == EGammaSpace::Linear &&
-			DestImage.Format == ERawImageFormat::G8 && DestImage.GammaSpace == EGammaSpace::Linear)
+	else if (SrcImage.Format == ERawImageFormat::RGBA16 && DestImage.Format == ERawImageFormat::BGRA8 &&
+		DestImage.GammaSpace == EGammaSpace::Linear )
+	{
+		const uint16 * SrcChannels = (const uint16*) SrcImage.RawData;
+		FColor* DestColors = (FColor*) DestImage.RawData;
+		ParallelLoop(TEXT("Texture.CopyImage.PF"), NumJobs, TexelsPerJob, NumTexels,
+			[DestColors, SrcChannels](int64 TexelIndex)
+			{
+				const uint16* Src = SrcChannels + TexelIndex * 4;
+				DestColors[TexelIndex] = FColor(
+					Requantize16to8(Src[0]),
+					Requantize16to8(Src[1]),
+					Requantize16to8(Src[2]),
+					Requantize16to8(Src[3]));
+			}
+		);
+	}
+	else if (SrcImage.Format == ERawImageFormat::BGRA8 && DestImage.Format == ERawImageFormat::G8 &&
+		SrcImage.GammaSpace == DestImage.GammaSpace )
 	{
-		const FColor* SrcColors = SrcImage.AsBGRA8().GetData();
-		uint8* DestLum = DestImage.AsG8().GetData();
+		const FColor* SrcColors = (const FColor*) SrcImage.RawData;
+		uint8* DestLum = (uint8*) DestImage.RawData;
 		ParallelLoop(TEXT("Texture.CopyImage.PF"), NumJobs, TexelsPerJob, NumTexels,
 			[DestLum, SrcColors](int64 TexelIndex)
 			{
@@ -524,14 +608,27 @@ static void CopyImage(const FImage& SrcImage, FImage& DestImage)
 			}
 		);
 	}
+	else if (SrcImage.Format == ERawImageFormat::G8 && DestImage.Format == ERawImageFormat::BGRA8 &&
+		SrcImage.GammaSpace == DestImage.GammaSpace )
+	{
+		const uint8* SrcLum = (const uint8*) SrcImage.RawData;
+		FColor* DestColors = (FColor*) DestImage.RawData;
+		ParallelLoop(TEXT("Texture.CopyImage.PF"), NumJobs, TexelsPerJob, NumTexels,
+			[SrcLum, DestColors](int64 TexelIndex)
+			{
+				uint8 Lum = SrcLum[TexelIndex];
+				DestColors[TexelIndex] = FColor(Lum,Lum,Lum,255);
+			}
+		);
+	}
 	else
 	{
 		TRACE_CPUPROFILER_EVENT_SCOPE(Texture.CopyImage.TempLinear);
 
 		// Arbitrary conversion, use 32-bit linear float as an intermediate format.
-		FImage TempImage(SrcImage.SizeX, SrcImage.SizeY, ERawImageFormat::RGBA32F);
-		CopyImage(SrcImage, TempImage);
-		CopyImage(TempImage, DestImage);
+		FImage TempImage(SrcImage.SizeX, SrcImage.SizeY, SrcImage.NumSlices, ERawImageFormat::RGBA32F, EGammaSpace::Linear);
+		FImageCore::CopyImage(SrcImage, TempImage);
+		FImageCore::CopyImage(TempImage, DestImage);
 	}
 }
 
@@ -611,29 +708,21 @@ static void ResizeImage(const FImage& SrcImage, FImage& DestImage)
 
 /* FImage constructors
  *****************************************************************************/
-
+ 
 FImage::FImage(int32 InSizeX, int32 InSizeY, int32 InNumSlices, ERawImageFormat::Type InFormat, EGammaSpace InGammaSpace)
-	: SizeX(InSizeX)
-	, SizeY(InSizeY)
-	, NumSlices(InNumSlices)
-	, Format(InFormat)
-	, GammaSpace(InGammaSpace)
+	: FImageInfo(InSizeX,InSizeY,InNumSlices,InFormat,InGammaSpace)
 {
 	InitImageStorage(*this);
 }
 
-
-FImage::FImage(int32 InSizeX, int32 InSizeY, ERawImageFormat::Type InFormat, EGammaSpace InGammaSpace)
-	: SizeX(InSizeX)
-	, SizeY(InSizeY)
-	, NumSlices(1)
-	, Format(InFormat)
-	, GammaSpace(InGammaSpace)
+void FImage::Init(const FImageInfo & Info)
 {
+	// assign the FImageInfo part of this :
+	FImageInfo * MyInfo = this; // implicit cast
+	*MyInfo = Info;
 	InitImageStorage(*this);
 }
 
-
 void FImage::Init(int32 InSizeX, int32 InSizeY, int32 InNumSlices, ERawImageFormat::Type InFormat, EGammaSpace InGammaSpace)
 {
 	SizeX = InSizeX;
@@ -658,16 +747,54 @@ void FImage::Init(int32 InSizeX, int32 InSizeY, ERawImageFormat::Type InFormat,
 
 /* FImage interface
  *****************************************************************************/
+ 
+void FImage::Swap(FImage & Other)
+{
+	::Swap(RawData,Other.RawData);
+	// FImageInfo can be moved with assignment
+	FImageInfo * Me = this;
+	FImageInfo * Him = &Other;
+	FImageInfo Temp;
+	Temp = *Me;
+	*Me = *Him;
+	*Him = Temp;
+}
+
+void FImage::ChangeFormat(ERawImageFormat::Type DestFormat, EGammaSpace DestGammaSpace)
+{
+	if ( Format == DestFormat &&
+		( GammaSpace == DestGammaSpace || ! ERawImageFormat::GetFormatNeedsGammaSpace(Format)) )
+	{
+		// no action needed
+	}
+	else
+	{
+		FImage Temp;
+		CopyTo(Temp,DestFormat,DestGammaSpace);
+		//*this = MoveTemp(Temp); // or swap?
+		Swap(Temp);
+	}
+}
 
 void FImage::CopyTo(FImage& DestImage, ERawImageFormat::Type DestFormat, EGammaSpace DestGammaSpace) const
 {
+	// if gamma correction is done, it's always *TO* sRGB , not to Pow22
+	// so if Pow22 was requested, change to sRGB
+	// so that Float->int->Float roundtrips correctly
+	if ( DestGammaSpace == EGammaSpace::Pow22 && GammaSpace != EGammaSpace::Pow22 )
+	{
+		// @@!! todo: log warning and fix call sites that trigger this
+
+		DestGammaSpace = EGammaSpace::sRGB;
+	}
+
 	DestImage.SizeX = SizeX;
 	DestImage.SizeY = SizeY;
 	DestImage.NumSlices = NumSlices;
 	DestImage.Format = DestFormat;
 	DestImage.GammaSpace = DestGammaSpace;
 	InitImageStorage(DestImage);
-	CopyImage(*this, DestImage);
+	FImageCore::CopyImage(*this, DestImage);
 }
 
 void FImage::ResizeTo(FImage& DestImage, int32 DestSizeX, int32 DestSizeY, ERawImageFormat::Type DestFormat, EGammaSpace DestGammaSpace) const
@@ -706,6 +833,27 @@ void FImage::ResizeTo(FImage& DestImage, int32 DestSizeX, int32 DestSizeY, ERawI
 	}
 }
 
+IMAGECORE_API FImageView FImageView::GetSlice(int32 SliceIndex) const
+{
+	check( SliceIndex >= 0 && SliceIndex < NumSlices );
+	FImageView Ret = *this;
+	Ret.NumSlices = 1;
+	int64 BytesPerSlice = Ret.GetImageSizeBytes();
+	Ret.RawData = (void *)((uint8 *)RawData + SliceIndex * BytesPerSlice);
+	return Ret;
+}	
+
+IMAGECORE_API FImageView FImage::GetSlice(int32 SliceIndex) const
+{
+	check( SliceIndex >= 0 && SliceIndex < NumSlices );
+	FImageView Ret = *this;
+	Ret.NumSlices = 1;
+	int64 BytesPerSlice = Ret.GetImageSizeBytes();
+	check( RawData.Num() == NumSlices * BytesPerSlice );
+	Ret.RawData = (void *)(RawData.GetData() + SliceIndex * BytesPerSlice);
+	return Ret;
+}	
+
 void FImage::Linearize(uint8 SourceEncoding, FImage& DestImage) const
 {
 	TRACE_CPUPROFILER_EVENT_SCOPE(Texture.Linearize);
@@ -721,29 +869,57 @@ void FImage::Linearize(uint8 SourceEncoding, FImage& DestImage) const
 	check(SrcImage.SizeX == DestImage.SizeX);
 	check(SrcImage.SizeY == DestImage.SizeY);
 	check(SrcImage.NumSlices == DestImage.NumSlices);
+
 	//NOTE: SrcImage has GammaSpace in addition to SourceEncoding
-	// @todo Oodle : the way this function behaves with SrcImage GammaSpace and SourceEncoding is not well defined
+	// if SourceEncoding == None, the SrcImage Gamma is used
+	// else SourceEncoding overrides SrcImage Gamma
+	// note that SrcImage Gamma applies only on U8 formats (GetFormatNeedsGammaSpace)
+	// but SourceEncoding applies on ALL formats (eg. float->float)
 
 	UE::Color::EEncoding SourceEncodingType = static_cast<UE::Color::EEncoding>(SourceEncoding);
 
 	if (SourceEncodingType == UE::Color::EEncoding::None)
 	{
 		// note SrcImage.GammaSpace is decoded by CopyImage, is that intended?
-		CopyImage(SrcImage, DestImage);
+		FImageCore::CopyImage(SrcImage, DestImage);
 		return;
 	}
 	else if (SourceEncodingType >= UE::Color::EEncoding::Max)
 	{
-		// note SrcImage.GammaSpace is decoded by CopyImage, is that intended?
 		UE_LOG(LogImageCore, Warning, TEXT("Invalid encoding %d, falling back to linearization using CopyImage."), SourceEncoding);
-		CopyImage(SrcImage, DestImage);
+		FImageCore::CopyImage(SrcImage, DestImage);
 		return;
 	}
 
-	// @todo Oodle : common case of sRGB decoding, just use CopyImage, it's much faster
-	// SourceEncodingType == UE::Color::EEncoding::sRGB
-	// this function calls a func pointer per pixel
+	// fast path common case of sRGB decoding, just use CopyImage, it's much faster
+	// note that CopyImage only does Gamma on 8-bit formats (GetFormatNeedsGammaSpace)
+	//	but Linearize does it on ALL formats
+	if (SourceEncodingType == UE::Color::EEncoding::sRGB &&
+		ERawImageFormat::GetFormatNeedsGammaSpace(SrcImage.Format) )
+	{
+		FImageView SrcImageView(SrcImage);
+		SrcImageView.GammaSpace = EGammaSpace::sRGB;
+		FImageCore::CopyImage(SrcImageView, DestImage);
+		return;
+	}
+	else if (SourceEncodingType == UE::Color::EEncoding::Gamma22 &&
+		ERawImageFormat::GetFormatNeedsGammaSpace(SrcImage.Format) )
+	{
+		FImageView SrcImageView(SrcImage);
+		SrcImageView.GammaSpace = EGammaSpace::Pow22;
+		FImageCore::CopyImage(SrcImageView, DestImage);
+		return;
+	}
+	else if (SourceEncodingType == UE::Color::EEncoding::Linear )
+	{
+		FImageView SrcImageView(SrcImage);
+		SrcImageView.GammaSpace = EGammaSpace::Linear;
+		FImageCore::CopyImage(SrcImageView, DestImage);
+		return;
+	}
 
+	// slow case
+	// this function calls a func pointer per pixel
 
 	const int64 NumTexels = int64(SrcImage.SizeX) * SrcImage.SizeY * SrcImage.NumSlices;
 
@@ -862,7 +1038,25 @@ void FImage::Linearize(uint8 SourceEncoding, FImage& DestImage) const
 	}
 }
 
-int32 FImage::GetBytesPerPixel() const
+IMAGECORE_API const TCHAR * ERawImageFormat::GetName(Type Format)
+{
+	switch (Format)
+	{
+	case ERawImageFormat::G8:  return TEXT("G8");
+	case ERawImageFormat::G16: return TEXT("G16");
+	case ERawImageFormat::R16F: return TEXT("R16F");
+	case ERawImageFormat::BGRA8: return TEXT("BGRA8");
+	case ERawImageFormat::BGRE8: return TEXT("BGRE8");
+	case ERawImageFormat::RGBA16: return TEXT("RGBA16");
+	case ERawImageFormat::RGBA16F: return TEXT("RGBA16F");
+	case ERawImageFormat::RGBA32F: return TEXT("RGBA32F");
+	default:
+		check(0);
+		return TEXT("invalid");
+	}
+}
+
+IMAGECORE_API int32 ERawImageFormat::GetBytesPerPixel(Type Format)
 {
 	int32 BytesPerPixel = 0;
 	switch (Format)
@@ -889,6 +1083,78 @@ int32 FImage::GetBytesPerPixel() const
 	case ERawImageFormat::RGBA32F:
 		BytesPerPixel = 16;
 		break;
+
+	default:
+		check(0);
+		break;
 	}
 	return BytesPerPixel;
 }
+
+
+IMAGECORE_API bool ERawImageFormat::IsHDR(Type Format)
+{
+	return Format == RGBA16F || Format == RGBA32F || Format == R16F || Format == BGRE8;
+}
+
+IMAGECORE_API FLinearColor ERawImageFormat::GetOnePixelLinear(const void * PixelData,Type Format,bool bSRGB)
+{
+	switch(Format)
+	{
+	case G8:
+	{
+		uint8 Gray = ((const uint8 *)PixelData)[0];
+		FColor Color(Gray,Gray,Gray);
+		if ( bSRGB )
+			return FLinearColor::FromSRGBColor(Color);
+		else
+			return Color.ReinterpretAsLinear();
+	}
+	case BGRA8:
+	{
+		FColor Color = *((const FColor *)PixelData);
+		if ( bSRGB )
+			return FLinearColor::FromSRGBColor(Color);
+		else
+			return Color.ReinterpretAsLinear();
+	}
+	case BGRE8:
+	{
+		FColor Color = *((const FColor *)PixelData);
+		return Color.FromRGBE();
+	}
+	case RGBA16:
+	{
+		const uint16 * Samples = (const uint16 *)PixelData;
+		return FLinearColor(
+			Samples[0]/65535.f, 
+			Samples[1]/65535.f, 
+			Samples[2]/65535.f, 
+			Samples[3]/65535.f);
+	}
+	case RGBA16F:
+	{
+		FFloat16Color Color = *((const FFloat16Color *)PixelData);
+		return Color.GetFloats();
+	}
+	case RGBA32F:
+	{
+		return *((FLinearColor *)PixelData);
+	}
+	case G16:
+	{
+		const uint16 * Samples = (const uint16 *)PixelData;
+		float Gray = Samples[0]/65535.f;
+		return FLinearColor(Gray,Gray,Gray,1.f);
+	}
+	case R16F:
+	{
+		float Gray = FPlatformMath::LoadHalf((const uint16 *)PixelData);
+		return FLinearColor(Gray,Gray,Gray,1.f);
+	}
+	
+	default:
+		check(0);
+		return FLinearColor();
+	}
+}
diff --git a/Engine/Source/Runtime/ImageCore/Private/ImageCoreUtils.cpp b/Engine/Source/Runtime/ImageCore/Private/ImageCoreUtils.cpp
new file mode 100644
--- /dev/null
+++ b/Engine/Source/Runtime/ImageCore/Private/ImageCoreUtils.cpp
@@ -0,0 +1,95 @@
+// Copyright Epic Games, Inc. All Rights Reserved.
+
+#include "ImageCoreUtils.h"
+
+DEFINE_LOG_CATEGORY_STATIC(LogImageCoreUtils, Log, All);
+
+
+EPixelFormat FImageCoreUtils::GetPixelFormatForRawImageFormat(ERawImageFormat::Type InFormat, 
+							ERawImageFormat::Type * pOutEquivalentFormat)
+{
+	ERawImageFormat::Type OutEquivalentFormatDummy;
+	if ( pOutEquivalentFormat == nullptr )
+	{
+		pOutEquivalentFormat = &OutEquivalentFormatDummy;
+	}
+		
+	// if *pOutEquivalentFormat != InFormat , then conversion is needed
+	*pOutEquivalentFormat = InFormat;
+
+	// do not map to the very closest EPixelFormat
+	//	instead map to a close one that is actually usable as Texture
+
+	switch(InFormat)
+	{
+	case ERawImageFormat::G8:    
+		return PF_G8;
+	case ERawImageFormat::BGRA8: 
+		return PF_B8G8R8A8;
+	case ERawImageFormat::BGRE8: 
+		*pOutEquivalentFormat = ERawImageFormat::RGBA16F;
+		return PF_FloatRGBA;
+	case ERawImageFormat::RGBA16:  
+		return PF_R16G16B16A16_UNORM;
+	case ERawImageFormat::G16: 
+		return PF_G16;
+	case ERawImageFormat::RGBA16F: 
+		return PF_FloatRGBA;
+	case ERawImageFormat::RGBA32F: 
+		*pOutEquivalentFormat = ERawImageFormat::RGBA16F;
+		return PF_FloatRGBA;
+	case ERawImageFormat::R16F:	
+		return PF_R16F;
+	
+	default:
+		check(0);
+		return PF_Unknown;
+	}
+}
+
+	// ETextureSourceFormat and ERawImageFormat::Type are one-to-one :
+IMAGECORE_API ERawImageFormat::Type FImageCoreUtils::ConvertToRawImageFormat(ETextureSourceFormat Format)
+{
+	switch(Format)
+	{
+	case TSF_G8: return ERawImageFormat::G8;
+	case TSF_BGRA8: return ERawImageFormat::BGRA8;
+	case TSF_BGRE8: return ERawImageFormat::BGRE8;
+	case TSF_RGBA16: return ERawImageFormat::RGBA16;
+	case TSF_RGBA16F: return ERawImageFormat::RGBA16F;
+
+	case TSF_G16: return ERawImageFormat::G16;
+	case TSF_RGBA32F: return ERawImageFormat::RGBA32F;
+	case TSF_R16F: return ERawImageFormat::R16F;
+
+	// these are mapped to TSF_BGRA8/TSF_BGRE8 on load, so the runtime will never see them :
+	case TSF_RGBA8_DEPRECATED:
+	case TSF_RGBE8_DEPRECATED:
+		UE_LOG(LogImageCoreUtils, Warning,TEXT("Deprecated format in ConvertToRawImageFormat not supported."));
+		return ERawImageFormat::Invalid;
+
+	default:
+	case TSF_Invalid:
+		check(0);
+		return ERawImageFormat::Invalid;
+	}
+}
+
+IMAGECORE_API ETextureSourceFormat FImageCoreUtils::ConvertToTextureSourceFormat(ERawImageFormat::Type Format)
+{
+	switch(Format)
+	{
+	case ERawImageFormat::G8:    return TSF_G8;
+	case ERawImageFormat::BGRA8: return TSF_BGRA8;
+	case ERawImageFormat::BGRE8: return TSF_BGRE8;
+	case ERawImageFormat::RGBA16:  return TSF_RGBA16;
+	case ERawImageFormat::RGBA16F: return TSF_RGBA16F;
+	case ERawImageFormat::RGBA32F: return TSF_RGBA32F;
+	case ERawImageFormat::G16: return TSF_G16;
+	case ERawImageFormat::R16F:	return TSF_R16F;
+	
+	default:
+		check(0);
+		return TSF_Invalid;
+	}
+}
diff --git a/Engine/Source/Runtime/ImageCore/Public/ImageCore.h b/Engine/Source/Runtime/ImageCore/Public/ImageCore.h
--- a/Engine/Source/Runtime/ImageCore/Public/ImageCore.h
+++ b/Engine/Source/Runtime/ImageCore/Public/ImageCore.h
@@ -2,6 +2,23 @@
 
 #pragma once
 
+/****
+* 
+* ImageCore : low level pixel surface types
+* 
+* FImage owns an array of pixels
+* FImageView can point at an array of pixels (FImage or otherwise)
+* both have an FImageInfo
+* 
+* ERawImageFormat::Type is a pixel format that can be used in FImage (or TextureSource)
+* 
+* helpers to load/save/convert FImage are in Engine in FImageUtils
+* 
+* ImageCore does not use Engine, only Core
+* it can be used in standalone apps that don't have Engine
+* 
+*/
+
 
 /* Dependencies
  *****************************************************************************/
@@ -15,18 +32,62 @@ namespace ERawImageFormat
 {
 	/**
 	 * Enumerates supported raw image formats.
+	 * should map one-to-one with ETextureSourceFormat
 	 */
 	enum Type : uint8
 	{
 		G8,
-		BGRA8,
+		BGRA8, // FColor
 		BGRE8,
 		RGBA16,
-		RGBA16F,
-		RGBA32F,
+		RGBA16F, // FFloat16Color
+		RGBA32F, // FLinearColor
 		G16,
 		R16F,
+		Invalid = 0xFF
 	};
+	
+	// RGBA -> G8 takes from the  *R* channel
+	// G8 -> RGBA replicated to gray
+ 
+	IMAGECORE_API int32 GetBytesPerPixel(Type Format);
+	
+	IMAGECORE_API const TCHAR * GetName(Type Format);
+	
+	IMAGECORE_API bool IsHDR(Type Format);
+	
+	// Get one pixel of Format type from PixelData and return in Linear color
+	IMAGECORE_API FLinearColor GetOnePixelLinear(const void * PixelData,Type Format,bool bSRGB);	
+	
+	// G8 and BGRA8 are affected by Gamma
+	//	16/32 is NOT
+	FORCEINLINE bool GetFormatNeedsGammaSpace(Type Format)
+	{
+		if ( Format == G8 || Format == BGRA8 )
+		{
+			return true;
+		}
+		else
+		{
+			// these formats ignore GammaSpace setting
+			return false;
+		}
+	}
+
+	// when converting from pixel bags that don't store gamma to FImage
+	//	you can use this to guess the gamma they probably wanted
+	//	do not use when a real gamma flag is available!
+	FORCEINLINE EGammaSpace GetDefaultGammaSpace(Type Format)
+	{
+		if ( GetFormatNeedsGammaSpace(Format) )
+		{
+			return EGammaSpace::sRGB;
+		}
+		else
+		{
+			return EGammaSpace::Linear;
+		}
+	}
 };
 
 namespace UE {
@@ -35,36 +96,171 @@ namespace UE {
 	}
 }
 
-
 /**
- * Structure for raw image data.
- */
-struct FImage
+* 
+* FImageInfo describes a 2d pixel surface
+* it can be used by FImage or FImageView
+* 
+*/
+struct FImageInfo
 {
-	/** Raw image data. */
-	TArray64<uint8> RawData;
-
 	/** Width of the image. */
-	int32 SizeX;
+	int32 SizeX = 0;
 	
 	/** Height of the image. */
-	int32 SizeY;
+	int32 SizeY = 0;
 	
 	/** Number of image slices. */
-	int32 NumSlices;
+	int32 NumSlices = 0;
 	
 	/** Format in which the image is stored. */
-	ERawImageFormat::Type Format;
+	ERawImageFormat::Type Format = ERawImageFormat::BGRA8;
 	
 	/** The gamma space the image is stored in. */
-	EGammaSpace GammaSpace;
+	EGammaSpace GammaSpace = EGammaSpace::sRGB;
+
+	FImageInfo() { }
+	
+	FImageInfo(int32 InSizeX, int32 InSizeY, int32 InNumSlices, ERawImageFormat::Type InFormat, EGammaSpace InGammaSpace)
+		: SizeX(InSizeX)
+		, SizeY(InSizeY)
+		, NumSlices(InNumSlices)
+		, Format(InFormat)
+		, GammaSpace(InGammaSpace)
+	{
+	}
+		
+	FORCEINLINE bool IsGammaCorrected() const
+	{
+		return GammaSpace != EGammaSpace::Linear;
+	}
+
+	/**
+	 * Gets the number of bytes per pixel.
+	 *
+	 * @return Bytes per pixel.
+	 */
+	FORCEINLINE int32 GetBytesPerPixel() const
+	{
+		return ERawImageFormat::GetBytesPerPixel(Format);
+	}
+	
+	FORCEINLINE int64 GetNumPixels() const
+	{
+		return (int64) SizeX * SizeY * NumSlices;
+	}
+	
+	FORCEINLINE int64 GetImageSizeBytes() const
+	{
+		return GetNumPixels() * GetBytesPerPixel();
+	}
+	
+	FORCEINLINE int64 GetSliceNumPixels() const
+	{
+		return (int64) SizeX * SizeY;
+	}
+	
+	FORCEINLINE int64 GetSliceSizeBytes() const
+	{
+		return GetSliceNumPixels() * GetBytesPerPixel();
+	}
+	
+	FORCEINLINE int32 GetWidth()  const { return SizeX; }
+	FORCEINLINE int32 GetHeight() const { return SizeY; }
+	
+	FORCEINLINE EGammaSpace GetGammaSpace() const
+	{
+		// Gamma is ignored unless GetFormatNeedsGammaSpace, so make sure it is Linear
+		check( GetFormatNeedsGammaSpace(Format) || GammaSpace == EGammaSpace::Linear );
+
+		return GammaSpace;
+	}
+};
+
+/***
+* 
+* Image functions should take an FImageView as input and output to FImage
+* 
+* FImageView can point at an FImage or any pixel surface
+* 
+* FImage hold allocations.  FImageView is non-owning.  Make sure the memory pointed at is not freed.
+* 
+*/
+struct FImageView : public FImageInfo
+{
+	void * RawData = nullptr;
+
+	FImageView() { }
+	
+	FImageView(const FImageInfo & InInfo,void * InRawData) : FImageInfo(InInfo), RawData(InRawData)
+	{
+	}
+
+	// Make an FImageView from an array of FColor pixels as BGRA8
+	FImageView(const FColor * InColors,int32 InSizeX,int32 InSizeY,EGammaSpace InGammaSpace = EGammaSpace::sRGB)
+	{
+		RawData = (void *) InColors;
+		SizeX = InSizeX;
+		SizeY = InSizeY;
+		NumSlices = 1;
+		Format = ERawImageFormat::BGRA8;
+		GammaSpace = InGammaSpace; 
+	}
+	
+	// Make an FImageView from an array of FLinearColor pixels as RGBA32F
+	FImageView(const FLinearColor * InColors,int32 InSizeX,int32 InSizeY)
+	{
+		RawData = (void *) InColors;
+		SizeX = InSizeX;
+		SizeY = InSizeY;
+		NumSlices = 1;
+		Format = ERawImageFormat::RGBA32F;
+		GammaSpace = EGammaSpace::Linear;
+	}
+	
+	// Make an FImageView from an array of FLinearColor pixels as RGBA16F
+	FImageView(const FFloat16Color * InColors,int32 InSizeX,int32 InSizeY)
+	{
+		RawData = (void *) InColors;
+		SizeX = InSizeX;
+		SizeY = InSizeY;
+		NumSlices = 1;
+		Format = ERawImageFormat::RGBA16F;
+		GammaSpace = EGammaSpace::Linear;
+	}
+
+	FImageView(void * InData,int32 InSizeX, int32 InSizeY, int32 InNumSlices, ERawImageFormat::Type InFormat, EGammaSpace InGammaSpace)
+		: FImageInfo(InSizeX,InSizeY,InNumSlices,InFormat,InGammaSpace),
+		RawData(InData)
+	{
+	}
+	
+	FImageView(void * InData,int32 InSizeX, int32 InSizeY, ERawImageFormat::Type InFormat)
+		: FImageInfo(InSizeX,InSizeY,1,InFormat, ERawImageFormat::GetDefaultGammaSpace(InFormat) ),
+		RawData(InData)
+	{
+	}
+
+	// get an FImageView to one 2d slice
+	IMAGECORE_API FImageView GetSlice(int32 SliceIndex) const;
+};
+
+/**
+ * Structure for raw image data.
+ */
+struct FImage : public FImageInfo
+{
+	/** Raw image data. */
+	TArray64<uint8> RawData;
 
 public:
 
 	/**
 	 * Default constructor.
 	 */
-	FImage() { }
+	FImage()
+	{
+	}
 
 	/**
 	 * Creates and initializes a new image with the specified number of slices.
@@ -73,9 +269,15 @@ public:
 	 * @param InSizeY - The image height.
 	 * @param InNumSlices - The number of slices.
 	 * @param InFormat - The image format.
-	 * @param bInSRGB - Whether the color values are in SRGB format.
+	 * @param InGammaSpace - (optional) Image Gamma Space.  GetDefaultGammaSpace() if not provided.
 	 */
-	IMAGECORE_API FImage(int32 InSizeX, int32 InSizeY, int32 InNumSlices, ERawImageFormat::Type InFormat, EGammaSpace InGammaSpace = EGammaSpace::Linear);
+	IMAGECORE_API FImage(int32 InSizeX, int32 InSizeY, int32 InNumSlices, ERawImageFormat::Type InFormat, EGammaSpace InGammaSpace);
+
+	// @@!! danger? changed default GammaSpace from Linear to Guess
+	FORCEINLINE FImage(int32 InSizeX, int32 InSizeY, int32 InNumSlices, ERawImageFormat::Type InFormat) :
+		FImage(InSizeX,InSizeY,InNumSlices,InFormat,ERawImageFormat::GetDefaultGammaSpace(InFormat))
+	{
+	}
 
 	/**
 	 * Creates and initializes a new image with a single slice.
@@ -83,13 +285,38 @@ public:
 	 * @param InSizeX - The image width.
 	 * @param InSizeY - The image height.
 	 * @param InFormat - The image format.
-	 * @param bInSRGB - Whether the color values are in SRGB format.
+	 * @param InGammaSpace - (optional) Image Gamma Space.  GetDefaultGammaSpace() if not provided.
 	 */
-	IMAGECORE_API FImage(int32 InSizeX, int32 InSizeY, ERawImageFormat::Type InFormat, EGammaSpace InGammaSpace = EGammaSpace::Linear);
+	FORCEINLINE FImage(int32 InSizeX, int32 InSizeY, ERawImageFormat::Type InFormat, EGammaSpace InGammaSpace) :
+		FImage(InSizeX,InSizeY,1,InFormat,InGammaSpace)
+	{
+	}
+
+	// @@!! danger? changed default GammaSpace from Linear to Guess
+	FORCEINLINE FImage(int32 InSizeX, int32 InSizeY, ERawImageFormat::Type InFormat) :
+		FImage(InSizeX,InSizeY,1,InFormat,ERawImageFormat::GetDefaultGammaSpace(InFormat))
+	{
+	}
 
 
 public:
 
+	// Swap the contents of this FImage with another
+	IMAGECORE_API void Swap(FImage & Other);
+	
+	// implicit conversion to FImageView
+	// allows FImage to be passed to functions that take FImageView
+	// functions that read images should use "const FImageView &" as their argument type
+	FORCEINLINE operator FImageView () const
+	{
+		// copy the shared FImageInfo part :
+		const FImageInfo & Info = *this;
+		return FImageView( Info , (void *) RawData.GetData() );
+	}
+
+	// get an FImageView to one 2d slice
+	IMAGECORE_API FImageView GetSlice(int32 SliceIndex) const;
+	
 	/**
 	 * Copies the image to a destination image with the specified format.
 	 *
@@ -98,6 +325,16 @@ public:
 	 * @param DestSRGB - Whether the destination image is in SRGB format.
 	 */
 	IMAGECORE_API void CopyTo(FImage& DestImage, ERawImageFormat::Type DestFormat, EGammaSpace DestGammaSpace) const;
+	
+	// CopyTo same format
+	IMAGECORE_API void CopyTo(FImage& DestImage) const
+	{
+		CopyTo(DestImage,Format,GammaSpace);
+	}
+	
+	// in-place format change :
+	// does nothing if already in the desired format
+	IMAGECORE_API void ChangeFormat(ERawImageFormat::Type DestFormat, EGammaSpace DestGammaSpace);
 
 	/**
 	 * Copies and resizes the image to a destination image with the specified size and format.
@@ -117,9 +354,16 @@ public:
 	 *
 	 * @param SourceEncoding - Opaque source encoding (matching UE::Color::EEncoding).
 	 * @param DestImage - The destination image.
+	 * 
+	 * SourceEncoding = None (0) to Linearize using the GammaSpace in this FImage
 	 */
 	IMAGECORE_API void Linearize(uint8 SourceEncoding, FImage& DestImage) const;
 
+	FORCEINLINE void Linearize(FImage& DestImage) const
+	{
+		Linearize(0,DestImage);
+	}
+
 	/**
 	 * Apply a color space transformation from the source chromaticities to the engine's working color space.
 	 *
@@ -132,15 +376,9 @@ public:
 	 */
 	IMAGECORE_API void TransformToWorkingColorSpace(const FVector2d& SourceRedChromaticity, const FVector2d& SourceGreenChromaticity, const FVector2d& SourceBlueChromaticity, const FVector2d& SourceWhiteChromaticity, UE::Color::EChromaticAdaptationMethod Method, double EqualityTolerance = 1.e-7);
 
-	/**
-	 * Gets the number of bytes per pixel.
-	 *
-	 * @return Bytes per pixel.
-	 */
-	IMAGECORE_API int32 GetBytesPerPixel() const;
-
 	/**
 	 * Initializes this image with the specified number of slices.
+	 * Allocates RawData.
 	 *
 	 * @param InSizeX - The image width.
 	 * @param InSizeY - The image height.
@@ -148,18 +386,36 @@ public:
 	 * @param InFormat - The image format.
 	 * @param bInSRGB - Whether the color values are in SRGB format.
 	 */
-	IMAGECORE_API void Init(int32 InSizeX, int32 InSizeY, int32 InNumSlices, ERawImageFormat::Type InFormat, EGammaSpace InGammaSpace = EGammaSpace::Linear);
+	IMAGECORE_API void Init(int32 InSizeX, int32 InSizeY, int32 InNumSlices, ERawImageFormat::Type InFormat, EGammaSpace InGammaSpace);
+	
+	FORCEINLINE void Init(int32 InSizeX, int32 InSizeY, int32 InNumSlices, ERawImageFormat::Type InFormat)
+	{
+		Init(InSizeX,InSizeY,InNumSlices,InFormat,ERawImageFormat::GetDefaultGammaSpace(InFormat));
+	}
 
 	/**
 	 * Initializes this image with a single slice.
+	 * Allocates RawData.
 	 *
 	 * @param InSizeX - The image width.
 	 * @param InSizeY - The image height.
 	 * @param InFormat - The image format.
 	 * @param bInSRGB - Whether the color values are in SRGB format.
 	 */
-	IMAGECORE_API void Init(int32 InSizeX, int32 InSizeY, ERawImageFormat::Type InFormat, EGammaSpace InGammaSpace = EGammaSpace::Linear);
+	IMAGECORE_API void Init(int32 InSizeX, int32 InSizeY, ERawImageFormat::Type InFormat, EGammaSpace InGammaSpace);
+	
+	FORCEINLINE void Init(int32 InSizeX, int32 InSizeY, ERawImageFormat::Type InFormat)
+	{
+		Init(InSizeX,InSizeY,InFormat,ERawImageFormat::GetDefaultGammaSpace(InFormat));
+	}
 
+	/**
+	 * Initializes this image with specified Info.
+	 * Allocates RawData.
+	 *
+	 * @param Info - The image description.  Can also pass an FImage or FImageView to this function.
+	 */
+	IMAGECORE_API void Init(const FImageInfo & Info);
 
 public:
 
@@ -262,14 +518,53 @@ public:
 		check(Format == ERawImageFormat::R16F);
 		return { (const class FFloat16*)RawData.GetData(), int64(RawData.Num() / sizeof(FFloat16)) };
 	}
-
-	FORCEINLINE bool IsGammaCorrected() const
-	{
-		return GammaSpace != EGammaSpace::Linear;
-	}
 };
 
+/* Functions
+ *****************************************************************************/
 
 IMAGECORE_API int32 ImageParallelForComputeNumJobsForPixels(int64 & OutNumPixelsPerJob,int64 NumPixels);
 
-IMAGECORE_API int32 ImageParallelForComputeNumJobsForRows(int32 & OutNumItemsPerJob,int32 SizeX,int32 SizeY);
\ No newline at end of file
+IMAGECORE_API int32 ImageParallelForComputeNumJobsForRows(int32 & OutNumItemsPerJob,int32 SizeX,int32 SizeY);
+
+namespace FImageCore
+{
+
+/**
+ * Copy FImageView data, converting pixels if necessary.
+ * Sizes must match. Dest must already be allocated.
+ * For a function that allocates Dest, use CopyTo().
+ * 
+ * Because FImage implicitly convers to FImageView, these functions can be used on FImage types as well.
+ *
+ * @param SrcImage - The source image to copy from.
+ * @param DestImage - The destination image to copy to. (the FImageView is const but what it points at is not)
+ */
+IMAGECORE_API void CopyImage(const FImageView & SrcImage,const FImageView & DestImage);
+
+/**
+ * Copy Image, swapping RB, SrcImage must be BGRA8 or RGBA16.
+ * Sizes must match. Dest must already be allocated.
+ * Src == Dest is okay for in-place transpose.
+ *
+ * @param SrcImage - The source image to copy from.
+ * @param DestImage - The destination image to copy to. (the FImageView is const but what it points at is not)
+ */
+IMAGECORE_API void CopyImageRGBABGRA(const FImageView & SrcImage,const FImageView & DestImage);
+
+/**
+* Swap RB channels on Image.  Image must be BGRA8 or RGBA16.
+* 
+* @param  Image is modified (the FImageView is const but what it points at is not)
+*/
+IMAGECORE_API void TransposeImageRGBABGRA(const FImageView & Image);
+
+};
+
+
+// non-namespaced CopyImage for backwards compatibility
+UE_DEPRECATED(5.1,"Use FImageCore::CopyImage")
+FORCEINLINE void CopyImage(const FImageView & SrcImage,const FImageView & DestImage)
+{
+	return FImageCore::CopyImage(SrcImage,DestImage);
+}
\ No newline at end of file
diff --git a/Engine/Source/Runtime/ImageCore/Public/ImageCoreUtils.h b/Engine/Source/Runtime/ImageCore/Public/ImageCoreUtils.h
new file mode 100644
--- /dev/null
+++ b/Engine/Source/Runtime/ImageCore/Public/ImageCoreUtils.h
@@ -0,0 +1,38 @@
+// Copyright Epic Games, Inc. All Rights Reserved.
+
+/*=============================================================================
+ImageCoreUtils.h: Image utility functions.
+=============================================================================*/
+
+#pragma once
+
+#include "CoreMinimal.h"
+#include "ImageCore.h"
+#include "Engine/TextureDefines.h"
+#include "PixelFormat.h"
+
+// ImageCore Module does not use Engine, but it can see TextureDefines
+// do not use Engine functions here
+// 
+// keep the ImageCore.h includes as minimal as possible
+// put any functions that need more headers here
+
+namespace FImageCoreUtils
+{
+
+	// ETextureSourceFormat and ERawImageFormat::Type are one-to-one :
+	IMAGECORE_API ERawImageFormat::Type ConvertToRawImageFormat(ETextureSourceFormat Format);
+	IMAGECORE_API ETextureSourceFormat ConvertToTextureSourceFormat(ERawImageFormat::Type Format);
+
+	// EPixelFormat is the graphics texture pixel format
+	// it is a much larger superset of ERawImageFormat
+	// GetPixelFormatForRawImageFormat does not map to the very closest EPixelFormat
+	//	instead map to a close one that is actually usable as Texture
+	// if *pOutEquivalentFormat != InFormat , then conversion is needed
+	//  and conversion can be done using CopyImage to OutEquivalentFormat 
+	IMAGECORE_API EPixelFormat GetPixelFormatForRawImageFormat(ERawImageFormat::Type InFormat, 
+							ERawImageFormat::Type * pOutEquivalentFormat=nullptr);
+
+
+
+};
diff --git a/Engine/Source/Runtime/ImageWrapper/ImageWrapper.Build.cs b/Engine/Source/Runtime/ImageWrapper/ImageWrapper.Build.cs
--- a/Engine/Source/Runtime/ImageWrapper/ImageWrapper.Build.cs
+++ b/Engine/Source/Runtime/ImageWrapper/ImageWrapper.Build.cs
@@ -12,11 +12,18 @@ public class ImageWrapper : ModuleRules
 				"Runtime/ImageWrapper/Private/Formats",
 			});
 
+		PublicIncludePathModuleNames.AddRange(new string[] {
+			"ImageCore"
+			}
+		);
+
 		PublicDefinitions.Add("WITH_UNREALPNG=1");
 		PublicDefinitions.Add("WITH_UNREALJPEG=1");
 
 		PrivateDependencyModuleNames.Add("Core");
+
 		PublicDependencyModuleNames.Add("LibTiff");
+		PublicDependencyModuleNames.Add("ImageCore");
 
 		AddEngineThirdPartyPrivateStaticDependencies(Target,
 			"zlib",
diff --git a/Engine/Source/Runtime/ImageWrapper/Private/Formats/BmpImageWrapper.cpp b/Engine/Source/Runtime/ImageWrapper/Private/Formats/BmpImageWrapper.cpp
--- a/Engine/Source/Runtime/ImageWrapper/Private/Formats/BmpImageWrapper.cpp
+++ b/Engine/Source/Runtime/ImageWrapper/Private/Formats/BmpImageWrapper.cpp
@@ -18,15 +18,14 @@ FBmpImageWrapper::FBmpImageWrapper(bool bInHasHeader, bool bInHalfHeight)
 {
 }
 
-
-void FBmpImageWrapper::Compress(int32 Quality)
-{
-	checkf(false, TEXT("BMP compression not supported"));
-}
-
-
 void FBmpImageWrapper::Uncompress(const ERGBFormat InFormat, const int32 InBitDepth)
 {
+	RawData.Empty();
+	if ( CompressedData.IsEmpty() )
+	{
+		return;
+	}
+
 	const uint8* Buffer = CompressedData.GetData();
 
 	if (!bHasHeader || ((CompressedData.Num()>=sizeof(FBitmapFileHeader)+sizeof(FBitmapInfoHeader)) && Buffer[0]=='B' && Buffer[1]=='M'))
@@ -38,6 +37,11 @@ void FBmpImageWrapper::Uncompress(const ERGBFormat InFormat, const int32 InBitDe
 
 void FBmpImageWrapper::UncompressBMPData(const ERGBFormat InFormat, const int32 InBitDepth)
 {
+	// always writes BGRA8 :
+	check( InFormat == ERGBFormat::BGRA );
+	check( InBitDepth == 8 );
+	check( ! CompressedData.IsEmpty() );
+
 	const uint8* Buffer = CompressedData.GetData();
 	const FBitmapInfoHeader* bmhdr = nullptr;
 	const uint8* Bits = nullptr;
@@ -241,12 +245,20 @@ bool FBmpImageWrapper::SetCompressed(const void* InCompressedData, int64 InCompr
 {
 	bool bResult = FImageWrapperBase::SetCompressed(InCompressedData, InCompressedSize);
 
-	return bResult && (bHasHeader ? LoadBMPHeader() : LoadBMPInfoHeader());	// Fetch the variables from the header info
+	bResult = bResult && (bHasHeader ? LoadBMPHeader() : LoadBMPInfoHeader());	// Fetch the variables from the header info
+
+	if ( ! bResult )
+	{
+		CompressedData.Reset();
+	}
+
+	return bResult;
 }
 
 
 bool FBmpImageWrapper::LoadBMPHeader()
 {
+	//note: not Endian correct
 	const FBitmapInfoHeader* bmhdr = (FBitmapInfoHeader *)(CompressedData.GetData() + sizeof(FBitmapFileHeader));
 	const FBitmapFileHeader* bmf   = (FBitmapFileHeader *)(CompressedData.GetData() + 0);
 	if ((CompressedData.Num() >= sizeof(FBitmapFileHeader) + sizeof(FBitmapInfoHeader)) && CompressedData.GetData()[0] == 'B' && CompressedData.GetData()[1] == 'M')
@@ -263,7 +275,11 @@ bool FBmpImageWrapper::LoadBMPHeader()
 			Width = bmhdr->biWidth;
 			Height = FMath::Abs(bmhdr->biHeight);
 			Format = ERGBFormat::BGRA;
-			BitDepth = bmhdr->biBitCount;
+			//  YUCK
+			//  BmpImageWrapper was reporting BitDepth *total* not per-channel (eg. 24)
+			//  some legacy code knew that and expected to see the bad values, beware
+			//BitDepth = bmhdr->biBitCount;
+			BitDepth = 8;
 
 			return true;
 		}
@@ -284,6 +300,7 @@ bool FBmpImageWrapper::LoadBMPHeader()
 
 bool FBmpImageWrapper::LoadBMPInfoHeader()
 {
+	//note: not Endian correct
 	const FBitmapInfoHeader* bmhdr = (FBitmapInfoHeader *)CompressedData.GetData();
 
 	if (bmhdr->biCompression != BCBI_RGB && bmhdr->biCompression != BCBI_BITFIELDS)
@@ -298,7 +315,9 @@ bool FBmpImageWrapper::LoadBMPInfoHeader()
 		Width = bmhdr->biWidth;
 		Height = FMath::Abs(bmhdr->biHeight);
 		Format = ERGBFormat::BGRA;
-		BitDepth = bmhdr->biBitCount;
+		//  BmpImageWrapper was reporting BitDepth *total* not per-channel (eg. 24)
+		//BitDepth = bmhdr->biBitCount;
+		BitDepth = 8;
 
 		return true;
 	}
@@ -314,3 +333,244 @@ bool FBmpImageWrapper::LoadBMPInfoHeader()
 
 	return false;
 }
+
+bool FBmpImageWrapper::CanSetRawFormat(const ERGBFormat InFormat, const int32 InBitDepth) const
+{
+	return ((InFormat == ERGBFormat::BGRA || InFormat == ERGBFormat::Gray) && InBitDepth == 8);
+}
+
+ERawImageFormat::Type FBmpImageWrapper::GetSupportedRawFormat(const ERawImageFormat::Type InFormat) const
+{
+	switch(InFormat)
+	{
+	case ERawImageFormat::G8:
+	case ERawImageFormat::BGRA8:
+		return InFormat; // directly supported
+	case ERawImageFormat::G16:
+		return ERawImageFormat::G8; // needs conversion
+	case ERawImageFormat::BGRE8:
+	case ERawImageFormat::RGBA16:
+	case ERawImageFormat::RGBA16F:
+	case ERawImageFormat::RGBA32F:
+	case ERawImageFormat::R16F:
+		return ERawImageFormat::BGRA8; // needs conversion
+	default:
+		check(0);
+		return ERawImageFormat::BGRA8;
+	};
+}
+
+
+void FBmpImageWrapper::Compress(int32 Quality)
+{
+	check( RawFormat == ERGBFormat::BGRA || RawFormat == ERGBFormat::Gray );
+	check( RawBitDepth == 8 );
+
+	// write 8,24, or 32 bit bmp
+
+	int64 NumPixels = Width*Height;
+	int64 RawDataSize = RawData.Num();
+
+	int RawBytesPerPel = (RawFormat == ERGBFormat::BGRA) ? 4 : 1;
+
+	check( RawDataSize == NumPixels*RawBytesPerPel );
+
+	int OutputBytesPerPel = RawBytesPerPel;
+
+	if ( RawBytesPerPel == 4 )
+	{
+		// scan for A to choose 24 bit output
+
+		const FColor * RawColors = (const FColor *)RawData.GetData();
+		bool bHasAnyAlpha = false;
+		for(int64 i=0;i<NumPixels;i++)
+		{
+			if ( RawColors[i].A != 255 )
+			{
+				bHasAnyAlpha = true;
+				break;
+			}
+		}
+
+		if ( ! bHasAnyAlpha )
+		{
+			OutputBytesPerPel = 3;
+		}
+	}
+
+	bool bWritePal = (RawBytesPerPel == 1);
+
+	int OutputRowBytes = (Width*OutputBytesPerPel + 3)&(~3);
+	int OutputPalBytes = bWritePal ? 1024 : 0;
+	int64 OutputImageBytes = OutputRowBytes * Height;
+
+	CompressedData.Empty( OutputImageBytes + OutputPalBytes + 1024 );
+
+	// scope to write headers:
+	{
+		// copied from FFileHelper::CreateBitmap
+		// 
+		// Types.
+		#if PLATFORM_SUPPORTS_PRAGMA_PACK
+			#pragma pack (push,1)
+		#endif
+		struct BITMAPFILEHEADER
+		{
+			uint16 bfType GCC_PACK(1);
+			uint32 bfSize GCC_PACK(1);
+			uint16 bfReserved1 GCC_PACK(1); 
+			uint16 bfReserved2 GCC_PACK(1);
+			uint32 bfOffBits GCC_PACK(1);
+		} FH = { };
+		struct BITMAPINFOHEADER
+		{
+			uint32 biSize GCC_PACK(1); 
+			int32  biWidth GCC_PACK(1);
+			int32  biHeight GCC_PACK(1);
+			uint16 biPlanes GCC_PACK(1);
+			uint16 biBitCount GCC_PACK(1);
+			uint32 biCompression GCC_PACK(1);
+			uint32 biSizeImage GCC_PACK(1);
+			int32  biXPelsPerMeter GCC_PACK(1); 
+			int32  biYPelsPerMeter GCC_PACK(1);
+			uint32 biClrUsed GCC_PACK(1);
+			uint32 biClrImportant GCC_PACK(1); 
+		} IH = { };
+		struct BITMAPV4HEADER
+		{
+			uint32 bV4RedMask GCC_PACK(1);
+			uint32 bV4GreenMask GCC_PACK(1);
+			uint32 bV4BlueMask GCC_PACK(1);
+			uint32 bV4AlphaMask GCC_PACK(1);
+			uint32 bV4CSType GCC_PACK(1);
+			uint32 bV4EndpointR[3] GCC_PACK(1);
+			uint32 bV4EndpointG[3] GCC_PACK(1);
+			uint32 bV4EndpointB[3] GCC_PACK(1);
+			uint32 bV4GammaRed GCC_PACK(1);
+			uint32 bV4GammaGreen GCC_PACK(1);
+			uint32 bV4GammaBlue GCC_PACK(1);
+		} IHV4 = { };
+		#if PLATFORM_SUPPORTS_PRAGMA_PACK
+			#pragma pack (pop)
+		#endif
+
+		bool bInWriteAlpha = ( OutputBytesPerPel == 4 );
+		uint32 InfoHeaderSize = sizeof(BITMAPINFOHEADER) + (bInWriteAlpha ? sizeof(BITMAPV4HEADER) : 0);
+
+		// File header.
+		FH.bfType       		= INTEL_ORDER16((uint16) ('B' + 256*'M'));
+		FH.bfSize       		= INTEL_ORDER32((uint32) (sizeof(BITMAPFILEHEADER) + InfoHeaderSize + OutputImageBytes + OutputPalBytes));
+		FH.bfOffBits    		= INTEL_ORDER32((uint32) (sizeof(BITMAPFILEHEADER) + InfoHeaderSize));
+		CompressedData.Append( (const uint8 *) &FH, sizeof(FH) );
+
+		// Info header.
+		IH.biSize               = INTEL_ORDER32((uint32) InfoHeaderSize);
+		IH.biWidth              = INTEL_ORDER32((uint32) Width);
+		IH.biHeight             = INTEL_ORDER32((uint32) Height);
+		IH.biPlanes             = INTEL_ORDER16((uint16) 1);
+		IH.biBitCount           = INTEL_ORDER16((uint16) OutputBytesPerPel * 8);
+		if(bInWriteAlpha)
+		{
+			IH.biCompression    = INTEL_ORDER32((uint32) 3); //BI_BITFIELDS
+		}
+		else
+		{
+			IH.biCompression    = INTEL_ORDER32((uint32) 0); //BI_RGB
+		}
+		IH.biSizeImage          = INTEL_ORDER32((uint32) OutputImageBytes);
+		if ( bWritePal )
+		{
+			IH.biClrUsed            = INTEL_ORDER32((uint32) 256);
+			IH.biClrImportant       = INTEL_ORDER32((uint32) 256);
+		}
+		CompressedData.Append( (const uint8 *) &IH, sizeof(IH) );
+
+		// If we're writing alpha, we need to write the extra portion of the V4 header
+		if (bInWriteAlpha)
+		{
+			IHV4.bV4RedMask     = INTEL_ORDER32((uint32) 0x00ff0000);
+			IHV4.bV4GreenMask   = INTEL_ORDER32((uint32) 0x0000ff00);
+			IHV4.bV4BlueMask    = INTEL_ORDER32((uint32) 0x000000ff);
+			IHV4.bV4AlphaMask   = INTEL_ORDER32((uint32) 0xff000000);
+			IHV4.bV4CSType      = INTEL_ORDER32((uint32) 'Win ');
+			CompressedData.Append( (const uint8 *) &IHV4, sizeof(IHV4) );
+		}
+	}
+
+	if ( bWritePal )
+	{
+		// write palette for G8 :
+		FColor Palette[256];
+
+		for(int i=0;i<256;i++)
+		{
+			Palette[i] = FColor(i,i,i,255);
+		}
+		
+		check( sizeof(Palette) == OutputPalBytes );
+		CompressedData.Append( (const uint8 *)Palette,1024 );
+	}
+
+	int64 HeaderBytes = CompressedData.Num();
+	CompressedData.SetNum( HeaderBytes + OutputImageBytes );
+	uint8 * PayloadPtr = CompressedData.GetData() + HeaderBytes;
+	
+	int OutputRowPadBytes = OutputRowBytes - Width*OutputBytesPerPel;
+	check( OutputRowPadBytes < 4 );
+	uint8 Zero[4] = { };
+
+	// write rows :
+	switch(OutputBytesPerPel)
+	{
+	case 1:
+	{
+		const uint8 * RawPtr = RawData.GetData();
+		for(int y=Height-1;y>=0;y--)
+		{
+			memcpy(PayloadPtr,RawPtr + y * Width,Width);
+			PayloadPtr += Width;
+			memcpy(PayloadPtr,Zero,OutputRowPadBytes);
+			PayloadPtr += OutputRowPadBytes;
+		}
+		break;
+	}
+
+	case 3:
+	{
+		const FColor * RawColors = (const FColor *) RawData.GetData();
+		for(int y=Height-1;y>=0;y--)
+		{
+			const FColor * RawRow = RawColors + y * Width;
+			for(int x=0;x<Width;x++)
+			{
+				*PayloadPtr++ = RawRow[x].B;
+				*PayloadPtr++ = RawRow[x].G;
+				*PayloadPtr++ = RawRow[x].R;
+			}
+			memcpy(PayloadPtr,Zero,OutputRowPadBytes);
+			PayloadPtr += OutputRowPadBytes;
+		}
+		break;
+	}
+
+	case 4:
+	{
+		check( OutputRowBytes == Width * 4 );
+		check( OutputRowPadBytes == 0 );
+
+		const uint8 * RawPtr = RawData.GetData();
+		for(int y=Height-1;y>=0;y--)
+		{
+			memcpy(PayloadPtr,RawPtr + y * OutputRowBytes,OutputRowBytes);
+			PayloadPtr += OutputRowBytes;
+		}
+		break;
+	}
+
+	default:
+		check(0);
+		break;
+	}
+	
+	check( PayloadPtr == CompressedData.GetData() + CompressedData.Num() );
+}
\ No newline at end of file
diff --git a/Engine/Source/Runtime/ImageWrapper/Private/Formats/BmpImageWrapper.h b/Engine/Source/Runtime/ImageWrapper/Private/Formats/BmpImageWrapper.h
--- a/Engine/Source/Runtime/ImageWrapper/Private/Formats/BmpImageWrapper.h
+++ b/Engine/Source/Runtime/ImageWrapper/Private/Formats/BmpImageWrapper.h
@@ -43,6 +43,9 @@ public:
 	virtual void Compress(int32 Quality) override;
 	virtual void Uncompress(const ERGBFormat InFormat, int32 InBitDepth) override;
 	virtual bool SetCompressed(const void* InCompressedData, int64 InCompressedSize) override;
+	
+	virtual bool CanSetRawFormat(const ERGBFormat InFormat, const int32 InBitDepth) const override;
+	virtual ERawImageFormat::Type GetSupportedRawFormat(const ERawImageFormat::Type InFormat) const override;
 
 private:
 
diff --git a/Engine/Source/Runtime/ImageWrapper/Private/Formats/ExrImageWrapper.cpp b/Engine/Source/Runtime/ImageWrapper/Private/Formats/ExrImageWrapper.cpp
--- a/Engine/Source/Runtime/ImageWrapper/Private/Formats/ExrImageWrapper.cpp
+++ b/Engine/Source/Runtime/ImageWrapper/Private/Formats/ExrImageWrapper.cpp
@@ -183,6 +183,34 @@ int32 GetChannelNames(ERGBFormat InRGBFormat, const char* const*& OutChannelName
 	return ChannelCount;
 }
 
+
+bool FExrImageWrapper::CanSetRawFormat(const ERGBFormat InFormat, const int32 InBitDepth) const
+{
+	return (InFormat == ERGBFormat::RGBAF || InFormat == ERGBFormat::GrayF) && (InBitDepth == 16 || InBitDepth == 32);
+}
+
+ERawImageFormat::Type FExrImageWrapper::GetSupportedRawFormat(const ERawImageFormat::Type InFormat) const
+{
+	switch(InFormat)
+	{
+	case ERawImageFormat::RGBA16F:
+	case ERawImageFormat::RGBA32F:
+	case ERawImageFormat::R16F:
+		return InFormat; // directly supported
+	case ERawImageFormat::G8:
+		return ERawImageFormat::R16F; // needs conversion
+	case ERawImageFormat::BGRA8:
+	case ERawImageFormat::BGRE8:
+		return ERawImageFormat::RGBA16F; // needs conversion
+	case ERawImageFormat::G16:
+	case ERawImageFormat::RGBA16:
+		return ERawImageFormat::RGBA32F; // needs conversion
+	default:
+		check(0);
+		return ERawImageFormat::BGRA8;
+	}
+}
+
 bool FExrImageWrapper::SetRaw(const void* InRawData, int64 InRawSize, const int32 InWidth, const int32 InHeight, const ERGBFormat InFormat, const int32 InBitDepth, const int32 InBytesPerRow)
 {
 	check(InRawData);
@@ -191,6 +219,11 @@ bool FExrImageWrapper::SetRaw(const void* InRawData, int64 InRawSize, const int3
 	check(InHeight > 0);
 	check(InBytesPerRow >= 0);
 
+	// FExrImageWrapper used to take RGBA 8-bit input
+	//	 and write it linearly
+	// the new image path now requires you to convert to float before coming in here
+	// so U8 will be converted to float *with* gamma correction
+
 	switch (InBitDepth)
 	{
 	case 8:
@@ -294,6 +327,7 @@ void FExrImageWrapper::Compress(int32 Quality)
 	if (RawBitDepth == 8)
 	{
 		// uint8 channels are linearly converted into FFloat16 channels.
+		// note: NO GAMMA CORRECTION
 		ConvertedRawData.SetNumUninitialized(sizeof(FFloat16) * RawData.Num());
 		FFloat16* Output = reinterpret_cast<FFloat16*>(ConvertedRawData.GetData());
 		for (int64 i = 0; i < RawData.Num(); ++i)
@@ -404,7 +438,7 @@ void FExrImageWrapper::Uncompress(const ERGBFormat InFormat, const int32 InBitDe
 
 	const char* const* ChannelNames;
 	int32 ChannelCount = GetChannelNames(InFormat, ChannelNames);
-	check(ChannelCount);
+	check(ChannelCount == 1 || ChannelCount == 4);
 
 	TArray<TArray64<uint8>> ChannelData;
 	ChannelData.SetNum(ChannelCount);
diff --git a/Engine/Source/Runtime/ImageWrapper/Private/Formats/ExrImageWrapper.h b/Engine/Source/Runtime/ImageWrapper/Private/Formats/ExrImageWrapper.h
--- a/Engine/Source/Runtime/ImageWrapper/Private/Formats/ExrImageWrapper.h
+++ b/Engine/Source/Runtime/ImageWrapper/Private/Formats/ExrImageWrapper.h
@@ -43,6 +43,9 @@ public:
 	virtual bool SetCompressed(const void* InCompressedData, int64 InCompressedSize) override;
 	virtual void Compress(int32 Quality) override;
 	virtual void Uncompress(const ERGBFormat InFormat, int32 InBitDepth) override;
+	
+	virtual bool CanSetRawFormat(const ERGBFormat InFormat, const int32 InBitDepth) const override;
+	virtual ERawImageFormat::Type GetSupportedRawFormat(const ERawImageFormat::Type InFormat) const override;
 };
 
 #endif // WITH_UNREALEXR
diff --git a/Engine/Source/Runtime/ImageWrapper/Private/Formats/HdrImageWrapper.cpp b/Engine/Source/Runtime/ImageWrapper/Private/Formats/HdrImageWrapper.cpp
--- a/Engine/Source/Runtime/ImageWrapper/Private/Formats/HdrImageWrapper.cpp
+++ b/Engine/Source/Runtime/ImageWrapper/Private/Formats/HdrImageWrapper.cpp
@@ -1,6 +1,7 @@
 // Copyright Epic Games, Inc. All Rights Reserved.
 
 #include "Formats/HdrImageWrapper.h"
+#include "ImageWrapperPrivate.h"
 
 #define LOCTEXT_NAMESPACE "HdrImageWrapper"
 
@@ -20,6 +21,7 @@ namespace UE::ImageWrapper::Private::HdrImageWrapper
 
 bool FHdrImageWrapper::SetCompressedFromView(TArrayView64<const uint8> Data)
 {
+	// CompressedData is just a View, does not take a copy
 	CompressedData = Data;
 
 	if (CompressedData.Num() < 11)
@@ -34,7 +36,8 @@ bool FHdrImageWrapper::SetCompressedFromView(TArrayView64<const uint8> Data)
 
 	GetHeaderLine(FileDataPtr, Line);
 
-	if (FCStringAnsi::Strcmp(Line, "#?RADIANCE"))
+	if (FCStringAnsi::Strcmp(Line, "#?RADIANCE") != 0 &&
+		FCStringAnsi::Strcmp(Line, "#?RGBE") != 0)
 	{
 		FreeCompressedData();
 		return false;
@@ -81,6 +84,7 @@ bool FHdrImageWrapper::SetCompressedFromView(TArrayView64<const uint8> Data)
 
 bool FHdrImageWrapper::SetCompressed(const void* InCompressedData, int64 InCompressedSize)
 {
+	// takes copy
 	CompressedDataHolder.Reset(InCompressedSize);
 	CompressedDataHolder.AddUninitialized(InCompressedSize);
 	FMemory::Memcpy(CompressedDataHolder.GetData(), InCompressedData, InCompressedSize);
@@ -88,24 +92,73 @@ bool FHdrImageWrapper::SetCompressed(const void* InCompressedData, int64 InCompr
 	return SetCompressedFromView(MakeArrayView(CompressedDataHolder));
 }
 
-bool FHdrImageWrapper::SetRaw(const void* InRawData, int64 InRawSize, const int32 InWidth, const int32 InHeight, const ERGBFormat InFormat, const int32 InBitDepth, const int32 InBytesPerRow)
+// CanSetRawFormat returns true if SetRaw will accept this format
+bool FHdrImageWrapper::CanSetRawFormat(const ERGBFormat InFormat, const int32 InBitDepth) const
 {
-	// Move the code that export hdr here
-	unimplemented();
-	return false;
+	return InFormat == ERGBFormat::BGRE && InBitDepth == 8;
 }
 
-bool FHdrImageWrapper::SetAnimationInfo(int32 InNumFrames, int32 InFramerate)
+// returns InFormat if supported, else maps to something supported
+ERawImageFormat::Type FHdrImageWrapper::GetSupportedRawFormat(const ERawImageFormat::Type InFormat) const
 {
-	unimplemented();
-	return false;
+	// only writes one format :
+	return ERawImageFormat::BGRE8;
+}
+
+bool FHdrImageWrapper::SetRaw(const void* InRawData, int64 InRawSize, const int32 InWidth, const int32 InHeight, const ERGBFormat InFormat, const int32 InBitDepth, const int32 InBytesPerRow)
+{
+	if ( ! CanSetRawFormat(InFormat,InBitDepth) )
+	{
+		UE_LOG(LogImageWrapper, Warning, TEXT("ImageWrapper unsupported format; check CanSetRawFormat; %d x %d"), (int)InFormat,InBitDepth);
+		return false;
+	}
+
+	RawDataHolder.Empty(InRawSize);
+	RawDataHolder.Append((const uint8 *)InRawData,InRawSize);
+	Width = InWidth;
+	Height = InHeight;
+
+	check( InFormat == ERGBFormat::BGRE );
+	check( InBitDepth == 8 );
+
+	return true;
 }
 
 TArray64<uint8> FHdrImageWrapper::GetCompressed(int32 Quality)
 {
-	// Move the code that export hdr here
-	unimplemented();
-	return {};
+	// must have set BGRE8 raw data :
+	int64 NumPixels = Width * Height;
+	check( RawDataHolder.Num() == NumPixels * 4 );
+
+	
+	const int32 MaxHeaderSize = 256;
+	char Header[MAX_SPRINTF];
+	FCStringAnsi::Sprintf(Header, "#?RADIANCE\nFORMAT=32-bit_rle_rgbe\n\n-Y %d +X %d\n", Height,Width);
+	Header[MaxHeaderSize - 1] = 0;
+	int32 HeaderLen = FMath::Min(FCStringAnsi::Strlen(Header), MaxHeaderSize);
+
+	FreeCompressedData();
+	TArray64<uint8> CompressedDataArray;
+	CompressedDataArray.SetNum( HeaderLen + NumPixels * 4 );
+	uint8 * CompressedPtr = CompressedDataArray.GetData();
+
+	memcpy(CompressedPtr,Header,HeaderLen);
+	CompressedPtr += HeaderLen;
+
+	// just put the BGRE8 bytes
+	// but need to swizzle to RGBE8 order :
+
+	FColor * ToColors = (FColor *)CompressedPtr;
+	const FColor * FromColors = (const FColor *)RawDataHolder.GetData();
+	for(int64 i=0;i<NumPixels;i++)
+	{
+		ToColors[i].B = FromColors[i].R;
+		ToColors[i].G = FromColors[i].G;
+		ToColors[i].R = FromColors[i].B;
+		ToColors[i].A = FromColors[i].A;
+	}
+
+	return MoveTemp(CompressedDataArray);
 }
 
 bool FHdrImageWrapper::GetRaw(const ERGBFormat InFormat, int32 InBitDepth, TArray64<uint8>& OutRawData)
@@ -158,12 +211,18 @@ ERGBFormat FHdrImageWrapper::GetFormat() const
 	return ERGBFormat::BGRE;
 }
 
-int32 FHdrImageWrapper::GetNumFrames() const
+bool FHdrImageWrapper::SetAnimationInfo_DEPRECATED(int32 InNumFrames, int32 InFramerate)
+{
+	unimplemented();
+	return false;
+}
+
+int32 FHdrImageWrapper::GetNumFrames_DEPRECATED() const
 {
 	return INDEX_NONE;
 }
 
-int32 FHdrImageWrapper::GetFramerate() const
+int32 FHdrImageWrapper::GetFramerate_DEPRECATED() const
 {
 	return INDEX_NONE;
 }
diff --git a/Engine/Source/Runtime/ImageWrapper/Private/Formats/IcnsImageWrapper.cpp b/Engine/Source/Runtime/ImageWrapper/Private/Formats/IcnsImageWrapper.cpp
--- a/Engine/Source/Runtime/ImageWrapper/Private/Formats/IcnsImageWrapper.cpp
+++ b/Engine/Source/Runtime/ImageWrapper/Private/Formats/IcnsImageWrapper.cpp
@@ -24,6 +24,20 @@ bool FIcnsImageWrapper::SetCompressed(const void* InCompressedData, int64 InComp
 }
 
 
+// CanSetRawFormat returns true if SetRaw will accept this format
+bool FIcnsImageWrapper::CanSetRawFormat(const ERGBFormat InFormat, const int32 InBitDepth) const
+{
+	//checkf(false, TEXT("ICNS compression not supported"));
+	return false;
+}
+
+// returns InFormat if supported, else maps to something supported
+ERawImageFormat::Type FIcnsImageWrapper::GetSupportedRawFormat(const ERawImageFormat::Type InFormat) const
+{
+	//checkf(false, TEXT("ICNS compression not supported"));
+	return ERawImageFormat::BGRA8;
+}
+
 bool FIcnsImageWrapper::SetRaw(const void* InRawData, int64 InRawSize, const int32 InWidth, const int32 InHeight, const ERGBFormat InFormat, const int32 InBitDepth, const int32 InBytesPerRow)
 {
 	// Only support tightly packed
diff --git a/Engine/Source/Runtime/ImageWrapper/Private/Formats/IcnsImageWrapper.h b/Engine/Source/Runtime/ImageWrapper/Private/Formats/IcnsImageWrapper.h
--- a/Engine/Source/Runtime/ImageWrapper/Private/Formats/IcnsImageWrapper.h
+++ b/Engine/Source/Runtime/ImageWrapper/Private/Formats/IcnsImageWrapper.h
@@ -24,4 +24,7 @@ public:
 	virtual bool SetRaw(const void* InRawData, int64 InRawSize, const int32 InWidth, const int32 InHeight, const ERGBFormat InFormat, const int32 InBitDepth, const int32 InBytesPerRow = 0) override;
 	virtual void Compress(int32 Quality) override;
 	virtual void Uncompress(const ERGBFormat InFormat, int32 InBitDepth) override;
+	
+	virtual bool CanSetRawFormat(const ERGBFormat InFormat, const int32 InBitDepth) const override;
+	virtual ERawImageFormat::Type GetSupportedRawFormat(const ERawImageFormat::Type InFormat) const override;
 };
diff --git a/Engine/Source/Runtime/ImageWrapper/Private/Formats/IcoImageWrapper.cpp b/Engine/Source/Runtime/ImageWrapper/Private/Formats/IcoImageWrapper.cpp
--- a/Engine/Source/Runtime/ImageWrapper/Private/Formats/IcoImageWrapper.cpp
+++ b/Engine/Source/Runtime/ImageWrapper/Private/Formats/IcoImageWrapper.cpp
@@ -65,6 +65,20 @@ FIcoImageWrapper::FIcoImageWrapper()
 
 /* FImageWrapper interface
  *****************************************************************************/
+ 
+// CanSetRawFormat returns true if SetRaw will accept this format
+bool FIcoImageWrapper::CanSetRawFormat(const ERGBFormat InFormat, const int32 InBitDepth) const
+{
+	//checkf(false, TEXT("ICO compression not supported"));
+	return false;
+}
+
+// returns InFormat if supported, else maps to something supported
+ERawImageFormat::Type FIcoImageWrapper::GetSupportedRawFormat(const ERawImageFormat::Type InFormat) const
+{
+	//checkf(false, TEXT("ICO compression not supported"));
+	return ERawImageFormat::BGRA8;
+}
 
 void FIcoImageWrapper::Compress( int32 Quality )
 {
diff --git a/Engine/Source/Runtime/ImageWrapper/Private/Formats/IcoImageWrapper.h b/Engine/Source/Runtime/ImageWrapper/Private/Formats/IcoImageWrapper.h
--- a/Engine/Source/Runtime/ImageWrapper/Private/Formats/IcoImageWrapper.h
+++ b/Engine/Source/Runtime/ImageWrapper/Private/Formats/IcoImageWrapper.h
@@ -27,6 +27,9 @@ public:
 	virtual void Uncompress(const ERGBFormat InFormat, int32 InBitDepth) override;
 	virtual bool SetCompressed(const void* InCompressedData, int64 InCompressedSize) override;
 	virtual bool GetRaw(const ERGBFormat InFormat, int32 InBitDepth, TArray64<uint8>& OutRawData) override;
+	
+	virtual bool CanSetRawFormat(const ERGBFormat InFormat, const int32 InBitDepth) const override;
+	virtual ERawImageFormat::Type GetSupportedRawFormat(const ERawImageFormat::Type InFormat) const override;
 
 protected:
 
diff --git a/Engine/Source/Runtime/ImageWrapper/Private/Formats/JpegImageWrapper.cpp b/Engine/Source/Runtime/ImageWrapper/Private/Formats/JpegImageWrapper.cpp
--- a/Engine/Source/Runtime/ImageWrapper/Private/Formats/JpegImageWrapper.cpp
+++ b/Engine/Source/Runtime/ImageWrapper/Private/Formats/JpegImageWrapper.cpp
@@ -4,6 +4,7 @@
 
 #include "Math/Color.h"
 #include "Misc/ScopeLock.h"
+#include "ImageWrapperPrivate.h"
 
 #if WITH_UNREALJPEG
 
@@ -133,6 +134,33 @@ bool FJpegImageWrapper::SetCompressed(const void* InCompressedData, int64 InComp
 }
 
 
+// CanSetRawFormat returns true if SetRaw will accept this format
+bool FJpegImageWrapper::CanSetRawFormat(const ERGBFormat InFormat, const int32 InBitDepth) const
+{
+	return ((InFormat == ERGBFormat::RGBA || InFormat == ERGBFormat::BGRA || InFormat == ERGBFormat::Gray) && InBitDepth == 8);
+}
+
+// returns InFormat if supported, else maps to something supported
+ERawImageFormat::Type FJpegImageWrapper::GetSupportedRawFormat(const ERawImageFormat::Type InFormat) const
+{
+	switch(InFormat)
+	{
+	case ERawImageFormat::G8:
+	case ERawImageFormat::BGRA8:
+		return InFormat; // directly supported
+	case ERawImageFormat::BGRE8:
+	case ERawImageFormat::RGBA16:
+	case ERawImageFormat::RGBA16F:
+	case ERawImageFormat::RGBA32F:
+	case ERawImageFormat::G16:
+	case ERawImageFormat::R16F:
+		return ERawImageFormat::BGRA8; // needs conversion
+	default:
+		check(0);
+		return ERawImageFormat::BGRA8;
+	};
+}
+
 bool FJpegImageWrapper::SetRaw(const void* InRawData, int64 InRawSize, const int32 InWidth, const int32 InHeight, const ERGBFormat InFormat, const int32 InBitDepth, const int32 InBytesPerRow)
 {
 	check((InFormat == ERGBFormat::RGBA || InFormat == ERGBFormat::BGRA || InFormat == ERGBFormat::Gray) && InBitDepth == 8);
@@ -159,7 +187,7 @@ void FJpegImageWrapper::Compress(int32 Quality)
 	{
 		ensure(Quality >= 1 && Quality <= 100);
 
-		#define JPEG_QUALITY_MIN 30
+		#define JPEG_QUALITY_MIN 40
 		// JPEG should not be used below quality JPEG_QUALITY_MIN
 		Quality = FMath::Clamp(Quality, JPEG_QUALITY_MIN, 100);
 	}
@@ -232,6 +260,7 @@ void FJpegImageWrapper::Uncompress(const ERGBFormat InFormat, int32 InBitDepth)
 	else
 	{
 		check(false);
+		return;
 	}
 
 	FScopeLock JPEGLock(&GJPEGSection);
@@ -242,14 +271,18 @@ void FJpegImageWrapper::Uncompress(const ERGBFormat InFormat, int32 InBitDepth)
 	uint8* OutData = jpgd::decompress_jpeg_image_from_memory(
 		CompressedData.GetData(), CompressedData.Num(), &Width, &Height, &NumColors, Channels, (int)InFormat);
 
-
-	RawData.Reset(Width * Height * Channels);
-	RawData.AddUninitialized(Width * Height * Channels);
 	if (OutData)
 	{
+		RawData.Reset(Width * Height * Channels);
+		RawData.AddUninitialized(Width * Height * Channels);
 		FMemory::Memcpy(RawData.GetData(), OutData, RawData.Num());
 		FMemory::Free(OutData);
 	}
+	else
+	{
+		UE_LOG(LogImageWrapper, Error, TEXT("JPEG Decompress Error"));
+		RawData.Empty();
+	}
 #endif
 }
 
@@ -344,6 +377,8 @@ void FJpegImageWrapper::UncompressTurbo(const ERGBFormat InFormat, int32 InBitDe
 
 	if (tjDecompress2(Decompressor, CompressedData.GetData(), CompressedData.Num(), RawData.GetData(), Width, 0, Height, PixelFormat, Flags) != 0)
 	{
+		UE_LOG(LogImageWrapper, Error, TEXT("JPEG Decompress Error"));
+		RawData.Empty();
 		return;
 	}
 }
diff --git a/Engine/Source/Runtime/ImageWrapper/Private/Formats/JpegImageWrapper.h b/Engine/Source/Runtime/ImageWrapper/Private/Formats/JpegImageWrapper.h
--- a/Engine/Source/Runtime/ImageWrapper/Private/Formats/JpegImageWrapper.h
+++ b/Engine/Source/Runtime/ImageWrapper/Private/Formats/JpegImageWrapper.h
@@ -36,6 +36,9 @@ public:
 	virtual bool SetRaw(const void* InRawData, int64 InRawSize, const int32 InWidth, const int32 InHeight, const ERGBFormat InFormat, const int32 InBitDepth, const int32 InBytesPerRow = 0) override;
 	virtual void Uncompress(const ERGBFormat InFormat, int32 InBitDepth) override;
 	virtual void Compress(int32 Quality) override;
+	
+	virtual bool CanSetRawFormat(const ERGBFormat InFormat, const int32 InBitDepth) const override;
+	virtual ERawImageFormat::Type GetSupportedRawFormat(const ERawImageFormat::Type InFormat) const override;
 
 private:
 
diff --git a/Engine/Source/Runtime/ImageWrapper/Private/Formats/PngImageWrapper.cpp b/Engine/Source/Runtime/ImageWrapper/Private/Formats/PngImageWrapper.cpp
--- a/Engine/Source/Runtime/ImageWrapper/Private/Formats/PngImageWrapper.cpp
+++ b/Engine/Source/Runtime/ImageWrapper/Private/Formats/PngImageWrapper.cpp
@@ -117,6 +117,34 @@ FPngImageWrapper::FPngImageWrapper()
 
 /* FImageWrapper interface
  *****************************************************************************/
+ 
+// CanSetRawFormat returns true if SetRaw will accept this format
+bool FPngImageWrapper::CanSetRawFormat(const ERGBFormat InFormat, const int32 InBitDepth) const
+{
+	return (InFormat == ERGBFormat::RGBA || InFormat == ERGBFormat::BGRA || InFormat == ERGBFormat::Gray) && 
+		( InBitDepth == 8 || InBitDepth == 16 );
+}
+
+// returns InFormat if supported, else maps to something supported
+ERawImageFormat::Type FPngImageWrapper::GetSupportedRawFormat(const ERawImageFormat::Type InFormat) const
+{
+	switch(InFormat)
+	{
+	case ERawImageFormat::G8:
+	case ERawImageFormat::BGRA8:
+	case ERawImageFormat::RGBA16:
+	case ERawImageFormat::G16:
+		return InFormat; // directly supported
+	case ERawImageFormat::BGRE8:
+	case ERawImageFormat::RGBA16F:
+	case ERawImageFormat::RGBA32F:
+	case ERawImageFormat::R16F:
+		return ERawImageFormat::RGBA16; // needs conversion
+	default:
+		check(0);
+		return ERawImageFormat::BGRA8;
+	};
+}
 
 void FPngImageWrapper::Compress(int32 Quality)
 {
@@ -156,6 +184,8 @@ void FPngImageWrapper::Compress(int32 Quality)
 		if (setjmp(png_jmpbuf(png_ptr)) != 0)
 #endif
 		{
+			CompressedData.Empty();
+			UE_LOG(LogImageWrapper, Error, TEXT("PNG Compress Error"));
 			return;
 		}
 
@@ -288,6 +318,8 @@ void FPngImageWrapper::UncompressPNGData(const ERGBFormat InFormat, const int32
 		if (setjmp(png_jmpbuf(png_ptr)) != 0)
 #endif
 		{
+			RawData.Empty();
+			UE_LOG(LogImageWrapper, Error, TEXT("PNG Decompress Error"));
 			return;
 		}
 
diff --git a/Engine/Source/Runtime/ImageWrapper/Private/Formats/PngImageWrapper.h b/Engine/Source/Runtime/ImageWrapper/Private/Formats/PngImageWrapper.h
--- a/Engine/Source/Runtime/ImageWrapper/Private/Formats/PngImageWrapper.h
+++ b/Engine/Source/Runtime/ImageWrapper/Private/Formats/PngImageWrapper.h
@@ -49,6 +49,9 @@ public:
 	virtual void Reset() override;
 	virtual bool SetCompressed(const void* InCompressedData, int64 InCompressedSize) override;
 	virtual void Uncompress(const ERGBFormat InFormat, int32 InBitDepth) override;
+	
+	virtual bool CanSetRawFormat(const ERGBFormat InFormat, const int32 InBitDepth) const override;
+	virtual ERawImageFormat::Type GetSupportedRawFormat(const ERawImageFormat::Type InFormat) const override;
 
 public:
 
diff --git a/Engine/Source/Runtime/ImageWrapper/Private/Formats/TgaImageWrapper.cpp b/Engine/Source/Runtime/ImageWrapper/Private/Formats/TgaImageWrapper.cpp
--- a/Engine/Source/Runtime/ImageWrapper/Private/Formats/TgaImageWrapper.cpp
+++ b/Engine/Source/Runtime/ImageWrapper/Private/Formats/TgaImageWrapper.cpp
@@ -6,9 +6,55 @@
 #include "TgaImageSupport.h"
 
 
+
+// CanSetRawFormat returns true if SetRaw will accept this format
+bool FTgaImageWrapper::CanSetRawFormat(const ERGBFormat InFormat, const int32 InBitDepth) const
+{
+	return ( InFormat == ERGBFormat::BGRA ) && ( InBitDepth == 8 );
+}
+
+// returns InFormat if supported, else maps to something supported
+ERawImageFormat::Type FTgaImageWrapper::GetSupportedRawFormat(const ERawImageFormat::Type InFormat) const
+{
+	return ERawImageFormat::BGRA8;
+}
+
 void FTgaImageWrapper::Compress(int32 Quality)
 {
-	checkf(false, TEXT("TGA compression not supported"));
+	CompressedData.Reset();
+
+	int64 SurfaceBytes = Width * Height * 4;
+	check( RawData.Num() == SurfaceBytes );
+
+	int64 CompresedSize = sizeof(FTGAFileHeader) + SurfaceBytes;
+	CompressedData.SetNum(CompresedSize);
+
+	uint8 * CompressedPtr = CompressedData.GetData();
+	
+	// super lazy TGA writer here :
+	//	always writes 32 bit BGRA
+
+	FTGAFileHeader header = { };
+
+	header.Width  = INTEL_ORDER16( (uint16) Width );
+	header.Height = INTEL_ORDER16( (uint16) Height );
+	header.ImageTypeCode = 2;
+	header.BitsPerPixel = 32;
+	header.ImageDescriptor = 8;
+
+	memcpy(CompressedPtr,&header,sizeof(header));
+	CompressedPtr += sizeof(header);
+
+	// write rows BGRA , bottom up :
+	for(int y = Height-1; y>= 0;y--)
+	{
+		int64 RowBytes = Width * 4;
+		const void * From = &RawData[y * RowBytes]; 
+		memcpy(CompressedPtr,From,RowBytes);
+		CompressedPtr += RowBytes;
+	}
+
+	check( CompressedPtr == CompressedData.GetData() + CompressedData.Num() );
 }
 
 bool FTgaImageWrapper::SetCompressed(const void* InCompressedData, int64 InCompressedSize)
@@ -21,36 +67,30 @@ bool FTgaImageWrapper::SetCompressed(const void* InCompressedData, int64 InCompr
 void FTgaImageWrapper::Uncompress(const ERGBFormat InFormat, const int32 InBitDepth)
 {
 	const int32 BytesPerPixel = ( InFormat == ERGBFormat::Gray ? 1 : 4 );
-	RawData.AddUninitialized( (int64)Width * Height * BytesPerPixel );
+	int64 TextureDataSize = (int64)Width * Height * BytesPerPixel;
+	RawData.Empty(TextureDataSize);
+	RawData.AddUninitialized(TextureDataSize);
 
-	const int32 TextureDataSize = RawData.Num();
 	uint32* TextureData = reinterpret_cast< uint32* >( RawData.GetData() );
 
 	FTGAFileHeader* TgaHeader = reinterpret_cast< FTGAFileHeader* >( CompressedData.GetData() );
 	if ( !DecompressTGA_helper( TgaHeader, CompressedData.Num(), TextureData, TextureDataSize ) )
 	{
 		SetError(TEXT("Error while decompressing a TGA"));
+		RawData.Reset();
 	}
 }
 
-bool FTgaImageWrapper::LoadTGAHeader()
+bool FTgaImageWrapper::IsTGAHeader(const void * CompressedData,int64 CompressedDataLength)
 {
-	check( CompressedData.Num() );
+	const FTGAFileHeader* TgaHeader = (FTGAFileHeader*)CompressedData;
 
-	const FTGAFileHeader* TgaHeader = (FTGAFileHeader*)CompressedData.GetData();
-
-	if ( CompressedData.Num() >= sizeof( FTGAFileHeader ) &&
+	if ( CompressedDataLength >= sizeof( FTGAFileHeader ) &&
 		( (TgaHeader->ColorMapType == 0 && TgaHeader->ImageTypeCode == 2) ||
 		  ( TgaHeader->ColorMapType == 0 && TgaHeader->ImageTypeCode == 3 ) || // ImageTypeCode 3 is greyscale
 		  ( TgaHeader->ColorMapType == 0 && TgaHeader->ImageTypeCode == 10 ) ||
 		  ( TgaHeader->ColorMapType == 1 && TgaHeader->ImageTypeCode == 1 && TgaHeader->BitsPerPixel == 8 ) ) )
 	{
-		Width = TgaHeader->Width;
-		Height = TgaHeader->Height;
-		ColorMapType = TgaHeader->ColorMapType;
-		ImageTypeCode = TgaHeader->ImageTypeCode;
-		BitDepth = TgaHeader->BitsPerPixel;
-
 		return true;
 	}
 	else
@@ -58,3 +98,40 @@ bool FTgaImageWrapper::LoadTGAHeader()
 		return false;
 	}
 }
+
+bool FTgaImageWrapper::LoadTGAHeader()
+{
+	check( CompressedData.Num() );
+
+	if( ! IsTGAHeader(CompressedData.GetData(),CompressedData.Num()) )
+	{
+		return false;
+	}
+
+	const FTGAFileHeader* TgaHeader = (FTGAFileHeader*)CompressedData.GetData();
+
+	Width = TgaHeader->Width;
+	Height = TgaHeader->Height;
+	BitDepth = 8;
+	
+	// must exactly match the logic in DecompressTGA_helper
+	if(TgaHeader->ColorMapType == 1 && TgaHeader->ImageTypeCode == 1 && TgaHeader->BitsPerPixel == 8)
+	{
+		Format = ERGBFormat::Gray;
+	}
+	else if(TgaHeader->ColorMapType == 0 && TgaHeader->ImageTypeCode == 3 && TgaHeader->BitsPerPixel == 8)
+	{
+		Format = ERGBFormat::Gray;
+	}
+	else
+	{
+		Format = ERGBFormat::BGRA;
+	}
+
+	ColorMapType = TgaHeader->ColorMapType;
+	ImageTypeCode = TgaHeader->ImageTypeCode;
+	
+	// TgaHeader->BitsPerPixel is bits per *color*
+
+	return true;
+}
diff --git a/Engine/Source/Runtime/ImageWrapper/Private/Formats/TgaImageWrapper.h b/Engine/Source/Runtime/ImageWrapper/Private/Formats/TgaImageWrapper.h
--- a/Engine/Source/Runtime/ImageWrapper/Private/Formats/TgaImageWrapper.h
+++ b/Engine/Source/Runtime/ImageWrapper/Private/Formats/TgaImageWrapper.h
@@ -17,6 +17,8 @@ public:
 	 * Load the header information, returns true if successful.
 	 */
 	bool LoadTGAHeader();
+	
+	static bool IsTGAHeader(const void * CompressedData,int64 CompressedDataLength);
 
 public:
 
@@ -25,6 +27,9 @@ public:
 	virtual void Compress(int32 Quality) override;
 	virtual bool SetCompressed(const void* InCompressedData, int64 InCompressedSize) override;
 	virtual void Uncompress(const ERGBFormat InFormat, int32 InBitDepth) override;
+	
+	virtual bool CanSetRawFormat(const ERGBFormat InFormat, const int32 InBitDepth) const override;
+	virtual ERawImageFormat::Type GetSupportedRawFormat(const ERawImageFormat::Type InFormat) const override;
 
 private:
 	/** The color type as defined in the header. */
diff --git a/Engine/Source/Runtime/ImageWrapper/Private/Formats/TiffImageWrapper.cpp b/Engine/Source/Runtime/ImageWrapper/Private/Formats/TiffImageWrapper.cpp
--- a/Engine/Source/Runtime/ImageWrapper/Private/Formats/TiffImageWrapper.cpp
+++ b/Engine/Source/Runtime/ImageWrapper/Private/Formats/TiffImageWrapper.cpp
@@ -406,8 +406,23 @@ namespace UE::ImageWrapper::Private
 
 	void FTiffImageWrapper::Compress(int32 Quality)
 	{
+		// should not get here because CanSetRawFormat is false
 		checkf(false, TEXT("TIFF compression not supported"));
 	}
+	
+	// CanSetRawFormat returns true if SetRaw will accept this format
+	bool FTiffImageWrapper::CanSetRawFormat(const ERGBFormat InFormat, const int32 InBitDepth) const
+	{
+		//checkf(false, TEXT("TIFF compression not supported"));
+		return false;
+	}
+
+	// returns InFormat if supported, else maps to something supported
+	ERawImageFormat::Type FTiffImageWrapper::GetSupportedRawFormat(const ERawImageFormat::Type InFormat) const
+	{
+		//checkf(false, TEXT("TIFF compression not supported"));
+		return ERawImageFormat::BGRA8;
+	}
 
 	void FTiffImageWrapper::Uncompress(const ERGBFormat InFormat, int32 InBitDepth)
 	{
diff --git a/Engine/Source/Runtime/ImageWrapper/Private/Formats/TiffImageWrapper.h b/Engine/Source/Runtime/ImageWrapper/Private/Formats/TiffImageWrapper.h
--- a/Engine/Source/Runtime/ImageWrapper/Private/Formats/TiffImageWrapper.h
+++ b/Engine/Source/Runtime/ImageWrapper/Private/Formats/TiffImageWrapper.h
@@ -29,6 +29,9 @@ namespace UE::ImageWrapper::Private
 		virtual void Uncompress(const ERGBFormat InFormat, int32 InBitDepth) override;
 		virtual bool SetCompressed(const void* InCompressedData, int64 InCompressedSize) override;
 		
+		virtual bool CanSetRawFormat(const ERGBFormat InFormat, const int32 InBitDepth) const override;
+		virtual ERawImageFormat::Type GetSupportedRawFormat(const ERawImageFormat::Type InFormat) const override;
+
 	private:
 
 		void ReleaseTiffImage();
diff --git a/Engine/Source/Runtime/ImageWrapper/Private/ImageWrapperBase.cpp b/Engine/Source/Runtime/ImageWrapper/Private/ImageWrapperBase.cpp
--- a/Engine/Source/Runtime/ImageWrapper/Private/ImageWrapperBase.cpp
+++ b/Engine/Source/Runtime/ImageWrapper/Private/ImageWrapperBase.cpp
@@ -1,6 +1,7 @@
 // Copyright Epic Games, Inc. All Rights Reserved.
 
 #include "ImageWrapperBase.h"
+#include "ImageWrapperPrivate.h"
 
 
 /* FImageWrapperBase structors
@@ -14,8 +15,8 @@ FImageWrapperBase::FImageWrapperBase()
 	, BitDepth(0)
 	, Width(0)
 	, Height(0)
-	, NumFrames(1)
-	, Framerate(0)
+	, NumFrames_DEPRECATED(1)
+	, Framerate_DEPRECATED(0)
 { }
 
 
@@ -25,6 +26,10 @@ FImageWrapperBase::FImageWrapperBase()
 void FImageWrapperBase::Reset()
 {
 	LastError.Empty();
+	
+	// not reset ?
+	//TArray64<uint8> RawData;
+	//TArray64<uint8> CompressedData;
 
 	RawFormat = ERGBFormat::Invalid;
 	RawBitDepth = 0;
@@ -33,8 +38,8 @@ void FImageWrapperBase::Reset()
 	BitDepth = 0;
 	Width = 0;
 	Height = 0;
-	NumFrames = 1;
-	Framerate = 0;
+	NumFrames_DEPRECATED = 1;
+	Framerate_DEPRECATED = 0;
 }
 
 
@@ -64,6 +69,10 @@ bool FImageWrapperBase::GetRaw(const ERGBFormat InFormat, int32 InBitDepth, TArr
 	if (LastError.IsEmpty())
 	{
 		OutRawData = MoveTemp(RawData);
+		if ( OutRawData.IsEmpty() )
+		{
+			return false;
+		}
 	}
 
 	return LastError.IsEmpty();
@@ -98,6 +107,12 @@ bool FImageWrapperBase::SetRaw(const void* InRawData, int64 InRawSize, const int
 
 	Reset();
 	CompressedData.Empty();		// Invalidates the compressed data too
+	
+	if ( ! CanSetRawFormat(InFormat,InBitDepth) )
+	{
+		UE_LOG(LogImageWrapper, Warning, TEXT("ImageWrapper unsupported format; check CanSetRawFormat; %d x %d"), (int)InFormat,InBitDepth);
+		return false;
+	}
 
 	RawData.Empty(InRawSize);
 	RawData.AddUninitialized(InRawSize);
@@ -113,10 +128,251 @@ bool FImageWrapperBase::SetRaw(const void* InRawData, int64 InRawSize, const int
 	return true;
 }
 
-bool FImageWrapperBase::SetAnimationInfo(int32 InNumFrames, int32 InFramerate)
+bool FImageWrapperBase::SetAnimationInfo_DEPRECATED(int32 InNumFrames, int32 InFramerate)
+{
+	NumFrames_DEPRECATED = InNumFrames;
+	Framerate_DEPRECATED = InFramerate;
+	
+	return true;
+}
+
+
+ERawImageFormat::Type IImageWrapper::ConvertRGBFormat(ERGBFormat RGBFormat,int BitDepth,bool * bIsExactMatch)
+{
+	bool bIsExactMatchDummy;
+	if ( ! bIsExactMatch )
+	{
+		bIsExactMatch = &bIsExactMatchDummy;
+	}
+		
+	switch(RGBFormat)
+	{
+	case ERGBFormat::RGBA:
+		if ( BitDepth == 8 )
+		{
+			*bIsExactMatch = false; // needs RB swap
+			return ERawImageFormat::BGRA8;
+		}
+		else if ( BitDepth == 16 )
+		{
+			*bIsExactMatch = true;
+			return ERawImageFormat::RGBA16;
+		}
+		break;
+			
+	case ERGBFormat::BGRA:
+		if ( BitDepth == 8 )
+		{
+			*bIsExactMatch = true;
+			return ERawImageFormat::BGRA8;
+		}
+		else if ( BitDepth == 16 )
+		{
+			*bIsExactMatch = false; // needs RB swap
+			return ERawImageFormat::RGBA16;
+		}
+		break;
+			
+	case ERGBFormat::Gray:
+		if ( BitDepth == 8 )
+		{
+			*bIsExactMatch = true;
+			return ERawImageFormat::G8;
+		}
+		else if ( BitDepth == 16 )
+		{
+			*bIsExactMatch = true;
+			return ERawImageFormat::G16;
+		}
+		break;
+			
+	case ERGBFormat::BGRE:
+		if ( BitDepth == 8 )
+		{
+			*bIsExactMatch = true;
+			return ERawImageFormat::BGRE8;
+		}
+		break;
+			
+	case ERGBFormat::RGBAF:
+		if ( BitDepth == 16 )
+		{
+			*bIsExactMatch = true;
+			return ERawImageFormat::RGBA16F;
+		}
+		else if ( BitDepth == 32 )
+		{
+			*bIsExactMatch = true;
+			return ERawImageFormat::RGBA32F;
+		}
+		break;
+			
+	case ERGBFormat::GrayF:
+		if ( BitDepth == 16 )
+		{
+			*bIsExactMatch = true;
+			return ERawImageFormat::R16F;
+		}
+		else if ( BitDepth == 32 )
+		{
+			*bIsExactMatch = false; // no single channel F32
+			return ERawImageFormat::RGBA32F; // promote F32 to 4xF32
+		}
+		break;
+
+	default:
+		break;
+	}
+
+	UE_LOG(LogImageWrapper, Warning, TEXT("ConvertRGBFormat not handled : %d,%d"), (int)RGBFormat,BitDepth );
+		
+	*bIsExactMatch = false;
+	return ERawImageFormat::Invalid;
+}
+	
+void IImageWrapper::ConvertRawImageFormat(ERawImageFormat::Type RawFormat, ERGBFormat & OutFormat,int & OutBitDepth)
+{
+	switch(RawFormat)
+	{
+	case ERawImageFormat::G8:
+		OutFormat = ERGBFormat::Gray;
+		OutBitDepth = 8;
+		break;
+	case ERawImageFormat::BGRA8:
+		OutFormat = ERGBFormat::BGRA;
+		OutBitDepth = 8;
+		break;
+	case ERawImageFormat::BGRE8:
+		OutFormat = ERGBFormat::BGRE;
+		OutBitDepth = 8;
+		break;
+	case ERawImageFormat::RGBA16:
+		OutFormat = ERGBFormat::RGBA;
+		OutBitDepth = 16;
+		break;
+	case ERawImageFormat::RGBA16F:
+		OutFormat = ERGBFormat::RGBAF;
+		OutBitDepth = 16;
+		break;
+	case ERawImageFormat::RGBA32F:
+		OutFormat = ERGBFormat::RGBAF;
+		OutBitDepth = 32;
+		break;
+	case ERawImageFormat::G16:
+		OutFormat = ERGBFormat::Gray;
+		OutBitDepth = 16;
+		break;
+	case ERawImageFormat::R16F:
+		OutFormat = ERGBFormat::GrayF;
+		OutBitDepth = 16;
+		break;
+	default:
+		check(0);
+		break;
+	}
+}
+
+bool IImageWrapper::GetRawImage(FImage & OutImage)
 {
-	NumFrames = InNumFrames;
-	Framerate = InFramerate;
+	TArray64<uint8> OutRawData;
+	if ( ! GetRaw(OutRawData) )
+	{
+		return false;
+	}
+
+	int Width = GetWidth();
+	int Height = GetHeight();
+	ERGBFormat RGBFormat = GetFormat();
+	int BitDepth = GetBitDepth();
+
+	bool bExactMatch;
+	ERawImageFormat::Type RawFormat = GetClosestRawImageFormat(&bExactMatch);
 	
+	if ( RawFormat == ERawImageFormat::Invalid )
+	{
+		return false;
+	}
+
+	EGammaSpace GammaSpace = EGammaSpace::Linear;
+	if ( GetSRGB() )
+	{
+		GammaSpace = EGammaSpace::sRGB;
+	}
+
+	if ( bExactMatch )
+	{
+		// no conversion required
+
+		OutImage.RawData = MoveTemp(OutRawData);
+		OutImage.SizeX = Width;
+		OutImage.SizeY = Height;
+		OutImage.NumSlices = 1;
+		OutImage.Format = RawFormat;
+		OutImage.GammaSpace = GammaSpace;
+	}
+	else
+	{
+		OutImage.Init( Width, Height, RawFormat, GammaSpace );
+
+		FImageView SrcImage = OutImage;
+		SrcImage.RawData = OutRawData.GetData();
+
+		switch(RGBFormat)
+		{
+		case ERGBFormat::RGBA:
+		{
+			// RGBA8 -> BGRA8
+			check( BitDepth == 8 );
+			check( RawFormat == ERawImageFormat::BGRA8 );
+			FImageCore::CopyImageRGBABGRA(SrcImage, OutImage );
+			break;
+		}
+			
+		case ERGBFormat::BGRA:
+		{
+			// BGRA16 -> RGBA16
+			check( BitDepth == 16 );
+			check( RawFormat == ERawImageFormat::RGBA16 );
+			FImageCore::CopyImageRGBABGRA(SrcImage, OutImage );
+			break;
+		}
+
+		case ERGBFormat::GrayF:
+		{
+			/*
+			// 1 channel F32 -> F16 :
+			//  because ERawImageFormat has no 1 channel F32 currently
+			check( BitDepth == 32 );
+			check( RawFormat == ERawImageFormat::R16F );
+			const float * Src = (const float *)OutRawData.GetData();
+			uint16 * Dst = (uint16 *)OutImage.RawData.GetData();
+			int64 NumPixels = OutImage.GetNumPixels();
+			for(int64 i=0;i<NumPixels;i++)
+			{
+				// clamp in F16 range? does StoreHalf do that?
+				FPlatformMath::StoreHalf(&Dst[i], Src[i]);
+			}
+			*/
+
+			// 1 channel F32 -> 4xF32 :
+			check( BitDepth == 32 );
+			check( RawFormat == ERawImageFormat::RGBA32F );
+			const float * Src = (const float *)OutRawData.GetData();
+			FLinearColor * Dst = (FLinearColor *)OutImage.RawData.GetData();
+			int64 NumPixels = OutImage.GetNumPixels();
+			for(int64 i=0;i<NumPixels;i++)
+			{
+				Dst[i] = FLinearColor( Src[i],Src[i],Src[i],1.f );
+			}
+
+			break;
+		}
+
+		default:
+			check(0);
+			return false;			
+		}
+	}
+
 	return true;
 }
diff --git a/Engine/Source/Runtime/ImageWrapper/Private/ImageWrapperBase.h b/Engine/Source/Runtime/ImageWrapper/Private/ImageWrapperBase.h
--- a/Engine/Source/Runtime/ImageWrapper/Private/ImageWrapperBase.h
+++ b/Engine/Source/Runtime/ImageWrapper/Private/ImageWrapperBase.h
@@ -48,6 +48,8 @@ public:
 	 * Compresses the data.
 	 *
 	 * @param Quality The compression quality.
+	 * 
+	 * returns void. call SetError() in your implementation if you fail.
 	 */
 	virtual void Compress(int32 Quality) = 0;
 
@@ -61,12 +63,14 @@ public:
 	 *
 	 * @param ErrorMessage The error message to set.
 	 */
-	virtual void SetError(const TCHAR* ErrorMessage);
+	void SetError(const TCHAR* ErrorMessage);
 
 	/**  
 	 * Function to uncompress our data 
 	 *
 	 * @param InFormat How we want to manipulate the RGB data
+	 * 
+	 * returns void. call SetError() in your implementation if you fail.
 	 */
 	virtual void Uncompress(const ERGBFormat InFormat, int32 InBitDepth) = 0;
 
@@ -98,19 +102,19 @@ public:
 		return Width;
 	}
 
-	virtual int32 GetNumFrames() const override
+	virtual int32 GetNumFrames_DEPRECATED() const override
 	{
-		return NumFrames;
+		return NumFrames_DEPRECATED;
 	}
 	
-	virtual int32 GetFramerate() const override
+	virtual int32 GetFramerate_DEPRECATED() const override
 	{
-		return Framerate;
+		return Framerate_DEPRECATED;
 	}
 
 	virtual bool SetCompressed(const void* InCompressedData, int64 InCompressedSize) override;
 	virtual bool SetRaw(const void* InRawData, int64 InRawSize, const int32 InWidth, const int32 InHeight, const ERGBFormat InFormat, const int32 InBitDepth, const int32 InBytesPerRow = 0) override;
-	virtual bool SetAnimationInfo(int32 InNumFrames, int32 InFramerate) override;
+	virtual bool SetAnimationInfo_DEPRECATED(int32 InNumFrames, int32 InFramerate) override;
 
 protected:
 
@@ -136,8 +140,8 @@ protected:
 	int32 Height;
 	
 	/** Animation information */
-	int32 NumFrames;
-	int32 Framerate;
+	int32 NumFrames_DEPRECATED;
+	int32 Framerate_DEPRECATED;
 
 	/** Last Error Message. */
 	FString LastError;
diff --git a/Engine/Source/Runtime/ImageWrapper/Private/ImageWrapperModule.cpp b/Engine/Source/Runtime/ImageWrapper/Private/ImageWrapperModule.cpp
--- a/Engine/Source/Runtime/ImageWrapper/Private/ImageWrapperModule.cpp
+++ b/Engine/Source/Runtime/ImageWrapper/Private/ImageWrapperModule.cpp
@@ -32,6 +32,9 @@ namespace
 
 	// Binary for #?RADIANCE
 	static const uint8 IMAGE_MAGIC_HDR[] = {0x23, 0x3f, 0x52, 0x41, 0x44, 0x49, 0x41, 0x4e, 0x43, 0x45, 0x0a};
+	
+	// Binary for #?RGBE
+	static const uint8 IMAGE_MAGIC_HDR2[] = {0x23, 0x3f, 0x52, 0x47, 0x42, 0x45 };
 
 	// Tiff has two magic bytes sequence
 	static const uint8 IMAGE_MAGIC_TIFF_LITTLE_ENDIAN[] = {0x49, 0x49, 0x2A, 0x00};
@@ -70,6 +73,8 @@ public:
 
 	virtual TSharedPtr<IImageWrapper> CreateImageWrapper(const EImageFormat InFormat) override
 	{
+		TRACE_CPUPROFILER_EVENT_SCOPE(ImageWrapper.Create);
+
 		TSharedPtr<IImageWrapper> ImageWrapper;
 
 		// Allocate a helper for the format type
@@ -155,7 +160,8 @@ public:
 		{
 			Format = EImageFormat::ICNS;
 		}
-		else if (StartsWith((uint8*)CompressedData, CompressedSize, IMAGE_MAGIC_HDR))
+		else if (StartsWith((uint8*)CompressedData, CompressedSize, IMAGE_MAGIC_HDR) ||
+			StartsWith((uint8*)CompressedData, CompressedSize, IMAGE_MAGIC_HDR2))
 		{
 			Format = EImageFormat::HDR;
 		}
@@ -167,10 +173,221 @@ public:
 		{
 			Format = EImageFormat::TIFF;
 		}
+		else if ( FTgaImageWrapper::IsTGAHeader(CompressedData, CompressedSize) )
+		{
+			Format = EImageFormat::TGA;
+		}
 
 		return Format;
 	}
 
+	virtual const TCHAR * GetExtension(EImageFormat Format) override
+	{
+		switch(Format)
+		{
+		case EImageFormat::PNG: return TEXT("png");
+		case EImageFormat::JPEG: return TEXT("jpg");
+		case EImageFormat::GrayscaleJPEG: return TEXT("jpg");
+		case EImageFormat::BMP: return TEXT("bmp");
+		case EImageFormat::ICO: return TEXT("ico");
+		case EImageFormat::EXR: return TEXT("exr");
+		case EImageFormat::ICNS: return TEXT("icns");
+		case EImageFormat::TGA: return TEXT("tga");
+		case EImageFormat::HDR: return TEXT("hdr");
+		case EImageFormat::TIFF: return TEXT("tiff");
+		case EImageFormat::Invalid:
+		default:
+			check(0);
+			return nullptr;
+		}
+	}
+	
+	virtual EImageFormat GetImageFormatFromExtension(const TCHAR * Name) override
+	{
+		const TCHAR * Dot = FCString::Strrchr(Name,TEXT('.'));
+		if ( Dot )
+		{
+			Name = Dot+1;
+		}
+		if ( FCString::Stricmp(Name,TEXT("png")) == 0 )
+		{
+			return EImageFormat::PNG;
+		}
+		else if ( FCString::Stricmp(Name,TEXT("jpg")) == 0
+		 || FCString::Stricmp(Name,TEXT("jpeg")) == 0 )
+		{
+			return EImageFormat::JPEG;
+		}
+		else if ( FCString::Stricmp(Name,TEXT("bmp")) == 0 )
+		{
+			return EImageFormat::BMP;
+		}
+		else if ( FCString::Stricmp(Name,TEXT("ico")) == 0 )
+		{
+			return EImageFormat::ICO;
+		}
+		else if ( FCString::Stricmp(Name,TEXT("exr")) == 0 )
+		{
+			return EImageFormat::EXR;
+		}
+		else if ( FCString::Stricmp(Name,TEXT("icns")) == 0 )
+		{
+			return EImageFormat::ICNS;
+		}
+		else if ( FCString::Stricmp(Name,TEXT("tga")) == 0 )
+		{
+			return EImageFormat::TGA;
+		}
+		else if ( FCString::Stricmp(Name,TEXT("hdr")) == 0 )
+		{
+			return EImageFormat::HDR;
+		}
+		else if ( FCString::Stricmp(Name,TEXT("tiff")) == 0 ||
+			 FCString::Stricmp(Name,TEXT("tif")) == 0 )
+		{
+			return EImageFormat::TIFF;
+		}
+		else
+		{
+			UE_LOG(LogImageWrapper,Warning,TEXT("GetImageFormatFromExtension not found : %s\n"),Name);
+			return EImageFormat::Invalid;
+		}
+	}
+	
+	virtual ERawImageFormat::Type ConvertRGBFormat(ERGBFormat RGBFormat,int BitDepth,bool * bIsExactMatch) override
+	{
+		return IImageWrapper::ConvertRGBFormat(RGBFormat,BitDepth,bIsExactMatch);
+	}
+	
+	virtual void ConvertRawImageFormat(ERawImageFormat::Type RawFormat, ERGBFormat & OutFormat,int & OutBitDepth) override
+	{
+		IImageWrapper::ConvertRawImageFormat(RawFormat,OutFormat,OutBitDepth);
+	}
+	
+	virtual EImageFormat GetDefaultOutputFormat(ERawImageFormat::Type RawFormat) override
+	{
+		switch(RawFormat)
+		{
+		case ERawImageFormat::G8:
+		case ERawImageFormat::BGRA8:
+		case ERawImageFormat::RGBA16:
+		case ERawImageFormat::G16:
+			return EImageFormat::PNG;
+
+		case ERawImageFormat::BGRE8:
+			return EImageFormat::HDR;
+
+		case ERawImageFormat::RGBA16F:
+		case ERawImageFormat::RGBA32F:
+		case ERawImageFormat::R16F:
+			return EImageFormat::EXR;
+		
+		default:
+			check(0);
+			return EImageFormat::Invalid;
+		}
+	}
+	
+	virtual bool DecompressImage(const void* InCompressedData, int64 InCompressedSize, FImage & OutImage) override
+	{
+		TRACE_CPUPROFILER_EVENT_SCOPE(ImageWrapper.Decompress);
+
+		EImageFormat ImageFormat = DetectImageFormat(InCompressedData,InCompressedSize);
+		if ( ImageFormat == EImageFormat::Invalid )
+		{
+			return false;
+		}
+		
+		TSharedPtr<IImageWrapper> ImageWrapper = CreateImageWrapper(ImageFormat);
+		if ( ! ImageWrapper.IsValid() )
+		{
+			return false;
+		}
+
+		if ( ! ImageWrapper->SetCompressed(InCompressedData,InCompressedSize) )
+		{
+			return false;
+		}
+
+		if ( ! ImageWrapper->GetRawImage(OutImage) )
+		{
+			return false;
+		}
+
+		return true;
+	}
+
+	virtual bool CompressImage(TArray64<uint8> & OutData, EImageFormat ToFormat, const FImageView & InImage, int32 Quality) override
+	{
+		TRACE_CPUPROFILER_EVENT_SCOPE(ImageWrapper.Compress);
+
+		TSharedPtr<IImageWrapper> ImageWrapper = CreateImageWrapper(ToFormat);
+		if ( ! ImageWrapper.IsValid() )
+		{
+			return false;
+		}
+
+		ERGBFormat RGBFormat;
+		int BitDepth;
+		ConvertRawImageFormat(InImage.Format,RGBFormat,BitDepth);
+
+		// can't save slices :
+		check(InImage.NumSlices == 1 );
+
+		FImageView WriteImage = InImage;
+		FImage TempImage;
+
+		if ( ! ImageWrapper->CanSetRawFormat(RGBFormat,BitDepth) )
+		{
+			// output image format may not support this pixel format
+			//	in that case conversion is needed
+			// 
+			// @@!! warn or set a flag when this conversion is lossy?
+
+			ERawImageFormat::Type NewFormat = ImageWrapper->GetSupportedRawFormat(InImage.Format);
+			check( NewFormat != InImage.Format );
+			
+			// we will write to the image wrapper using "NewFormat"
+			// do a blit to convert InImage to NewFormat
+			// 
+
+			TempImage.Init(InImage.SizeX,InImage.SizeY,NewFormat,ERawImageFormat::GetDefaultGammaSpace(NewFormat));
+			WriteImage = TempImage;
+
+			FImageCore::CopyImage(InImage,WriteImage);
+			
+			// re-get the RGBFormat we will set :
+			ConvertRawImageFormat(NewFormat,RGBFormat,BitDepth);
+			check( ImageWrapper->CanSetRawFormat(RGBFormat,BitDepth) );
+		}
+		else if ( InImage.GetGammaSpace() != ERawImageFormat::GetDefaultGammaSpace(InImage.Format) )
+		{
+			// ImageWrapper SetRaw() does not accept gamma information
+			// it assumes U8 = SRGB and Float = Linear
+
+			// eg. if you save a U8 surface with linear gamma
+			// go ahead and write the U8 bytes without changing them
+			// I think that's probably what is intended
+			// when someone writes Linear U8 to PNG they don't want me to do a gamma transform
+			// @@!! TODO: write the gamma setting to the output when possible, eg. for DDS output
+		}
+
+		if ( ! ImageWrapper->SetRaw(WriteImage.RawData,WriteImage.GetImageSizeBytes(),WriteImage.SizeX,WriteImage.SizeY,RGBFormat,BitDepth) )
+		{
+			return false;
+		}
+
+		OutData = ImageWrapper->GetCompressed(Quality);
+
+		if ( OutData.IsEmpty() )
+		{
+			return false;
+		}
+
+		return true;
+	}
+
+
 public:
 
 	//~ IModuleInterface interface
diff --git a/Engine/Source/Runtime/ImageWrapper/Private/TgaImageSupport.cpp b/Engine/Source/Runtime/ImageWrapper/Private/TgaImageSupport.cpp
--- a/Engine/Source/Runtime/ImageWrapper/Private/TgaImageSupport.cpp
+++ b/Engine/Source/Runtime/ImageWrapper/Private/TgaImageSupport.cpp
@@ -291,6 +291,8 @@ bool DecompressTGA_helper( const FTGAFileHeader* TgaHeader, const int64 TGABuffe
 	bool bSuccess = false;
 	if ( TgaHeader->ImageTypeCode == 10 ) // 10 = RLE compressed 
 	{
+		check( TextureDataSize == TgaHeader->Width * TgaHeader->Height * 4 );
+
 		// RLE compression: CHUNKS: 1 -byte header, high bit 0 = raw, 1 = compressed
 		// bits 0-6 are a 7-bit count; count+1 = number of raw pixels following, or rle pixels to be expanded. 
 		if(TgaHeader->BitsPerPixel == 32)
@@ -313,6 +315,8 @@ bool DecompressTGA_helper( const FTGAFileHeader* TgaHeader, const int64 TGABuffe
 	}
 	else if(TgaHeader->ImageTypeCode == 2) // 2 = Uncompressed RGB
 	{
+		check( TextureDataSize == TgaHeader->Width * TgaHeader->Height * 4 );
+
 		if(TgaHeader->BitsPerPixel == 32)
 		{
 			bSuccess = TgaImageSupportImpl::DecompressTGA_32bpp(TgaHeader, TGABufferLenght, TextureData);
@@ -334,11 +338,15 @@ bool DecompressTGA_helper( const FTGAFileHeader* TgaHeader, const int64 TGABuffe
 	// Support for alpha stored as pseudo-color 8-bit TgaHeader
 	else if(TgaHeader->ColorMapType == 1 && TgaHeader->ImageTypeCode == 1 && TgaHeader->BitsPerPixel == 8)
 	{
+		check( TextureDataSize == TgaHeader->Width * TgaHeader->Height * 1 );
+
 		bSuccess = TgaImageSupportImpl::DecompressTGA_8bpp(TgaHeader, TGABufferLenght, (uint8*)TextureData);
 	}
 	// standard grayscale
 	else if(TgaHeader->ColorMapType == 0 && TgaHeader->ImageTypeCode == 3 && TgaHeader->BitsPerPixel == 8)
 	{
+		check( TextureDataSize == TgaHeader->Width * TgaHeader->Height * 1 );
+
 		bSuccess = TgaImageSupportImpl::DecompressTGA_8bpp(TgaHeader, TGABufferLenght, (uint8*)TextureData);
 	}
 	else
@@ -364,6 +372,8 @@ bool DecompressTGA_helper( const FTGAFileHeader* TgaHeader, const int64 TGABuffe
 		int32 NumBlocksX = TgaHeader->Width;
 		int32 NumBlocksY = TgaHeader->Height;
 		int32 BlockBytes = TgaHeader->BitsPerPixel == 8 ? 1 : 4;
+		
+		check( TextureDataSize == NumBlocksX * NumBlocksY * BlockBytes );
 
 		uint8* MipData = (uint8*)TextureData;
 
diff --git a/Engine/Source/Runtime/ImageWrapper/Public/Formats/HdrImageWrapper.h b/Engine/Source/Runtime/ImageWrapper/Public/Formats/HdrImageWrapper.h
--- a/Engine/Source/Runtime/ImageWrapper/Public/Formats/HdrImageWrapper.h
+++ b/Engine/Source/Runtime/ImageWrapper/Public/Formats/HdrImageWrapper.h
@@ -9,6 +9,7 @@
 // http://paulbourke.net/dataformats/pic/
 
 /** To load the HDR file image format. Does not support all possible types HDR formats (e.g. xyze is not supported) */
+//  does not use ImageWrapperBase , unlike all other imagewrappers
 class IMAGEWRAPPER_API FHdrImageWrapper : public IImageWrapper
 {
 public:
@@ -19,19 +20,23 @@ public:
 	// IIMageWrapper Interface begin
 	virtual bool SetCompressed(const void* InCompressedData, int64 InCompressedSize) override;
 	virtual bool SetRaw(const void* InRawData, int64 InRawSize, const int32 InWidth, const int32 InHeight, const ERGBFormat InFormat, const int32 InBitDepth, const int32 InBytesPerRow) override;
-	virtual bool SetAnimationInfo(int32 InNumFrames, int32 InFramerate) override;
 	virtual TArray64<uint8> GetCompressed(int32 Quality = 0) override;
 	virtual bool GetRaw(const ERGBFormat InFormat, int32 InBitDepth, TArray64<uint8>& OutRawData) override;
+	
+	virtual bool CanSetRawFormat(const ERGBFormat InFormat, const int32 InBitDepth) const override;
+	virtual ERawImageFormat::Type GetSupportedRawFormat(const ERawImageFormat::Type InFormat) const override;
 
 	virtual int32 GetWidth() const override;
 	virtual int32 GetHeight() const override;
 	virtual int32 GetBitDepth() const override;
 	virtual ERGBFormat GetFormat() const override;
 
-	virtual int32 GetNumFrames() const override;
-	virtual int32 GetFramerate() const override;
+	virtual bool SetAnimationInfo_DEPRECATED(int32 InNumFrames, int32 InFramerate) override;
+	virtual int32 GetNumFrames_DEPRECATED() const override;
+	virtual int32 GetFramerate_DEPRECATED() const override;
 	// IImageWrapper Interface end
 
+	// GetErrorMessage : nice idea, but not virtual, never called
 	const FText& GetErrorMessage() const;
 
 	void FreeCompressedData();
@@ -51,6 +56,7 @@ private:
 	const uint8* RGBDataStart = nullptr;
 
 	TArray64<uint8> CompressedDataHolder;
+	TArray64<uint8> RawDataHolder;
 
 	/** INDEX_NONE if not valid */
 	int32 Width = INDEX_NONE;
diff --git a/Engine/Source/Runtime/ImageWrapper/Public/IImageWrapper.h b/Engine/Source/Runtime/ImageWrapper/Public/IImageWrapper.h
--- a/Engine/Source/Runtime/ImageWrapper/Public/IImageWrapper.h
+++ b/Engine/Source/Runtime/ImageWrapper/Public/IImageWrapper.h
@@ -6,7 +6,7 @@
 #include "Containers/ArrayView.h"
 #include "CoreTypes.h"
 #include "Templates/SharedPointer.h"
-
+#include "ImageCore.h"
 
 /**
  * Enumerates the types of image formats this class can handle.
@@ -79,6 +79,12 @@ enum class ERGBFormat : int8
  * Enumerates available image compression qualities.
  * 
  * JPEG interprets Quality as 1-100
+ * JPEG default quality is 85 , Uncompressed means 100
+ * 
+ * Negative qualities set PNG zlib level
+ * PNG interprets "Uncompressed" as zlib level 0 (none)
+ * 
+ * EXR respects the "Uncompressed" flag
  */
 enum class EImageCompressionQuality : uint8
 {
@@ -89,22 +95,35 @@ enum class EImageCompressionQuality : uint8
 
 /**
  * Interface for image wrappers.
+ *
+ * to Encode:
+ *	SetRaw() then GetCompressed()
+ * to Decode :
+ *  SetCompressed() then GetRaw()
+ *
+ * in general, direct use of the IImageWrapper interface is now discouraged
+ * use ImageWrapperModule CompressImage/DecompressImage instead.
  */
 class IImageWrapper
 {
 public:
 
+
 	/**  
-	 * Sets the compressed data.
+	 * Sets the compressed data.  Can then call GetRaw().
 	 *
 	 * @param InCompressedData The memory address of the start of the compressed data.
 	 * @param InCompressedSize The size of the compressed data parsed.
 	 * @return true if data was the expected format.
+	 * 
+	 * after SetCompressed, image info queries like GetWidth and GetBitDepth are allowed
+	 * call GetRaw to get the decoded bits
+	 * decompression is not done until GetRaw
 	 */
 	virtual bool SetCompressed(const void* InCompressedData, int64 InCompressedSize) = 0;
 
 	/**  
-	 * Sets the compressed data.
+	 * Sets the raw image data.  Prepares to call GetCompressed() next.
 	 *
 	 * @param InRawData The memory address of the start of the raw data.
 	 * @param InRawSize The size of the compressed data parsed.
@@ -114,50 +133,59 @@ public:
 	 * @param InBitDepth the bit-depth per channel, normally 8.
 	 * @param InBytesPerRow the number of bytes between rows, 0 = tightly packed rows with no padding.
 	 * @return true if data was the expected format.
+	 * 
+	 * you must not SetRaw() with a format unless it passes CanSetRawFormat()
+	 * deprecated : avoid direct calls to SetRaw(), use ImageWrapperModule CompressImage instead
+	 * do not use InBytesPerRow, it is ignored
+	 * SetRaw does not take gamma information
+	 *  assumes U8 = SRGB and all else = Linear
 	 */
 	virtual bool SetRaw(const void* InRawData, int64 InRawSize, const int32 InWidth, const int32 InHeight, const ERGBFormat InFormat, const int32 InBitDepth, const int32 InBytesPerRow = 0) = 0;
 
-	/**
-	 * Set information for animated formats
-	 * @param InNumFrames The number of frames in the animation (the RawData from SetRaw will need to be a multiple of NumFrames)
-	 * @param InFramerate The playback rate of the animation
-	 * @return true if successful
+	/* CanSetRawFormat returns true if SetRaw will accept this format */
+	virtual bool CanSetRawFormat(const ERGBFormat InFormat, const int32 InBitDepth) const = 0;
+
+	/* returns InFormat if supported, else maps to something supported
+	 * the returned format will pass CanSetRawFormat()
 	 */
-	virtual bool SetAnimationInfo(int32 InNumFrames, int32 InFramerate) = 0;
+	virtual ERawImageFormat::Type GetSupportedRawFormat(const ERawImageFormat::Type InFormat) const = 0;
+
 
 	/**
-	 * Gets the compressed data.
+	 * Gets the compressed data.  (call SetRaw first)
 	 * (Note: It may consume the data set in the SetCompressed function if it was set before)
 	 * 
-	 * @return Array of the compressed data.
+	 * @return Array of the compressed data.  returns empty array on failure
+	 * 
 	 */
 	virtual TArray64<uint8> GetCompressed(int32 Quality = 0) = 0;
 
 	/**
-	 * Get the raw version of the image and write to the array view
-	 * (Note: It may consume the data set in the SetRaw function if it was set before)
-	 *
-	 * @param InFormat How we want to manipulate the RGB data.
-	 * @param InBitDepth The output bit-depth per channel, normally 8.
-	 * @param OutRawData Will contain the uncompressed raw data.
-	 * @return true on success, false otherwise.
-	 */
-	virtual bool GetRaw(const ERGBFormat InFormat, int32 InBitDepth, TArrayView64<uint8> OutRawData)
+	* GetRaw after SetCompressed
+	* fills the raw data in the native format/depth contained in the file
+	* Do not use the GetRaw() variants that take format/depth arguments.
+	* 
+	* @param OutRawData    Filled with raw image data.
+	*
+	* prefer GetRawImage instead.
+	*/
+	bool GetRaw(TArray64<uint8>& OutRawData)
 	{
-		// Todo We should make this function virtual pure and implement this properly in each ImageWrapper
-		TArray64<uint8> TmpRawData;
-		if (GetRaw(InFormat, InBitDepth, TmpRawData))
-		{
-			if (ensureMsgf(TmpRawData.Num() == OutRawData.Num(), TEXT("The view doesn't have the proper size to receive the texture.")))
-			{
-				FPlatformMemory::Memcpy(OutRawData.GetData(), TmpRawData.GetData(), OutRawData.Num());
-				return true;
-			}
-		}
-
-		return false;
+		ERGBFormat Format = GetFormat();
+		int32 BitDepth = GetBitDepth();
+		return GetRaw(Format,BitDepth,OutRawData);
 	}
+	
+	/* Decode the image file data from SetCompressed() into an FImage
+	* OutImage is allocated and attributes are filled
+	*	 passed-in contents of OutImage are destroyed
+	*	OutImage.Format is ignored
+	*
+	* @param OutImage	 Filled with the image
+	*/
+	bool GetRawImage(FImage & OutImage);
 
+	
 	/**  
 	 * Gets the raw data.
 	 * (Note: It may consume the data set in the SetRaw function if it was set before)
@@ -166,9 +194,12 @@ public:
 	 * @param InBitDepth The output bit-depth per channel, normally 8.
 	 * @param OutRawData Will contain the uncompressed raw data.
 	 * @return true on success, false otherwise.
+	 *
+	 * this is often broken, should only be used with InFormat == GetFormat()
+	 * DEPRECATED , use GetRaw() with 1 argument or GetRawImage()
 	 */
 	virtual bool GetRaw(const ERGBFormat InFormat, int32 InBitDepth, TArray64<uint8>& OutRawData) = 0;
-
+	
 	/**
 	 * Gets the raw data in a TArray. Only use this if you're certain that the image is less than 2 GB in size.
 	 * Prefer using the overload which takes a TArray64 in general.
@@ -178,6 +209,9 @@ public:
 	 * @param InBitDepth The output bit-depth per channel, normally 8.
 	 * @param OutRawData Will contain the uncompressed raw data.
 	 * @return true on success, false otherwise.
+	 *
+	 * this is often broken, should only be used with InFormat == GetFormat()
+	 * DEPRECATED , use GetRaw() with 1 argument or GetRawImage()
 	 */
 	bool GetRaw(const ERGBFormat InFormat, int32 InBitDepth, TArray<uint8>& OutRawData)
 	{
@@ -192,6 +226,33 @@ public:
 			return false;
 		}
 	}
+	
+	/**
+	 * Get the raw version of the image and write to the array view
+	 * (Note: It may consume the data set in the SetRaw function if it was set before)
+	 *
+	 * @param InFormat How we want to manipulate the RGB data.
+	 * @param InBitDepth The output bit-depth per channel, normally 8.
+	 * @param OutRawData Will contain the uncompressed raw data.
+	 * @return true on success, false otherwise.
+	 *
+	 * this is often broken, should only be used with InFormat == GetFormat()
+	 * DEPRECATED , use GetRaw() with 1 argument or GetRawImage()
+	 */
+	bool GetRaw(const ERGBFormat InFormat, int32 InBitDepth, TArrayView64<uint8> OutRawData)
+	{
+		TArray64<uint8> TmpRawData;
+		if (GetRaw(InFormat, InBitDepth, TmpRawData))
+		{
+			if (ensureMsgf(TmpRawData.Num() == OutRawData.Num(), TEXT("The view doesn't have the proper size to receive the texture.")))
+			{
+				FPlatformMemory::Memcpy(OutRawData.GetData(), TmpRawData.GetData(), OutRawData.Num());
+				return true;
+			}
+		}
+
+		return false;
+	}
 
 	/**
 	 * Gets the width of the image.
@@ -213,6 +274,9 @@ public:
 	 * Gets the bit depth of the image.
 	 *
 	 * @return The bit depth per-channel of the image.
+	 *
+	 * Beware several of the old wrappers (BMP,TGA) incorrectly used to return bits per *color* not per channel
+	 * they now correctly return per-channel.
 	 */
 	virtual int32 GetBitDepth() const = 0;
 
@@ -223,16 +287,54 @@ public:
 	 * @return The format the image data is in
 	 */
 	virtual ERGBFormat GetFormat() const = 0;
+	
+	/* Should the pixels be treated as sRGB encoded? (or Linear)
+	* 
+	 * @@!! TODO: implement GetSRGB() on the formats that can report it (eg. DDS?)
+	 */
+	virtual bool GetSRGB() const
+	{
+		// sRGB is guessed from bit depth
+		// 8 = on
+		return GetBitDepth() == 8 && GetFormat() != ERGBFormat::BGRE;
+	}
+
+	// call these from ImageWrapperModule.h , see documentation there
+	IMAGEWRAPPER_API static void ConvertRawImageFormat(ERawImageFormat::Type RawFormat, ERGBFormat & OutFormat,int & OutBitDepth);
+	IMAGEWRAPPER_API static ERawImageFormat::Type ConvertRGBFormat(ERGBFormat RGBFormat,int BitDepth,bool * bIsExactMatch = nullptr);
+	
+	/* get the current image format, mapped to an ERawImageFormat
+	 * if ! *bIsExactMatch , conversion is needed
+	 * can call after SetCompressed()
+	 */
+	ERawImageFormat::Type GetClosestRawImageFormat(bool * bIsExactMatch = nullptr) const
+	{
+		ERGBFormat Format = GetFormat();
+		int BitDepth = GetBitDepth();
+		ERawImageFormat::Type Ret = ConvertRGBFormat(Format,BitDepth,bIsExactMatch);
+		return Ret;
+	}
+
+
+	// @@!! is this animation stuff used at all? -> it seems not
+
+	/**
+	 * Set information for animated formats
+	 * @param InNumFrames The number of frames in the animation (the RawData from SetRaw will need to be a multiple of NumFrames)
+	 * @param InFramerate The playback rate of the animation
+	 * @return true if successful
+	 */
+	virtual bool SetAnimationInfo_DEPRECATED(int32 InNumFrames, int32 InFramerate) = 0;
 
 	/**
 	 * @return The number of frames in an animated image
 	 */
-	virtual int32 GetNumFrames() const = 0;
+	virtual int32 GetNumFrames_DEPRECATED() const = 0;
 
 	/**
 	 * @return The playback framerate of animated images (or 0 for non-animated)
 	 */
-	virtual int32 GetFramerate() const = 0;
+	virtual int32 GetFramerate_DEPRECATED() const = 0;
 
 public:
 
diff --git a/Engine/Source/Runtime/ImageWrapper/Public/IImageWrapperModule.h b/Engine/Source/Runtime/ImageWrapper/Public/IImageWrapperModule.h
--- a/Engine/Source/Runtime/ImageWrapper/Public/IImageWrapperModule.h
+++ b/Engine/Source/Runtime/ImageWrapper/Public/IImageWrapperModule.h
@@ -5,9 +5,8 @@
 #include "CoreTypes.h"
 #include "Modules/ModuleInterface.h"
 #include "Templates/SharedPointer.h"
-
-class IImageWrapper;
-enum class EImageFormat : int8;
+#include "ImageCore.h"
+#include "IImageWrapper.h"
 
 
 /**
@@ -18,11 +17,40 @@ class IImageWrapperModule
 {
 public:
 
+
+	/**  
+	 * Convert input FImage into a file-format encoded array.
+ 	 * in ImageWrapper land, "Compress" means "put in file format"
+	 * OutData is filled with the file format encoded data
+	 * lossy conversion of the pixel type may be done if necessary
+	 * eg. if you pass F32 float pixels to write to BMP, they will be converted to SRGB U8 BGRA8
+	 * 
+	 * @param OutData		Filled with image-format data
+	 * @param ToFormat		Image format to encode to
+	 * @param InImage		Image to encode
+	 * @param Quality		50-100 for JPEG, or EImageCompressionQuality
+	 * 
+	**/
+	virtual bool CompressImage(TArray64<uint8> & OutData, EImageFormat ToFormat, const FImageView & InImage, int32 Quality = 0) = 0;
+	
+	/* Read an image from file format encoded data.
+	 * ImageWrapper calls this a "decompress"
+	 * OutImage is allocated and filled, any existing contents are discarded
+	 * 
+	 * @param InCompressedData	Image format encoded bytes to read; format is automatically deduced from the content
+	 * @param InCompressedSize	Size of InCompressedData in bytes
+	 * @param OutImage			Filled with Image that is read.  Allocated.
+	 */
+	virtual bool DecompressImage(const void* InCompressedData, int64 InCompressedSize, FImage & OutImage) = 0;
+
 	/**  
-	 * Create a helper of a specific type 
+	 * Create an IImageWrapper helper of a specific type 
 	 *
 	 * @param InFormat - The type of image we want to deal with
 	 * @return The helper base class to manage the data
+	 * 
+	 * EImageFormat is a compressor / file format , not a pixel format
+	 * Deprecated.  Prefer CompressImage/DecompressImage.
 	 */
 	virtual TSharedPtr<IImageWrapper> CreateImageWrapper(const EImageFormat InFormat) = 0;
 
@@ -36,6 +64,30 @@ public:
 	 */
 	virtual EImageFormat DetectImageFormat(const void* InCompressedData, int64 InCompressedSize) = 0;
 
+	/* Name can be a full name like "xx.png" or just the extension part, like "png"
+	 * returns EImageFormat::Invalid if no supported image extension found
+	 */
+	virtual EImageFormat GetImageFormatFromExtension(const TCHAR * Name) = 0;
+
+	/* returns extension, not including the "." , 3 or 4 chars */
+	virtual const TCHAR * GetExtension(EImageFormat Format) = 0;
+
+	/* get a good default output image format for a pixel format */
+	virtual EImageFormat GetDefaultOutputFormat(ERawImageFormat::Type RawFormat) = 0;
+
+	/* Convert an ImageWrapper style {ERGBFormat+BitDepth} into an ERawImageFormat for FImage
+	* returns ERawImageFormat::Invalid if no mapping is possible
+	* bIsExactMatch is filled with whether the formats are an exact match or not
+	*	if not, conversion is needed
+	*/
+	virtual ERawImageFormat::Type ConvertRGBFormat(ERGBFormat Format,int BitDepth,bool * bIsExactMatch = nullptr) = 0;
+	
+	/**
+	* Convert an FImage ERawImageFormat into an ImageWrapper style {ERGBFormat+BitDepth} 
+	* mapping is always possible and requires no conversion
+	*/
+	virtual void ConvertRawImageFormat(ERawImageFormat::Type RawFormat, ERGBFormat & OutFormat,int & OutBitDepth) = 0;
+
 public:
 
 	/** Virtual destructor. */
diff --git a/Engine/Source/Runtime/ImageWriteQueue/ImageWriteQueue.Build.cs b/Engine/Source/Runtime/ImageWriteQueue/ImageWriteQueue.Build.cs
--- a/Engine/Source/Runtime/ImageWriteQueue/ImageWriteQueue.Build.cs
+++ b/Engine/Source/Runtime/ImageWriteQueue/ImageWriteQueue.Build.cs
@@ -10,6 +10,7 @@ public class ImageWriteQueue : ModuleRules
 			new string[] {
 				"Core",
 				"CoreUObject",
+				"ImageCore",
 				"ImageWrapper",
 			}
 		);
diff --git a/Engine/Source/Runtime/ImageWriteQueue/Private/ImageWriteTask.cpp b/Engine/Source/Runtime/ImageWriteQueue/Private/ImageWriteTask.cpp
--- a/Engine/Source/Runtime/ImageWriteQueue/Private/ImageWriteTask.cpp
+++ b/Engine/Source/Runtime/ImageWriteQueue/Private/ImageWriteTask.cpp
@@ -12,80 +12,6 @@
 #include "Modules/ModuleManager.h"
 #include "IImageWrapperModule.h"
 
-struct FGlobalImageWrappers
-{
-	IImageWrapper* FindOrCreateImageWrapper(EImageFormat InFormat)
-	{
-		FScopeLock ScopeLock(&ImageWrappersCriticalSection);
-
-		// Try and find an available image wrapper of the correct format first
-		for (int32 Index = 0; Index < AvailableImageWrappers.Num(); ++Index)
-		{
-			if (AvailableImageWrappers[Index].Get<0>() == InFormat)
-			{
-				IImageWrapper* Wrapper = AvailableImageWrappers[Index].Get<1>();
-				AvailableImageWrappers.RemoveAtSwap(Index, 1, false);
-				return Wrapper;
-			}
-		}
-
-		// Create a new one if none other could be used
-		IImageWrapperModule* ImageWrapperModule = FModuleManager::GetModulePtr<IImageWrapperModule>("ImageWrapper");
-		if (!ensure(ImageWrapperModule))
-		{
-			return nullptr;
-		}
-
-		TSharedPtr<IImageWrapper> NewImageWrapper = ImageWrapperModule->CreateImageWrapper(InFormat);
-		if (!ensureMsgf(NewImageWrapper.IsValid(), TEXT("Unable to create an image wrapper for the desired format.")))
-		{
-			return nullptr;
-		}
-
-		AllImageWrappers.Add(MakeTuple(InFormat, NewImageWrapper.ToSharedRef()));
-		return NewImageWrapper.Get();
-	}
-
-	void ReturnImageWrapper(IImageWrapper* InWrapper)
-	{
-		FScopeLock ScopeLock(&ImageWrappersCriticalSection);
-
-		// Try and find an available image wrapper of the correct format first
-		for (const TTuple<EImageFormat, TSharedRef<IImageWrapper>>& Pair : AllImageWrappers)
-		{
-			if (&Pair.Get<1>().Get() == InWrapper)
-			{
-				AvailableImageWrappers.Add(MakeTuple(Pair.Get<0>(), InWrapper));
-				return;
-			}
-		}
-
-		checkf(false, TEXT("Unable to find image wrapper in list of owned wrappers - this is invalid"));
-	}
-
-private:
-	FCriticalSection ImageWrappersCriticalSection;
-	TArray< TTuple<EImageFormat, IImageWrapper*> > AvailableImageWrappers;
-	TArray< TTuple<EImageFormat, TSharedRef<IImageWrapper>> > AllImageWrappers;
-} GImageWrappers;
-
-
-
-static const TCHAR* GetFormatExtension(EImageFormat InImageFormat)
-{
-	switch (InImageFormat)
-	{
-	case EImageFormat::PNG:           return TEXT(".png");
-	case EImageFormat::JPEG:          return TEXT(".jpg");
-	case EImageFormat::GrayscaleJPEG: return TEXT(".jpg");
-	case EImageFormat::BMP:           return TEXT(".bmp");
-	case EImageFormat::ICO:           return TEXT(".ico");
-	case EImageFormat::EXR:           return TEXT(".exr");
-	case EImageFormat::ICNS:          return TEXT(".icns");
-	}
-	return nullptr;
-}
-
 bool FImageWriteTask::RunTask()
 {
 	bool bSuccess = WriteToDisk();
@@ -106,46 +32,10 @@ void FImageWriteTask::OnAbandoned()
 	}
 }
 
-bool FImageWriteTask::InitializeWrapper(IImageWrapper* InWrapper, EImageFormat WrapperFormat)
-{
-	const void* RawPtr = nullptr;
-	int64 SizeBytes = 0;
-
-	if (PixelData->GetRawData(RawPtr, SizeBytes))
-	{
-		uint8      BitDepth    = PixelData->GetBitDepth();
-		FIntPoint  Size        = PixelData->GetSize();
-		ERGBFormat PixelLayout = PixelData->GetPixelLayout();
-
-		return InWrapper->SetRaw(RawPtr, SizeBytes, Size.X, Size.Y, PixelLayout, BitDepth);
-	}
-
-	return false;
-}
-
-bool FImageWriteTask::WriteBitmap()
-{
-	uint8      NumChannels = PixelData->GetNumChannels();
-	uint8      BitDepth    = PixelData->GetBitDepth();
-	FIntPoint  Size        = PixelData->GetSize();
-
-	if (BitDepth != 8 || NumChannels != 4)
-	{
-		return false;
-	}
-
-	const void* RawPtr = nullptr;
-	int64 SizeBytes = 0;
-
-	if (PixelData->GetRawData(RawPtr, SizeBytes))
-	{
-		return FFileHelper::CreateBitmap(*Filename, Size.X, Size.Y, static_cast<const FColor*>(RawPtr));
-	}
-	return false;
-}
-
 void FImageWriteTask::PreProcess()
 {
+	TRACE_CPUPROFILER_EVENT_SCOPE(ImageWriteTask.PreProcess);
+	
 	for (const FPixelPreProcessor& PreProcessor : PixelPreProcessors)
 	{
 		// PreProcessors are assumed to be valid. Fetch the Data pointer each time
@@ -157,11 +47,16 @@ void FImageWriteTask::PreProcess()
 
 bool FImageWriteTask::WriteToDisk()
 {
-	// Ensure that the payload filename has the correct extension for the format (have to special case jpeg since they can be both *.jpg and *.jpeg)
-	const TCHAR* FormatExtension = GetFormatExtension(Format);
-	if (FormatExtension && !Filename.EndsWith(FormatExtension) && (Format != EImageFormat::JPEG || !Filename.EndsWith(TEXT(".jpeg"))))
+	TRACE_CPUPROFILER_EVENT_SCOPE(ImageWriteTask.WriteToDisk);
+
+	static FName ImageWrapperName("ImageWrapper");
+	IImageWrapperModule& ImageWrapperModule = FModuleManager::LoadModuleChecked<IImageWrapperModule>(ImageWrapperName);
+		
+	// Ensure that the payload filename has the correct extension for the format
+	if ( ImageWrapperModule.GetImageFormatFromExtension(*Filename) != Format )
 	{
-		Filename = FPaths::GetBaseFilename(Filename, false) + FormatExtension;
+		const TCHAR* FormatExtension = ImageWrapperModule.GetExtension(Format);
+		Filename = FPaths::GetBaseFilename(Filename, false) + TEXT('.') + FormatExtension;
 	}
 
 	bool bSuccess = EnsureWritableFile();
@@ -169,43 +64,41 @@ bool FImageWriteTask::WriteToDisk()
 	if (bSuccess)
 	{
 		PreProcess();
+		
+		FImagePixelData* Data = PixelData.Get();
+		FImageView Image = Data->GetImageView();
+		if ( Image.RawData == nullptr )
+		{
+			UE_LOG(LogImageWriteQueue, Error, TEXT("Failed to write image to '%s'. Couldn't get pixels."), *Filename);
+			return false;
+		}
 
-		// bitmap support with IImageWrapper is flaky so it needs its own codepath for now
-		if (Format == EImageFormat::BMP)
+		TArray64<uint8> CompressedFile;
+		if ( ! ImageWrapperModule.CompressImage(CompressedFile,Format,Image,CompressionQuality) )
 		{
-			bSuccess = WriteBitmap();
+			UE_LOG(LogImageWriteQueue, Error, TEXT("Failed to write image to '%s'. CompressImage failed."), *Filename);
+			return false;
 		}
-		else
+		
+		uint64 TotalNumberOfBytes, NumberOfFreeBytes;
+		if (FPlatformMisc::GetDiskTotalAndFreeSpace(FPaths::GetPath(Filename), TotalNumberOfBytes, NumberOfFreeBytes))
 		{
-			IImageWrapper* ImageWrapper = GImageWrappers.FindOrCreateImageWrapper(Format);
-			if (ImageWrapper)
+			const uint64 HeadRoom = 1024*1024;
+			if (NumberOfFreeBytes < (uint64)CompressedFile.Num() + HeadRoom)
 			{
-				if (InitializeWrapper(ImageWrapper, Format))
-				{
-					const TArray64<uint8> CompressedFile = ImageWrapper->GetCompressed(CompressionQuality);
-					uint64 TotalNumberOfBytes, NumberOfFreeBytes;
-					if (FPlatformMisc::GetDiskTotalAndFreeSpace(FPaths::GetPath(Filename), TotalNumberOfBytes, NumberOfFreeBytes))
-					{
-						if (NumberOfFreeBytes < (uint64)CompressedFile.Num() + 4096)
-						{
-							UE_LOG(LogImageWriteQueue, Error, TEXT("Failed to write image to '%s'. There is not enough free space on the disk."), *Filename);
-							return false;
-						}
-					}
-					else
-					{
-						uint32 ErrorCode = FPlatformMisc::GetLastError();
-						TCHAR ErrorBuffer[1024];
-						FPlatformMisc::GetSystemErrorMessage(ErrorBuffer, 1024, ErrorCode);
-						UE_LOG(LogImageWriteQueue, Warning, TEXT("Failed to check free space for %s. Error: %u (%s)"), *FPaths::GetPath(Filename), ErrorCode, ErrorBuffer);
-					}
-					IFileManager* FileManager = &IFileManager::Get();
-					bSuccess = FFileHelper::SaveArrayToFile(CompressedFile, *Filename);
-				}
-
-				GImageWrappers.ReturnImageWrapper(ImageWrapper);
+				UE_LOG(LogImageWriteQueue, Error, TEXT("Failed to write image to '%s'. There is not enough free space on the disk."), *Filename);
+				return false;
 			}
 		}
+		else
+		{
+			uint32 ErrorCode = FPlatformMisc::GetLastError();
+			TCHAR ErrorBuffer[1024];
+			FPlatformMisc::GetSystemErrorMessage(ErrorBuffer, 1024, ErrorCode);
+			UE_LOG(LogImageWriteQueue, Warning, TEXT("Failed to check free space for %s. Error: %u (%s)"), *FPaths::GetPath(Filename), ErrorCode, ErrorBuffer);
+		}
+
+		bSuccess = FFileHelper::SaveArrayToFile(CompressedFile, *Filename);
 	}
 
 	if (!bSuccess)
diff --git a/Engine/Source/Runtime/ImageWriteQueue/Public/ImagePixelData.h b/Engine/Source/Runtime/ImageWriteQueue/Public/ImagePixelData.h
--- a/Engine/Source/Runtime/ImageWriteQueue/Public/ImagePixelData.h
+++ b/Engine/Source/Runtime/ImageWriteQueue/Public/ImagePixelData.h
@@ -6,11 +6,13 @@
 #include "Math/IntPoint.h"
 #include "IImageWrapper.h"
 #include "Templates/UniquePtr.h"
+#include "ImageCore.h"
 
 class FFloat16Color;
 template<typename PixelType> struct TImagePixelDataTraits;
 
 
+// todo : use ImageCore ERawImageFormat instead
 enum class EImagePixelType
 {
 	Color,
@@ -25,9 +27,16 @@ struct IImagePixelDataPayload
 
 typedef TSharedPtr<IImagePixelDataPayload, ESPMode::ThreadSafe> FImagePixelPayloadPtr;
 
+
+// todo : use ImageCore FImage instead
+//  get rid of this whole class and TImagePixelData too
+//	just use an FImage instead
 struct FImagePixelData
 {
 	virtual ~FImagePixelData() {}
+	
+	// NOTE : before, U8 would be written to EXR *linear*
+	//	now it will get gamma corrected if bSRGB (which is on by default)
 
 	/**
 	 * Retrieve the type of this data
@@ -68,6 +77,52 @@ struct FImagePixelData
 	{
 		return NumChannels;
 	}
+	
+	/**
+	 * Get the pixel data as an FImage
+	 */
+	FImageView GetImageView() const
+	{
+		const void* RawPtr    = nullptr;
+		int64       SizeBytes = 0;
+
+		FImageView Ret;
+
+		if ( ! GetRawData(RawPtr, SizeBytes) )
+		{
+			return FImageView();
+		}
+		
+		check( NumChannels == 4 );
+
+		switch(Type)
+		{
+		case EImagePixelType::Color:
+			check(PixelLayout == ERGBFormat::BGRA );
+			check(BitDepth == 8 );
+			Ret = FImageView( (const FColor *)RawPtr, Size.X, Size.Y, bSRGB ? EGammaSpace::sRGB : EGammaSpace::Linear );
+			break;
+
+		case EImagePixelType::Float16:
+			check(PixelLayout == ERGBFormat::RGBAF );
+			check(BitDepth == 16 );
+			Ret = FImageView( (const FFloat16Color *)RawPtr, Size.X, Size.Y );
+			break;
+
+		case EImagePixelType::Float32:
+			check(PixelLayout == ERGBFormat::RGBAF );
+			check(BitDepth == 32 );
+			Ret = FImageView( (const FLinearColor *)RawPtr, Size.X, Size.Y );
+			break;
+
+		default:
+			check(0);
+			return FImageView();
+		}
+
+		check( Ret.GetImageSizeBytes() == SizeBytes );
+		return Ret;
+	}
 
 	/**
 	 * Check that this data is the size it should be
@@ -140,6 +195,9 @@ struct FImagePixelData
 	template<typename T>
 	const T* GetPayload() const { return static_cast<T*>(Payload.Get()); }
 
+	bool GetSRGB() const { return bSRGB; }
+	void SetSRGB(bool InSRGB) { bSRGB = InSRGB; }
+
 protected:
 
 	FImagePixelData(const FIntPoint& InSize, EImagePixelType InPixelType, ERGBFormat InPixelLayout, uint8 InBitDepth, uint8 InNumChannels, FImagePixelPayloadPtr InPayload)
@@ -149,7 +207,10 @@ protected:
 		, BitDepth(InBitDepth)
 		, NumChannels(InNumChannels)
 		, Payload(InPayload)
-	{}
+	{
+		// FColor is sRGB by default, floats are Linear
+		bSRGB = ( InPixelType == EImagePixelType::Color );
+	}
 
 private:
 
@@ -175,6 +236,9 @@ private:
 	/** Number of channels in the data */
 	uint8 NumChannels;
 
+	/** Is FColor SRGB or Linear?  Floats are always Linear and ignore this */
+	bool bSRGB;
+
 	/** Optional user-specified payload */
 	FImagePixelPayloadPtr Payload;
 };
diff --git a/Engine/Source/Runtime/ImageWriteQueue/Public/ImageWriteTask.h b/Engine/Source/Runtime/ImageWriteQueue/Public/ImageWriteTask.h
--- a/Engine/Source/Runtime/ImageWriteQueue/Public/ImageWriteTask.h
+++ b/Engine/Source/Runtime/ImageWriteQueue/Public/ImageWriteTask.h
@@ -99,13 +99,6 @@ private:
 	 */
 	bool InitializeWrapper(IImageWrapper* InWrapper, EImageFormat WrapperFormat);
 
-	/**
-	 * Special case implementation for writing bitmap data due to deficiencies in the IImageWriter API (it can't set raw pixel data without trying to compress it, which asserts)
-	 *
-	 * @return true a bitmap was written out at the specified filename, false otherwise
-	 */
-	bool WriteBitmap();
-
 
 	/**
 	 * Run over all the processors for the pixel data
@@ -114,76 +107,6 @@ private:
 };
 
 
-/**
- * A pixel preprocessor for use with FImageWriteTask::PixelPreProcessor that does gamma correction as part of the threaded work
- */
-template<typename PixelType> struct TAsyncGammaCorrect;
-
-
-template<>
-struct TAsyncGammaCorrect<FColor>
-{
-	float Gamma;
-	TAsyncGammaCorrect(float InGamma) : Gamma(InGamma) {}
-
-	void operator()(FImagePixelData* PixelData)
-	{
-		check(PixelData->GetType() == EImagePixelType::Color);
-
-		TImagePixelData<FColor>* ColorData = static_cast<TImagePixelData<FColor>*>(PixelData);
-		for (FColor& Pixel : ColorData->Pixels)
-		{
-			Pixel.A = (uint8)FMath::RoundToFloat(FMath::Pow(Pixel.A / 255.f, Gamma) * 255.f);
-			Pixel.R = (uint8)FMath::RoundToFloat(FMath::Pow(Pixel.R / 255.f, Gamma) * 255.f);
-			Pixel.G = (uint8)FMath::RoundToFloat(FMath::Pow(Pixel.G / 255.f, Gamma) * 255.f);
-			Pixel.B = (uint8)FMath::RoundToFloat(FMath::Pow(Pixel.B / 255.f, Gamma) * 255.f);
-		}
-	}
-};
-
-template<>
-struct TAsyncGammaCorrect<FFloat16Color>
-{
-	float Gamma;
-	TAsyncGammaCorrect(float InGamma) : Gamma(InGamma) {}
-
-	void operator()(FImagePixelData* PixelData)
-	{
-		check(PixelData->GetType() == EImagePixelType::Float16);
-
-		TImagePixelData<FFloat16Color>* Float16ColorData = static_cast<TImagePixelData<FFloat16Color>*>(PixelData);
-		for (FFloat16Color& Pixel : Float16ColorData->Pixels)
-		{
-			Pixel.A = FMath::Pow(Pixel.A.GetFloat(), Gamma);
-			Pixel.R = FMath::Pow(Pixel.R.GetFloat(), Gamma);
-			Pixel.G = FMath::Pow(Pixel.G.GetFloat(), Gamma);
-			Pixel.B = FMath::Pow(Pixel.B.GetFloat(), Gamma);
-		}
-	}
-};
-
-template<>
-struct TAsyncGammaCorrect<FLinearColor>
-{
-	float Gamma;
-	TAsyncGammaCorrect(float InGamma) : Gamma(InGamma) {}
-
-	void operator()(FImagePixelData* PixelData)
-	{
-		check(PixelData->GetType() == EImagePixelType::Float32);
-
-		TImagePixelData<FLinearColor>* LinearColorData = static_cast<TImagePixelData<FLinearColor>*>(PixelData);
-		for (FLinearColor& Pixel : LinearColorData->Pixels)
-		{
-			Pixel.A = FMath::Pow(Pixel.A, Gamma);
-			Pixel.R = FMath::Pow(Pixel.R, Gamma);
-			Pixel.G = FMath::Pow(Pixel.G, Gamma);
-			Pixel.B = FMath::Pow(Pixel.B, Gamma);
-		}
-	}
-};
-
-
 /**
  * A pixel preprocessor for use with FImageWriteTask::PixelPreProcessor that overwrites the alpha channel with a fixed value as part of the threaded work
  */
diff --git a/Engine/Source/Runtime/MediaUtils/Private/MediaRecorder.cpp b/Engine/Source/Runtime/MediaUtils/Private/MediaRecorder.cpp
--- a/Engine/Source/Runtime/MediaUtils/Private/MediaRecorder.cpp
+++ b/Engine/Source/Runtime/MediaUtils/Private/MediaRecorder.cpp
@@ -115,27 +115,6 @@ namespace MediaRecorderHelpers
 
 		return MoveTemp(PixelData);
 	};
-
-	TUniquePtr<TImagePixelData<FColor>> ConvertSampleInterpretedAsFloat16ColorToFColor(TSharedPtr<IMediaTextureSample, ESPMode::ThreadSafe> InSample, const FIntPoint InSize, int32 InNumChannels)
-	{
-		const int64 NumberOfTexel = InSize.Y * InSize.X;
-		TUniquePtr<TImagePixelData<FColor>> PixelData = MakeUnique<TImagePixelData<FColor>>(InSize);
-		PixelData->Pixels.Reset(NumberOfTexel);
-
-		const FFloat16Color* Buffer = reinterpret_cast<const FFloat16Color*>(InSample->GetBuffer());
-		for (int64 i = 0 ; i < NumberOfTexel; i++)
-		{
-			FColor Output;
-			Output.R = (uint8)(FMath::Clamp<float>(Buffer[0].R, 0.0f, 1.0f) * 255);
-			Output.G = (uint8)(FMath::Clamp<float>(Buffer[0].G, 0.0f, 1.0f) * 255);
-			Output.B = (uint8)(FMath::Clamp<float>(Buffer[0].B, 0.0f, 1.0f) * 255);
-			Output.A = (uint8)(FMath::Clamp<float>(Buffer[0].A, 0.0f, 1.0f) * 255);
-			PixelData->Pixels.Add(Output);
-			Buffer++;
-		}
-
-		return MoveTemp(PixelData);
-	}
 }
 
 /* FMediaRecorder implementation
@@ -252,9 +231,6 @@ void FMediaRecorder::TickRecording()
 				continue;
 			}
 
-			bool bIsGammaCorrectionPreProcessingEnabled = (TargetImageFormat == EImageFormat::EXR && Sample->IsOutputSrgb());
-			bool IsFloatFormatConversionNeeded = ((Sample->GetFormat() == EMediaTextureSampleFormat::FloatRGBA) && (TargetImageFormat != EImageFormat::EXR));
-
 			TUniquePtr<FImageWriteTask> ImageTask = MakeUnique<FImageWriteTask>();
 
 			// Set PixelData
@@ -273,21 +249,14 @@ void FMediaRecorder::TickRecording()
 				}
 
 				// Should we move the color buffer into a raw image data container.
-				bool bUseFMediaImagePixelData = bSetAlpha || (Sample->GetStride() != Size.X * NumChannels) || bIsGammaCorrectionPreProcessingEnabled || IsFloatFormatConversionNeeded;
+				bool bUseFMediaImagePixelData = bSetAlpha || (Sample->GetStride() != Size.X * NumChannels * (BitDepth/8));
 
 				if (bUseFMediaImagePixelData)
 				{
 					const int32 NumberOfTexel = Size.Y * Size.X;
 					if (Sample->GetFormat() == EMediaTextureSampleFormat::FloatRGBA)
 					{
-						if (IsFloatFormatConversionNeeded)
-						{
-							ImageTask->PixelData = MediaRecorderHelpers::ConvertSampleInterpretedAsFloat16ColorToFColor(Sample, Size, NumChannels);
-						}
-						else
-						{
-							ImageTask->PixelData = MediaRecorderHelpers::CreatePixelData<FFloat16Color>(Sample, Size, NumChannels);
-						}
+						ImageTask->PixelData = MediaRecorderHelpers::CreatePixelData<FFloat16Color>(Sample, Size, NumChannels);
 					}
 					else
 					{
@@ -319,11 +288,11 @@ void FMediaRecorder::TickRecording()
 					}
 				}
 			}
-
-			if (bIsGammaCorrectionPreProcessingEnabled)
+			
+			if (Sample->GetFormat() == EMediaTextureSampleFormat::CharBGRA)
 			{
-				const float DefaultGammaValue = 2.2f;
-				ImageTask->PixelPreProcessors.Add(TAsyncGammaCorrect<FColor>(DefaultGammaValue));
+				// IsOutputSrgb almost always true
+				ImageTask->PixelData->SetSRGB( Sample->IsOutputSrgb() );
 			}
 
 			ImageTask->Format = TargetImageFormat;
diff --git a/Engine/Source/Runtime/Renderer/Private/PostProcess/PostProcessVisualizeBuffer.cpp b/Engine/Source/Runtime/Renderer/Private/PostProcess/PostProcessVisualizeBuffer.cpp
--- a/Engine/Source/Runtime/Renderer/Private/PostProcess/PostProcessVisualizeBuffer.cpp
+++ b/Engine/Source/Runtime/Renderer/Private/PostProcess/PostProcessVisualizeBuffer.cpp
@@ -274,14 +274,8 @@ void AddDumpToFilePass(FRDGBuilder& GraphBuilder, FScreenPassTexture Input, cons
 			// Always write full alpha
 			ImageTask->PixelPreProcessors.Add(TAsyncAlphaWrite<FColor>(255));
 
-			if (ImageTask->Format == EImageFormat::EXR)
-			{
-				// Write FColors with a gamma curve. This replicates behavior that previously existed in ExrImageWrapper.cpp (see following overloads) that assumed
-				// any 8 bit output format needed linearizing, but this is not a safe assumption to make at such a low level:
-				// void ExtractAndConvertChannel(const uint8*Src, uint32 SrcChannels, uint32 x, uint32 y, float* ChannelOUT)
-				// void ExtractAndConvertChannel(const uint8*Src, uint32 SrcChannels, uint32 x, uint32 y, FFloat16* ChannelOUT)
-				ImageTask->PixelPreProcessors.Add(TAsyncGammaCorrect<FColor>(2.2f));
-			}
+			// ImageTask->PixelData should be sRGB
+			//  it will gamma correct automatically if written to EXR
 		}
 
 		HighResScreenshotConfig.ImageWriteQueue->Enqueue(MoveTemp(ImageTask));
diff --git a/Engine/Source/Runtime/TextureUtilitiesCommon/Private/NormalMapIdentification.cpp b/Engine/Source/Runtime/TextureUtilitiesCommon/Private/NormalMapIdentification.cpp
--- a/Engine/Source/Runtime/TextureUtilitiesCommon/Private/NormalMapIdentification.cpp
+++ b/Engine/Source/Runtime/TextureUtilitiesCommon/Private/NormalMapIdentification.cpp
@@ -151,16 +151,29 @@ public:
 
 	FLinearColor DoSampleColor( int32 X, int32 Y )
 	{
-		FLinearColor Result;
-		const uint8* PixelToSample = SourceTextureData + ((Y * TextureSizeX + X) * 8);
+		const uint8* PixelToSample = SourceTextureData + ((Y * TextureSizeX + X) * sizeof(FFloat16Color));
 
 		// this assume the normal map to be in linear (not the case if Photoshop converts a 8bit normalmap to float and saves it as 16bit dds)
-		Result.R = (float)((FFloat16*)PixelToSample)[0];
-		Result.G = (float)((FFloat16*)PixelToSample)[1];
-		Result.B = (float)((FFloat16*)PixelToSample)[2];
-		Result.A = (float)((FFloat16*)PixelToSample)[3];
+		
+		return ((const FFloat16Color*)PixelToSample)->GetFloats();
+	}
 
-		return Result;
+	float ScaleAndBiasComponent( float Value ) const
+	{
+		// no need to scale and bias floating point components.
+		return Value;
+	}
+};
+class SampleNormalMapPixelF32 : public NormalMapSamplerBase
+{
+public:
+	SampleNormalMapPixelF32() {}
+	~SampleNormalMapPixelF32() {}
+
+	FLinearColor DoSampleColor( int32 X, int32 Y )
+	{
+		const uint8* PixelToSample = SourceTextureData + ((Y * TextureSizeX + X) * sizeof(FLinearColor));
+		return *((const FLinearColor*)PixelToSample);
 	}
 
 	float ScaleAndBiasComponent( float Value ) const
@@ -351,6 +364,8 @@ public:
  */
 static bool IsTextureANormalMap( UTexture* Texture )
 {
+	TRACE_CPUPROFILER_EVENT_SCOPE(IsTextureANormalMap);
+
 #if NORMALMAP_IDENTIFICATION_TIMING
 	double StartSeconds = FPlatformTime::Seconds();
 #endif
@@ -381,9 +396,9 @@ static bool IsTextureANormalMap( UTexture* Texture )
 				bIsNormalMap = Analyzer.DoesTextureLookLikelyToBeANormalMap( Texture );
 			}
 			break;
-		case TSF_RGBA8:
+		case TSF_RGBA32F:
 			{
-				TNormalMapAnalyzer<SampleNormalMapPixelRGBA8> Analyzer;
+				TNormalMapAnalyzer<SampleNormalMapPixelF32> Analyzer;
 				bIsNormalMap = Analyzer.DoesTextureLookLikelyToBeANormalMap( Texture );
 			}
 			break;
