Exclude an unnecessary second-stage compilation pass after duplicating a Blueprint with a missing/invalid parent class, which could lead to a crash in certain scenarios.

diff --git a/Engine/Source/Editor/UnrealEd/Private/Kismet2/BlueprintEditorUtils.cpp b/Engine/Source/Editor/UnrealEd/Private/Kismet2/BlueprintEditorUtils.cpp
--- a/Engine/Source/Editor/UnrealEd/Private/Kismet2/BlueprintEditorUtils.cpp
+++ b/Engine/Source/Editor/UnrealEd/Private/Kismet2/BlueprintEditorUtils.cpp
@@ -2016,12 +2016,12 @@ void FBlueprintEditorUtils::PostDuplicateBlueprint(UBlueprint* Blueprint, bool b
 			UObject* NewCDO = Blueprint->GeneratedClass->GetDefaultObject();
 			check(NewCDO != nullptr);
 			UEditorEngine::CopyPropertiesForUnrelatedObjects(OldCDO, NewCDO);
-		}
 
-		if (!FBlueprintDuplicationScopeFlags::HasAnyFlag(FBlueprintDuplicationScopeFlags::NoExtraCompilation))
-		{
-			// And compile again to make sure they go into the generated class, get cleaned up, etc...
-			FKismetEditorUtilities::CompileBlueprint(Blueprint, EBlueprintCompileOptions::SkipGarbageCollection);
+			if (!FBlueprintDuplicationScopeFlags::HasAnyFlag(FBlueprintDuplicationScopeFlags::NoExtraCompilation))
+			{
+				// And compile again to make sure they go into the generated class, get cleaned up, etc...
+				FKismetEditorUtilities::CompileBlueprint(Blueprint, EBlueprintCompileOptions::SkipGarbageCollection);
+			}
 		}
 
 		// it can still keeps references to some external objects
