Make texture exporters generic
so all 2d formats support VT and UDIM export
allow JPEG export of stored jpeg bits (lossless)

diff --git a/Engine/Source/Editor/UnrealEd/Classes/Exporters/RenderTargetExporterHDR.h b/Engine/Source/Editor/UnrealEd/Classes/Exporters/RenderTargetExporterHDR.h
--- a/Engine/Source/Editor/UnrealEd/Classes/Exporters/RenderTargetExporterHDR.h
+++ b/Engine/Source/Editor/UnrealEd/Classes/Exporters/RenderTargetExporterHDR.h
@@ -1,7 +1,7 @@
 // Copyright Epic Games, Inc. All Rights Reserved.
 
 //~=============================================================================
-// TextureExporterHDR
+// URenderTargetExporterPNG+URenderTargetExporterEXR
 //~=============================================================================
 
 #pragma once
@@ -11,15 +11,9 @@
 #include "Exporters/Exporter.h"
 #include "RenderTargetExporterHDR.generated.h"
 
-UCLASS()
-class UNREALED_API URenderTargetExporterHDR : public UExporter
-{
-	GENERATED_UCLASS_BODY()
 
-	//~ Begin UExporter Interface
-	virtual bool ExportBinary(UObject* Object, const TCHAR* Type, FArchive& Ar, FFeedbackContext* Warn, int32 FileIndex = 0, uint32 PortFlags = 0) override;
-	//~ End UExporter Interface
-};
+// note : URenderTargetExporterHDR deleted
+//	export to HDR is discouraged, use PNG or EXR
 
 
 UCLASS()
diff --git a/Engine/Source/Editor/UnrealEd/Classes/Exporters/TextureExporterBMP.h b/Engine/Source/Editor/UnrealEd/Classes/Exporters/TextureExporterBMP.h
--- a/Engine/Source/Editor/UnrealEd/Classes/Exporters/TextureExporterBMP.h
+++ b/Engine/Source/Editor/UnrealEd/Classes/Exporters/TextureExporterBMP.h
@@ -9,20 +9,15 @@
 #include "CoreMinimal.h"
 #include "UObject/ObjectMacros.h"
 #include "Exporters/Exporter.h"
+#include "Exporters/TextureExporterGeneric.h"
 #include "TextureExporterBMP.generated.h"
 
 UCLASS()
-class UTextureExporterBMP : public UExporter
+class UTextureExporterBMP : public UTextureExporterGeneric
 {
 	GENERATED_UCLASS_BODY()
-
-	virtual bool SupportsObject(UObject* Object) const override;
-	virtual int32 GetFileCount(UObject* Object) const override;
-	virtual FString GetUniqueFilename(const TCHAR* Filename, int32 FileIndex, int32 FileCount) override;
-	virtual bool ExportBinary(UObject* Object, const TCHAR* Type, FArchive& Ar, FFeedbackContext* Warn, int32 FileIndex = 0, uint32 PortFlags = 0) override;
-
-protected:
-	class UTexture2D* GetExportTexture(UObject* Object) const;
+	
+	virtual bool SupportsTexture(UTexture2D* Texture) const override;
 };
 
 UCLASS()
diff --git a/Engine/Source/Editor/UnrealEd/Classes/Exporters/TextureExporterEXR.h b/Engine/Source/Editor/UnrealEd/Classes/Exporters/TextureExporterEXR.h
--- a/Engine/Source/Editor/UnrealEd/Classes/Exporters/TextureExporterEXR.h
+++ b/Engine/Source/Editor/UnrealEd/Classes/Exporters/TextureExporterEXR.h
@@ -9,19 +9,20 @@
 #include "CoreMinimal.h"
 #include "UObject/ObjectMacros.h"
 #include "Exporters/Exporter.h"
+#include "Exporters/TextureExporterGeneric.h"
 #include "TextureExporterEXR.generated.h"
 
 UCLASS()
-class UNREALED_API UTextureExporterEXR : public UExporter
+class UTextureExporterEXR : public UTextureExporterGeneric
 {
 	GENERATED_UCLASS_BODY()
-
-
-	//~ Begin UExporter Interface
-	virtual bool SupportsObject(UObject* Object) const override;
-	virtual bool ExportBinary( UObject* Object, const TCHAR* Type, FArchive& Ar, FFeedbackContext* Warn, int32 FileIndex = 0, uint32 PortFlags=0 ) override;
-	//~ End UExporter Interface
+	
+	virtual bool SupportsTexture(UTexture2D* Texture) const override;
 };
 
-
+UCLASS()
+class UVirtualTextureBuilderExporterEXR : public UTextureExporterEXR
+{
+	GENERATED_UCLASS_BODY()
+};
 
diff --git a/Engine/Source/Editor/UnrealEd/Classes/Exporters/TextureExporterGeneric.h b/Engine/Source/Editor/UnrealEd/Classes/Exporters/TextureExporterGeneric.h
new file mode 100644
--- /dev/null
+++ b/Engine/Source/Editor/UnrealEd/Classes/Exporters/TextureExporterGeneric.h
@@ -0,0 +1,30 @@
+// Copyright Epic Games, Inc. All Rights Reserved.
+
+//=============================================================================
+// TextureExporterGeneric
+//=============================================================================
+
+#pragma once
+
+#include "CoreMinimal.h"
+#include "UObject/ObjectMacros.h"
+#include "Exporters/Exporter.h"
+#include "TextureExporterGeneric.generated.h"
+
+UCLASS()
+class UTextureExporterGeneric : public UExporter
+{
+	GENERATED_UCLASS_BODY()
+
+	virtual bool SupportsObject(UObject* Object) const override;
+	virtual int32 GetFileCount(UObject* Object) const override;
+	virtual FString GetUniqueFilename(const TCHAR* Filename, int32 FileIndex, int32 FileCount) override;
+	virtual bool ExportBinary(UObject* Object, const TCHAR* Type, FArchive& Ar, FFeedbackContext* Warn, int32 FileIndex = 0, uint32 PortFlags = 0) override;
+	
+	// derived class exporter must implement :
+	virtual bool SupportsTexture(UTexture2D* Texture) const { return false; };
+
+protected:
+	class UTexture2D* GetExportTexture(UObject* Object) const;
+};
+
diff --git a/Engine/Source/Editor/UnrealEd/Classes/Exporters/TextureExporterHDR.h b/Engine/Source/Editor/UnrealEd/Classes/Exporters/TextureExporterHDR.h
--- a/Engine/Source/Editor/UnrealEd/Classes/Exporters/TextureExporterHDR.h
+++ b/Engine/Source/Editor/UnrealEd/Classes/Exporters/TextureExporterHDR.h
@@ -9,18 +9,19 @@
 #include "CoreMinimal.h"
 #include "UObject/ObjectMacros.h"
 #include "Exporters/Exporter.h"
+#include "Exporters/TextureExporterGeneric.h"
 #include "TextureExporterHDR.generated.h"
 
 UCLASS()
-class UNREALED_API UTextureExporterHDR : public UExporter
+class UTextureExporterHDR : public UTextureExporterGeneric
 {
 	GENERATED_UCLASS_BODY()
-
-	//~ Begin UExporter Interface
-	virtual bool SupportsObject(UObject* Object) const override;
-	virtual bool ExportBinary(UObject* Object, const TCHAR* Type, FArchive& Ar, FFeedbackContext* Warn, int32 FileIndex = 0, uint32 PortFlags = 0) override;
-	//~ End UExporter Interface
+	
+	virtual bool SupportsTexture(UTexture2D* Texture) const override;
 };
 
-
-
+UCLASS()
+class UVirtualTextureBuilderExporterHDR : public UTextureExporterHDR
+{
+	GENERATED_UCLASS_BODY()
+};
diff --git a/Engine/Source/Editor/UnrealEd/Classes/Exporters/TextureExporterPNG.h b/Engine/Source/Editor/UnrealEd/Classes/Exporters/TextureExporterPNG.h
--- a/Engine/Source/Editor/UnrealEd/Classes/Exporters/TextureExporterPNG.h
+++ b/Engine/Source/Editor/UnrealEd/Classes/Exporters/TextureExporterPNG.h
@@ -9,10 +9,29 @@
 #include "CoreMinimal.h"
 #include "UObject/ObjectMacros.h"
 #include "Exporters/Exporter.h"
+#include "Exporters/TextureExporterGeneric.h"
 #include "TextureExporterPNG.generated.h"
 
 UCLASS()
-class UNREALED_API UTextureExporterPNG : public UExporter
+class UTextureExporterPNG : public UTextureExporterGeneric
+{
+	GENERATED_UCLASS_BODY()
+	
+	virtual bool SupportsTexture(UTexture2D* Texture) const override;
+};
+
+UCLASS()
+class UVirtualTextureBuilderExporterPNG : public UTextureExporterPNG
+{
+	GENERATED_UCLASS_BODY()
+};
+
+
+
+// JPEG is not UTextureExporterGeneric
+// JPEG exporter writes stored JPEG bits in TextureSource compressed payloads
+UCLASS()
+class UNREALED_API UTextureExporterJPEG : public UExporter
 {
 	GENERATED_UCLASS_BODY()
 
@@ -25,3 +44,5 @@ class UNREALED_API UTextureExporterPNG : public UExporter
 
 
 
+
+
diff --git a/Engine/Source/Editor/UnrealEd/Private/Factories/EditorFactories.cpp b/Engine/Source/Editor/UnrealEd/Private/Factories/EditorFactories.cpp
--- a/Engine/Source/Editor/UnrealEd/Private/Factories/EditorFactories.cpp
+++ b/Engine/Source/Editor/UnrealEd/Private/Factories/EditorFactories.cpp
@@ -2715,14 +2715,15 @@ UTextureFactory::UTextureFactory(const FObjectInitializer& ObjectInitializer)
 	Formats.Add( TEXT( "bmp;Texture" ) );
 	Formats.Add( TEXT( "pcx;Texture" ) );
 	Formats.Add( TEXT( "tga;Texture" ) );
-	Formats.Add( TEXT( "float;Texture" ) );
+	Formats.Add( TEXT( "float;Texture" ) ); // @@!! what's this ?
 	Formats.Add( TEXT( "psd;Texture" ) );
 	Formats.Add( TEXT( "dds;Texture (Cubemap or 2D)" ) );
-	Formats.Add( TEXT( "hdr;Cubemap Texture (LongLat unwrap)" ) );
+	Formats.Add( TEXT( "hdr;Texture (HDR) (LongLat unwrap or 2D)" ) );
 	Formats.Add( TEXT( "ies;IES Texture (Standard light profiles)" ) );
 	Formats.Add( TEXT( "png;Texture" ) );
 	Formats.Add( TEXT( "jpg;Texture" ) );
 	Formats.Add( TEXT( "jpeg;Texture" ) );
+	// note: conflicts with .exr ImgMedia import in the import dialog
 	Formats.Add( TEXT( "exr;Texture (HDR)" ) );
 	Formats.Add( TEXT( "tif;Texture (TIFF)" ) );
 	Formats.Add( TEXT( "tiff;Texture (TIFF)" ) );
@@ -4681,7 +4682,7 @@ bool UTextureExporterPCX::ExportBinary( UObject* Object, const TCHAR* Type, FArc
 }
 
 
-static bool ExportTexture2DGeneric( UObject* Object,FArchive& Ar, const TCHAR * ImageFormat, FFeedbackContext* Warn, int LayerIndex=0)
+static bool ExportTexture2DGeneric( UObject* Object,FArchive& Ar, const TCHAR * ImageFormat, FFeedbackContext* Warn, int FileIndex=0)
 {
 	UTexture2D* Texture = CastChecked<UTexture2D>( Object );
 	if ( Texture == nullptr )
@@ -4695,8 +4696,20 @@ static bool ExportTexture2DGeneric( UObject* Object,FArchive& Ar, const TCHAR *
 		return false;
 	}
 	
+	int NumBlocks = Texture->Source.GetNumBlocks();
+	int NumLayers = Texture->Source.GetNumLayers();
+	int BlockIndex = 0;
+	int LayerIndex = 0;
+
+	int NumFiles = NumBlocks * NumLayers;
+	check( FileIndex < NumFiles );
+	BlockIndex = FileIndex / NumLayers;
+	LayerIndex = FileIndex % NumLayers; 
+
+	UE_LOG(LogEditorFactories, Display, TEXT("Exporting Texture as %s Layer %d/%d Block %d/%d"), ImageFormat, LayerIndex,NumLayers, BlockIndex,NumBlocks);
+
 	FImage Image;
-	if ( ! Texture->Source.GetMipImage(Image,0,LayerIndex,0) )
+	if ( ! Texture->Source.GetMipImage(Image,BlockIndex,LayerIndex,0) )
 	{
 		Warn->Logf(ELogVerbosity::Error, TEXT("ExportTexture failed : could not GetMipImage"));
 		return false;
@@ -4715,21 +4728,23 @@ static bool ExportTexture2DGeneric( UObject* Object,FArchive& Ar, const TCHAR *
 }
 
 /*------------------------------------------------------------------------------
-	UTextureExporterBMP implementation.
+	UTextureExporterGeneric implementation.
+
+	UTextureExporterGeneric base class for all standard 2d texture exporters
 ------------------------------------------------------------------------------*/
-UTextureExporterBMP::UTextureExporterBMP(const FObjectInitializer& ObjectInitializer)
+
+// note that 
+
+UTextureExporterGeneric::UTextureExporterGeneric(const FObjectInitializer& ObjectInitializer)
 	: Super(ObjectInitializer)
 {
 	SupportedClass = UTexture2D::StaticClass();
 	PreferredFormatIndex = 0;
-	FormatExtension.Add(TEXT("BMP"));
-	FormatDescription.Add(TEXT("Windows Bitmap"));
-
 }
 
-UTexture2D* UTextureExporterBMP::GetExportTexture(UObject* Object) const
+UTexture2D* UTextureExporterGeneric::GetExportTexture(UObject* Object) const
 {
-	// for UVirtualTextureBuilderExporterBMP
+	// for UVirtualTextureBuilderExporter
 	UVirtualTextureBuilder* VirtualTextureBuilder = Cast<UVirtualTextureBuilder>(Object);
 	if (VirtualTextureBuilder != nullptr)
 	{
@@ -4738,7 +4753,10 @@ UTexture2D* UTextureExporterBMP::GetExportTexture(UObject* Object) const
 	return Cast<UTexture2D>(Object);
 }
 
-bool UTextureExporterBMP::SupportsObject(UObject* Object) const
+// note that UTextureExporterGeneric will actually be registered as an exporter itself
+// but will return false for SupportsObject because it has SupportsTexture == false
+
+bool UTextureExporterGeneric::SupportsObject(UObject* Object) const
 {
 	bool bSupportsObject = false;
 	if (Super::SupportsObject(Object))
@@ -4746,49 +4764,162 @@ bool UTextureExporterBMP::SupportsObject(UObject* Object) const
 		UTexture2D* Texture = GetExportTexture(Object);
 		if ( Texture )
 		{
-			bSupportsObject = Texture->Source.GetFormat() == TSF_BGRA8 || Texture->Source.GetFormat() == TSF_G8;
+			if ( ! Texture->Source.IsValid() )
+			{
+				// should not get here
+				check(0);
+				return false;
+			}
+
+			return SupportsTexture(Texture);
 		}
 	}
-	return bSupportsObject;
+	return false;
 }
 
-int32 UTextureExporterBMP::GetFileCount(UObject* Object) const
+int32 UTextureExporterGeneric::GetFileCount(UObject* Object) const
 {
-	// for UVirtualTextureBuilderExporterBMP
 	UTexture2D* Texture = GetExportTexture(Object);
-	return (Texture != nullptr) ? Texture->Source.GetNumLayers() : 1;
+	check(Texture != nullptr);
+
+	// standard textures will have NumBlocks == NumLayers == 1
+	// VT can have them > 1
+	// UDIM gives you NumBlocks
+	int NumBlocks = Texture->Source.GetNumBlocks();
+	int NumLayers = Texture->Source.GetNumLayers();
+	int FileCount = NumBlocks * NumLayers;
+	check( FileCount > 0 );
+	return FileCount;
 }
 
-FString UTextureExporterBMP::GetUniqueFilename(const TCHAR* Filename, int32 FileIndex, int32 FileCount)
+FString UTextureExporterGeneric::GetUniqueFilename(const TCHAR* Filename, int32 FileIndex, int32 FileCount)
 {
-	// for UVirtualTextureBuilderExporterBMP
-	return (FileCount == 1) ? Filename : FString::Printf(TEXT("%s%d%s"), *FPaths::GetBaseFilename(Filename, false), FileIndex, *FPaths::GetExtension(Filename, true));
+	check( FileCount > 0 );
+	if ( FileCount == 1 )
+	{
+		// standard single texture
+		return Filename;
+	}
+	else
+	{
+		//	for VT with layers/blocks
+		// 
+		// @@!! would like to use LayerIndex+BlockIndex (rather than FileIndex)
+		//  but I'm not given the export Object
+		//	so not possible at the moment
+		// @todo: change GetUniqueFilename API to pass the Object
+		//
+		// it's possible you could even be able to map back to the original UDIM file names (xxx.1001 etc.)
+		//	using the TextureSource asset import info
+		return FString::Printf(TEXT("%s%d%s"), *FPaths::GetBaseFilename(Filename, false), FileIndex, *FPaths::GetExtension(Filename, true));
+	}
 }
 
-bool UTextureExporterBMP::ExportBinary( UObject* Object, const TCHAR* Type, FArchive& Ar, FFeedbackContext* Warn, int32 FileIndex, uint32 PortFlags )
+bool UTextureExporterGeneric::ExportBinary( UObject* Object, const TCHAR* Type, FArchive& Ar, FFeedbackContext* Warn, int32 FileIndex, uint32 PortFlags )
 {
 	UTexture2D* Texture = GetExportTexture(Object);
-	if(Texture == nullptr)
-	{
-		return false;
-	}
+	check( Texture != nullptr );
+	check( SupportsTexture(Texture) );
 	
 	// FileIndex for layers for VT :
 
-	return ExportTexture2DGeneric(Texture,Ar,TEXT("BMP"),Warn,FileIndex);
+	// Type == Image extension
+	return ExportTexture2DGeneric(Texture,Ar,Type,Warn,FileIndex);
+}
+
+//===========================================================
+// formats derived from UTextureExporterGeneric
+// just need to add their FormatExtension
+// and implement SupportsTexture
+
+UTextureExporterBMP::UTextureExporterBMP(const FObjectInitializer& ObjectInitializer)
+	: UTextureExporterGeneric(ObjectInitializer)
+{
+	FormatExtension.Add(TEXT("BMP"));
+	FormatDescription.Add(TEXT("Windows Bitmap"));
+}
+
+bool UTextureExporterBMP::SupportsTexture(UTexture2D* Texture) const
+{
+	// U8 formats only :
+	return Texture->Source.GetFormat() == TSF_BGRA8 || Texture->Source.GetFormat() == TSF_G8;
 }
 
 UVirtualTextureBuilderExporterBMP::UVirtualTextureBuilderExporterBMP(const FObjectInitializer& ObjectInitializer)
 	: UTextureExporterBMP(ObjectInitializer)
 {
 	SupportedClass = UVirtualTextureBuilder::StaticClass();
-	// already done by UTextureExporterBMP
-	//   I just change the SupportedClass
-	//PreferredFormatIndex = 0;
-	//FormatExtension.Add(TEXT("BMP"));
-	//FormatDescription.Add(TEXT("BMP Windows Bitmap"));
 }
 
+//========================================================
+
+UTextureExporterHDR::UTextureExporterHDR(const FObjectInitializer& ObjectInitializer)
+	: UTextureExporterGeneric(ObjectInitializer)
+{
+	FormatExtension.Add(TEXT("HDR"));
+	FormatDescription.Add(TEXT("HDR Radiance RGBE texture"));
+}
+
+bool UTextureExporterHDR::SupportsTexture(UTexture2D* Texture) const
+{
+	// only BGRE is allowed to go out to HDR now (lossless)
+	bool bSupportsObject = Texture->Source.GetFormat() == TSF_BGRE8;
+	return bSupportsObject;
+}
+
+UVirtualTextureBuilderExporterHDR::UVirtualTextureBuilderExporterHDR(const FObjectInitializer& ObjectInitializer)
+	: UTextureExporterHDR(ObjectInitializer)
+{
+	SupportedClass = UVirtualTextureBuilder::StaticClass();
+}
+
+//========================================================
+
+UTextureExporterPNG::UTextureExporterPNG(const FObjectInitializer& ObjectInitializer)
+	: UTextureExporterGeneric(ObjectInitializer)
+{
+	FormatExtension.Add(TEXT("PNG"));
+	FormatDescription.Add(TEXT("PNG"));
+}
+
+bool UTextureExporterPNG::SupportsTexture(UTexture2D* Texture) const
+{
+	ETextureSourceFormat TSF = Texture->Source.GetFormat();
+	ERawImageFormat::Type RawFormat = FImageCoreUtils::ConvertToRawImageFormat(TSF);
+	// supports all non-HDR formats :
+	return ! ERawImageFormat::IsHDR(RawFormat);
+}
+
+UVirtualTextureBuilderExporterPNG::UVirtualTextureBuilderExporterPNG(const FObjectInitializer& ObjectInitializer)
+	: UTextureExporterPNG(ObjectInitializer)
+{
+	SupportedClass = UVirtualTextureBuilder::StaticClass();
+}
+
+//========================================================
+
+UTextureExporterEXR::UTextureExporterEXR(const FObjectInitializer& ObjectInitializer)
+	: UTextureExporterGeneric(ObjectInitializer)
+{
+	FormatExtension.Add(TEXT("EXR"));
+	FormatDescription.Add(TEXT("EXR HDR float texture"));
+}
+
+bool UTextureExporterEXR::SupportsTexture(UTexture2D* Texture) const
+{
+	ETextureSourceFormat TSF = Texture->Source.GetFormat();
+	ERawImageFormat::Type RawFormat = FImageCoreUtils::ConvertToRawImageFormat(TSF);
+	// supports all HDR formats :
+	return ERawImageFormat::IsHDR(RawFormat);
+}
+
+UVirtualTextureBuilderExporterEXR::UVirtualTextureBuilderExporterEXR(const FObjectInitializer& ObjectInitializer)
+	: UTextureExporterEXR(ObjectInitializer)
+{
+	SupportedClass = UVirtualTextureBuilder::StaticClass();
+}
+
+
 /*------------------------------------------------------------------------------
 	URenderTargetExporterHDR implementation.
 	Exports render targets.
@@ -4822,27 +4953,6 @@ static bool ExportRenderTarget2DGeneric(UObject* Object, FArchive& Ar, const TCH
 	return true;
 }
 
-// HDR is terrible, use EXR or PNG instead
-//  left for now for backwards compatibility
-//  EXR and PNG exporters are prioritized in AssetTools
-URenderTargetExporterHDR::URenderTargetExporterHDR(const FObjectInitializer& ObjectInitializer)
-	: Super(ObjectInitializer)
-{
-	SupportedClass = UTextureRenderTarget2D::StaticClass();
-	PreferredFormatIndex = 0;
-	FormatExtension.Add(TEXT("HDR"));
-	FormatDescription.Add(TEXT("HDR Radiance RGBE lossy rendertarget"));
-
-	// HDR can export both RGBA8 and 16F render targets
-	//  so it has no SupportsObject() check
-	//  though it's poor for both
-}
-
-bool URenderTargetExporterHDR::ExportBinary(UObject* Object, const TCHAR* Type, FArchive& Ar, FFeedbackContext* Warn, int32 FileIndex, uint32 PortFlags)
-{
-	return ExportRenderTarget2DGeneric(Object,Ar,TEXT("HDR"));
-}
-
 URenderTargetExporterPNG::URenderTargetExporterPNG(const FObjectInitializer& ObjectInitializer)
 	: Super(ObjectInitializer)
 {
@@ -4949,106 +5059,57 @@ URenderTargetCubeExporterHDR::URenderTargetCubeExporterHDR(const FObjectInitiali
 	SupportedClass = UTextureRenderTargetCube::StaticClass();
 }
 
-/*------------------------------------------------------------------------------
-	UTextureExporterHDR implementation.
-	Export UTexture2D as .HDR
-------------------------------------------------------------------------------*/
-UTextureExporterHDR::UTextureExporterHDR(const FObjectInitializer& ObjectInitializer)
+UTextureExporterJPEG::UTextureExporterJPEG(const FObjectInitializer& ObjectInitializer)
 	: Super(ObjectInitializer)
 {
 	SupportedClass = UTexture2D::StaticClass();
 	PreferredFormatIndex = 0;
-	FormatExtension.Add(TEXT("HDR"));
-	FormatDescription.Add(TEXT("HDR Radiance RGBE texture"));
+	FormatExtension.Add(TEXT("JPG"));
+	FormatDescription.Add(TEXT("JPEG original imported into uasset"));
 }
 
-bool UTextureExporterHDR::SupportsObject(UObject* Object) const
+bool UTextureExporterJPEG::SupportsObject(UObject* Object) const
 {
-	bool bSupportsObject = false;
 	if (Super::SupportsObject(Object))
 	{
 		UTexture2D* Texture = Cast<UTexture2D>(Object);
 
 		if (Texture)
 		{
-			// only BGRE is allowed to go out to HDR now (lossless)
-			bSupportsObject = Texture->Source.GetFormat() == TSF_BGRE8;
-		}
-	}
-	return bSupportsObject;
-}
-
-bool UTextureExporterHDR::ExportBinary(UObject* Object, const TCHAR* Type, FArchive& Ar, FFeedbackContext* Warn, int32 FileIndex, uint32 PortFlags)
-{
-	return ExportTexture2DGeneric(Object,Ar,TEXT("HDR"),Warn);
-}
-
-/*------------------------------------------------------------------------------
-	UTextureExporterPNG implementation.
-------------------------------------------------------------------------------*/
-UTextureExporterPNG::UTextureExporterPNG(const FObjectInitializer& ObjectInitializer)
-	: Super(ObjectInitializer)
-{
-	SupportedClass = UTexture2D::StaticClass();
-	PreferredFormatIndex = 0;
-	FormatExtension.Add(TEXT("PNG"));
-	FormatDescription.Add(TEXT("PNG"));
-}
+			// we do NOT do lossy recompression
 
-bool UTextureExporterPNG::SupportsObject(UObject* Object) const
-{
-	if (Super::SupportsObject(Object))
-	{
-		UTexture2D* Texture = Cast<UTexture2D>(Object);
+			// Check it has JPEG BulkData :
+			if ( Texture->Source.GetSourceCompression() == TSCF_JPEG &&
+				Texture->Source.GetSizeOnDisk() > 0 )
+			{
+				ETextureSourceFormat TSF = Texture->Source.GetFormat();
+				ERawImageFormat::Type RawFormat = FImageCoreUtils::ConvertToRawImageFormat(TSF);
+				check( RawFormat == ERawImageFormat::G8 || RawFormat == ERawImageFormat::BGRA8 );
 
-		if (Texture)
-		{
-			ETextureSourceFormat TSF = Texture->Source.GetFormat();
-			ERawImageFormat::Type RawFormat = FImageCoreUtils::ConvertToRawImageFormat(TSF);
-			// supports all non-HDR formats :
-			return ! ERawImageFormat::IsHDR(RawFormat);
+				return true;
+			}
 		}
 	}
 	return false;
 }
 
-bool UTextureExporterPNG::ExportBinary( UObject* Object, const TCHAR* Type, FArchive& Ar, FFeedbackContext* Warn, int32 FileIndex, uint32 PortFlags )
+bool UTextureExporterJPEG::ExportBinary( UObject* Object, const TCHAR* Type, FArchive& Ar, FFeedbackContext* Warn, int32 FileIndex, uint32 PortFlags )
 {
-	return ExportTexture2DGeneric(Object,Ar,TEXT("PNG"),Warn);
-}
+	UTexture2D* Texture = Cast<UTexture2D>(Object);
+	check( Texture != nullptr );
 
-/*------------------------------------------------------------------------------
-	UTextureExporterEXR implementation.
-------------------------------------------------------------------------------*/
-UTextureExporterEXR::UTextureExporterEXR(const FObjectInitializer& ObjectInitializer)
-	: Super(ObjectInitializer)
-{
-	SupportedClass = UTexture2D::StaticClass();
-	PreferredFormatIndex = 0;
-	FormatExtension.Add(TEXT("EXR"));
-	FormatDescription.Add(TEXT("EXR HDR float texture"));
-}
+	check( Texture->Source.GetSourceCompression() == TSCF_JPEG &&
+			Texture->Source.GetSizeOnDisk() > 0 );
 
-bool UTextureExporterEXR::SupportsObject(UObject* Object) const
-{
-	if (Super::SupportsObject(Object))
-	{
-		UTexture2D* Texture = Cast<UTexture2D>(Object);
+	// just write the JPEG data we already have :
+	
+	UE_LOG(LogEditorFactories, Display, TEXT("Exporting Texture as JPEG stored bits (no lossy decompress or recompress)."));
 
-		if (Texture)
-		{
-			ETextureSourceFormat TSF = Texture->Source.GetFormat();
-			ERawImageFormat::Type RawFormat = FImageCoreUtils::ConvertToRawImageFormat(TSF);
-			// supports all HDR formats :
-			return ERawImageFormat::IsHDR(RawFormat);
-		}
-	}
-	return false;
-}
+	Texture->Source.OperateOnLoadedBulkData( [&](const FSharedBuffer& BulkDataBuffer) {
+		Ar.Serialize( const_cast<void *>(BulkDataBuffer.GetData()), BulkDataBuffer.GetSize());
+	} );
 
-bool UTextureExporterEXR::ExportBinary( UObject* Object, const TCHAR* Type, FArchive& Ar, FFeedbackContext* Warn, int32 FileIndex, uint32 PortFlags )
-{
-	return ExportTexture2DGeneric(Object,Ar,TEXT("EXR"),Warn);
+	return true;
 }
 
 /*------------------------------------------------------------------------------
@@ -5090,6 +5151,8 @@ bool UTextureExporterTGA::ExportBinary( UObject* Object, const TCHAR* Type, FArc
 		return false;
 	}
 
+	// legacy code would claim to support 16 bit, but then just drop the bottom 8 bits
+	// this is no longer used
 	const bool bIsRGBA16 = Texture->Source.GetFormat() == TSF_RGBA16;
 
 	if (bIsRGBA16)
