two EXR import extensions in Editor import dialog
are selected by alphabet order when you do *.exr in the dialog

diff --git a/Engine/Plugins/Media/ImgMedia/Source/ImgMediaEditor/Private/Factories/ImgMediaSourceFactory.cpp b/Engine/Plugins/Media/ImgMedia/Source/ImgMediaEditor/Private/Factories/ImgMediaSourceFactory.cpp
--- a/Engine/Plugins/Media/ImgMedia/Source/ImgMediaEditor/Private/Factories/ImgMediaSourceFactory.cpp
+++ b/Engine/Plugins/Media/ImgMedia/Source/ImgMediaEditor/Private/Factories/ImgMediaSourceFactory.cpp
@@ -11,6 +11,8 @@
 UImgMediaSourceFactory::UImgMediaSourceFactory(const FObjectInitializer& ObjectInitializer)
 	: Super(ObjectInitializer)
 {
+	// note: conflicts with .exr image import by TextureFactory in the import dialog
+	// if you filter for *.exr it chooses the first extension alphabetically
 	Formats.Add(TEXT("exr;EXR ImgMedia Image Sequence"));
 
 	SupportedClass = UImgMediaSource::StaticClass();
diff --git a/Engine/Source/Editor/UnrealEd/Private/Factories/EditorFactories.cpp b/Engine/Source/Editor/UnrealEd/Private/Factories/EditorFactories.cpp
--- a/Engine/Source/Editor/UnrealEd/Private/Factories/EditorFactories.cpp
+++ b/Engine/Source/Editor/UnrealEd/Private/Factories/EditorFactories.cpp
@@ -2724,7 +2724,8 @@ UTextureFactory::UTextureFactory(const FObjectInitializer& ObjectInitializer)
 	Formats.Add( TEXT( "jpg;Texture" ) );
 	Formats.Add( TEXT( "jpeg;Texture" ) );
 	// note: conflicts with .exr ImgMedia import in the import dialog
-	Formats.Add( TEXT( "exr;Texture (HDR)" ) );
+	// alphabetic order affects the select when you enter "*.exr" in the import dialog box
+	Formats.Add( TEXT( "exr;EXR image HDR Texture2d" ) );
 	Formats.Add( TEXT( "tif;Texture (TIFF)" ) );
 	Formats.Add( TEXT( "tiff;Texture (TIFF)" ) );
 
