DumpMaterialExpressionsCommandlet - This commandlet will dump to a plain text file an info table of all material expressions in the engine and the plugins enabled on the project.
Output material expression fields:
*) Name - Name of the UMaterialExpression class (shows in material editor node search)
*) DisplayName - Optional display name set on the class (shows in material editor node search and overrides Name)
*) Caption - Text displayed on the material expression node
*) Description
*) Tooltip
Output location <Project>\Saved\MaterialEditor\MaterialExpressions.txt

diff --git a/Engine/Source/Editor/UnrealEd/Classes/Commandlets/DumpMaterialExpressionsCommandlet.h b/Engine/Source/Editor/UnrealEd/Classes/Commandlets/DumpMaterialExpressionsCommandlet.h
new file mode 100644
--- /dev/null
+++ b/Engine/Source/Editor/UnrealEd/Classes/Commandlets/DumpMaterialExpressionsCommandlet.h
@@ -0,0 +1,18 @@
+// Copyright Epic Games, Inc. All Rights Reserved.
+
+#pragma once
+
+#include "CoreMinimal.h"
+#include "UObject/ObjectMacros.h"
+#include "Commandlets/Commandlet.h"
+#include "DumpMaterialExpressionsCommandlet.generated.h"
+
+UCLASS(config=Editor)
+class UDumpMaterialExpressionsCommandlet : public UCommandlet
+{
+	GENERATED_UCLASS_BODY()
+
+	//~ Begin UCommandlet Interface
+	virtual int32 Main(const FString& Params) override;
+	//~ End UCommandlet Interface
+};
diff --git a/Engine/Source/Editor/UnrealEd/Private/Commandlets/DumpMaterialExpressionsCommandlet.cpp b/Engine/Source/Editor/UnrealEd/Private/Commandlets/DumpMaterialExpressionsCommandlet.cpp
new file mode 100644
--- /dev/null
+++ b/Engine/Source/Editor/UnrealEd/Private/Commandlets/DumpMaterialExpressionsCommandlet.cpp
@@ -0,0 +1,157 @@
+// Copyright Epic Games, Inc. All Rights Reserved.
+
+#include "Commandlets/DumpMaterialExpressionsCommandlet.h"
+#include "Materials/MaterialExpression.h"
+#include "UObject/UObjectIterator.h"
+#include "HAL/FileManager.h"
+
+DEFINE_LOG_CATEGORY_STATIC(LogDumpMaterialExpressionsCommandlet, Log, All);
+
+UDumpMaterialExpressionsCommandlet::UDumpMaterialExpressionsCommandlet(const FObjectInitializer& ObjectInitializer)
+	: Super(ObjectInitializer)
+{
+}
+
+int32 UDumpMaterialExpressionsCommandlet::Main(const FString& Params)
+{
+	TArray<FString> Tokens;
+	TArray<FString> Switches;
+	TMap<FString, FString> ParamVals;
+	UCommandlet::ParseCommandLine(*Params, Tokens, Switches, ParamVals);
+
+	if (Switches.Contains("help"))
+	{
+		UE_LOG(LogDumpMaterialExpressionsCommandlet, Log, TEXT("DumpMaterialExpressions"));
+		UE_LOG(LogDumpMaterialExpressionsCommandlet, Log, TEXT("This commandlet will dump to a plain text file an info table of all material expressions in the engine and the plugins enabled on the project."));
+		return 0;
+	}
+
+	struct FMaterialExpressionInfo
+	{
+		UMaterialExpression* MaterialExpression;
+		FString Name;
+		FString DisplayName;
+		FString Caption;
+		FString Description;
+		FString Tooltip;
+	};
+	TArray<FMaterialExpressionInfo> MaterialExpressionInfos;
+
+	static const FName NAME_DisplayName(TEXT("DisplayName"));
+
+	const FString NameField = TEXT("NAME");
+	const FString DisplayNameField = TEXT("DISPLAYNAME");
+	const FString CaptionField = TEXT("CAPTION");
+	const FString DescriptionField = TEXT("DESCRIPTION");
+	const FString TooltipField = TEXT("TOOLTIP");
+
+	int32 MaxNameLength = NameField.Len();
+	int32 MaxDisplayNameLength = DisplayNameField.Len();
+	int32 MaxCaptionLength = CaptionField.Len();
+	int32 MaxDescriptionLength = DescriptionField.Len();
+
+	// Collect all default material expression objects
+	for (TObjectIterator<UClass> It; It; ++It)
+	{
+		UClass* Class = *It;
+		// Skip the base UMaterialExpression class
+		if (!Class->HasAnyClassFlags(CLASS_Abstract))
+		{
+			if (UMaterialExpression* DefaultExpression = Cast<UMaterialExpression>(Class->GetDefaultObject()))
+			{
+				TArray<FString> MultilineCaption;
+				DefaultExpression->GetCaption(MultilineCaption);
+				FString Caption;
+				for (const FString& Line : MultilineCaption)
+				{
+					Caption += Line;
+				}
+
+				TArray<FString> MultilineToolTip;
+				DefaultExpression->GetExpressionToolTip(MultilineToolTip);
+				FString Tooltip;
+				for (const FString& Line : MultilineToolTip)
+				{
+					Tooltip += Line;
+				}
+
+				FMaterialExpressionInfo ExpressionInfo;
+				ExpressionInfo.MaterialExpression = DefaultExpression;
+				ExpressionInfo.Name = Class->GetName().Mid(FCString::Strlen(TEXT("MaterialExpression")));;
+				ExpressionInfo.DisplayName = Class->GetMetaData(NAME_DisplayName);
+				ExpressionInfo.Caption = Caption;
+				ExpressionInfo.Description = DefaultExpression->GetDescription();
+				ExpressionInfo.Tooltip = Tooltip;
+				MaterialExpressionInfos.Add(ExpressionInfo);
+
+				MaxNameLength = FMath::Max(MaxNameLength, ExpressionInfo.Name.Len());
+				MaxDisplayNameLength = FMath::Max(MaxDisplayNameLength, ExpressionInfo.DisplayName.Len());
+				MaxCaptionLength = FMath::Max(MaxCaptionLength, ExpressionInfo.Caption.Len());
+				MaxDescriptionLength = FMath::Max(MaxDescriptionLength, ExpressionInfo.Description.Len());
+			}
+		}
+	}
+
+	// Additional padding for spacing
+	const int32 AdditionalPadding = 3;
+	MaxNameLength += AdditionalPadding;
+	MaxDisplayNameLength += AdditionalPadding;
+	MaxCaptionLength += AdditionalPadding;
+	MaxDescriptionLength += AdditionalPadding;
+
+	auto GetFormattedText = [](const FString& InText) -> FString
+	{
+		FString OutText = InText;
+		OutText = InText.Replace(TEXT("\n"), TEXT(" "));
+		if (OutText.IsEmpty())
+		{
+			OutText = TEXT("N/A");
+		}
+		return OutText;
+	};
+
+	auto GenerateSpacePadding = [](int32 MaxLen, int32 TextLen) -> FString
+	{
+		FString Padding;
+		for (int i = 0; i < MaxLen - TextLen; ++i)
+		{
+			Padding += TEXT(" ");
+		}
+		return Padding;
+	};
+
+	// Write the material expression list to a text file
+	const FString OutputFilePath = FPaths::Combine(*FPaths::ProjectSavedDir(), TEXT("MaterialEditor"), TEXT("MaterialExpressions.txt"));
+	FArchive* FileWriter = IFileManager::Get().CreateFileWriter(*OutputFilePath);
+
+	auto WriteLine = [FileWriter, GenerateSpacePadding, MaxNameLength, MaxDisplayNameLength, MaxCaptionLength, MaxDescriptionLength](const FString& Name, const FString& DisplayName, const FString& Caption, const FString& Description, const FString& Tooltip)
+	{
+		FString NamePadding = GenerateSpacePadding(MaxNameLength, Name.Len());
+		FString DisplayNamePadding = GenerateSpacePadding(MaxDisplayNameLength, DisplayName.Len());
+		FString CaptionPadding = GenerateSpacePadding(MaxCaptionLength, Caption.Len());
+		FString DescriptionPadding = GenerateSpacePadding(MaxDescriptionLength, Description.Len());
+
+		FString OutputLine = (Name + NamePadding);
+		OutputLine += (DisplayName + DisplayNamePadding);
+		OutputLine += (Caption + CaptionPadding);
+		OutputLine += (Description + DescriptionPadding);
+		OutputLine += (Tooltip + TEXT("\n"));
+		FileWriter->Serialize(const_cast<ANSICHAR*>(StringCast<ANSICHAR>(*OutputLine).Get()), OutputLine.Len());
+	};
+
+	WriteLine(NameField, DisplayNameField, CaptionField, DescriptionField, TooltipField);
+	for (FMaterialExpressionInfo& ExpressionInfo : MaterialExpressionInfos)
+	{
+		WriteLine(GetFormattedText(ExpressionInfo.Name), GetFormattedText(ExpressionInfo.DisplayName), GetFormattedText(ExpressionInfo.Caption), GetFormattedText(ExpressionInfo.Description), GetFormattedText(ExpressionInfo.Tooltip));
+	}
+
+	FString OutputLine = FString::Printf(TEXT("\nTotal %d material expressions found."), MaterialExpressionInfos.Num());
+	FileWriter->Serialize(const_cast<ANSICHAR*>(StringCast<ANSICHAR>(*OutputLine).Get()), OutputLine.Len());
+
+	FileWriter->Close();
+	delete FileWriter;
+
+	UE_LOG(LogDumpMaterialExpressionsCommandlet, Log, TEXT("Total %d material expressions are written to %s"), MaterialExpressionInfos.Num(), *OutputFilePath);
+
+	return 0;
+}
