Fix an issue where a deleted in-memory asset would be surfaced as a dirty package by editor file utilities

diff --git a/Engine/Source/Editor/UnrealEd/Private/ObjectTools.cpp b/Engine/Source/Editor/UnrealEd/Private/ObjectTools.cpp
--- a/Engine/Source/Editor/UnrealEd/Private/ObjectTools.cpp
+++ b/Engine/Source/Editor/UnrealEd/Private/ObjectTools.cpp
@@ -2172,7 +2172,7 @@ namespace ObjectTools
 	void CleanupAfterSuccessfulDelete (const TArray<UPackage*>& PotentialPackagesToDelete, bool bPerformReferenceCheck)
 	{
 		TArray<UPackage*> PackagesToDelete = PotentialPackagesToDelete;
-		TArray<UPackage*> EmptyPackagesToUnload;
+		TArray<UPackage*> PackagesToUnload;
 		TArray<FString> PackageFilesToDelete;
 		TArray<TSharedRef<ISourceControlState, ESPMode::ThreadSafe> > PackageSCCStates;
 		ISourceControlProvider& SourceControlProvider = ISourceControlModule::Get().GetProvider();
@@ -2209,21 +2209,24 @@ namespace ObjectTools
 				UPackage* CurrentPackage = Cast<UPackage>(Package);
 
 				FString PackageFilename;
-				if( Package != nullptr && FPackageName::DoesPackageExist( Package->GetName(), &PackageFilename ) )
+
+				if( CurrentPackage == nullptr )
 				{
-					PackageFilesToDelete.Add(PackageFilename);
-					CurrentPackage->SetDirtyFlag(false);
+					PackagesToDelete.RemoveAt(PackageIdx);
 				}
 				else
 				{
-					// Could not determine filename for package so we can not delete
-					PackagesToDelete.RemoveAt(PackageIdx);
+					CurrentPackage->SetDirtyFlag(false);
 
-					if (UPackage::IsEmptyPackage(CurrentPackage))
+					if (FPackageName::DoesPackageExist(Package->GetName(), &PackageFilename))
+					{
+						PackageFilesToDelete.Add(PackageFilename);
+					}
+					else
 					{
-						// If the package is empty, unload it anyway
-						EmptyPackagesToUnload.Add(CurrentPackage);
-						CurrentPackage->SetDirtyFlag(false);
+						// Could not determine filename for package so we can not delete, but we should unload
+						PackagesToDelete.RemoveAt(PackageIdx);
+						PackagesToUnload.Add(CurrentPackage);
 					}
 				}
 			}
@@ -2251,12 +2254,12 @@ namespace ObjectTools
 		}
 
 		// Unload the packages and collect garbage.
-		if ( PackagesToDelete.Num() > 0 || EmptyPackagesToUnload.Num() > 0 )
+		if ( PackagesToDelete.Num() > 0 || PackagesToUnload.Num() > 0 )
 		{
 			TArray<UPackage*> AllPackagesToUnload;
-			AllPackagesToUnload.Reserve(PackagesToDelete.Num() + EmptyPackagesToUnload.Num());
+			AllPackagesToUnload.Reserve(PackagesToDelete.Num() + PackagesToUnload.Num());
 			AllPackagesToUnload.Append(PackagesToDelete);
-			AllPackagesToUnload.Append(EmptyPackagesToUnload);
+			AllPackagesToUnload.Append(PackagesToUnload);
 
 			UPackageTools::UnloadPackages(AllPackagesToUnload);
 		}
