Fixed localization commandlets incorrectly detecting engine (no project) execution

#rnx

diff --git a/Engine/Source/Editor/UnrealEd/Classes/Commandlets/GatherTextCommandletBase.h b/Engine/Source/Editor/UnrealEd/Classes/Commandlets/GatherTextCommandletBase.h
--- a/Engine/Source/Editor/UnrealEd/Classes/Commandlets/GatherTextCommandletBase.h
+++ b/Engine/Source/Editor/UnrealEd/Classes/Commandlets/GatherTextCommandletBase.h
@@ -79,6 +79,9 @@ public:
 	bool ShouldSplitPlatformForPath(const FString& InPath, FName* OutPlatformName = nullptr) const;
 	FName GetSplitPlatformNameFromPath(const FString& InPath) const;
 
+	// Utility to get the correct base path (engine or project) for the current environment
+	static const FString& GetProjectBasePath();
+
 protected:
 	TSharedPtr< FLocTextHelper > GatherManifestHelper;
 
diff --git a/Engine/Source/Editor/UnrealEd/Private/Commandlets/ExtractLocResCommandlet.cpp b/Engine/Source/Editor/UnrealEd/Private/Commandlets/ExtractLocResCommandlet.cpp
--- a/Engine/Source/Editor/UnrealEd/Private/Commandlets/ExtractLocResCommandlet.cpp
+++ b/Engine/Source/Editor/UnrealEd/Private/Commandlets/ExtractLocResCommandlet.cpp
@@ -1,6 +1,7 @@
 // Copyright Epic Games, Inc. All Rights Reserved.
 
 #include "Commandlets/ExtractLocResCommandlet.h"
+#include "Commandlets/GatherTextCommandletBase.h"
 #include "Misc/Paths.h"
 #include "Misc/FileHelper.h"
 #include "HAL/FileManager.h"
@@ -23,15 +24,7 @@ int32 UExtractLocResCommandlet::Main(const FString& Params)
 		TArray<FString> PotentialLocResFileNames;
 		LocResParamPtr->ParseIntoArray(PotentialLocResFileNames, TEXT(";"));
 
-		FString ProjectBasePath;
-		if (!FPaths::ProjectDir().IsEmpty())
-		{
-			ProjectBasePath = FPaths::ProjectDir();
-		}
-		else
-		{
-			ProjectBasePath = FPaths::EngineDir();
-		}
+		const FString& ProjectBasePath = UGatherTextCommandletBase::GetProjectBasePath();
 
 		for (FString& PotentialLocResFileName : PotentialLocResFileNames)
 		{
diff --git a/Engine/Source/Editor/UnrealEd/Private/Commandlets/GatherTextCommandlet.cpp b/Engine/Source/Editor/UnrealEd/Private/Commandlets/GatherTextCommandlet.cpp
--- a/Engine/Source/Editor/UnrealEd/Private/Commandlets/GatherTextCommandlet.cpp
+++ b/Engine/Source/Editor/UnrealEd/Private/Commandlets/GatherTextCommandlet.cpp
@@ -42,15 +42,7 @@ int32 UGatherTextCommandlet::Main( const FString& Params )
 	{
 		ConfigParamPtr->ParseIntoArray(GatherTextConfigPaths, TEXT(";"));
 
-		FString ProjectBasePath;
-		if (!FPaths::ProjectDir().IsEmpty())
-		{
-			ProjectBasePath = FPaths::ProjectDir();
-		}
-		else
-		{
-			ProjectBasePath = FPaths::EngineDir();
-		}
+		const FString& ProjectBasePath = UGatherTextCommandletBase::GetProjectBasePath();
 
 		for (FString& GatherTextConfigPath : GatherTextConfigPaths)
 		{
diff --git a/Engine/Source/Editor/UnrealEd/Private/Commandlets/GatherTextCommandletBase.cpp b/Engine/Source/Editor/UnrealEd/Private/Commandlets/GatherTextCommandletBase.cpp
--- a/Engine/Source/Editor/UnrealEd/Private/Commandlets/GatherTextCommandletBase.cpp
+++ b/Engine/Source/Editor/UnrealEd/Private/Commandlets/GatherTextCommandletBase.cpp
@@ -1,6 +1,7 @@
 // Copyright Epic Games, Inc. All Rights Reserved.
 
 #include "Commandlets/GatherTextCommandletBase.h"
+#include "Misc/App.h"
 #include "Misc/Paths.h"
 #include "Misc/ConfigCacheIni.h"
 #include "Misc/PackageName.h"
@@ -108,7 +109,6 @@ bool UGatherTextCommandletBase::GetStringFromConfig( const TCHAR* Section, const
 
 void ResolveLocalizationPath(FString& InOutPath)
 {
-	static const bool bIsEngineTarget = FPaths::ProjectDir().IsEmpty();
 	static const FString AbsoluteEnginePath = FPaths::ConvertRelativePathToFull(FPaths::EngineDir()) / FString();
 	static const FString AbsoluteProjectPath = FPaths::ConvertRelativePathToFull(FPaths::ProjectDir()) / FString();
 
@@ -117,7 +117,9 @@ void ResolveLocalizationPath(FString& InOutPath)
 
 	if (FPaths::IsRelative(InOutPath))
 	{
-		InOutPath.InsertAt(0, bIsEngineTarget ? AbsoluteEnginePath : AbsoluteProjectPath);
+		static const FString AbsoluteTargetPath = FPaths::ConvertRelativePathToFull(UGatherTextCommandletBase::GetProjectBasePath()) / FString();
+
+		InOutPath.InsertAt(0, AbsoluteTargetPath);
 	}
 
 	FPaths::CollapseRelativeDirectories(InOutPath);
@@ -156,6 +158,12 @@ int32 UGatherTextCommandletBase::GetPathArrayFromConfig( const TCHAR* Section, c
 	return count;
 }
 
+const FString& UGatherTextCommandletBase::GetProjectBasePath()
+{
+	static const FString ProjectBasePath = FApp::HasProjectName() ? FPaths::ProjectDir() : FPaths::EngineDir();
+	return ProjectBasePath;
+}
+
 
 FFuzzyPathMatcher::FFuzzyPathMatcher(const TArray<FString>& InIncludePathFilters, const TArray<FString>& InExcludePathFilters)
 {
diff --git a/Engine/Source/Editor/UnrealEd/Private/Commandlets/GatherTextFromSourceCommandlet.cpp b/Engine/Source/Editor/UnrealEd/Private/Commandlets/GatherTextFromSourceCommandlet.cpp
--- a/Engine/Source/Editor/UnrealEd/Private/Commandlets/GatherTextFromSourceCommandlet.cpp
+++ b/Engine/Source/Editor/UnrealEd/Private/Commandlets/GatherTextFromSourceCommandlet.cpp
@@ -280,15 +280,7 @@ int32 UGatherTextFromSourceCommandlet::Main( const FString& Params )
 	// Parse all source files for macros and add entries to SourceParsedEntries
 	for ( FString& SourceFile : FilesToProcess)
 	{
-		FString ProjectBasePath;
-		if (!FPaths::ProjectDir().IsEmpty())
-		{
-			ProjectBasePath = FPaths::ProjectDir();
-		}
-		else
-		{
-			ProjectBasePath = FPaths::EngineDir();
-		}
+		const FString& ProjectBasePath = UGatherTextCommandletBase::GetProjectBasePath();
 
 		ParseCtxt.Filename = SourceFile;
 		ParseCtxt.FileTypes = ParseCtxt.Filename.EndsWith(TEXT(".ini")) ? EGatherTextSourceFileTypes::Ini : EGatherTextSourceFileTypes::Cpp;
