Fix build on macOS

#rnx

diff --git a/Engine/Source/Editor/UnrealEd/Classes/Exporters/TextureExporterGeneric.h b/Engine/Source/Editor/UnrealEd/Classes/Exporters/TextureExporterGeneric.h
--- a/Engine/Source/Editor/UnrealEd/Classes/Exporters/TextureExporterGeneric.h
+++ b/Engine/Source/Editor/UnrealEd/Classes/Exporters/TextureExporterGeneric.h
@@ -9,6 +9,7 @@
 #include "CoreMinimal.h"
 #include "UObject/ObjectMacros.h"
 #include "Exporters/Exporter.h"
+#include "Engine/Texture2D.h"
 #include "TextureExporterGeneric.generated.h"
 
 UCLASS()
