Persona - ReImport new Skeleton Issues
#fix forcefully close and reopen SkeletalMesh editor in case it is reimported with a new target Skeleton

diff --git a/Engine/Source/Editor/UnrealEd/Private/Factories/EditorFactories.cpp b/Engine/Source/Editor/UnrealEd/Private/Factories/EditorFactories.cpp
--- a/Engine/Source/Editor/UnrealEd/Private/Factories/EditorFactories.cpp
+++ b/Engine/Source/Editor/UnrealEd/Private/Factories/EditorFactories.cpp
@@ -6541,6 +6541,17 @@ EReimportResult::Type UReimportFbxSkeletalMeshFactory::Reimport( UObject* Obj, i
 
 	if (!bOperationCanceled)
 	{
+		// Replacing the skeleton requires setting up a fresh editor
+		const bool bReplacingSkeleton = ImportOptions->SkeletonForAnimation != SkeletalMesh->GetSkeleton();
+		bool bShouldReopenAsset = false;
+		if (bReplacingSkeleton)
+		{
+			if (UAssetEditorSubsystem* AssetEditorSubsystem = GEditor->GetEditorSubsystem<UAssetEditorSubsystem>())
+			{
+				bShouldReopenAsset = AssetEditorSubsystem->CloseAllEditorsForAsset(SkeletalMesh) != 0;
+			}
+		}
+		
 		FScopedSuspendAlternateSkinWeightPreview ScopedSuspendAlternateSkinWeightPreview(SkeletalMesh);
 		{
 			FScopedSkeletalMeshPostEditChange ScopedPostEditChange(SkeletalMesh);
@@ -6605,6 +6616,15 @@ EReimportResult::Type UReimportFbxSkeletalMeshFactory::Reimport( UObject* Obj, i
 				FSkinWeightsUtilities::ReimportAlternateSkinWeight(SkeletalMesh, 0);
 			}
 
+			// Reopen asset if it was open previously
+			if (bShouldReopenAsset)
+			{
+				if (UAssetEditorSubsystem* AssetEditorSubsystem = GEditor->GetEditorSubsystem<UAssetEditorSubsystem>())
+				{
+					AssetEditorSubsystem->OpenEditorForAsset(SkeletalMesh);
+				}
+			}
+
 			// Reimporting can have dangerous effects if the mesh is still in the transaction buffer.  Reset the transaction buffer if this is the case
 			if (!IsRunningCommandlet() && GEditor->IsObjectInTransactionBuffer(SkeletalMesh))
 			{
