Add CreationName & CreationDescription output fields to DumpMaterialExpressionCommandlet

diff --git a/Engine/Source/Editor/UnrealEd/Private/Commandlets/DumpMaterialExpressionsCommandlet.cpp b/Engine/Source/Editor/UnrealEd/Private/Commandlets/DumpMaterialExpressionsCommandlet.cpp
--- a/Engine/Source/Editor/UnrealEd/Private/Commandlets/DumpMaterialExpressionsCommandlet.cpp
+++ b/Engine/Source/Editor/UnrealEd/Private/Commandlets/DumpMaterialExpressionsCommandlet.cpp
@@ -23,6 +23,12 @@ int32 UDumpMaterialExpressionsCommandlet::Main(const FString& Params)
 	{
 		UE_LOG(LogDumpMaterialExpressionsCommandlet, Log, TEXT("DumpMaterialExpressions"));
 		UE_LOG(LogDumpMaterialExpressionsCommandlet, Log, TEXT("This commandlet will dump to a plain text file an info table of all material expressions in the engine and the plugins enabled on the project."));
+		UE_LOG(LogDumpMaterialExpressionsCommandlet, Log, TEXT("The output fields include:"));
+		UE_LOG(LogDumpMaterialExpressionsCommandlet, Log, TEXT("Name - The class name of the material expression"));
+		UE_LOG(LogDumpMaterialExpressionsCommandlet, Log, TEXT("CreationName - The name displayed in the search dropdown menu to add an expression"));
+		UE_LOG(LogDumpMaterialExpressionsCommandlet, Log, TEXT("CreationDescription - The tooltip displayed on the CreationName in the search dropdown menu"));
+		UE_LOG(LogDumpMaterialExpressionsCommandlet, Log, TEXT("Caption - The caption displayed on the material expression node"));
+		UE_LOG(LogDumpMaterialExpressionsCommandlet, Log, TEXT("Tooltip - The tooltip displayed on the material expression node"));
 		return 0;
 	}
 
@@ -30,7 +36,8 @@ int32 UDumpMaterialExpressionsCommandlet::Main(const FString& Params)
 	{
 		UMaterialExpression* MaterialExpression;
 		FString Name;
-		FString DisplayName;
+		FString CreationName;
+		FString CreationDescription;
 		FString Caption;
 		FString Description;
 		FString Tooltip;
@@ -40,13 +47,15 @@ int32 UDumpMaterialExpressionsCommandlet::Main(const FString& Params)
 	static const FName NAME_DisplayName(TEXT("DisplayName"));
 
 	const FString NameField = TEXT("NAME");
-	const FString DisplayNameField = TEXT("DISPLAYNAME");
+	const FString CreationNameField = TEXT("CREATIONNAME");
+	const FString CreationDescriptionField = TEXT("CREATIONDESCRIPTION");
 	const FString CaptionField = TEXT("CAPTION");
 	const FString DescriptionField = TEXT("DESCRIPTION");
 	const FString TooltipField = TEXT("TOOLTIP");
 
 	int32 MaxNameLength = NameField.Len();
-	int32 MaxDisplayNameLength = DisplayNameField.Len();
+	int32 MaxCreationNameLength = CreationNameField.Len();
+	int32 MaxCreationDescriptionLength = CreationDescriptionField.Len();
 	int32 MaxCaptionLength = CaptionField.Len();
 	int32 MaxDescriptionLength = DescriptionField.Len();
 
@@ -75,17 +84,22 @@ int32 UDumpMaterialExpressionsCommandlet::Main(const FString& Params)
 					Tooltip += Line;
 				}
 
+				FString DisplayName = Class->GetMetaData(NAME_DisplayName);
+				FString CreationName = DefaultExpression->GetCreationName().ToString();
+
 				FMaterialExpressionInfo ExpressionInfo;
 				ExpressionInfo.MaterialExpression = DefaultExpression;
-				ExpressionInfo.Name = Class->GetName().Mid(FCString::Strlen(TEXT("MaterialExpression")));;
-				ExpressionInfo.DisplayName = Class->GetMetaData(NAME_DisplayName);
+				ExpressionInfo.Name = Class->GetName().Mid(FCString::Strlen(TEXT("MaterialExpression")));
+				ExpressionInfo.CreationName = (CreationName.IsEmpty() ? (DisplayName.IsEmpty() ? ExpressionInfo.Name : DisplayName) : CreationName);
+				ExpressionInfo.CreationDescription = DefaultExpression->GetCreationDescription().ToString();
 				ExpressionInfo.Caption = Caption;
 				ExpressionInfo.Description = DefaultExpression->GetDescription();
 				ExpressionInfo.Tooltip = Tooltip;
 				MaterialExpressionInfos.Add(ExpressionInfo);
 
 				MaxNameLength = FMath::Max(MaxNameLength, ExpressionInfo.Name.Len());
-				MaxDisplayNameLength = FMath::Max(MaxDisplayNameLength, ExpressionInfo.DisplayName.Len());
+				MaxCreationNameLength = FMath::Max(MaxCreationNameLength, ExpressionInfo.CreationName.Len());
+				MaxCreationDescriptionLength = FMath::Max(MaxCreationDescriptionLength, ExpressionInfo.CreationDescription.Len());
 				MaxCaptionLength = FMath::Max(MaxCaptionLength, ExpressionInfo.Caption.Len());
 				MaxDescriptionLength = FMath::Max(MaxDescriptionLength, ExpressionInfo.Description.Len());
 			}
@@ -95,7 +109,8 @@ int32 UDumpMaterialExpressionsCommandlet::Main(const FString& Params)
 	// Additional padding for spacing
 	const int32 AdditionalPadding = 3;
 	MaxNameLength += AdditionalPadding;
-	MaxDisplayNameLength += AdditionalPadding;
+	MaxCreationNameLength += AdditionalPadding;
+	MaxCreationDescriptionLength += AdditionalPadding;
 	MaxCaptionLength += AdditionalPadding;
 	MaxDescriptionLength += AdditionalPadding;
 
@@ -124,25 +139,28 @@ int32 UDumpMaterialExpressionsCommandlet::Main(const FString& Params)
 	const FString OutputFilePath = FPaths::Combine(*FPaths::ProjectSavedDir(), TEXT("MaterialEditor"), TEXT("MaterialExpressions.txt"));
 	FArchive* FileWriter = IFileManager::Get().CreateFileWriter(*OutputFilePath);
 
-	auto WriteLine = [FileWriter, GenerateSpacePadding, MaxNameLength, MaxDisplayNameLength, MaxCaptionLength, MaxDescriptionLength](const FString& Name, const FString& DisplayName, const FString& Caption, const FString& Description, const FString& Tooltip)
+	auto WriteLine = [FileWriter, GenerateSpacePadding, MaxNameLength, MaxCreationNameLength, MaxCreationDescriptionLength, MaxCaptionLength, MaxDescriptionLength]
+		(const FString& Name, const FString& CreationName, const FString& CreationDescription, const FString& Caption, const FString& Description, const FString& Tooltip)
 	{
 		FString NamePadding = GenerateSpacePadding(MaxNameLength, Name.Len());
-		FString DisplayNamePadding = GenerateSpacePadding(MaxDisplayNameLength, DisplayName.Len());
+		FString CreationNamePadding = GenerateSpacePadding(MaxCreationNameLength, CreationName.Len());
+		FString CreationDescriptionPadding = GenerateSpacePadding(MaxCreationDescriptionLength, CreationDescription.Len());
 		FString CaptionPadding = GenerateSpacePadding(MaxCaptionLength, Caption.Len());
 		FString DescriptionPadding = GenerateSpacePadding(MaxDescriptionLength, Description.Len());
 
 		FString OutputLine = (Name + NamePadding);
-		OutputLine += (DisplayName + DisplayNamePadding);
+		OutputLine += (CreationName + CreationNamePadding);
+		OutputLine += (CreationDescription + CreationDescriptionPadding);
 		OutputLine += (Caption + CaptionPadding);
 		OutputLine += (Description + DescriptionPadding);
 		OutputLine += (Tooltip + TEXT("\n"));
 		FileWriter->Serialize(const_cast<ANSICHAR*>(StringCast<ANSICHAR>(*OutputLine).Get()), OutputLine.Len());
 	};
 
-	WriteLine(NameField, DisplayNameField, CaptionField, DescriptionField, TooltipField);
+	WriteLine(NameField, CreationNameField, CreationDescriptionField, CaptionField, DescriptionField, TooltipField);
 	for (FMaterialExpressionInfo& ExpressionInfo : MaterialExpressionInfos)
 	{
-		WriteLine(GetFormattedText(ExpressionInfo.Name), GetFormattedText(ExpressionInfo.DisplayName), GetFormattedText(ExpressionInfo.Caption), GetFormattedText(ExpressionInfo.Description), GetFormattedText(ExpressionInfo.Tooltip));
+		WriteLine(GetFormattedText(ExpressionInfo.Name), GetFormattedText(ExpressionInfo.CreationName), GetFormattedText(ExpressionInfo.CreationDescription), GetFormattedText(ExpressionInfo.Caption), GetFormattedText(ExpressionInfo.Description), GetFormattedText(ExpressionInfo.Tooltip));
 	}
 
 	FString OutputLine = FString::Printf(TEXT("\nTotal %d material expressions found."), MaterialExpressionInfos.Num());
