Packaging: Fixed incorrect generation of layer information for subchunks leading to package generator crashing with an exception.

Subchunks are generated when having MaxChunkSize parameter set to a non-zero value in the project packaging properties. The problem was caused by the generated pakchunklayers.txt having information only for the main chunks, but not for subchunks. The package generator assumes that pakchunklayers.txt has the same number of lines as the generated pakchunklist.txt (storing a list of all chunks and subchunks).

This issue is a regression introduced in CL 7637286, which modified the code to output layer information for empty chunks.

diff --git a/Engine/Source/Editor/UnrealEd/Private/Commandlets/AssetRegistryGenerator.cpp b/Engine/Source/Editor/UnrealEd/Private/Commandlets/AssetRegistryGenerator.cpp
--- a/Engine/Source/Editor/UnrealEd/Private/Commandlets/AssetRegistryGenerator.cpp
+++ b/Engine/Source/Editor/UnrealEd/Private/Commandlets/AssetRegistryGenerator.cpp
@@ -729,6 +729,12 @@ bool FAssetRegistryGenerator::GenerateStreamingInstallManifest(int64 InOverrideC
 			// add this pakfilelist to our master list of pakfilelists
 			FString PakChunkListLine = FString::Printf(TEXT("%s%s\r\n"), *PakChunkFilename, *PakChunkOptions);
 			PakChunkListFile->Serialize(TCHAR_TO_ANSI(*PakChunkListLine), PakChunkListLine.Len());
+
+			// Add layer information for this subchunk (we serialize it for the main chunk outside of this loop, hence the check).
+			if (SubChunkIndex > 0)
+			{
+				ChunkLayerFile->Serialize(TCHAR_TO_ANSI(*LayerString), LayerString.Len());
+			}
 		}
 	}
 
