Add material expression Type and ShowInCreateMenu fields to DumpMateriLExpressionCommandlet

diff --git a/Engine/Source/Editor/UnrealEd/Private/Commandlets/DumpMaterialExpressionsCommandlet.cpp b/Engine/Source/Editor/UnrealEd/Private/Commandlets/DumpMaterialExpressionsCommandlet.cpp
--- a/Engine/Source/Editor/UnrealEd/Private/Commandlets/DumpMaterialExpressionsCommandlet.cpp
+++ b/Engine/Source/Editor/UnrealEd/Private/Commandlets/DumpMaterialExpressionsCommandlet.cpp
@@ -2,6 +2,16 @@
 
 #include "Commandlets/DumpMaterialExpressionsCommandlet.h"
 #include "Materials/MaterialExpression.h"
+#include "Materials/MaterialExpressionComment.h"
+#include "Materials/MaterialExpressionMaterialLayerOutput.h"
+#include "Materials/MaterialExpressionParameter.h"
+#include "Materials/MaterialExpressionNamedReroute.h"
+#include "Materials/MaterialExpressionExecBegin.h"
+#include "Materials/MaterialExpressionExecEnd.h"
+#include "Materials/MaterialExpressionPinBase.h"
+#include "Materials/MaterialExpressionFunctionInput.h"
+#include "Materials/MaterialExpressionFunctionOutput.h"
+#include "Materials/MaterialExpressionComposite.h"
 #include "UObject/UObjectIterator.h"
 #include "HAL/FileManager.h"
 
@@ -25,8 +35,10 @@ int32 UDumpMaterialExpressionsCommandlet::Main(const FString& Params)
 		UE_LOG(LogDumpMaterialExpressionsCommandlet, Log, TEXT("This commandlet will dump to a plain text file an info table of all material expressions in the engine and the plugins enabled on the project."));
 		UE_LOG(LogDumpMaterialExpressionsCommandlet, Log, TEXT("The output fields include:"));
 		UE_LOG(LogDumpMaterialExpressionsCommandlet, Log, TEXT("Name - The class name of the material expression"));
-		UE_LOG(LogDumpMaterialExpressionsCommandlet, Log, TEXT("CreationName - The name displayed in the search dropdown menu to add an expression"));
-		UE_LOG(LogDumpMaterialExpressionsCommandlet, Log, TEXT("CreationDescription - The tooltip displayed on the CreationName in the search dropdown menu"));
+		UE_LOG(LogDumpMaterialExpressionsCommandlet, Log, TEXT("Type - ControlFlow | HLSLGenerator | CLASS_Deprecated"));
+		UE_LOG(LogDumpMaterialExpressionsCommandlet, Log, TEXT("ShowInCreateMenu - If the expression appears in the create node dropdown menu"));
+		UE_LOG(LogDumpMaterialExpressionsCommandlet, Log, TEXT("CreationName - The name displayed in the create node dropdown menu to add an expression"));
+		UE_LOG(LogDumpMaterialExpressionsCommandlet, Log, TEXT("CreationDescription - The tooltip displayed on the CreationName in the create node dropdown menu"));
 		UE_LOG(LogDumpMaterialExpressionsCommandlet, Log, TEXT("Caption - The caption displayed on the material expression node"));
 		UE_LOG(LogDumpMaterialExpressionsCommandlet, Log, TEXT("Tooltip - The tooltip displayed on the material expression node"));
 		return 0;
@@ -41,19 +53,23 @@ int32 UDumpMaterialExpressionsCommandlet::Main(const FString& Params)
 		FString Caption;
 		FString Description;
 		FString Tooltip;
+		FString	Type;
+		FString ShowInCreateMenu;
 	};
 	TArray<FMaterialExpressionInfo> MaterialExpressionInfos;
 
-	static const FName NAME_DisplayName(TEXT("DisplayName"));
-
 	const FString NameField = TEXT("NAME");
-	const FString CreationNameField = TEXT("CREATIONNAME");
-	const FString CreationDescriptionField = TEXT("CREATIONDESCRIPTION");
+	const FString TypeField = TEXT("TYPE");
+	const FString ShowInCreateMenuField = TEXT("SHOW_IN_CREATE_MENU");
+	const FString CreationNameField = TEXT("CREATION_NAME");
+	const FString CreationDescriptionField = TEXT("CREATION_DESCRIPTION");
 	const FString CaptionField = TEXT("CAPTION");
 	const FString DescriptionField = TEXT("DESCRIPTION");
 	const FString TooltipField = TEXT("TOOLTIP");
 
 	int32 MaxNameLength = NameField.Len();
+	int32 MaxTypeLength = TypeField.Len();
+	int32 MaxShowInCreateMenuLength = ShowInCreateMenuField.Len();
 	int32 MaxCreationNameLength = CreationNameField.Len();
 	int32 MaxCreationDescriptionLength = CreationDescriptionField.Len();
 	int32 MaxCaptionLength = CaptionField.Len();
@@ -68,6 +84,34 @@ int32 UDumpMaterialExpressionsCommandlet::Main(const FString& Params)
 		{
 			if (UMaterialExpression* DefaultExpression = Cast<UMaterialExpression>(Class->GetDefaultObject()))
 			{
+				const bool bClassDeprecated = Class->HasAnyClassFlags(CLASS_Deprecated);
+				const bool bControlFlow = Class->HasMetaData("MaterialControlFlow");
+				const bool bNewHLSLGenerator = Class->HasMetaData("MaterialNewHLSLGenerator");
+
+				// If the expression is listed in the material node creation dropdown menu
+				// See class exclusions in:
+				//    MaterialExpressionClasses::InitMaterialExpressionClasses()
+				//    FMaterialEditorUtilities::AddMaterialExpressionCategory() 
+				//    MaterialExpressions.cpp IsAllowedExpressionType()
+				bool bShowInCreateMenu = !(bClassDeprecated
+									  || Class == UMaterialExpressionComment::StaticClass()
+									  || Class == UMaterialExpressionMaterialLayerOutput::StaticClass()
+									  || Class == UMaterialExpressionParameter::StaticClass()
+									  || Class == UMaterialExpressionNamedRerouteUsage::StaticClass()
+									  || Class == UMaterialExpressionExecBegin::StaticClass()
+									  || Class == UMaterialExpressionExecEnd::StaticClass()
+									  || Class == UMaterialExpressionPinBase::StaticClass()
+									  || Class == UMaterialExpressionFunctionInput::StaticClass()
+									  || Class == UMaterialExpressionFunctionOutput::StaticClass()
+									  || Class == UMaterialExpressionComposite::StaticClass());
+
+				FString ExpressionType;
+				if (bControlFlow)									{ ExpressionType = TEXT("ControlFlow"); }
+				if (!ExpressionType.IsEmpty() && bNewHLSLGenerator)	{ ExpressionType += TEXT("|"); }
+				if (bNewHLSLGenerator)								{ ExpressionType += TEXT("HLSLGenerator"); }
+				if (!ExpressionType.IsEmpty() && bClassDeprecated)	{ ExpressionType += TEXT("|"); }
+				if (bClassDeprecated)								{ ExpressionType += TEXT("CLASS_Deprecated"); }
+
 				TArray<FString> MultilineCaption;
 				DefaultExpression->GetCaption(MultilineCaption);
 				FString Caption;
@@ -84,7 +128,7 @@ int32 UDumpMaterialExpressionsCommandlet::Main(const FString& Params)
 					Tooltip += Line;
 				}
 
-				FString DisplayName = Class->GetMetaData(NAME_DisplayName);
+				FString DisplayName = Class->GetMetaData(TEXT("DisplayName"));
 				FString CreationName = DefaultExpression->GetCreationName().ToString();
 
 				FMaterialExpressionInfo ExpressionInfo;
@@ -95,9 +139,13 @@ int32 UDumpMaterialExpressionsCommandlet::Main(const FString& Params)
 				ExpressionInfo.Caption = Caption;
 				ExpressionInfo.Description = DefaultExpression->GetDescription();
 				ExpressionInfo.Tooltip = Tooltip;
+				ExpressionInfo.Type = ExpressionType;
+				ExpressionInfo.ShowInCreateMenu = bShowInCreateMenu ? TEXT("Yes") : TEXT("No");
 				MaterialExpressionInfos.Add(ExpressionInfo);
 
 				MaxNameLength = FMath::Max(MaxNameLength, ExpressionInfo.Name.Len());
+				MaxTypeLength = FMath::Max(MaxTypeLength, ExpressionInfo.Type.Len());
+				MaxShowInCreateMenuLength = FMath::Max(MaxShowInCreateMenuLength, ExpressionInfo.ShowInCreateMenu.Len());
 				MaxCreationNameLength = FMath::Max(MaxCreationNameLength, ExpressionInfo.CreationName.Len());
 				MaxCreationDescriptionLength = FMath::Max(MaxCreationDescriptionLength, ExpressionInfo.CreationDescription.Len());
 				MaxCaptionLength = FMath::Max(MaxCaptionLength, ExpressionInfo.Caption.Len());
@@ -109,6 +157,8 @@ int32 UDumpMaterialExpressionsCommandlet::Main(const FString& Params)
 	// Additional padding for spacing
 	const int32 AdditionalPadding = 3;
 	MaxNameLength += AdditionalPadding;
+	MaxTypeLength += AdditionalPadding;
+	MaxShowInCreateMenuLength += AdditionalPadding;
 	MaxCreationNameLength += AdditionalPadding;
 	MaxCreationDescriptionLength += AdditionalPadding;
 	MaxCaptionLength += AdditionalPadding;
@@ -139,16 +189,20 @@ int32 UDumpMaterialExpressionsCommandlet::Main(const FString& Params)
 	const FString OutputFilePath = FPaths::Combine(*FPaths::ProjectSavedDir(), TEXT("MaterialEditor"), TEXT("MaterialExpressions.txt"));
 	FArchive* FileWriter = IFileManager::Get().CreateFileWriter(*OutputFilePath);
 
-	auto WriteLine = [FileWriter, GenerateSpacePadding, MaxNameLength, MaxCreationNameLength, MaxCreationDescriptionLength, MaxCaptionLength, MaxDescriptionLength]
-		(const FString& Name, const FString& CreationName, const FString& CreationDescription, const FString& Caption, const FString& Description, const FString& Tooltip)
+	auto WriteLine = [FileWriter, GenerateSpacePadding, MaxNameLength, MaxTypeLength, MaxShowInCreateMenuLength, MaxCreationNameLength, MaxCreationDescriptionLength, MaxCaptionLength, MaxDescriptionLength]
+		(const FString& Name, const FString& Type, const FString& ShowInCreateMenu, const FString& CreationName, const FString& CreationDescription, const FString& Caption, const FString& Description, const FString& Tooltip)
 	{
 		FString NamePadding = GenerateSpacePadding(MaxNameLength, Name.Len());
+		FString TypePadding = GenerateSpacePadding(MaxTypeLength, Type.Len());
+		FString ShowInCreateMenuPadding = GenerateSpacePadding(MaxShowInCreateMenuLength, ShowInCreateMenu.Len());
 		FString CreationNamePadding = GenerateSpacePadding(MaxCreationNameLength, CreationName.Len());
 		FString CreationDescriptionPadding = GenerateSpacePadding(MaxCreationDescriptionLength, CreationDescription.Len());
 		FString CaptionPadding = GenerateSpacePadding(MaxCaptionLength, Caption.Len());
 		FString DescriptionPadding = GenerateSpacePadding(MaxDescriptionLength, Description.Len());
 
 		FString OutputLine = (Name + NamePadding);
+		OutputLine += (Type + TypePadding);
+		OutputLine += (ShowInCreateMenu + ShowInCreateMenuPadding);
 		OutputLine += (CreationName + CreationNamePadding);
 		OutputLine += (CreationDescription + CreationDescriptionPadding);
 		OutputLine += (Caption + CaptionPadding);
@@ -157,10 +211,12 @@ int32 UDumpMaterialExpressionsCommandlet::Main(const FString& Params)
 		FileWriter->Serialize(const_cast<ANSICHAR*>(StringCast<ANSICHAR>(*OutputLine).Get()), OutputLine.Len());
 	};
 
-	WriteLine(NameField, CreationNameField, CreationDescriptionField, CaptionField, DescriptionField, TooltipField);
+	WriteLine(NameField, TypeField, ShowInCreateMenuField, CreationNameField, CreationDescriptionField, CaptionField, DescriptionField, TooltipField);
 	for (FMaterialExpressionInfo& ExpressionInfo : MaterialExpressionInfos)
 	{
-		WriteLine(GetFormattedText(ExpressionInfo.Name), GetFormattedText(ExpressionInfo.CreationName), GetFormattedText(ExpressionInfo.CreationDescription), GetFormattedText(ExpressionInfo.Caption), GetFormattedText(ExpressionInfo.Description), GetFormattedText(ExpressionInfo.Tooltip));
+		WriteLine(GetFormattedText(ExpressionInfo.Name), GetFormattedText(ExpressionInfo.Type), GetFormattedText(ExpressionInfo.ShowInCreateMenu), 
+					GetFormattedText(ExpressionInfo.CreationName), GetFormattedText(ExpressionInfo.CreationDescription), GetFormattedText(ExpressionInfo.Caption), 
+					GetFormattedText(ExpressionInfo.Description), GetFormattedText(ExpressionInfo.Tooltip));
 	}
 
 	FString OutputLine = FString::Printf(TEXT("\nTotal %d material expressions found."), MaterialExpressionInfos.Num());
