USD: Fix dome light import now opening a pop up dialog for HDR parse method.

diff --git a/Engine/Plugins/Importers/USDImporter/Source/USDUtilities/Private/USDShadeConversion.cpp b/Engine/Plugins/Importers/USDImporter/Source/USDUtilities/Private/USDShadeConversion.cpp
--- a/Engine/Plugins/Importers/USDImporter/Source/USDUtilities/Private/USDShadeConversion.cpp
+++ b/Engine/Plugins/Importers/USDImporter/Source/USDUtilities/Private/USDShadeConversion.cpp
@@ -1157,6 +1157,7 @@ namespace UE
 					TextureFactory->SuppressImportOverwriteDialog();
 					TextureFactory->bUseHashAsGuid = true;
 					TextureFactory->LODGroup = LODGroup;
+					TextureFactory->HDRImportShouldBeLongLatCubeMap = EAppReturnType::YesAll;
 
 					if ( bIsSupportedUdimTexture )
 					{
diff --git a/Engine/Source/Editor/UnrealEd/Classes/Factories/TextureFactory.h b/Engine/Source/Editor/UnrealEd/Classes/Factories/TextureFactory.h
--- a/Engine/Source/Editor/UnrealEd/Classes/Factories/TextureFactory.h
+++ b/Engine/Source/Editor/UnrealEd/Classes/Factories/TextureFactory.h
@@ -148,6 +148,10 @@ class UNREALED_API UTextureFactory : public UFactory, public IImportSettingsPars
 	UPROPERTY(Transient)
 	ETextureSourceColorSpace ColorSpaceMode;
 
+	/* Store YesAll/NoAll responses: */
+	UPROPERTY(Transient)
+	TEnumAsByte<EAppReturnType::Type> HDRImportShouldBeLongLatCubeMap = EAppReturnType::Retry;
+
 public:
 	UTextureFactory(const FObjectInitializer& ObjectInitializer = FObjectInitializer::Get());
 
@@ -202,9 +206,6 @@ protected:
 	/** Keep track of if we are doing a reimport */
 	bool bIsDoingAReimport = false;
 
-	/* Store YesAll/NoAll responses: */
-	EAppReturnType::Type HDRImportShouldBeLongLatCubeMap = EAppReturnType::Retry;
-
 private:
 	/** This variable is static because in StaticImportObject() the type of the factory is not known. */
 	static bool bSuppressImportOverwriteDialog;
