CIS fixes for EditorAssetSubsystem

diff --git a/Engine/Source/Editor/UnrealEd/Private/Subsystems/EditorAssetSubsystem.cpp b/Engine/Source/Editor/UnrealEd/Private/Subsystems/EditorAssetSubsystem.cpp
--- a/Engine/Source/Editor/UnrealEd/Private/Subsystems/EditorAssetSubsystem.cpp
+++ b/Engine/Source/Editor/UnrealEd/Private/Subsystems/EditorAssetSubsystem.cpp
@@ -1,12 +1,21 @@
-﻿// Copyright Epic Games, Inc. All Rights Reserved.
+// Copyright Epic Games, Inc. All Rights Reserved.
 
 #include "Subsystems/EditorAssetSubsystem.h"
 
+#include "AssetRegistry/IAssetRegistry.h"
+#include "AssetRegistry/AssetRegistryModule.h"
 #include "AssetToolsModule.h"
+#include "Editor.h"
 #include "EditorScriptingHelpers.h"
 #include "Experimental/Async/LazyEvent.h"
+#include "FileHelpers.h"
+#include "HAL/FileManager.h"
 #include "Misc/PackageName.h"
+#include "Modules/ModuleManager.h"
+#include "ObjectTools.h"
+#include "PackageTools.h"
 #include "Templates/ValueOrError.h"
+#include "UObject/MetaData.h"
 
 DEFINE_LOG_CATEGORY_STATIC(LogEditorAssetSubsystem, Log, All);
 
diff --git a/Engine/Source/Editor/UnrealEd/Public/Subsystems/EditorAssetSubsystem.h b/Engine/Source/Editor/UnrealEd/Public/Subsystems/EditorAssetSubsystem.h
--- a/Engine/Source/Editor/UnrealEd/Public/Subsystems/EditorAssetSubsystem.h
+++ b/Engine/Source/Editor/UnrealEd/Public/Subsystems/EditorAssetSubsystem.h
@@ -1,9 +1,10 @@
-﻿// Copyright Epic Games, Inc. All Rights Reserved.
+// Copyright Epic Games, Inc. All Rights Reserved.
 
 #pragma once
 
 #include "CoreMinimal.h"
 #include "AssetData.h"
+#include "EditorSubsystem.h"
 
 #include "EditorAssetSubsystem.generated.h"
 
