#UE Fixing warning about ambiguous class that is emitted when a world partition actor package has GetClass() called on it during the asset registry generation part of the cook. GetClass() is using ANY_PACKAGE to search for the class, which is wrong when called on assets whose type may be generated from a blueprint, since blueprints can have the same name in different paths

diff --git a/Engine/Source/Editor/UnrealEd/Private/Commandlets/AssetRegistryGenerator.cpp b/Engine/Source/Editor/UnrealEd/Private/Commandlets/AssetRegistryGenerator.cpp
--- a/Engine/Source/Editor/UnrealEd/Private/Commandlets/AssetRegistryGenerator.cpp
+++ b/Engine/Source/Editor/UnrealEd/Private/Commandlets/AssetRegistryGenerator.cpp
@@ -1282,8 +1282,7 @@ void FAssetRegistryGenerator::FinalizeChunkIDs(const TSet<FName>& InCookedPackag
 		const_cast<FAssetData&>(AssetData).ChunkIDs.Empty();
 
 		// Update whether the owner package contains a map
-		UClass* AssetClass = AssetData.GetClass();
-		if (AssetClass && (AssetClass->IsChildOf(UWorld::StaticClass()) || AssetClass->IsChildOf(ULevel::StaticClass())))
+		if ((AssetData.PackageFlags & PKG_ContainsMap) != 0)
 		{
 			PackagesContainingMaps.Add(AssetData.PackageName);
 		}
