FAssetData: Callsites need to handle AssetData.GetClass() returning null. Added IsInstanceOf(UClass*) for the common operation GetClass() && GetClass->IsChildOf(BaseClass).
#rnx

diff --git a/Engine/Plugins/2D/Paper2D/Source/Paper2DEditor/Private/PaperFlipbookActorFactory.cpp b/Engine/Plugins/2D/Paper2D/Source/Paper2DEditor/Private/PaperFlipbookActorFactory.cpp
--- a/Engine/Plugins/2D/Paper2D/Source/Paper2DEditor/Private/PaperFlipbookActorFactory.cpp
+++ b/Engine/Plugins/2D/Paper2D/Source/Paper2DEditor/Private/PaperFlipbookActorFactory.cpp
@@ -48,7 +48,7 @@ void UPaperFlipbookActorFactory::PostCreateBlueprint(UObject* Asset, AActor* CDO
 
 bool UPaperFlipbookActorFactory::CanCreateActorFrom(const FAssetData& AssetData, FText& OutErrorMsg)
 {
-	if (AssetData.IsValid() && AssetData.GetClass()->IsChildOf(UPaperFlipbook::StaticClass()))
+	if (AssetData.IsValid() && AssetData.IsInstanceOf(UPaperFlipbook::StaticClass()))
 	{
 		return true;
 	}
diff --git a/Engine/Plugins/2D/Paper2D/Source/Paper2DEditor/Private/PaperSpriteActorFactory.cpp b/Engine/Plugins/2D/Paper2D/Source/Paper2DEditor/Private/PaperSpriteActorFactory.cpp
--- a/Engine/Plugins/2D/Paper2D/Source/Paper2DEditor/Private/PaperSpriteActorFactory.cpp
+++ b/Engine/Plugins/2D/Paper2D/Source/Paper2DEditor/Private/PaperSpriteActorFactory.cpp
@@ -60,7 +60,7 @@ void UPaperSpriteActorFactory::PostCreateBlueprint(UObject* Asset, AActor* CDO)
 
 bool UPaperSpriteActorFactory::CanCreateActorFrom(const FAssetData& AssetData, FText& OutErrorMsg)
 {
-	if (AssetData.IsValid() && AssetData.GetClass()->IsChildOf(UPaperSprite::StaticClass()))
+	if (AssetData.IsValid() && AssetData.IsInstanceOf(UPaperSprite::StaticClass()))
 	{
 		return true;
 	}
diff --git a/Engine/Plugins/2D/Paper2D/Source/Paper2DEditor/Private/TerrainSplineEditing/TerrainSplineActorFactory.cpp b/Engine/Plugins/2D/Paper2D/Source/Paper2DEditor/Private/TerrainSplineEditing/TerrainSplineActorFactory.cpp
--- a/Engine/Plugins/2D/Paper2D/Source/Paper2DEditor/Private/TerrainSplineEditing/TerrainSplineActorFactory.cpp
+++ b/Engine/Plugins/2D/Paper2D/Source/Paper2DEditor/Private/TerrainSplineEditing/TerrainSplineActorFactory.cpp
@@ -42,7 +42,7 @@ bool UTerrainSplineActorFactory::CanCreateActorFrom(const FAssetData& AssetData,
 {
 	if (GetDefault<UPaperRuntimeSettings>()->bEnableTerrainSplineEditing)
 	{
-		if (AssetData.IsValid() && AssetData.GetClass()->IsChildOf(UPaperTerrainMaterial::StaticClass()))
+		if (AssetData.IsValid() && AssetData.IsInstanceOf(UPaperTerrainMaterial::StaticClass()))
 		{
 			return true;
 		}
diff --git a/Engine/Plugins/Animation/IKRig/Source/IKRigEditor/Private/RetargetEditor/SIKRetargetAssetBrowser.cpp b/Engine/Plugins/Animation/IKRig/Source/IKRigEditor/Private/RetargetEditor/SIKRetargetAssetBrowser.cpp
--- a/Engine/Plugins/Animation/IKRig/Source/IKRigEditor/Private/RetargetEditor/SIKRetargetAssetBrowser.cpp
+++ b/Engine/Plugins/Animation/IKRig/Source/IKRigEditor/Private/RetargetEditor/SIKRetargetAssetBrowser.cpp
@@ -178,7 +178,7 @@ void SIKRetargetAssetBrowser::OnAssetDoubleClicked(const FAssetData& AssetData)
 bool SIKRetargetAssetBrowser::OnShouldFilterAsset(const struct FAssetData& AssetData)
 {
 	// is this an animation asset?
-	if (!AssetData.GetClass()->IsChildOf(UAnimationAsset::StaticClass()))
+	if (!AssetData.IsInstanceOf(UAnimationAsset::StaticClass()))
 	{
 		return true;
 	}
diff --git a/Engine/Plugins/Animation/IKRig/Source/IKRigEditor/Private/RigEditor/SIKRigAssetBrowser.cpp b/Engine/Plugins/Animation/IKRig/Source/IKRigEditor/Private/RigEditor/SIKRigAssetBrowser.cpp
--- a/Engine/Plugins/Animation/IKRig/Source/IKRigEditor/Private/RigEditor/SIKRigAssetBrowser.cpp
+++ b/Engine/Plugins/Animation/IKRig/Source/IKRigEditor/Private/RigEditor/SIKRigAssetBrowser.cpp
@@ -1,4 +1,4 @@
-﻿// Copyright Epic Games, Inc. All Rights Reserved.
+// Copyright Epic Games, Inc. All Rights Reserved.
 
 #include "RigEditor/SIKRigAssetBrowser.h"
 
@@ -83,7 +83,7 @@ void SIKRigAssetBrowser::OnAssetDoubleClicked(const FAssetData& AssetData)
 bool SIKRigAssetBrowser::OnShouldFilterAsset(const struct FAssetData& AssetData)
 {
 	// is this an animation asset?
-	if (!AssetData.GetClass()->IsChildOf(UAnimationAsset::StaticClass()))
+	if (!AssetData.IsInstanceOf(UAnimationAsset::StaticClass()))
 	{
 		return true;
 	}
diff --git a/Engine/Plugins/Editor/MaterialAnalyzer/Source/Private/SMaterialAnalyzer.cpp b/Engine/Plugins/Editor/MaterialAnalyzer/Source/Private/SMaterialAnalyzer.cpp
--- a/Engine/Plugins/Editor/MaterialAnalyzer/Source/Private/SMaterialAnalyzer.cpp
+++ b/Engine/Plugins/Editor/MaterialAnalyzer/Source/Private/SMaterialAnalyzer.cpp
@@ -222,7 +222,7 @@ void SMaterialAnalyzer::SetupAssetRegistryCallbacks()
 
 void SMaterialAnalyzer::OnAssetAdded(const FAssetData& InAssetData)
 {
-	if(InAssetData.GetClass()->IsChildOf<UMaterialInterface>())
+	if(InAssetData.IsInstanceOf(UMaterialInterface::StaticClass()))
 	{
 		RecentlyAddedAssetData.Add(InAssetData);
 	}
diff --git a/Engine/Plugins/Enterprise/VariantManagerContent/Source/VariantManagerContentEditor/Private/LevelVariantSetsActorFactory.cpp b/Engine/Plugins/Enterprise/VariantManagerContent/Source/VariantManagerContentEditor/Private/LevelVariantSetsActorFactory.cpp
--- a/Engine/Plugins/Enterprise/VariantManagerContent/Source/VariantManagerContentEditor/Private/LevelVariantSetsActorFactory.cpp
+++ b/Engine/Plugins/Enterprise/VariantManagerContent/Source/VariantManagerContentEditor/Private/LevelVariantSetsActorFactory.cpp
@@ -24,7 +24,7 @@ bool ULevelVariantSetsActorFactory::CanCreateActorFrom( const FAssetData& AssetD
 		return true;
 	}
 
-	if ( AssetData.IsValid() && !AssetData.GetClass()->IsChildOf( ULevelVariantSets::StaticClass() ) )
+	if ( AssetData.IsValid() && !AssetData.IsInstanceOf( ULevelVariantSets::StaticClass() ) )
 	{
 		OutErrorMsg = NSLOCTEXT("CanCreateActor", "NoLevelVariantSetsAsset", "A valid variant sets asset must be specified.");
 		return false;
diff --git a/Engine/Plugins/Experimental/ChaosSolverPlugin/Source/ChaosSolverEditor/Private/Chaos/ActorFactoryChaosSolver.cpp b/Engine/Plugins/Experimental/ChaosSolverPlugin/Source/ChaosSolverEditor/Private/Chaos/ActorFactoryChaosSolver.cpp
--- a/Engine/Plugins/Experimental/ChaosSolverPlugin/Source/ChaosSolverEditor/Private/Chaos/ActorFactoryChaosSolver.cpp
+++ b/Engine/Plugins/Experimental/ChaosSolverPlugin/Source/ChaosSolverEditor/Private/Chaos/ActorFactoryChaosSolver.cpp
@@ -26,7 +26,7 @@ UActorFactoryChaosSolver::UActorFactoryChaosSolver(const FObjectInitializer& Obj
 
 bool UActorFactoryChaosSolver::CanCreateActorFrom(const FAssetData& AssetData, FText& OutErrorMsg)
 {
-	if (!AssetData.IsValid() || !AssetData.GetClass()->IsChildOf(UChaosSolver::StaticClass()))
+	if (!AssetData.IsValid() || !AssetData.IsInstanceOf(UChaosSolver::StaticClass()))
 	{
 		OutErrorMsg = NSLOCTEXT("CanCreateActor", "NoChaosSolverSpecified", "No ChaosSolver was specified.");
 		return false;
diff --git a/Engine/Plugins/Experimental/GeometryCollectionPlugin/Source/GeometryCollectionEditor/Private/GeometryCollection/ActorFactoryGeometryCollection.cpp b/Engine/Plugins/Experimental/GeometryCollectionPlugin/Source/GeometryCollectionEditor/Private/GeometryCollection/ActorFactoryGeometryCollection.cpp
--- a/Engine/Plugins/Experimental/GeometryCollectionPlugin/Source/GeometryCollectionEditor/Private/GeometryCollection/ActorFactoryGeometryCollection.cpp
+++ b/Engine/Plugins/Experimental/GeometryCollectionPlugin/Source/GeometryCollectionEditor/Private/GeometryCollection/ActorFactoryGeometryCollection.cpp
@@ -27,7 +27,7 @@ UActorFactoryGeometryCollection::UActorFactoryGeometryCollection(const FObjectIn
 
 bool UActorFactoryGeometryCollection::CanCreateActorFrom(const FAssetData& AssetData, FText& OutErrorMsg)
 {
-	if (!AssetData.IsValid() || !AssetData.GetClass()->IsChildOf(UGeometryCollection::StaticClass()))
+	if (!AssetData.IsValid() || !AssetData.IsInstanceOf(UGeometryCollection::StaticClass()))
 	{
 		OutErrorMsg = NSLOCTEXT("CanCreateActor", "NoGeometryCollectionSpecified", "No GeometryCollection mesh was specified.");
 		return false;
diff --git a/Engine/Plugins/FX/Niagara/Source/NiagaraEditor/Private/NiagaraEditorCommon.cpp b/Engine/Plugins/FX/Niagara/Source/NiagaraEditor/Private/NiagaraEditorCommon.cpp
--- a/Engine/Plugins/FX/Niagara/Source/NiagaraEditor/Private/NiagaraEditorCommon.cpp
+++ b/Engine/Plugins/FX/Niagara/Source/NiagaraEditor/Private/NiagaraEditorCommon.cpp
@@ -1785,7 +1785,7 @@ UActorFactoryNiagara::UActorFactoryNiagara(const FObjectInitializer& ObjectIniti
 
 bool UActorFactoryNiagara::CanCreateActorFrom(const FAssetData& AssetData, FText& OutErrorMsg)
 {
-	if (!AssetData.IsValid() || !AssetData.GetClass()->IsChildOf(UNiagaraSystem::StaticClass()))
+	if (!AssetData.IsValid() || !AssetData.IsInstanceOf(UNiagaraSystem::StaticClass()))
 	{
 		OutErrorMsg = NSLOCTEXT("CanCreateActor", "NoSystem", "A valid Niagara System must be specified.");
 		return false;
diff --git a/Engine/Plugins/Media/MediaCompositing/Source/MediaCompositingEditor/Private/Sequencer/MediaSequenceRecorderExtender.cpp b/Engine/Plugins/Media/MediaCompositing/Source/MediaCompositingEditor/Private/Sequencer/MediaSequenceRecorderExtender.cpp
--- a/Engine/Plugins/Media/MediaCompositing/Source/MediaCompositingEditor/Private/Sequencer/MediaSequenceRecorderExtender.cpp
+++ b/Engine/Plugins/Media/MediaCompositing/Source/MediaCompositingEditor/Private/Sequencer/MediaSequenceRecorderExtender.cpp
@@ -336,7 +336,7 @@ bool FMediaSequenceRecorderExtender::OnRecordingMediaPlayerListAllowDrop(TShared
 		bResult = true;
 		for (const FAssetData& AssetData : AssetDragDropOperation->GetAssets())
 		{
-			if (!AssetData.IsValid() || !AssetData.GetClass()->IsChildOf<UMediaPlayer>())
+			if (!AssetData.IsValid() || !AssetData.IsInstanceOf(UMediaPlayer::StaticClass()))
 			{
 				bResult = false;
 				break;
@@ -354,7 +354,7 @@ bool FMediaSequenceRecorderExtender::OnRecordingMediaPlayerListAllowDrop(TShared
 	{
 		for (const FAssetData& AssetData : AssetDragDropOperation->GetAssets())
 		{
-			if (AssetData.IsValid() && AssetData.GetClass()->IsChildOf<UMediaPlayer>())
+			if (AssetData.IsValid() && AssetData.IsInstanceOf(UMediaPlayer::StaticClass()))
 			{
 				UMediaPlayer* MediaPlayer = Cast<UMediaPlayer>(AssetData.GetAsset());
 				if (MediaPlayer)
diff --git a/Engine/Plugins/MovieScene/MatineeToLevelSequence/Source/MatineeToLevelSequence/Private/MatineeCameraShakeToNewCameraShakeConverter.cpp b/Engine/Plugins/MovieScene/MatineeToLevelSequence/Source/MatineeToLevelSequence/Private/MatineeCameraShakeToNewCameraShakeConverter.cpp
--- a/Engine/Plugins/MovieScene/MatineeToLevelSequence/Source/MatineeToLevelSequence/Private/MatineeCameraShakeToNewCameraShakeConverter.cpp
+++ b/Engine/Plugins/MovieScene/MatineeToLevelSequence/Source/MatineeToLevelSequence/Private/MatineeCameraShakeToNewCameraShakeConverter.cpp
@@ -184,7 +184,7 @@ private:
 				AssetRegistry->GetAssetsByPackageName(AssetDependency, AssetDependencyDatas);
 				for (const FAssetData& AssetDependencyData : AssetDependencyDatas)
 				{
-					const bool bIsCameraAnimDependency = AssetDependencyData.GetClass()->IsChildOf<UCameraAnim>();
+					const bool bIsCameraAnimDependency = AssetDependencyData.IsInstanceOf(UCameraAnim::StaticClass());
 					if (bIsCameraAnimDependency)
 					{
 						return false;
diff --git a/Engine/Plugins/Runtime/GeometryCache/Source/GeometryCacheEd/Private/ActorFactoryGeometryCache.cpp b/Engine/Plugins/Runtime/GeometryCache/Source/GeometryCacheEd/Private/ActorFactoryGeometryCache.cpp
--- a/Engine/Plugins/Runtime/GeometryCache/Source/GeometryCacheEd/Private/ActorFactoryGeometryCache.cpp
+++ b/Engine/Plugins/Runtime/GeometryCache/Source/GeometryCacheEd/Private/ActorFactoryGeometryCache.cpp
@@ -16,7 +16,7 @@ UActorFactoryGeometryCache::UActorFactoryGeometryCache(const FObjectInitializer&
 
 bool UActorFactoryGeometryCache::CanCreateActorFrom(const FAssetData& AssetData, FText& OutErrorMsg)
 {
-	if (!AssetData.IsValid() || !AssetData.GetClass()->IsChildOf(UGeometryCache::StaticClass()))
+	if (!AssetData.IsValid() || !AssetData.IsInstanceOf(UGeometryCache::StaticClass()))
 	{
 		OutErrorMsg = FText::FromString("A valid GeometryCache must be specified.");
 		return false;
diff --git a/Engine/Plugins/Runtime/HairStrands/Source/HairStrandsEditor/Private/GroomActorFactory.cpp b/Engine/Plugins/Runtime/HairStrands/Source/HairStrandsEditor/Private/GroomActorFactory.cpp
--- a/Engine/Plugins/Runtime/HairStrands/Source/HairStrandsEditor/Private/GroomActorFactory.cpp
+++ b/Engine/Plugins/Runtime/HairStrands/Source/HairStrandsEditor/Private/GroomActorFactory.cpp
@@ -18,7 +18,7 @@ UGroomActorFactory::UGroomActorFactory(const FObjectInitializer& ObjectInitializ
 
 bool UGroomActorFactory::CanCreateActorFrom( const FAssetData& AssetData, FText& OutErrorMsg )
 {
-	if ( !AssetData.IsValid() || !AssetData.GetClass()->IsChildOf( UGroomAsset::StaticClass() ) )
+	if ( !AssetData.IsValid() || !AssetData.IsInstanceOf( UGroomAsset::StaticClass() ) )
 	{
 		OutErrorMsg = NSLOCTEXT("CanCreateActor", "NoGroomAsset", "A valid groom asset must be specified.");
 		return false;
diff --git a/Engine/Source/Developer/AssetTools/Private/AssetTools.cpp b/Engine/Source/Developer/AssetTools/Private/AssetTools.cpp
--- a/Engine/Source/Developer/AssetTools/Private/AssetTools.cpp
+++ b/Engine/Source/Developer/AssetTools/Private/AssetTools.cpp
@@ -358,6 +358,11 @@ void UAssetToolsImpl::GetAssetTypeActionsList( TArray<TWeakPtr<IAssetTypeActions
 
 TWeakPtr<IAssetTypeActions> UAssetToolsImpl::GetAssetTypeActionsForClass(const UClass* Class) const
 {
+	if (!Class)
+	{
+		return TSharedPtr<IAssetTypeActions>();
+	}
+
 	TSharedPtr<IAssetTypeActions> MostDerivedAssetTypeActions;
 
 	for (int32 TypeActionsIdx = 0; TypeActionsIdx < AssetTypeActionsList.Num(); ++TypeActionsIdx)
diff --git a/Engine/Source/Developer/AssetTools/Private/AssetViewUtils.cpp b/Engine/Source/Developer/AssetTools/Private/AssetViewUtils.cpp
--- a/Engine/Source/Developer/AssetTools/Private/AssetViewUtils.cpp
+++ b/Engine/Source/Developer/AssetTools/Private/AssetViewUtils.cpp
@@ -1297,7 +1297,7 @@ bool AssetViewUtils::IsValidObjectPathForCreate(const FString& ObjectPath, const
 	if (ExistingAsset.IsValid())
 	{
 		// An asset of a different type already exists at this location, inform the user and continue
-		if (ObjectClass && !ExistingAsset.GetClass()->IsChildOf(ObjectClass))
+		if (ObjectClass && !ExistingAsset.IsInstanceOf(ObjectClass))
 		{
 			OutErrorMessage = FText::Format(LOCTEXT("RenameAssetOtherTypeAlreadyExists", "An asset of type '{0}' already exists at this location with the name '{1}'."), FText::FromName(ExistingAsset.AssetClass), FText::FromString(ObjectName));
 			
diff --git a/Engine/Source/Developer/AssetTools/Private/AssetVtConversion.cpp b/Engine/Source/Developer/AssetTools/Private/AssetVtConversion.cpp
--- a/Engine/Source/Developer/AssetTools/Private/AssetVtConversion.cpp
+++ b/Engine/Source/Developer/AssetTools/Private/AssetVtConversion.cpp
@@ -46,7 +46,7 @@ void GetReferencersData(UObject *Object, UClass *MatchClass, TArray<FAssetData>
 		{
 			if (MatchClass != nullptr)
 			{
-				if (AssetData.GetClass()->IsChildOf(MatchClass))
+				if (AssetData.IsInstanceOf(MatchClass))
 				{
 					OutAssetDatas.AddUnique(AssetData);
 				}
diff --git a/Engine/Source/Developer/CookedEditor/Private/CookedEditorPackageManager.cpp b/Engine/Source/Developer/CookedEditor/Private/CookedEditorPackageManager.cpp
--- a/Engine/Source/Developer/CookedEditor/Private/CookedEditorPackageManager.cpp
+++ b/Engine/Source/Developer/CookedEditor/Private/CookedEditorPackageManager.cpp
@@ -196,7 +196,7 @@ bool FIniCookedEditorPackageManager::AllowAssetToBeGathered(const struct FAssetD
 {
 	for (UClass* Class : DisallowedAssetClassesToGather)
 	{
-		if (AssetData.GetClass()->IsChildOf(Class))
+		if (AssetData.IsInstanceOf(Class))
 		{
 			return false;
 		}
diff --git a/Engine/Source/Editor/AnimationModifiers/Private/SAnimationModifiersTab.cpp b/Engine/Source/Editor/AnimationModifiers/Private/SAnimationModifiersTab.cpp
--- a/Engine/Source/Editor/AnimationModifiers/Private/SAnimationModifiersTab.cpp
+++ b/Engine/Source/Editor/AnimationModifiers/Private/SAnimationModifiersTab.cpp
@@ -274,7 +274,7 @@ void SAnimationModifiersTab::FindAnimSequencesForSkeleton(TArray<UAnimSequence *
 		for (const FAssetData& Asset : Assets)
 		{
 			// Only add assets whos class is of UAnimSequence
-			if (Asset.GetClass()->IsChildOf(UAnimSequence::StaticClass()))
+			if (Asset.IsInstanceOf(UAnimSequence::StaticClass()))
 			{
 				ReferencedAnimSequences.Add(CastChecked<UAnimSequence>(Asset.GetAsset()));
 			}
diff --git a/Engine/Source/Editor/AudioEditor/Private/SoundSubmixGraphSchema.cpp b/Engine/Source/Editor/AudioEditor/Private/SoundSubmixGraphSchema.cpp
--- a/Engine/Source/Editor/AudioEditor/Private/SoundSubmixGraphSchema.cpp
+++ b/Engine/Source/Editor/AudioEditor/Private/SoundSubmixGraphSchema.cpp
@@ -154,7 +154,7 @@ void USoundSubmixGraphSchema::GetAssetsGraphHoverMessage(const TArray<FAssetData
 
 	for (const FAssetData& Data : Assets)
 	{
-		if (!Data.GetClass()->IsChildOf(USoundSubmixBase::StaticClass()))
+		if (!Data.IsInstanceOf(USoundSubmixBase::StaticClass()))
 		{
 			OutOkIcon = false;
 			OutTooltipText = TEXT("Asset(s) must all be Submixes.");
diff --git a/Engine/Source/Editor/Blutility/Private/BlutilityContentBrowserExtensions.cpp b/Engine/Source/Editor/Blutility/Private/BlutilityContentBrowserExtensions.cpp
--- a/Engine/Source/Editor/Blutility/Private/BlutilityContentBrowserExtensions.cpp
+++ b/Engine/Source/Editor/Blutility/Private/BlutilityContentBrowserExtensions.cpp
@@ -81,7 +81,7 @@ public:
 						else
 						{
 							// Is the asset the right kind?
-							bPassesClassFilter = Asset.GetClass()->IsChildOf(SupportedClass);
+							bPassesClassFilter = Asset.IsInstanceOf(SupportedClass);
 						}
 
 						if (bPassesClassFilter)
diff --git a/Engine/Source/Editor/FoliageEdit/Private/ActorFactoryProceduralFoliage.cpp b/Engine/Source/Editor/FoliageEdit/Private/ActorFactoryProceduralFoliage.cpp
--- a/Engine/Source/Editor/FoliageEdit/Private/ActorFactoryProceduralFoliage.cpp
+++ b/Engine/Source/Editor/FoliageEdit/Private/ActorFactoryProceduralFoliage.cpp
@@ -31,7 +31,7 @@ bool UActorFactoryProceduralFoliage::PreSpawnActor(UObject* Asset, FTransform& I
 
 bool UActorFactoryProceduralFoliage::CanCreateActorFrom(const FAssetData& AssetData, FText& OutErrorMsg)
 {
-	if (!AssetData.IsValid() || !AssetData.GetClass()->IsChildOf(UProceduralFoliageSpawner::StaticClass()))
+	if (!AssetData.IsValid() || !AssetData.IsInstanceOf(UProceduralFoliageSpawner::StaticClass()))
 	{
 		OutErrorMsg = NSLOCTEXT("CanCreateActor", "NoProceduralFoliageSpawner", "A valid ProceduralFoliageSpawner must be specified.");
 		return false;
diff --git a/Engine/Source/Editor/FoliageEdit/Private/SFoliagePalette.cpp b/Engine/Source/Editor/FoliageEdit/Private/SFoliagePalette.cpp
--- a/Engine/Source/Editor/FoliageEdit/Private/SFoliagePalette.cpp
+++ b/Engine/Source/Editor/FoliageEdit/Private/SFoliagePalette.cpp
@@ -837,7 +837,7 @@ EVisibility SFoliagePalette::GetFoliageDropTargetVisibility() const
 			{
 				for (const UClass* Filter : ClassFilters)
 				{
-					if (AssetData.GetClass()->IsChildOf(Filter))
+					if (AssetData.IsInstanceOf(Filter))
 					{
 						return EVisibility::Visible;
 					}
diff --git a/Engine/Source/Editor/Kismet/Private/SSCSEditor.cpp b/Engine/Source/Editor/Kismet/Private/SSCSEditor.cpp
--- a/Engine/Source/Editor/Kismet/Private/SSCSEditor.cpp
+++ b/Engine/Source/Editor/Kismet/Private/SSCSEditor.cpp
@@ -202,7 +202,7 @@ FReply SSCSEditor::TryHandleAssetDragDropOperation(const FDragDropEvent& DragDro
 						PotentialActorClass = BPClass->GeneratedClass;
 					}
 				}
-				else if (AssetClass->IsChildOf(UClass::StaticClass()))
+				else if (AssetClass && AssetClass->IsChildOf(UClass::StaticClass()))
 				{
 					UClass* AssetAsClass = CastChecked<UClass>(Asset);
 					if (AssetAsClass->IsChildOf(UActorComponent::StaticClass()))
diff --git a/Engine/Source/Editor/LevelEditor/Private/LevelEditorCreateActorMenu.cpp b/Engine/Source/Editor/LevelEditor/Private/LevelEditorCreateActorMenu.cpp
--- a/Engine/Source/Editor/LevelEditor/Private/LevelEditorCreateActorMenu.cpp
+++ b/Engine/Source/Editor/LevelEditor/Private/LevelEditorCreateActorMenu.cpp
@@ -247,13 +247,14 @@ static void GetContentBrowserSelectionFactoryMenuEntries( FAssetData& TargetAsse
 		TargetAssetData = SelectedAssets.Top();
 	}
 
-	if ( TargetAssetData.GetClass() == UClass::StaticClass() )
+	UClass* AssetClass = TargetAssetData.GetClass();
+	if (AssetClass == UClass::StaticClass() )
 	{
 		UClass* Class = Cast<UClass>( TargetAssetData.GetAsset() );
 
 		bPlaceable = AssetSelectionUtils::IsClassPlaceable( Class );
 	}
-	else if (TargetAssetData.GetClass()->IsChildOf<UBlueprint>())
+	else if (AssetClass && AssetClass->IsChildOf<UBlueprint>())
 	{
 		// For blueprints, attempt to determine placeability from its tag information
 
diff --git a/Engine/Source/Editor/Persona/Private/AnimGraphNodeDetails.cpp b/Engine/Source/Editor/Persona/Private/AnimGraphNodeDetails.cpp
--- a/Engine/Source/Editor/Persona/Private/AnimGraphNodeDetails.cpp
+++ b/Engine/Source/Editor/Persona/Private/AnimGraphNodeDetails.cpp
@@ -271,7 +271,7 @@ bool FAnimGraphNodeDetails::OnShouldFilterAnimAsset( const FAssetData& AssetData
 	{
 		const UClass* AssetClass = AssetData.GetClass();
 		// If node is an 'asset player', only let you select the right kind of asset for it
-		if (!NodeToFilterFor->IsChildOf(UAnimGraphNode_AssetPlayerBase::StaticClass()) || SupportNodeClassForAsset(AssetClass, NodeToFilterFor))
+		if (!NodeToFilterFor->IsChildOf(UAnimGraphNode_AssetPlayerBase::StaticClass()) || (AssetClass && SupportNodeClassForAsset(AssetClass, NodeToFilterFor)))
 		{
 			return false;
 		}
diff --git a/Engine/Source/Editor/Persona/Private/SAnimationSequenceBrowser.cpp b/Engine/Source/Editor/Persona/Private/SAnimationSequenceBrowser.cpp
--- a/Engine/Source/Editor/Persona/Private/SAnimationSequenceBrowser.cpp
+++ b/Engine/Source/Editor/Persona/Private/SAnimationSequenceBrowser.cpp
@@ -347,7 +347,7 @@ public:
 		FAssetData ItemAssetData;
 		if (InItem.Legacy_TryGetAssetData(ItemAssetData))
 		{
-			return !ItemAssetData.GetClass()->IsChildOf(USoundWave::StaticClass());
+			return !ItemAssetData.IsInstanceOf(USoundWave::StaticClass());
 		}
 		return false;
 	}
@@ -593,7 +593,7 @@ TSharedPtr<SWidget> SAnimationSequenceBrowser::OnGetAssetContextMenu(const TArra
 			);
 		}
 
-		if (SelectedAssets.Num() == 1 && SelectedAssets[0].GetClass()->IsChildOf(USoundWave::StaticClass()))
+		if (SelectedAssets.Num() == 1 && SelectedAssets[0].IsInstanceOf(USoundWave::StaticClass()))
 		{
 			FToolMenuSection& Section = InMenu->AddSection("SequenceAudioOptions", LOCTEXT("AssetHeading", "Asset"));
 
@@ -1539,7 +1539,7 @@ EVisibility SAnimationSequenceBrowser::GetHistoryVisibility() const
 
 bool SAnimationSequenceBrowser::HandleFilterAsset(const FAssetData& InAssetData) const
 {
-	if (InAssetData.GetClass()->IsChildOf(UAnimationAsset::StaticClass()))
+	if (InAssetData.IsInstanceOf(UAnimationAsset::StaticClass()))
 	{
 		const USkeleton* DesiredSkeleton = PersonaToolkitPtr.Pin()->GetSkeleton();
 		if (DesiredSkeleton)
diff --git a/Engine/Source/Editor/PlacementMode/Private/PlacementModeModule.cpp b/Engine/Source/Editor/PlacementMode/Private/PlacementModeModule.cpp
--- a/Engine/Source/Editor/PlacementMode/Private/PlacementModeModule.cpp
+++ b/Engine/Source/Editor/PlacementMode/Private/PlacementModeModule.cpp
@@ -764,13 +764,14 @@ bool FPlacementModeModule::PassesFilters(const TSharedPtr<FPlaceableItem>& Item)
 		if (PredicatePair.Value(Item))
 		{
 			bool bPlaceable = true;
-			if (Item->AssetData.GetClass() == UClass::StaticClass())
+			UClass* AssetClass = Item->AssetData.GetClass();
+			if (AssetClass == UClass::StaticClass())
 			{
 				UClass* Class = Cast<UClass>(Item->AssetData.GetAsset());
 
 				bPlaceable = AssetSelectionUtils::IsClassPlaceable(Class);
 			}
-			else if (Item->AssetData.GetClass()->IsChildOf<UBlueprint>())
+			else if (AssetClass && AssetClass->IsChildOf<UBlueprint>())
 			{
 				// For blueprints, attempt to determine placeability from its tag information
 
diff --git a/Engine/Source/Editor/PlacementMode/Private/SPlacementModeTools.cpp b/Engine/Source/Editor/PlacementMode/Private/SPlacementModeTools.cpp
--- a/Engine/Source/Editor/PlacementMode/Private/SPlacementModeTools.cpp
+++ b/Engine/Source/Editor/PlacementMode/Private/SPlacementModeTools.cpp
@@ -482,8 +482,12 @@ FReply SPlacementAssetMenuEntry::OnMouseButtonUp(const FGeometry& MyGeometry, co
 		if (!Item->Factory)
 		{
 			// If no actor factory was found or failed, add the actor from the uclass
-			UObject* ClassObject = Item->AssetData.GetClass()->GetDefaultObject();
-			FActorFactoryAssetProxy::GetFactoryForAssetObject(ClassObject);
+			UClass* AssetClass = Item->AssetData.GetClass();
+			if (AssetClass)
+			{
+				UObject* ClassObject = AssetClass->GetDefaultObject();
+				FActorFactoryAssetProxy::GetFactoryForAssetObject(ClassObject);
+			}
 		}
 		NewActor = FLevelEditorActionCallbacks::AddActor(Factory, Item->AssetData, &TempTransform);
 		if (NewActor && GCurrentLevelEditingViewportClient)
diff --git a/Engine/Source/Editor/PlacementMode/Public/IPlacementModeModule.h b/Engine/Source/Editor/PlacementMode/Public/IPlacementModeModule.h
--- a/Engine/Source/Editor/PlacementMode/Public/IPlacementModeModule.h
+++ b/Engine/Source/Editor/PlacementMode/Public/IPlacementModeModule.h
@@ -136,7 +136,7 @@ struct FPlaceableItem
 	{
 		UClass* Class = AssetData.GetClass() == UClass::StaticClass() ? Cast<UClass>(AssetData.GetAsset()) : nullptr;
 		const bool bIsVolume = Class && Class->IsChildOf<AVolume>();
-		const bool bIsShape = Class ? false : AssetData.GetClass()->IsChildOf(UStaticMesh::StaticClass());
+		const bool bIsShape = Class ? false : AssetData.IsInstanceOf(UStaticMesh::StaticClass());
 
 		// Use the class unless its a volume or shape.  Those need custom names as the factory that spawns them does not properly represent what is being spawned.
 		if (Class && !bIsVolume && !bIsShape)
diff --git a/Engine/Source/Editor/PropertyEditor/Private/UserInterface/PropertyEditor/SPropertyEditorArray.h b/Engine/Source/Editor/PropertyEditor/Private/UserInterface/PropertyEditor/SPropertyEditorArray.h
--- a/Engine/Source/Editor/PropertyEditor/Private/UserInterface/PropertyEditor/SPropertyEditorArray.h
+++ b/Engine/Source/Editor/PropertyEditor/Private/UserInterface/PropertyEditor/SPropertyEditorArray.h
@@ -110,7 +110,7 @@ private:
 				for (FAssetData AssetData : DragDropOp->GetAssets())
 				{
 					// if the type matches
-					if (AssetData.GetClass()->IsChildOf(ObjectProperty->PropertyClass))
+					if (AssetData.IsInstanceOf(ObjectProperty->PropertyClass))
 					{
 						PropertyEditor->AddGivenItem(AssetData.ObjectPath.ToString());
 					}
@@ -158,7 +158,7 @@ private:
 				for (FAssetData AssetData : DragDropOp->GetAssets())
 				{
 					// if the type does not match
-					if (!AssetData.GetClass()->IsChildOf(ObjectProperty->PropertyClass))
+					if (!AssetData.IsInstanceOf(ObjectProperty->PropertyClass))
 					{
 						bHasOnlyValidElements = false;
 						break;
diff --git a/Engine/Source/Editor/PropertyEditor/Private/UserInterface/PropertyEditor/SPropertyEditorAsset.cpp b/Engine/Source/Editor/PropertyEditor/Private/UserInterface/PropertyEditor/SPropertyEditorAsset.cpp
--- a/Engine/Source/Editor/PropertyEditor/Private/UserInterface/PropertyEditor/SPropertyEditorAsset.cpp
+++ b/Engine/Source/Editor/PropertyEditor/Private/UserInterface/PropertyEditor/SPropertyEditorAsset.cpp
@@ -1390,6 +1390,12 @@ bool SPropertyEditorAsset::CanSetBasedOnCustomClasses( const FAssetData& InAsset
 
 bool SPropertyEditorAsset::IsClassAllowed(const UClass* InClass) const
 {
+	if (!InClass)
+	{
+		// A null class will not match any filters. If we have an allow list, this means failure, otherwise it means success.
+		return AllowedClassFilters.Num() == 0;
+	}
+
 	bool bClassAllowed = true;
 	if (AllowedClassFilters.Num() > 0)
 	{
diff --git a/Engine/Source/Editor/SubobjectEditor/Private/SSubobjectEditor.cpp b/Engine/Source/Editor/SubobjectEditor/Private/SSubobjectEditor.cpp
--- a/Engine/Source/Editor/SubobjectEditor/Private/SSubobjectEditor.cpp
+++ b/Engine/Source/Editor/SubobjectEditor/Private/SSubobjectEditor.cpp
@@ -2638,7 +2638,7 @@ FReply SSubobjectEditor::TryHandleAssetDragDropOperation(const FDragDropEvent& D
 						PotentialActorClass = BPClass->GeneratedClass;
 					}
 				}
-				else if (AssetClass->IsChildOf(UClass::StaticClass()))
+				else if (AssetClass && AssetClass->IsChildOf(UClass::StaticClass()))
 				{
 					UClass* AssetAsClass = CastChecked<UClass>(Asset);
 					if (AssetAsClass->IsChildOf(UActorComponent::StaticClass()))
diff --git a/Engine/Source/Editor/UnrealEd/Classes/ActorFactories/ActorFactoryVolume.h b/Engine/Source/Editor/UnrealEd/Classes/ActorFactories/ActorFactoryVolume.h
--- a/Engine/Source/Editor/UnrealEd/Classes/ActorFactories/ActorFactoryVolume.h
+++ b/Engine/Source/Editor/UnrealEd/Classes/ActorFactories/ActorFactoryVolume.h
@@ -22,7 +22,7 @@ public:
 			return true;
 		}
 
-		if (AssetData.IsValid() && !AssetData.GetClass()->IsChildOf(AVolume::StaticClass()))
+		if (AssetData.IsValid() && !AssetData.IsInstanceOf(AVolume::StaticClass()))
 		{
 			return false;
 		}
diff --git a/Engine/Source/Editor/UnrealEd/Private/AssetSelection.cpp b/Engine/Source/Editor/UnrealEd/Private/AssetSelection.cpp
--- a/Engine/Source/Editor/UnrealEd/Private/AssetSelection.cpp
+++ b/Engine/Source/Editor/UnrealEd/Private/AssetSelection.cpp
@@ -792,14 +792,15 @@ void FActorFactoryAssetProxy::GenerateActorFactoryMenuItems( const FAssetData& A
 UActorFactory* FActorFactoryAssetProxy::GetFactoryForAsset( const FAssetData& AssetData, bool bRequireValidObject/*=false*/ )
 {
 	UObject* Asset = NULL;
+	UClass* AssetClass = AssetData.GetClass();
 	
 	if ( AssetData.IsAssetLoaded() )
 	{
 		Asset = AssetData.GetAsset();
 	}
-	else if ( !bRequireValidObject )
+	else if ( !bRequireValidObject && AssetClass )
 	{
-		Asset = AssetData.GetClass()->GetDefaultObject();
+		Asset = AssetClass->GetDefaultObject();
 	}
 
 	return FActorFactoryAssetProxy::GetFactoryForAssetObject( Asset );
diff --git a/Engine/Source/Editor/UnrealEd/Private/ConsolidateWindow.cpp b/Engine/Source/Editor/UnrealEd/Private/ConsolidateWindow.cpp
--- a/Engine/Source/Editor/UnrealEd/Private/ConsolidateWindow.cpp
+++ b/Engine/Source/Editor/UnrealEd/Private/ConsolidateWindow.cpp
@@ -913,7 +913,11 @@ FReply SConsolidateToolWidget::OnDragOver( const FGeometry& MyGeometry, const FD
 
 		if ( Object == NULL )
 		{
-			Object = ExtractedDroppedAsset[Index].GetClass()->GetDefaultObject();
+			UClass* ObjectClass = ExtractedDroppedAsset[Index].GetClass();
+			if (ObjectClass)
+			{
+				Object = ObjectClass->GetDefaultObject();
+			}
 		}
 
 		if ( Object != NULL )
diff --git a/Engine/Source/Editor/UnrealEd/Private/EditorServer.cpp b/Engine/Source/Editor/UnrealEd/Private/EditorServer.cpp
--- a/Engine/Source/Editor/UnrealEd/Private/EditorServer.cpp
+++ b/Engine/Source/Editor/UnrealEd/Private/EditorServer.cpp
@@ -609,7 +609,7 @@ void UEditorEngine::LoadAndSelectAssets( TArray<FAssetData>& Assets, UClass* Typ
 		for ( int32 CurrentAssetIndex=0; CurrentAssetIndex < Assets.Num(); CurrentAssetIndex++ )
 		{
 			FAssetData& SelectedAsset = Assets[CurrentAssetIndex];
-			if ( TypeOfAsset == NULL || SelectedAsset.GetClass()->IsChildOf( TypeOfAsset ) )
+			if ( TypeOfAsset == NULL || SelectedAsset.IsInstanceOf( TypeOfAsset ) )
 			{
 				// GetAsset() will load the asset if necessary
 				UObject* LoadedAsset = SelectedAsset.GetAsset();
diff --git a/Engine/Source/Editor/UnrealEd/Private/Factories/ActorFactory.cpp b/Engine/Source/Editor/UnrealEd/Private/Factories/ActorFactory.cpp
--- a/Engine/Source/Editor/UnrealEd/Private/Factories/ActorFactory.cpp
+++ b/Engine/Source/Editor/UnrealEd/Private/Factories/ActorFactory.cpp
@@ -538,7 +538,7 @@ UActorFactoryStaticMesh::UActorFactoryStaticMesh(const FObjectInitializer& Objec
 
 bool UActorFactoryStaticMesh::CanCreateActorFrom( const FAssetData& AssetData, FText& OutErrorMsg )
 {
-	if ( !AssetData.IsValid() || !AssetData.GetClass()->IsChildOf( UStaticMesh::StaticClass() ) )
+	if ( !AssetData.IsValid() || !AssetData.IsInstanceOf( UStaticMesh::StaticClass() ) )
 	{
 		OutErrorMsg = NSLOCTEXT("CanCreateActor", "NoStaticMesh", "A valid static mesh must be specified.");
 		return false;
@@ -674,7 +674,7 @@ bool UActorFactoryDeferredDecal::CanCreateActorFrom( const FAssetData& AssetData
 	}
 
 	//But if an asset is specified it must be based-on a deferred decal umaterial
-	if ( !AssetData.GetClass()->IsChildOf( UMaterialInterface::StaticClass() ) )
+	if ( !AssetData.IsInstanceOf( UMaterialInterface::StaticClass() ) )
 	{
 		OutErrorMsg = NSLOCTEXT("CanCreateActor", "NoMaterial", "A valid material must be specified.");
 		return false;
@@ -685,7 +685,7 @@ bool UActorFactoryDeferredDecal::CanCreateActorFrom( const FAssetData& AssetData
 
 	uint32 SanityCheck = 0;
 	FAssetData CurrentAssetData = AssetData;
-	while( SanityCheck < 1000 && !CurrentAssetData.GetClass()->IsChildOf( UMaterial::StaticClass() ) )
+	while( SanityCheck < 1000 && !CurrentAssetData.IsInstanceOf( UMaterial::StaticClass() ) )
 	{
 		const FString ObjectPath = CurrentAssetData.GetTagValueRef<FString>( "Parent" );
 		if ( ObjectPath.IsEmpty() )
@@ -710,7 +710,7 @@ bool UActorFactoryDeferredDecal::CanCreateActorFrom( const FAssetData& AssetData
 		return false;
 	}
 
-	if ( !CurrentAssetData.GetClass()->IsChildOf( UMaterial::StaticClass() ) )
+	if ( !CurrentAssetData.IsInstanceOf( UMaterial::StaticClass() ) )
 	{
 		return false;
 	}
@@ -814,7 +814,7 @@ UActorFactoryEmitter::UActorFactoryEmitter(const FObjectInitializer& ObjectIniti
 
 bool UActorFactoryEmitter::CanCreateActorFrom( const FAssetData& AssetData, FText& OutErrorMsg )
 {
-	if ( !AssetData.IsValid() || !AssetData.GetClass()->IsChildOf( UParticleSystem::StaticClass() ) )
+	if ( !AssetData.IsValid() || !AssetData.IsInstanceOf( UParticleSystem::StaticClass() ) )
 	{
 		OutErrorMsg = NSLOCTEXT("CanCreateActor", "NoParticleSystem", "A valid particle system must be specified.");
 		return false;
@@ -917,7 +917,7 @@ UActorFactoryPhysicsAsset::UActorFactoryPhysicsAsset(const FObjectInitializer& O
 
 bool UActorFactoryPhysicsAsset::CanCreateActorFrom( const FAssetData& AssetData, FText& OutErrorMsg )
 {
-	if ( !AssetData.IsValid() || !AssetData.GetClass()->IsChildOf( UPhysicsAsset::StaticClass() ) )
+	if ( !AssetData.IsValid() || !AssetData.IsInstanceOf( UPhysicsAsset::StaticClass() ) )
 	{
 		OutErrorMsg = NSLOCTEXT("CanCreateActor", "NoPhysicsAsset", "A valid physics asset must be specified.");
 		return false;
@@ -1023,7 +1023,7 @@ bool UActorFactoryAnimationAsset::CanCreateActorFrom( const FAssetData& AssetDat
 	FAssetRegistryModule& AssetRegistryModule = FModuleManager::LoadModuleChecked<FAssetRegistryModule>(TEXT("AssetRegistry"));
 	IAssetRegistry& AssetRegistry = AssetRegistryModule.Get();
 
-	if ( AssetData.GetClass()->IsChildOf( UAnimSequenceBase::StaticClass() ) )
+	if ( AssetData.IsInstanceOf( UAnimSequenceBase::StaticClass() ) )
 	{
 		const FString SkeletonPath = AssetData.GetTagValueRef<FString>("Skeleton");
 		if ( SkeletonPath.IsEmpty() ) 
@@ -1179,12 +1179,13 @@ bool UActorFactorySkeletalMesh::CanCreateActorFrom( const FAssetData& AssetData,
 	IAssetRegistry& AssetRegistry = AssetRegistryModule.Get();
 	FAssetData SkeletalMeshData;
 
-	if ( AssetData.GetClass()->IsChildOf( USkeletalMesh::StaticClass() ) )
+	UClass* AssetClass = AssetData.GetClass();
+	if ( AssetClass && AssetClass->IsChildOf( USkeletalMesh::StaticClass() ) )
 	{
 		SkeletalMeshData = AssetData;
 	}
 
-	if ( !SkeletalMeshData.IsValid() && (AssetData.GetClass()->IsChildOf(UAnimBlueprint::StaticClass()) || AssetData.GetClass()->IsChildOf(UAnimBlueprintGeneratedClass::StaticClass())))
+	if ( !SkeletalMeshData.IsValid() && AssetClass && (AssetClass->IsChildOf(UAnimBlueprint::StaticClass()) || AssetClass->IsChildOf(UAnimBlueprintGeneratedClass::StaticClass())))
 	{
 		const FString TargetSkeletonPath = AssetData.GetTagValueRef<FString>( "TargetSkeleton" );
 		if ( TargetSkeletonPath.IsEmpty() )
@@ -1222,7 +1223,7 @@ bool UActorFactorySkeletalMesh::CanCreateActorFrom( const FAssetData& AssetData,
 		}
 	}
 
-	if ( !SkeletalMeshData.IsValid() && AssetData.GetClass()->IsChildOf( USkeleton::StaticClass() ) )
+	if ( !SkeletalMeshData.IsValid() && AssetClass && AssetClass->IsChildOf( USkeleton::StaticClass() ) )
 	{
 		// so I'm changing this to load directly not using tags and values
 		USkeleton* Skeleton = Cast<USkeleton>(AssetData.GetAsset());
@@ -1258,9 +1259,10 @@ bool UActorFactorySkeletalMesh::CanCreateActorFrom( const FAssetData& AssetData,
 		return false;
 	}
 
-	if(USkeletalMesh* SkeletalMeshCDO = Cast<USkeletalMesh>(AssetData.GetClass()->GetDefaultObject()))
+	check(AssetClass);
+	if (USkeletalMesh* SkeletalMeshCDO = Cast<USkeletalMesh>(AssetClass->GetDefaultObject()))
 	{
-		if(SkeletalMeshCDO->HasCustomActorFactory())
+		if (SkeletalMeshCDO->HasCustomActorFactory())
 		{
 			return false;
 		}
@@ -1483,7 +1485,7 @@ bool UActorFactoryAmbientSound::CanCreateActorFrom( const FAssetData& AssetData,
 		return true;
 	}
 
-	if ( AssetData.IsValid() && !AssetData.GetClass()->IsChildOf( USoundBase::StaticClass() ) )
+	if ( AssetData.IsValid() && !AssetData.IsInstanceOf( USoundBase::StaticClass() ) )
 	{
 		OutErrorMsg = NSLOCTEXT("CanCreateActor", "NoSoundAsset", "A valid sound asset must be specified.");
 		return false;
@@ -1539,7 +1541,7 @@ UActorFactoryClass::UActorFactoryClass(const FObjectInitializer& ObjectInitializ
 
 bool UActorFactoryClass::CanCreateActorFrom( const FAssetData& AssetData, FText& OutErrorMsg )
 {
-	if ( AssetData.IsValid() && AssetData.GetClass()->IsChildOf( UClass::StaticClass() ) )
+	if ( AssetData.IsValid() && AssetData.IsInstanceOf( UClass::StaticClass() ) )
 	{
 		UClass* ActualClass = Cast<UClass>(AssetData.GetAsset());
 		if ( (nullptr != ActualClass) && ActualClass->IsChildOf(AActor::StaticClass()) )
@@ -1554,7 +1556,7 @@ bool UActorFactoryClass::CanCreateActorFrom( const FAssetData& AssetData, FText&
 
 AActor* UActorFactoryClass::GetDefaultActor( const FAssetData& AssetData )
 {
-	if ( AssetData.IsValid() && AssetData.GetClass()->IsChildOf( UClass::StaticClass() ) )
+	if ( AssetData.IsValid() && AssetData.IsInstanceOf( UClass::StaticClass() ) )
 	{
 		UClass* ActualClass = Cast<UClass>(StaticLoadObject(UClass::StaticClass(), nullptr, *AssetData.ObjectPath.ToString(), nullptr, LOAD_NoWarn, nullptr));
 			
@@ -1606,7 +1608,7 @@ UActorFactoryBlueprint::UActorFactoryBlueprint(const FObjectInitializer& ObjectI
 
 bool UActorFactoryBlueprint::CanCreateActorFrom( const FAssetData& AssetData, FText& OutErrorMsg )
 {
-	if ( !AssetData.IsValid() || !AssetData.GetClass()->IsChildOf( UBlueprint::StaticClass() ) )
+	if ( !AssetData.IsValid() || !AssetData.IsInstanceOf( UBlueprint::StaticClass() ) )
 	{
 		OutErrorMsg = NSLOCTEXT("CanCreateActor", "NoBlueprint", "No Blueprint was specified, or the specified Blueprint needs to be compiled.");
 		return false;
@@ -1659,7 +1661,7 @@ AActor* UActorFactoryBlueprint::GetDefaultActor( const FAssetData& AssetData )
 
 UClass* UActorFactoryBlueprint::GetDefaultActorClass(const FAssetData& AssetData)
 {
-	if (!AssetData.IsValid() || !AssetData.GetClass()->IsChildOf(UBlueprint::StaticClass()))
+	if (!AssetData.IsValid() || !AssetData.IsInstanceOf(UBlueprint::StaticClass()))
 	{
 		return nullptr;
 	}
@@ -1884,7 +1886,7 @@ UActorFactoryVectorFieldVolume::UActorFactoryVectorFieldVolume(const FObjectInit
 
 bool UActorFactoryVectorFieldVolume::CanCreateActorFrom( const FAssetData& AssetData, FText& OutErrorMsg )
 {
-	if ( !AssetData.IsValid() || !AssetData.GetClass()->IsChildOf( UVectorField::StaticClass() ) )
+	if ( !AssetData.IsValid() || !AssetData.IsInstanceOf( UVectorField::StaticClass() ) )
 	{
 		OutErrorMsg = NSLOCTEXT("CanCreateActor", "NoVectorField", "No vector field was specified.");
 		return false;
@@ -1967,7 +1969,7 @@ bool UActorFactoryBoxVolume::CanCreateActorFrom( const FAssetData& AssetData, FT
 		return true;
 	}
 
-	if ( AssetData.IsValid() && !AssetData.GetClass()->IsChildOf( AVolume::StaticClass() ) )
+	if ( AssetData.IsValid() && !AssetData.IsInstanceOf( AVolume::StaticClass() ) )
 	{
 		return false;
 	}
@@ -2004,7 +2006,7 @@ bool UActorFactorySphereVolume::CanCreateActorFrom( const FAssetData& AssetData,
 		return true;
 	}
 
-	if ( AssetData.IsValid() && !AssetData.GetClass()->IsChildOf( AVolume::StaticClass() ) )
+	if ( AssetData.IsValid() && !AssetData.IsInstanceOf( AVolume::StaticClass() ) )
 	{
 		return false;
 	}
@@ -2065,7 +2067,7 @@ bool UActorFactoryLevelSequence::CanCreateActorFrom( const FAssetData& AssetData
 		return true;
 	}
 
-	if ( AssetData.IsValid() && !AssetData.GetClass()->IsChildOf( ULevelSequence::StaticClass() ) )
+	if ( AssetData.IsValid() && !AssetData.IsInstanceOf( ULevelSequence::StaticClass() ) )
 	{
 		OutErrorMsg = NSLOCTEXT("CanCreateActor", "NoLevelSequenceAsset", "A valid sequencer asset must be specified.");
 		return false;
diff --git a/Engine/Source/Editor/UnrealEd/Private/Tests/AutomationEditorCommon.cpp b/Engine/Source/Editor/UnrealEd/Private/Tests/AutomationEditorCommon.cpp
--- a/Engine/Source/Editor/UnrealEd/Private/Tests/AutomationEditorCommon.cpp
+++ b/Engine/Source/Editor/UnrealEd/Private/Tests/AutomationEditorCommon.cpp
@@ -710,7 +710,7 @@ void FAutomationEditorCommonUtils::CollectGameContentTests(TArray<FString>& OutB
 				Filename = FPackageName::LongPackageNameToFilename(Filename);
 				if (FAutomationTestFramework::Get().ShouldTestContent(Filename))
 				{
-					FString BeautifiedFilename = FString::Printf(TEXT("%s.%s"), *Asset.GetClass()->GetFName().ToString(), *Asset.AssetName.ToString());
+					FString BeautifiedFilename = FString::Printf(TEXT("%s.%s"), *Asset.AssetClass.ToString(), *Asset.AssetName.ToString());
 					OutBeautifiedNames.Add(BeautifiedFilename);
 					OutTestCommands.Add(Asset.ObjectPath.ToString());
 				}
diff --git a/Engine/Source/Editor/VirtualTexturingEditor/Private/RuntimeVirtualTextureAssetTypeActions.cpp b/Engine/Source/Editor/VirtualTexturingEditor/Private/RuntimeVirtualTextureAssetTypeActions.cpp
--- a/Engine/Source/Editor/VirtualTexturingEditor/Private/RuntimeVirtualTextureAssetTypeActions.cpp
+++ b/Engine/Source/Editor/VirtualTexturingEditor/Private/RuntimeVirtualTextureAssetTypeActions.cpp
@@ -44,7 +44,7 @@ namespace
 			{
 				if (MatchClass != nullptr)
 				{
-					if (AssetData.GetClass()->IsChildOf(MatchClass))
+					if (AssetData.IsInstanceOf(MatchClass))
 					{
 						OutAssetDatas.AddUnique(AssetData);
 					}
