Adding UCompileShadersTestBedCommandlet which is meant as a simple way to profile/test/optimization shader compilation.

- Currently it only compiles the global shaders.
- Having a commandlet is faster and easier to run repeatable tests than dealing with starting up the editor.

# How to run
ProjectName -run=CompileShadersTestBed -ddc=Cold -ddc=NoShared -noshaderddc -trace=default,counters

This ensures you don't get results in any DDC and that you always compile shaders.

diff --git a/Engine/Source/Editor/UnrealEd/Classes/Commandlets/CompileShadersTestBedCommandlet.h b/Engine/Source/Editor/UnrealEd/Classes/Commandlets/CompileShadersTestBedCommandlet.h
new file mode 100644
--- /dev/null
+++ b/Engine/Source/Editor/UnrealEd/Classes/Commandlets/CompileShadersTestBedCommandlet.h
@@ -0,0 +1,17 @@
+// Copyright Epic Games, Inc. All Rights Reserved.
+
+#pragma once
+
+#include "CoreMinimal.h"
+#include "Commandlets/Commandlet.h"
+#include "CompileShadersTestBedCommandlet.generated.h"
+
+UCLASS(config=Editor)
+class UCompileShadersTestBedCommandlet : public UCommandlet
+{
+	GENERATED_UCLASS_BODY()
+
+	//~ Begin UCommandlet Interface
+	virtual int32 Main(const FString& Params) override;
+	//~ End UCommandlet Interface
+};
diff --git a/Engine/Source/Editor/UnrealEd/Private/Commandlets/CompileShadersTestBedCommandlet.cpp b/Engine/Source/Editor/UnrealEd/Private/Commandlets/CompileShadersTestBedCommandlet.cpp
new file mode 100644
--- /dev/null
+++ b/Engine/Source/Editor/UnrealEd/Private/Commandlets/CompileShadersTestBedCommandlet.cpp
@@ -0,0 +1,44 @@
+// Copyright Epic Games, Inc. All Rights Reserved.
+
+#include "Commandlets/CompileShadersTestBedCommandlet.h"
+#include "GlobalShader.h"
+#include "ShaderCompiler.h"
+
+DEFINE_LOG_CATEGORY_STATIC(LogCompileShadersTestBedCommandlet, Log, All);
+
+UCompileShadersTestBedCommandlet::UCompileShadersTestBedCommandlet(const FObjectInitializer& ObjectInitializer)
+	: Super(ObjectInitializer)
+{
+}
+
+int32 UCompileShadersTestBedCommandlet::Main(const FString& Params)
+{
+	TRACE_CPUPROFILER_EVENT_SCOPE(UCompileShadersTestBedCommandlet::Main);
+
+	TArray<FString> Tokens;
+	TArray<FString> Switches;
+	TMap<FString, FString> ParamVals;
+	UCommandlet::ParseCommandLine(*Params, Tokens, Switches, ParamVals);
+
+	// Display help
+	if (Switches.Contains("help"))
+	{
+		UE_LOG(LogCompileShadersTestBedCommandlet, Log, TEXT("CompileShadersTestBed"));
+		UE_LOG(LogCompileShadersTestBedCommandlet, Log, TEXT("This commandlet compiles global and default material shaders.  Used to profile and test shader compilation."));
+		return 0;
+	}
+
+	PRIVATE_GAllowCommandletRendering = true;
+
+	{
+		TRACE_CPUPROFILER_EVENT_SCOPE(Run1);
+
+		CompileGlobalShaderMap(true);
+
+		TArray<int32> ShaderMapIds;
+		ShaderMapIds.Add(GlobalShaderMapId);
+		GShaderCompilingManager->FinishCompilation(TEXT("Global"), ShaderMapIds);
+	}
+
+	return 0;
+}
