Fix an editor crash when attempting to create a second new macro library graph after having renamed the first one.

This is a regression that was introduced into 5.0 at 18928212 (UE-141499).

diff --git a/Engine/Source/Editor/UnrealEd/Private/Kismet2/BlueprintEditorUtils.cpp b/Engine/Source/Editor/UnrealEd/Private/Kismet2/BlueprintEditorUtils.cpp
--- a/Engine/Source/Editor/UnrealEd/Private/Kismet2/BlueprintEditorUtils.cpp
+++ b/Engine/Source/Editor/UnrealEd/Private/Kismet2/BlueprintEditorUtils.cpp
@@ -2356,13 +2356,23 @@ UEdGraph* FBlueprintEditorUtils::CreateNewGraph(UObject* ParentScope, const FNam
 	// Ensure this name isn't already being used for a graph
 	if (GraphName != NAME_None)
 	{
-		UEdGraph* ExistingGraph = FindObject<UEdGraph>(ParentScope, *(GraphName.ToString()));
-		ensureMsgf(!ExistingGraph, TEXT("Graph %s already exists: %s"), *GraphName.ToString(), *ExistingGraph->GetFullName());
-
-		// Rename the old graph out of the way; but we have already failed at this point
-		if (ExistingGraph)
+		if (UObject* ExistingObject = FindObject<UObject>(ParentScope, *(GraphName.ToString())))
 		{
-			ExistingGraph->Rename(nullptr, ExistingGraph->GetOuter(), REN_DoNotDirty | REN_ForceNoResetLoaders);
+			if (!ensureMsgf(!ExistingObject->IsA<UEdGraph>(), TEXT("Graph %s already exists: %s"), *GraphName.ToString(), *ExistingObject->GetFullName()))
+			{
+				// Rename the old graph out of the way; but we have already failed at this point
+				ExistingObject->Rename(nullptr, ExistingObject->GetOuter(), REN_DoNotDirty | REN_ForceNoResetLoaders);
+			}
+			else if (ExistingObject->IsA<UObjectRedirector>())
+			{
+				const UBlueprint* Blueprint = Cast<UBlueprint>(ParentScope);
+				if (Blueprint && Blueprint->BlueprintType == BPTYPE_MacroLibrary)
+				{
+					// When renaming a graph inside a macro library, we may have dropped a redirector after a previous
+					// rename (see RenameGraph). If we're now reusing it, move the redirector aside to free up the name.
+					ExistingObject->Rename(nullptr, GetTransientPackage(), REN_DontCreateRedirectors);
+				}
+			}
 		}
 
 		// Construct new graph with the supplied name
