Add keywords to output fields of DumpMaterialExpressionsCommandlet

diff --git a/Engine/Source/Editor/UnrealEd/Private/Commandlets/DumpMaterialExpressionsCommandlet.cpp b/Engine/Source/Editor/UnrealEd/Private/Commandlets/DumpMaterialExpressionsCommandlet.cpp
--- a/Engine/Source/Editor/UnrealEd/Private/Commandlets/DumpMaterialExpressionsCommandlet.cpp
+++ b/Engine/Source/Editor/UnrealEd/Private/Commandlets/DumpMaterialExpressionsCommandlet.cpp
@@ -48,6 +48,7 @@ int32 UDumpMaterialExpressionsCommandlet::Main(const FString& Params)
 	{
 		UMaterialExpression* MaterialExpression;
 		FString Name;
+		FString Keywords;
 		FString CreationName;
 		FString CreationDescription;
 		FString Caption;
@@ -61,6 +62,7 @@ int32 UDumpMaterialExpressionsCommandlet::Main(const FString& Params)
 	const FString NameField = TEXT("NAME");
 	const FString TypeField = TEXT("TYPE");
 	const FString ShowInCreateMenuField = TEXT("SHOW_IN_CREATE_MENU");
+	const FString KeywordsField = TEXT("KEYWORDS");
 	const FString CreationNameField = TEXT("CREATION_NAME");
 	const FString CreationDescriptionField = TEXT("CREATION_DESCRIPTION");
 	const FString CaptionField = TEXT("CAPTION");
@@ -70,6 +72,7 @@ int32 UDumpMaterialExpressionsCommandlet::Main(const FString& Params)
 	int32 MaxNameLength = NameField.Len();
 	int32 MaxTypeLength = TypeField.Len();
 	int32 MaxShowInCreateMenuLength = ShowInCreateMenuField.Len();
+	int32 MaxKeywordsLength = KeywordsField.Len();
 	int32 MaxCreationNameLength = CreationNameField.Len();
 	int32 MaxCreationDescriptionLength = CreationDescriptionField.Len();
 	int32 MaxCaptionLength = CaptionField.Len();
@@ -134,6 +137,7 @@ int32 UDumpMaterialExpressionsCommandlet::Main(const FString& Params)
 				FMaterialExpressionInfo ExpressionInfo;
 				ExpressionInfo.MaterialExpression = DefaultExpression;
 				ExpressionInfo.Name = Class->GetName().Mid(FCString::Strlen(TEXT("MaterialExpression")));
+				ExpressionInfo.Keywords = DefaultExpression->GetKeywords().ToString();
 				ExpressionInfo.CreationName = (CreationName.IsEmpty() ? (DisplayName.IsEmpty() ? ExpressionInfo.Name : DisplayName) : CreationName);
 				ExpressionInfo.CreationDescription = DefaultExpression->GetCreationDescription().ToString();
 				ExpressionInfo.Caption = Caption;
@@ -146,6 +150,7 @@ int32 UDumpMaterialExpressionsCommandlet::Main(const FString& Params)
 				MaxNameLength = FMath::Max(MaxNameLength, ExpressionInfo.Name.Len());
 				MaxTypeLength = FMath::Max(MaxTypeLength, ExpressionInfo.Type.Len());
 				MaxShowInCreateMenuLength = FMath::Max(MaxShowInCreateMenuLength, ExpressionInfo.ShowInCreateMenu.Len());
+				MaxKeywordsLength = FMath::Max(MaxKeywordsLength, ExpressionInfo.Keywords.Len());
 				MaxCreationNameLength = FMath::Max(MaxCreationNameLength, ExpressionInfo.CreationName.Len());
 				MaxCreationDescriptionLength = FMath::Max(MaxCreationDescriptionLength, ExpressionInfo.CreationDescription.Len());
 				MaxCaptionLength = FMath::Max(MaxCaptionLength, ExpressionInfo.Caption.Len());
@@ -159,6 +164,7 @@ int32 UDumpMaterialExpressionsCommandlet::Main(const FString& Params)
 	MaxNameLength += AdditionalPadding;
 	MaxTypeLength += AdditionalPadding;
 	MaxShowInCreateMenuLength += AdditionalPadding;
+	MaxKeywordsLength += AdditionalPadding;
 	MaxCreationNameLength += AdditionalPadding;
 	MaxCreationDescriptionLength += AdditionalPadding;
 	MaxCaptionLength += AdditionalPadding;
@@ -189,12 +195,13 @@ int32 UDumpMaterialExpressionsCommandlet::Main(const FString& Params)
 	const FString OutputFilePath = FPaths::Combine(*FPaths::ProjectSavedDir(), TEXT("MaterialEditor"), TEXT("MaterialExpressions.txt"));
 	FArchive* FileWriter = IFileManager::Get().CreateFileWriter(*OutputFilePath);
 
-	auto WriteLine = [FileWriter, GenerateSpacePadding, MaxNameLength, MaxTypeLength, MaxShowInCreateMenuLength, MaxCreationNameLength, MaxCreationDescriptionLength, MaxCaptionLength, MaxDescriptionLength]
-		(const FString& Name, const FString& Type, const FString& ShowInCreateMenu, const FString& CreationName, const FString& CreationDescription, const FString& Caption, const FString& Description, const FString& Tooltip)
+	auto WriteLine = [FileWriter, GenerateSpacePadding, MaxNameLength, MaxTypeLength, MaxShowInCreateMenuLength, MaxKeywordsLength, MaxCreationNameLength, MaxCreationDescriptionLength, MaxCaptionLength, MaxDescriptionLength]
+		(const FString& Name, const FString& Type, const FString& ShowInCreateMenu, const FString& Keywords, const FString& CreationName, const FString& CreationDescription, const FString& Caption, const FString& Description, const FString& Tooltip)
 	{
 		FString NamePadding = GenerateSpacePadding(MaxNameLength, Name.Len());
 		FString TypePadding = GenerateSpacePadding(MaxTypeLength, Type.Len());
 		FString ShowInCreateMenuPadding = GenerateSpacePadding(MaxShowInCreateMenuLength, ShowInCreateMenu.Len());
+		FString KeywordsPadding = GenerateSpacePadding(MaxKeywordsLength, Keywords.Len());
 		FString CreationNamePadding = GenerateSpacePadding(MaxCreationNameLength, CreationName.Len());
 		FString CreationDescriptionPadding = GenerateSpacePadding(MaxCreationDescriptionLength, CreationDescription.Len());
 		FString CaptionPadding = GenerateSpacePadding(MaxCaptionLength, Caption.Len());
@@ -203,6 +210,7 @@ int32 UDumpMaterialExpressionsCommandlet::Main(const FString& Params)
 		FString OutputLine = (Name + NamePadding);
 		OutputLine += (Type + TypePadding);
 		OutputLine += (ShowInCreateMenu + ShowInCreateMenuPadding);
+		OutputLine += (Keywords + KeywordsPadding);
 		OutputLine += (CreationName + CreationNamePadding);
 		OutputLine += (CreationDescription + CreationDescriptionPadding);
 		OutputLine += (Caption + CaptionPadding);
@@ -211,10 +219,10 @@ int32 UDumpMaterialExpressionsCommandlet::Main(const FString& Params)
 		FileWriter->Serialize(const_cast<ANSICHAR*>(StringCast<ANSICHAR>(*OutputLine).Get()), OutputLine.Len());
 	};
 
-	WriteLine(NameField, TypeField, ShowInCreateMenuField, CreationNameField, CreationDescriptionField, CaptionField, DescriptionField, TooltipField);
+	WriteLine(NameField, TypeField, ShowInCreateMenuField, KeywordsField, CreationNameField, CreationDescriptionField, CaptionField, DescriptionField, TooltipField);
 	for (FMaterialExpressionInfo& ExpressionInfo : MaterialExpressionInfos)
 	{
-		WriteLine(GetFormattedText(ExpressionInfo.Name), GetFormattedText(ExpressionInfo.Type), GetFormattedText(ExpressionInfo.ShowInCreateMenu), 
+		WriteLine(GetFormattedText(ExpressionInfo.Name), GetFormattedText(ExpressionInfo.Type), GetFormattedText(ExpressionInfo.ShowInCreateMenu), GetFormattedText(ExpressionInfo.Keywords),
 					GetFormattedText(ExpressionInfo.CreationName), GetFormattedText(ExpressionInfo.CreationDescription), GetFormattedText(ExpressionInfo.Caption), 
 					GetFormattedText(ExpressionInfo.Description), GetFormattedText(ExpressionInfo.Tooltip));
 	}
