Fixed WorldPartitionDataLayerToAssetCommandlet to remove DataLayerInstances but keep the DeprecatedDataLayerNameToDataLayerInstance to do the runtime remapping of deprecated BP calls.

diff --git a/Engine/Source/Editor/UnrealEd/Private/Commandlets/WorldPartitionDataLayerToAssetCommandlet.cpp b/Engine/Source/Editor/UnrealEd/Private/Commandlets/WorldPartitionDataLayerToAssetCommandlet.cpp
--- a/Engine/Source/Editor/UnrealEd/Private/Commandlets/WorldPartitionDataLayerToAssetCommandlet.cpp
+++ b/Engine/Source/Editor/UnrealEd/Private/Commandlets/WorldPartitionDataLayerToAssetCommandlet.cpp
@@ -737,7 +737,8 @@ bool UDataLayerToAssetCommandlet::CommitConversion(TStrongObjectPtr<UDataLayerTo
 	{
 		if (ConversionInfo->DataLayerToConvert != nullptr)
 		{
-			if (WorldDataLayers->RemoveDataLayer(ConversionInfo->DataLayerToConvert))
+			// Remove directly from DataLayerInstances as RemoveDataLayer method also cleans DeprecatedDataLayerNameToDataLayerInstance which is used for runtime conversion
+			if (WorldDataLayers->DataLayerInstances.Remove(ConversionInfo->DataLayerToConvert))
 			{
 				UE_LOG(LogDataLayerToAssetCommandlet, Log, TEXT("Deleted old data layer %s, it is now converted to Data Layer Asset %s and Data Layer Instance %s"),
 					*ConversionInfo->DataLayerToConvert->GetDataLayerShortName(), *ConversionInfo->DataLayerAsset->GetFullName(), *ConversionInfo->DataLayerInstance->GetDataLayerFName().ToString());
@@ -767,4 +768,4 @@ bool UDataLayerToAssetCommandlet::IsAssetInConversionFolder(const TObjectPtr<UDa
 	return FCString::Stristr(*DataLayerAsset->GetPackage()->GetName(), *ConversionFolder) != nullptr;
 }
 
-PRAGMA_ENABLE_DEPRECATION_WARNINGS
\ No newline at end of file
+PRAGMA_ENABLE_DEPRECATION_WARNINGS
