FCookAttachments: remove unused variable.
#rnx

diff --git a/Engine/Source/Editor/UnrealEd/Private/TargetDomain/TargetDomainUtils.h b/Engine/Source/Editor/UnrealEd/Private/TargetDomain/TargetDomainUtils.h
--- a/Engine/Source/Editor/UnrealEd/Private/TargetDomain/TargetDomainUtils.h
+++ b/Engine/Source/Editor/UnrealEd/Private/TargetDomain/TargetDomainUtils.h
@@ -36,7 +36,6 @@ FCbObject BuildDefinitionListToObject(TConstArrayView<UE::DerivedData::FBuildDef
 
 struct FCookAttachments
 {
-	const ITargetPlatform* TargetPlatform;
 	TArray<FName> BuildDependencies;
 	TArray<FName> RuntimeOnlyDependencies;
 	TArray<UE::DerivedData::FBuildDefinition> BuildDefinitionList;
