WorldPartitionResaveActorsBuilder will now always include actors that needs to be resaved when running in ResaveDirtyActorDescsOnly mode.

#rnx

diff --git a/Engine/Source/Editor/UnrealEd/Private/WorldPartition/WorldPartitionResaveActorsBuilder.cpp b/Engine/Source/Editor/UnrealEd/Private/WorldPartition/WorldPartitionResaveActorsBuilder.cpp
--- a/Engine/Source/Editor/UnrealEd/Private/WorldPartition/WorldPartitionResaveActorsBuilder.cpp
+++ b/Engine/Source/Editor/UnrealEd/Private/WorldPartition/WorldPartitionResaveActorsBuilder.cpp
@@ -197,13 +197,13 @@ bool UWorldPartitionResaveActorsBuilder::RunInternal(UWorld* World, const FCellI
 						// Gather dependant objects that also needs to be moved along with the actor
 						TArray<UObject*> DependantObjects;
 						ForEachObjectWithPackage(Package, [&DependantObjects](UObject* Object)
+						{
+							if (!Cast<UMetaData>(Object))
 							{
-								if (!Cast<UMetaData>(Object))
-								{
-									DependantObjects.Add(Object);
-								}
-								return true;
-							}, false);
+								DependantObjects.Add(Object);
+							}
+							return true;
+						}, false);
 
 						// Move dependant objects into the new world package temporarily
 						for (UObject* DependantObject : DependantObjects)
@@ -277,11 +277,14 @@ bool UWorldPartitionResaveActorsBuilder::RunInternal(UWorld* World, const FCellI
 
 				if (bResaveDirtyActorDescsOnly)
 				{
-					TUniquePtr<FWorldPartitionActorDesc> NewActorDesc = Actor->CreateActorDesc();
-
-					if (ActorDesc->Equals(NewActorDesc.Get()))
+					if (!ActorDesc->IsResaveNeeded())
 					{
-						return true;
+						TUniquePtr<FWorldPartitionActorDesc> NewActorDesc = Actor->CreateActorDesc();
+
+						if (ActorDesc->Equals(NewActorDesc.Get()))
+						{
+							return true;
+						}
 					}
 
 					UE_LOG(LogWorldPartitionResaveActorsBuilder, Log, TEXT("Package %s needs to be resaved."), *Package->GetName());
