Adding compilation of DefaultMaterials to the ShadersTestBedCommandlet.

diff --git a/Engine/Source/Editor/UnrealEd/Private/Commandlets/CompileShadersTestBedCommandlet.cpp b/Engine/Source/Editor/UnrealEd/Private/Commandlets/CompileShadersTestBedCommandlet.cpp
--- a/Engine/Source/Editor/UnrealEd/Private/Commandlets/CompileShadersTestBedCommandlet.cpp
+++ b/Engine/Source/Editor/UnrealEd/Private/Commandlets/CompileShadersTestBedCommandlet.cpp
@@ -2,6 +2,7 @@
 
 #include "Commandlets/CompileShadersTestBedCommandlet.h"
 #include "GlobalShader.h"
+#include "Materials/Material.h"
 #include "ShaderCompiler.h"
 
 DEFINE_LOG_CATEGORY_STATIC(LogCompileShadersTestBedCommandlet, Log, All);
@@ -31,7 +32,17 @@ int32 UCompileShadersTestBedCommandlet::Main(const FString& Params)
 	PRIVATE_GAllowCommandletRendering = true;
 
 	{
-		TRACE_CPUPROFILER_EVENT_SCOPE(Run1);
+		TRACE_CPUPROFILER_EVENT_SCOPE(DefaultMaterials);
+
+		for (int32 Domain = 0; Domain < MD_MAX; ++Domain)
+		{
+			UMaterial::GetDefaultMaterial(static_cast<EMaterialDomain>(Domain))->CacheShaders(EMaterialShaderPrecompileMode::Background);
+		}
+		GShaderCompilingManager->FinishAllCompilation();
+	}
+
+	{
+		TRACE_CPUPROFILER_EVENT_SCOPE(GlobalShaders);
 
 		CompileGlobalShaderMap(true);
 
