fix HighResShot not clipping rectangle to image bounds, allowed read of memory out of bounds
also make non-PIE Editor screenshot use requested image format extension
PIE/in-game screenshot was previously fixed

diff --git a/Engine/Source/Editor/UnrealEd/Private/EditorViewportClient.cpp b/Engine/Source/Editor/UnrealEd/Private/EditorViewportClient.cpp
--- a/Engine/Source/Editor/UnrealEd/Private/EditorViewportClient.cpp
+++ b/Engine/Source/Editor/UnrealEd/Private/EditorViewportClient.cpp
@@ -59,8 +59,8 @@
 #include "GPUSkinCacheVisualizationData.h"
 #include "UnrealWidget.h"
 #include "EdModeInteractiveToolsContext.h"
-
 #include "CustomEditorStaticScreenPercentage.h"
+#include "IImageWrapperModule.h"
 
 ICustomEditorStaticScreenPercentage* GCustomEditorStaticScreenPercentage = nullptr;
 
@@ -5768,21 +5768,31 @@ bool FEditorViewportClient::ProcessScreenShots(FViewport* InViewport)
 			// Clip the bitmap to just the capture region if valid
 			if (!SourceRect.IsEmpty())
 			{
-				FColor* const Data = Bitmap.GetData();
 				const int32 OldWidth = BitmapSize.X;
 				const int32 OldHeight = BitmapSize.Y;
-				const int32 NewWidth = SourceRect.Width();
-				const int32 NewHeight = SourceRect.Height();
-				const int32 CaptureTopRow = SourceRect.Min.Y;
-				const int32 CaptureLeftColumn = SourceRect.Min.X;
 
-				for (int32 Row = 0; Row < NewHeight; Row++)
+				//clamp in bounds:
+				int CaptureMinX = FMath::Clamp(SourceRect.Min.X,0,OldWidth);
+				int CaptureMinY = FMath::Clamp(SourceRect.Min.Y,0,OldHeight);
+				
+				int CaptureMaxX = FMath::Clamp(SourceRect.Max.X,0,OldWidth);
+				int CaptureMaxY = FMath::Clamp(SourceRect.Max.Y,0,OldHeight);
+				
+				int32 NewWidth  = CaptureMaxX - CaptureMinX;
+				int32 NewHeight = CaptureMaxY - CaptureMinY;
+ 
+				if ( NewWidth > 0 && NewHeight > 0 )
 				{
-					FMemory::Memmove(Data + Row * NewWidth, Data + (Row + CaptureTopRow) * OldWidth + CaptureLeftColumn, NewWidth * sizeof(*Data));
-				}
+					FColor* const Data = Bitmap.GetData();
+
+					for (int32 Row = 0; Row < NewHeight; Row++)
+					{
+						FMemory::Memmove(Data + Row * NewWidth, Data + (Row + CaptureMinY) * OldWidth + CaptureMinX, NewWidth * sizeof(*Data));
+					}
 
-				Bitmap.RemoveAt(NewWidth * NewHeight, OldWidth * OldHeight - NewWidth * NewHeight, false);
-				BitmapSize = FIntPoint(NewWidth, NewHeight);
+					Bitmap.RemoveAt(NewWidth * NewHeight, OldWidth * OldHeight - NewWidth * NewHeight, false);
+					BitmapSize = FIntPoint(NewWidth, NewHeight);
+				}
 			}
 
 			if (GIsAutomationTesting)
@@ -5810,10 +5820,27 @@ bool FEditorViewportClient::ProcessScreenShots(FViewport* InViewport)
 			HighResScreenshotConfig.PopulateImageTaskParams(*ImageTask);
 			ImageTask->Filename = FScreenshotRequest::GetFilename();
 
-			// Save the bitmap to disc
+			// Filename can be PNG but EImageFormat can be EXR
+
+			// PopulateImageTaskParams sets Format, ignores extension on file name
+			// ignore EXR from PopulateImageTaskParams, we are always 8-bit here
+			//if ( ImageTask->Format != EImageFormat::EXR )
+			{
+				// if not high dynamic range, get format from filename :
+				IImageWrapperModule& ImageWrapperModule = FModuleManager::LoadModuleChecked<IImageWrapperModule>(FName("ImageWrapper"));
+				EImageFormat ImageFormat = ImageWrapperModule.GetImageFormatFromExtension(*ImageTask->Filename);
+				if ( ImageFormat != EImageFormat::Invalid )
+				{
+					ImageTask->Format = ImageFormat;
+				}
+			}
+
+			// Save the bitmap to disk
 			TFuture<bool> CompletionFuture = HighResScreenshotConfig.ImageWriteQueue->Enqueue(MoveTemp(ImageTask));
 			if (CompletionFuture.IsValid())
 			{
+				// @@!! this queues it then immediately waits? what's the point of ImageWriteQueue then?
+				// just use FImageUtils::Save
 				bIsScreenshotSaved = CompletionFuture.Get();
 			}
 		}
diff --git a/Engine/Source/Editor/UnrealEd/Private/UnrealEdSrv.cpp b/Engine/Source/Editor/UnrealEd/Private/UnrealEdSrv.cpp
--- a/Engine/Source/Editor/UnrealEd/Private/UnrealEdSrv.cpp
+++ b/Engine/Source/Editor/UnrealEd/Private/UnrealEdSrv.cpp
@@ -1189,6 +1189,8 @@ bool UUnrealEdEngine::Exec( UWorld* InWorld, const TCHAR* Stream, FOutputDevice&
 	}
 	else if( FParse::Command(&Str, TEXT("HighResShot") ) )
 	{
+		// this is HighResShot from the Editor NOT in PIE
+		// Editor PIE HighResShot is in GameViewportClient
 		if (GetHighResScreenshotConfig().ParseConsoleCommand(Str, Ar))
 		{
 			TakeHighResScreenShots();
diff --git a/Engine/Source/Runtime/Engine/Private/GameViewportClient.cpp b/Engine/Source/Runtime/Engine/Private/GameViewportClient.cpp
--- a/Engine/Source/Runtime/Engine/Private/GameViewportClient.cpp
+++ b/Engine/Source/Runtime/Engine/Private/GameViewportClient.cpp
@@ -1964,21 +1964,31 @@ bool UGameViewportClient::ProcessScreenShots(FViewport* InViewport)
 				// Clip the bitmap to just the capture region if valid
 				if (!SourceRect.IsEmpty())
 				{
-					FColor* const Data = Bitmap.GetData();
 					const int32 OldWidth = Size.X;
 					const int32 OldHeight = Size.Y;
-					const int32 NewWidth = SourceRect.Width();
-					const int32 NewHeight = SourceRect.Height();
-					const int32 CaptureTopRow = SourceRect.Min.Y;
-					const int32 CaptureLeftColumn = SourceRect.Min.X;
-
-					for (int32 Row = 0; Row < NewHeight; Row++)
+					
+					//clamp in bounds:
+					int CaptureMinX = FMath::Clamp(SourceRect.Min.X,0,OldWidth);
+					int CaptureMinY = FMath::Clamp(SourceRect.Min.Y,0,OldHeight);
+				
+					int CaptureMaxX = FMath::Clamp(SourceRect.Max.X,0,OldWidth);
+					int CaptureMaxY = FMath::Clamp(SourceRect.Max.Y,0,OldHeight);
+				
+					int32 NewWidth  = CaptureMaxX - CaptureMinX;
+					int32 NewHeight = CaptureMaxY - CaptureMinY;
+ 
+					if ( NewWidth > 0 && NewHeight > 0 )
 					{
-						FMemory::Memmove(Data + Row * NewWidth, Data + (Row + CaptureTopRow) * OldWidth + CaptureLeftColumn, NewWidth * sizeof(*Data));
-					}
+						FColor* const Data = Bitmap.GetData();
 
-					Bitmap.RemoveAt(NewWidth * NewHeight, OldWidth * OldHeight - NewWidth * NewHeight, false);
-					Size = FIntVector(NewWidth, NewHeight, 0);
+						for (int32 Row = 0; Row < NewHeight; Row++)
+						{
+							FMemory::Memmove(Data + Row * NewWidth, Data + (Row + CaptureMinY) * OldWidth + CaptureMinX, NewWidth * sizeof(*Data));
+						}
+
+						Bitmap.RemoveAt(NewWidth * NewHeight, OldWidth * OldHeight - NewWidth * NewHeight, false);
+						Size = FIntVector(NewWidth, NewHeight, 0);
+					}
 				}
 
 				if ( FPaths::GetExtension(ScreenShotName).IsEmpty() )
