new DDSFile DDS import and export
supports 2d,cube,cube array,volume
supports many dxgi pixel formats

diff --git a/Engine/Source/Editor/UnrealEd/Classes/Exporters/TextureExporterBMP.h b/Engine/Source/Editor/UnrealEd/Classes/Exporters/TextureExporterBMP.h
--- a/Engine/Source/Editor/UnrealEd/Classes/Exporters/TextureExporterBMP.h
+++ b/Engine/Source/Editor/UnrealEd/Classes/Exporters/TextureExporterBMP.h
@@ -17,7 +17,7 @@ class UTextureExporterBMP : public UTextureExporterGeneric
 {
 	GENERATED_UCLASS_BODY()
 	
-	virtual bool SupportsTexture(UTexture2D* Texture) const override;
+	virtual bool SupportsTexture(UTexture* Texture) const override;
 };
 
 UCLASS()
diff --git a/Engine/Source/Editor/UnrealEd/Classes/Exporters/TextureExporterDDS.h b/Engine/Source/Editor/UnrealEd/Classes/Exporters/TextureExporterDDS.h
new file mode 100644
--- /dev/null
+++ b/Engine/Source/Editor/UnrealEd/Classes/Exporters/TextureExporterDDS.h
@@ -0,0 +1,35 @@
+// Copyright Epic Games, Inc. All Rights Reserved.
+
+//=============================================================================
+// TextureExporterDDS
+//=============================================================================
+
+#pragma once
+
+#include "CoreMinimal.h"
+#include "UObject/ObjectMacros.h"
+#include "Exporters/Exporter.h"
+#include "Exporters/TextureExporterGeneric.h"
+#include "Engine/Texture.h"
+#include "TextureExporterDDS.generated.h"
+
+UCLASS()
+class UTextureExporterDDS : public UTextureExporterGeneric
+{
+	GENERATED_UCLASS_BODY()
+
+	virtual bool SupportsObject(UObject* Object) const override;
+
+	// should not get called :
+	virtual bool SupportsTexture(UTexture* Texture) const override { check(0); return true; };
+
+	virtual bool ExportBinary(UObject* Object, const TCHAR* Type, FArchive& Ar, FFeedbackContext* Warn, int32 FileIndex = 0, uint32 PortFlags = 0) override;
+	
+};
+
+
+UCLASS()
+class UVirtualTextureBuilderExporterDDS : public UTextureExporterDDS
+{
+	GENERATED_UCLASS_BODY()
+};
\ No newline at end of file
diff --git a/Engine/Source/Editor/UnrealEd/Classes/Exporters/TextureExporterEXR.h b/Engine/Source/Editor/UnrealEd/Classes/Exporters/TextureExporterEXR.h
--- a/Engine/Source/Editor/UnrealEd/Classes/Exporters/TextureExporterEXR.h
+++ b/Engine/Source/Editor/UnrealEd/Classes/Exporters/TextureExporterEXR.h
@@ -17,7 +17,7 @@ class UTextureExporterEXR : public UTextureExporterGeneric
 {
 	GENERATED_UCLASS_BODY()
 	
-	virtual bool SupportsTexture(UTexture2D* Texture) const override;
+	virtual bool SupportsTexture(UTexture* Texture) const override;
 };
 
 UCLASS()
diff --git a/Engine/Source/Editor/UnrealEd/Classes/Exporters/TextureExporterGeneric.h b/Engine/Source/Editor/UnrealEd/Classes/Exporters/TextureExporterGeneric.h
--- a/Engine/Source/Editor/UnrealEd/Classes/Exporters/TextureExporterGeneric.h
+++ b/Engine/Source/Editor/UnrealEd/Classes/Exporters/TextureExporterGeneric.h
@@ -9,7 +9,7 @@
 #include "CoreMinimal.h"
 #include "UObject/ObjectMacros.h"
 #include "Exporters/Exporter.h"
-#include "Engine/Texture2D.h"
+#include "Engine/Texture.h"
 #include "TextureExporterGeneric.generated.h"
 
 UCLASS()
@@ -23,9 +23,9 @@ class UTextureExporterGeneric : public UExporter
 	virtual bool ExportBinary(UObject* Object, const TCHAR* Type, FArchive& Ar, FFeedbackContext* Warn, int32 FileIndex = 0, uint32 PortFlags = 0) override;
 	
 	// derived class exporter must implement :
-	virtual bool SupportsTexture(UTexture2D* Texture) const { return false; };
+	virtual bool SupportsTexture(UTexture* Texture) const { return false; };
 
 protected:
-	class UTexture2D* GetExportTexture(UObject* Object) const;
+	class UTexture* GetExportTexture(UObject* Object) const;
 };
 
diff --git a/Engine/Source/Editor/UnrealEd/Classes/Exporters/TextureExporterHDR.h b/Engine/Source/Editor/UnrealEd/Classes/Exporters/TextureExporterHDR.h
--- a/Engine/Source/Editor/UnrealEd/Classes/Exporters/TextureExporterHDR.h
+++ b/Engine/Source/Editor/UnrealEd/Classes/Exporters/TextureExporterHDR.h
@@ -17,7 +17,7 @@ class UTextureExporterHDR : public UTextureExporterGeneric
 {
 	GENERATED_UCLASS_BODY()
 	
-	virtual bool SupportsTexture(UTexture2D* Texture) const override;
+	virtual bool SupportsTexture(UTexture* Texture) const override;
 };
 
 UCLASS()
diff --git a/Engine/Source/Editor/UnrealEd/Classes/Exporters/TextureExporterPNG.h b/Engine/Source/Editor/UnrealEd/Classes/Exporters/TextureExporterPNG.h
--- a/Engine/Source/Editor/UnrealEd/Classes/Exporters/TextureExporterPNG.h
+++ b/Engine/Source/Editor/UnrealEd/Classes/Exporters/TextureExporterPNG.h
@@ -17,7 +17,7 @@ class UTextureExporterPNG : public UTextureExporterGeneric
 {
 	GENERATED_UCLASS_BODY()
 	
-	virtual bool SupportsTexture(UTexture2D* Texture) const override;
+	virtual bool SupportsTexture(UTexture* Texture) const override;
 };
 
 UCLASS()
diff --git a/Engine/Source/Editor/UnrealEd/Classes/Factories/TextureFactory.h b/Engine/Source/Editor/UnrealEd/Classes/Factories/TextureFactory.h
--- a/Engine/Source/Editor/UnrealEd/Classes/Factories/TextureFactory.h
+++ b/Engine/Source/Editor/UnrealEd/Classes/Factories/TextureFactory.h
@@ -242,6 +242,8 @@ private:
 
 	/** used by CreateTexture() */
 	UTexture* ImportTexture(UClass* Class, UObject* InParent, FName Name, EObjectFlags Flags, const TCHAR* Type, const uint8*& Buffer, const uint8* BufferEnd, FFeedbackContext* Warn);
+	
+	UTexture * ImportDDS(const uint8* Buffer,int64 Length,UObject* InParent,FName Name, EObjectFlags Flags,EImageImportFlags ImportFlags,FFeedbackContext* Warn);
 
 	UTexture* ImportTextureUDIM(UClass* Class, UObject* InParent, FName Name, EObjectFlags Flags, const TCHAR* Type, const TMap<int32, FString>& UDIMIndexToFile, FFeedbackContext* Warn);
 
diff --git a/Engine/Source/Editor/UnrealEd/Private/Factories/EditorFactories.cpp b/Engine/Source/Editor/UnrealEd/Private/Factories/EditorFactories.cpp
--- a/Engine/Source/Editor/UnrealEd/Private/Factories/EditorFactories.cpp
+++ b/Engine/Source/Editor/UnrealEd/Private/Factories/EditorFactories.cpp
@@ -129,6 +129,8 @@
 #include "Engine/TextureLightProfile.h"
 #include "SoundCueGraph/SoundCueGraphNode.h"
 #include "Exporters/TextureCubeExporterHDR.h"
+#include "Exporters/TextureExporterGeneric.h"
+#include "Exporters/TextureExporterDDS.h"
 #include "Exporters/TextureExporterBMP.h"
 #include "Exporters/TextureExporterHDR.h"
 #include "Exporters/RenderTargetExporterHDR.h"
@@ -165,9 +167,9 @@
 #include "Sound/SoundCue.h"
 #include "Sound/SoundMix.h"
 #include "Engine/TextureCube.h"
-#include "Engine/VolumeTexture.h"
 #include "Engine/Texture2DArray.h"
 #include "Engine/VolumeTexture.h"
+#include "Engine/TextureCubeArray.h"
 #include "Engine/TextureRenderTarget.h"
 #include "Engine/TextureRenderTarget2D.h"
 #include "Engine/TextureRenderTarget2DArray.h"
@@ -195,8 +197,7 @@
 #include "AssetToolsModule.h"
 #include "IAssetTools.h"
 
-#include "DDSLoader.h"
-#include "Formats/HdrImageWrapper.h"
+#include "DDSFile.h"
 #include "Factories/TIFFLoader.h"
 #include "IESConverter.h"
 #include "IImageWrapper.h"
@@ -225,17 +226,7 @@
 #include "VT/VirtualTexture.h"
 #include "VT/VirtualTextureBuilder.h"
 
-#if PLATFORM_WINDOWS
-	// Needed for DDS support.
-	#include "Windows/WindowsHWrapper.h"
-	#include "Windows/AllowWindowsPlatformTypes.h"
-		#include <ddraw.h>
-	#include "Windows/HideWindowsPlatformTypes.h"
-#endif
-
-#if WITH_EDITOR
-	#include "CubemapUnwrapUtils.h"
-#endif
+#include "CubemapUnwrapUtils.h"
 
 #include "Components/BrushComponent.h"
 #include "EngineUtils.h"
@@ -1236,6 +1227,8 @@ UObject* ULevelFactory::FactoryCreateText
 		else // This actor is new, but rejected to import its properties, so just delete...
 		{
 			Actor->Destroy();
+			// @@!! Actor pointer continues to be used below?
+			// continue;
 		}
 
 		// If this is a newly imported brush, validate it.  If it's a newly imported dynamic brush, rebuild it first.
@@ -2715,7 +2708,7 @@ UTextureFactory::UTextureFactory(const FObjectInitializer& ObjectInitializer)
 	Formats.Add( TEXT( "bmp;Texture" ) );
 	Formats.Add( TEXT( "pcx;Texture" ) );
 	Formats.Add( TEXT( "tga;Texture" ) );
-	Formats.Add( TEXT( "float;Texture" ) ); // @@!! what's this ?
+	Formats.Add( TEXT( "float;Texture" ) ); // @@!! what's this ? seems to be listed but have no loader?
 	Formats.Add( TEXT( "psd;Texture" ) );
 	Formats.Add( TEXT( "dds;Texture (Cubemap or 2D)" ) );
 	Formats.Add( TEXT( "hdr;Texture (HDR) (LongLat unwrap or 2D)" ) );
@@ -3361,60 +3354,10 @@ bool UTextureFactory::ImportImage(const uint8* Buffer, int64 Length, FFeedbackCo
 		return true;
 	}
 	
-	//
-	// DDS Texture
-	// @@!! move this out with the other DDS cube/aray loader
-	FDDSLoadHelper  DDSLoadHelper(Buffer, Length);
-	if (DDSLoadHelper.IsValid2DTexture())
-	{
-		// DDS 2d texture
-		if (!IsImportResolutionValid(DDSLoadHelper.DDSHeader->dwWidth, DDSLoadHelper.DDSHeader->dwHeight, bAllowNonPowerOfTwo, Warn))
-		{
-			Warn->Logf(ELogVerbosity::Error, TEXT("DDS has invalid dimensions."));
-			return false;
-		}
-
-		ETextureSourceFormat SourceFormat = DDSLoadHelper.ComputeSourceFormat();
-
-		// Invalid DDS format
-		if (SourceFormat == TSF_Invalid)
-		{
-			Warn->Logf(ELogVerbosity::Error, TEXT("DDS uses an unsupported format."));
-			return false;
-		}
-
-		uint32 MipMapCount = DDSLoadHelper.ComputeMipMapCount();
-		if (SourceFormat != TSF_Invalid && MipMapCount > 0)
-		{
-			OutImage.Init2DWithMips(
-				DDSLoadHelper.DDSHeader->dwWidth,
-				DDSLoadHelper.DDSHeader->dwHeight,
-				MipMapCount,
-				SourceFormat,
-				DDSLoadHelper.GetDDSDataPointer()
-			);
-
-			if (MipMapCount > 1)
-			{
-				// if the source has mips we keep the mips by default, unless the user changes that
-				MipGenSettings = TMGS_LeaveExistingMips;
-			}
-
-			if (FTextureSource::IsHDR(SourceFormat))
-			{
-				// the loader can suggest a compression setting
-				OutImage.CompressionSettings = TC_HDR;
-				OutImage.SRGB = false;
-			}
-
-			return true;
-		}
-	}
-
-
 	//
 	// Legacy TIFF import (for the platforms that doesn't have libtiff)
-	//
+	// @@!!  this is not used?; all the Editor platforms have libtiff?  see libtiff.build.cs
+	//		remove me?
 	FTiffLoadHelper TiffLoaderHelper;
 	if (TiffLoaderHelper.IsValid())
 	{
@@ -3497,6 +3440,7 @@ UTexture* UTextureFactory::ImportTextureUDIM(UClass* Class, UObject* InParent, F
 
 	if (SourceImages.Num() < 2)
 	{
+		// @@!! seems like it should continue if Count == 1
 		return nullptr;
 	}
 
@@ -3519,8 +3463,9 @@ UTexture* UTextureFactory::ImportTextureUDIM(UClass* Class, UObject* InParent, F
 	{
 		Texture->SRGB = false;
 	}
-	else if (ColorSpaceMode == ETextureSourceColorSpace::SRGB)
+	else 
 	{
+		check(ColorSpaceMode == ETextureSourceColorSpace::SRGB);
 		Texture->SRGB = true;
 	}
 
@@ -3532,171 +3477,243 @@ UTexture* UTextureFactory::ImportTextureUDIM(UClass* Class, UObject* InParent, F
 	return Texture;
 }
 
-UTexture* UTextureFactory::ImportTexture(UClass* Class, UObject* InParent, FName Name, EObjectFlags Flags, const TCHAR* Type, const uint8*& Buffer, const uint8* BufferEnd, FFeedbackContext* Warn)
+UTexture * UTextureFactory::ImportDDS(const uint8* Buffer,int64 Length,UObject* InParent,FName Name, EObjectFlags Flags,EImageImportFlags ImportFlags,FFeedbackContext* Warn)
 {
-	TRACE_CPUPROFILER_EVENT_SCOPE(ImportTexture);
+	UE::DDS::EDDSError Error;
+	UE::DDS::FDDSFile* DDS = UE::DDS::FDDSFile::CreateFromDDSInMemory(Buffer,Length,&Error);
+	if ( DDS == nullptr )
+	{
+		// NotADds is okay/expected , we try this on all image buffers
+		check( Error != UE::DDS::EDDSError::OK ); 
+		if ( Error != UE::DDS::EDDSError::NotADds )
+		{
+			UE_LOG(LogEditorFactories,Warning,TEXT("Failed to load DDS (Error=%d)"),(int)Error);			
+		}
+		return nullptr;
+	}
 
-	bool bAllowNonPowerOfTwo = false;
-	GConfig->GetBool( TEXT("TextureImporter"), TEXT("AllowNonPowerOfTwoTextures"), bAllowNonPowerOfTwo, GEditorIni );
+	// delete DDS at scope exit :
+	TUniquePtr<UE::DDS::FDDSFile> DDSPtr(DDS);
 
-	// Validate it.
-	const int64 Length = BufferEnd - Buffer;
+	if (!IsImportResolutionValid(DDS->Width, DDS->Height, !!(ImportFlags & EImageImportFlags::AllowNonPowerOfTwo), Warn))
+	{
+		Warn->Logf(ELogVerbosity::Error, TEXT("DDS resolution not supported"));
+		return nullptr;
+	}
 
-	// We want to allow certain texture types to potentially return the compressed data rather than raw uncompressed
-	EImageImportFlags ImportFlags = EImageImportFlags::AllowReturnOfCompressedData;
-	if (bAllowNonPowerOfTwo)
+	// change X8 formats to A8 :	
+	DDS->ConvertRGBXtoRGBA();
+	// change RGBA8 to BGRA8 before DXGIFormatGetClosestRawFormat :
+	DDS->ConvertChannelOrder(UE::DDS::EChannelOrder::BGRA);
+
+	// map format to RawFormat and ETextureSourceFormat
+	ERawImageFormat::Type RawFormat = UE::DDS::DXGIFormatGetClosestRawFormat(DDS->DXGIFormat);
+	if ( RawFormat == ERawImageFormat::Invalid )
 	{
-		ImportFlags |= EImageImportFlags::AllowNonPowerOfTwo;
+		Warn->Logf(ELogVerbosity::Error, TEXT("DDS DXGIFormat not supported : %d : %s"), (int)DDS->DXGIFormat, UE::DDS::DXGIFormatGetName(DDS->DXGIFormat) );
+		return nullptr;
 	}
 
-	// Do unusual loaders first before Generic
+	UTexture * Texture;
+	int NumSlices;
+	if ( DDS->IsValidTexture2D() )
+	{
+		Texture = CreateTexture2D(InParent, Name, Flags);
+		NumSlices = 1;
 
-	//
-	// DDS Cubemap
-	//
-	FDDSLoadHelper DDSLoadHelper(Buffer, Length);
-	if(DDSLoadHelper.IsValidCubemapTexture())
+		UE_LOG(LogEditorFactories,Display,TEXT("DDS imported as 2d Texture"));
+	}
+	else if ( DDS->IsValidTextureCube() )
 	{
-		if(!IsImportResolutionValid(DDSLoadHelper.DDSHeader->dwWidth, DDSLoadHelper.DDSHeader->dwHeight, bAllowNonPowerOfTwo, Warn))
+		if ( DDS->ArraySize > 6 )
 		{
-			Warn->Logf(ELogVerbosity::Error, TEXT("DDS uses an unsupported format"));
-			return nullptr;
+			UObject* NewObject = CreateOrOverwriteAsset(UTextureCubeArray::StaticClass(), InParent, Name, Flags);
+			Texture = NewObject ? CastChecked<UTextureCubeArray>(NewObject) : nullptr;
+			
+			UE_LOG(LogEditorFactories,Display,TEXT("DDS imported as Cube Texture Array"));
 		}
-
-		int32 NumMips = DDSLoadHelper.ComputeMipMapCount();
-		ETextureSourceFormat Format = DDSLoadHelper.ComputeSourceFormat();
-		if (Format == TSF_Invalid)
+		else
 		{
-			Warn->Logf(ELogVerbosity::Error, TEXT("DDS file contains data in an unsupported format."));
-			return nullptr;
+			Texture = CreateTextureCube(InParent, Name, Flags);
+			
+			UE_LOG(LogEditorFactories,Display,TEXT("DDS imported as Cube Texture"));
 		}
 
-		if (NumMips > MAX_TEXTURE_MIP_COUNT)
-		{
-			Warn->Logf(ELogVerbosity::Error, TEXT("DDS file contains an unsupported number of mipmap levels."));
-			return nullptr;
-		}
+		NumSlices = DDS->ArraySize;
+	}
+	else if ( DDS->IsValidTextureArray() )
+	{
+		Texture = CreateTexture2DArray(InParent, Name, Flags);
+		NumSlices = DDS->ArraySize;
+
+		UE_LOG(LogEditorFactories,Display,TEXT("DDS imported as Texture Array"));
+	}
+	else if ( DDS->IsValidTextureVolume() )
+	{
+		UObject* NewObject = CreateOrOverwriteAsset(UVolumeTexture::StaticClass(), InParent, Name, Flags);
+		Texture = NewObject ? CastChecked<UVolumeTexture>(NewObject) : nullptr;
+		NumSlices = DDS->Depth;
+
+		UE_LOG(LogEditorFactories,Display,TEXT("DDS imported as Volume"));
+	}
+	else
+	{
+		UE_LOG(LogEditorFactories,Error,TEXT("DDS is not a valid Unreal type"));
+		return nullptr;
+	}
+	if ( Texture == nullptr )
+	{
+		UE_LOG(LogEditorFactories,Error,TEXT("DDS could not create texture"));
+		return nullptr;
+	}
+
+	ETextureSourceFormat TSFormat = FImageCoreUtils::ConvertToTextureSourceFormat(RawFormat);
 
-		// create the cube texture
-		UTextureCube* TextureCube = CreateTextureCube( InParent, Name, Flags );
+	int MipCount = (int)DDS->MipCount;
+
+	if ( MipCount > MAX_TEXTURE_MIP_COUNT )
+	{
+		// unexpected, because IsImportResolutionValid has already limited rsolution
+		UE_LOG(LogEditorFactories,Warning,TEXT("DDS exceeds MAX_TEXTURE_MIP_COUNT"));
+		MipCount = MAX_TEXTURE_MIP_COUNT;
+	}
+
+	#if 1
+	if ( DDS->Dimension == 3 )
+	{
+		// Unreal doesn't handle volume mips correctly?
+		//  (slices doesn't mip down)
+		// force mip import off for now
+		// @@!! fix Unreal volume mips?
+		// see also GetMippedNumSlices
+		MipCount = 1;
+	}
+	#endif
+
+	Texture->Source.Init(DDS->Width,DDS->Height,NumSlices,MipCount,TSFormat);
+
+	for(int MipIndex=0;MipIndex<MipCount;MipIndex++)
+	{
+		FTextureSource::FMipLock MipLock(FTextureSource::ELockState::ReadWrite,&(Texture->Source),MipIndex);
 
-		if ( TextureCube )
+		if ( DDS->Dimension == 3 )
 		{
-			TextureCube->Source.Init(
-				DDSLoadHelper.GetSizeX(),
-				DDSLoadHelper.GetSizeY(),
-				DDSLoadHelper.GetSliceCount(),
-				NumMips,
-				Format
-				);
-			if(Format == TSF_RGBA16F || Format == TSF_RGBA32F)
-			{
-				TextureCube->CompressionSettings = TC_HDR;
-				TextureCube->SRGB = false;
-			}
+			check( DDS->Mips.Num() == DDS->MipCount );
 
-			uint8* DestMipData[MAX_TEXTURE_MIP_COUNT] = {0};
-			int32 MipSize[MAX_TEXTURE_MIP_COUNT] = {0};
-			for (int32 MipIndex = 0; MipIndex < NumMips; ++MipIndex)
+			#if 0
+			// hack for Unreal not mipping volumes correctly (NumSlices doesn't go down like it should)
+			//  (moot because we now force MipCount to 1 for volumes)
+			// @@!! if the fix is made in GetMippedNumSlices , this is unnecessary
+			if ( MipIndex > 0 && MipLock.Image.NumSlices > (int32)DDS->Mips[MipIndex].Depth )
 			{
-				DestMipData[MipIndex] = TextureCube->Source.LockMip(MipIndex);
-				MipSize[MipIndex] = TextureCube->Source.CalcMipSize(MipIndex) / 6;
+				// memset the miplock to zero : last slices won't be written :
+				memset(MipLock.Image.RawData,0,MipLock.Image.GetImageSizeBytes());
+				// only copy the portion the DDS has :
+				MipLock.Image.NumSlices = DDS->Mips[MipIndex].Depth;
 			}
+			#endif
 
-			for (int32 SliceIndex = 0; SliceIndex < 6; ++SliceIndex)
-			{
-				const uint8* SrcMipData = DDSLoadHelper.GetDDSDataPointer((ECubeFace)SliceIndex);
-				for (int32 MipIndex = 0; MipIndex < NumMips; ++MipIndex)
-				{
-					FMemory::Memcpy(DestMipData[MipIndex] + MipSize[MipIndex] * SliceIndex, SrcMipData, MipSize[MipIndex]);
-					SrcMipData += MipSize[MipIndex];
-				}
-			}
+			check( DDS->Mips[MipIndex].Depth == MipLock.Image.NumSlices );
 
-			for (int32 MipIndex = 0; MipIndex < NumMips; ++MipIndex)
+			if ( ! DDS->GetMipImage( MipLock.Image, MipIndex ) )
 			{
-				TextureCube->Source.UnlockMip(MipIndex);
+				// Texture is not returned; it will be deleted by GC
+				Texture->MarkAsGarbage();
+				UE_LOG(LogEditorFactories,Error,TEXT("DDS could not convert pixel format"));
+				return nullptr;
 			}
+		}
+		else
+		{
+			// DDS->Mips[] contains both mips and arrays
+			check( DDS->Mips.Num() == DDS->MipCount * NumSlices );
 
-			if (NumMips > 1)
+			for(int SliceIndex=0;SliceIndex<NumSlices;SliceIndex++)
 			{
-				TextureCube->MipGenSettings = TMGS_LeaveExistingMips;
+				FImageView DestSlice = MipLock.Image.GetSlice(SliceIndex);
+				
+				// DDS Mips[] array has whole mip chain of each slice, then next slice
+				// we have the opposite (all slices of top mip first, then next mip)
+				int DDSMipIndex = SliceIndex * DDS->MipCount + MipIndex;
+
+				if ( ! DDS->GetMipImage( DestSlice, DDSMipIndex ) )
+				{
+					// Texture is not returned; it will be deleted by GC
+					Texture->MarkAsGarbage();
+					UE_LOG(LogEditorFactories,Error,TEXT("DDS could not convert pixel format"));
+					return nullptr;
+				}
 			}
 		}
-		
-		UE_LOG(LogEditorFactories,Display,TEXT("DDS imported as TextureCube.") );
-
-		return TextureCube;
 	}
-	
-	// DDS Texture array.
-	if (DDSLoadHelper.IsValidArrayTexture())
+
+	bool bSRGB = false;
+
+	if ( ERawImageFormat::GetFormatNeedsGammaSpace(RawFormat) )
 	{
-		if (!IsImportResolutionValid(DDSLoadHelper.DDSHeader->dwWidth, DDSLoadHelper.DDSHeader->dwHeight, bAllowNonPowerOfTwo, Warn))
+		if ( DDS->CreateFlags & UE::DDS::FDDSFile::CREATE_FLAG_WAS_D3D9 )
 		{
-			Warn->Logf(ELogVerbosity::Error, TEXT("DDS uses an unsupported format"));
-			return nullptr;
+			// no SRGB info in Dx9 format
+			// assume SRGB yes
+			bSRGB = true;
 		}
-
-		int32 NumMips = DDSLoadHelper.ComputeMipMapCount();
-		ETextureSourceFormat Format = DDSLoadHelper.ComputeSourceFormat();
-		if (Format == TSF_Invalid)
+		else if ( UE::DDS::DXGIFormatHasLinearAndSRGBForm(DDS->DXGIFormat) )
 		{
-			Warn->Logf(ELogVerbosity::Error, TEXT("DDS file contains data in an unsupported format."));
-			return nullptr;
+			// Dx10 file with format that has linear/srgb pair
+			//	( _UNORM when _UNORM_SRGB)
+			bSRGB = UE::DDS::DXGIFormatIsSRGB(DDS->DXGIFormat);
+		}
+		else
+		{
+			// Dx10 format that doesn't have linear/srgb pairs
+			
+			// R8G8_UNORM and R8_UNORM have no _SRGB pair
+			// so no way to clearly indicate SRGB or Linear for them
+			// assume SRGB yes
+			bSRGB = true;
 		}
+	}
 
-		// Create the array texture
-		UTexture2DArray* TextureArray = CreateTexture2DArray(InParent, Name, Flags);
+	if ( ERawImageFormat::IsHDR(RawFormat) )
+	{
+		Texture->CompressionSettings = TC_HDR;
+		check( bSRGB == false );
+	}
 
-		if (TextureArray)
-		{
-			TextureArray->Source.Init(
-				DDSLoadHelper.GetSizeX(),
-				DDSLoadHelper.GetSizeY(),
-				DDSLoadHelper.GetSliceCount(),
-				NumMips,
-				Format
-			);
-			if (Format == TSF_RGBA16F || Format == TSF_RGBA32F)
-			{
-				TextureArray->CompressionSettings = TC_HDR;
-				TextureArray->SRGB = false;
-			}
+	Texture->SRGB = bSRGB;
+	
+	if ( MipCount > 1)
+	{
+		// if the source has mips we keep the mips by default, unless the user changes that
+		Texture->MipGenSettings = TMGS_LeaveExistingMips;
+	}
+	
+	return Texture;
+}
 
-			uint8* DestMipData[MAX_TEXTURE_MIP_COUNT] = { 0 };
-			int32 MipSize[MAX_TEXTURE_MIP_COUNT] = { 0 };
-			for (int32 MipIndex = 0; MipIndex < NumMips; ++MipIndex)
-			{
-				DestMipData[MipIndex] = TextureArray->Source.LockMip(MipIndex);
-				MipSize[MipIndex] = TextureArray->Source.CalcMipSize(MipIndex) / DDSLoadHelper.GetSliceCount();
-			}
+UTexture* UTextureFactory::ImportTexture(UClass* Class, UObject* InParent, FName Name, EObjectFlags Flags, const TCHAR* Type, const uint8*& Buffer, const uint8* BufferEnd, FFeedbackContext* Warn)
+{
+	TRACE_CPUPROFILER_EVENT_SCOPE(ImportTexture);
 
-			for (uint32 SliceIndex = 0; SliceIndex < DDSLoadHelper.GetSliceCount(); ++SliceIndex)
-			{
-				const uint8* SrcMipData = DDSLoadHelper.GetDDSDataPointer((ECubeFace)SliceIndex);
-				for (int32 MipIndex = 0; MipIndex < NumMips; ++MipIndex)
-				{
-					FMemory::Memcpy(DestMipData[MipIndex] + MipSize[MipIndex] * SliceIndex, SrcMipData, MipSize[MipIndex]);
-					SrcMipData += MipSize[MipIndex];
-				}
-			}
+	bool bAllowNonPowerOfTwo = false;
+	GConfig->GetBool( TEXT("TextureImporter"), TEXT("AllowNonPowerOfTwoTextures"), bAllowNonPowerOfTwo, GEditorIni );
 
-			for (int32 MipIndex = 0; MipIndex < NumMips; ++MipIndex)
-			{
-				TextureArray->Source.UnlockMip(MipIndex);
-			}
+	// Validate it.
+	const int64 Length = BufferEnd - Buffer;
 
-			if (NumMips > 1)
-			{
-				// If the source has mips we keep the mips by default, unless the user changes that
-				MipGenSettings = TMGS_LeaveExistingMips;
-			}
-		}
-		
-		UE_LOG(LogEditorFactories,Display,TEXT("DDS imported as TextureArray.") );
+	// We want to allow certain texture types to potentially return the compressed data rather than raw uncompressed
+	EImageImportFlags ImportFlags = EImageImportFlags::AllowReturnOfCompressedData;
+	if (bAllowNonPowerOfTwo)
+	{
+		ImportFlags |= EImageImportFlags::AllowNonPowerOfTwo;
+	}
 
-		return TextureArray;
+	// Do unusual loaders first before Generic
+
+	UTexture * DDSTexture = ImportDDS(Buffer,Length,InParent,Name,Flags,ImportFlags,Warn);
+	if ( DDSTexture != nullptr )
+	{
+		return DDSTexture;
 	}
 
 	//
@@ -3750,19 +3767,19 @@ UTexture* UTextureFactory::ImportTexture(UClass* Class, UObject* InParent, FName
 		if (HDRImportShouldBeLongLatCubeMap == EAppReturnType::Yes ||
 			HDRImportShouldBeLongLatCubeMap == EAppReturnType::YesAll)
 		{
-			FHdrImageWrapper HdrImageWrapper;
-			if (HdrImageWrapper.SetCompressedFromView(TArrayView64<const uint8>(Buffer, Length)))
+			TSharedPtr<IImageWrapper> HdrImageWrapper = ImageWrapperModule.CreateImageWrapper(EImageFormat::HDR);
+			if (HdrImageWrapper->SetCompressed(Buffer, Length))
 			{
 				TArray64<uint8> UnCompressedData;
-				if (HdrImageWrapper.GetRaw(ERGBFormat::BGRE, 8, UnCompressedData))
+				if (HdrImageWrapper->GetRaw(ERGBFormat::BGRE, 8, UnCompressedData))
 				{
 					// create the cube texture
 					UTextureCube* TextureCube = CreateTextureCube(InParent, Name, Flags);
 					if ( TextureCube )
 					{
 						TextureCube->Source.Init(
-							HdrImageWrapper.GetWidth(),
-							HdrImageWrapper.GetHeight(),
+							HdrImageWrapper->GetWidth(),
+							HdrImageWrapper->GetHeight(),
 							/*NumSlices=*/ 1,
 							/*NumMips=*/ 1,
 							TSF_BGRE8,
@@ -3779,7 +3796,7 @@ UTexture* UTextureFactory::ImportTexture(UClass* Class, UObject* InParent, FName
 				}
 			}
 
-			Warn->Log(ELogVerbosity::Error, HdrImageWrapper.GetErrorMessage().ToString());
+			Warn->Log(ELogVerbosity::Error, TEXT("HDR import failed"));
 			return nullptr;
 		}
 		// else didn't want cubemap
@@ -3789,7 +3806,7 @@ UTexture* UTextureFactory::ImportTexture(UClass* Class, UObject* InParent, FName
 	//
 	// IES File (usually measured real world light profiles)
 	//
-	if(!FCString::Stricmp(Type, TEXT("ies")))
+	if( FCString::Stricmp(Type, TEXT("ies")) == 0)
 	{
 		// checks for .IES extension to avoid wasting loading large assets just to reject them during header parsing
 		FIESConverter IESConverter(Buffer, Length);
@@ -3819,6 +3836,11 @@ UTexture* UTextureFactory::ImportTexture(UClass* Class, UObject* InParent, FName
 
 			return Texture;
 		}
+		else
+		{
+			Warn->Log(ELogVerbosity::Error, TEXT("IES import failed"));
+			return nullptr;
+		}
 	}
 
 	//
@@ -4683,6 +4705,87 @@ bool UTextureExporterPCX::ExportBinary( UObject* Object, const TCHAR* Type, FArc
 }
 
 
+/*------------------------------------------------------------------------------
+	UTextureExporterDDS implementation.
+
+------------------------------------------------------------------------------*/
+
+UTextureExporterDDS::UTextureExporterDDS(const FObjectInitializer& ObjectInitializer)
+	: Super(ObjectInitializer)
+{
+	// note UTexture not UTexture2D , picks up cube etc.
+	SupportedClass = UTexture::StaticClass();
+	PreferredFormatIndex = 0;
+	FormatExtension.Add(TEXT("DDS"));
+	FormatDescription.Add(TEXT("DDS DirectDraw Surface"));
+}
+
+bool UTextureExporterDDS::SupportsObject(UObject* Object) const
+{
+	// don't call Super (UTextureExporterGeneric), skip to UExporter
+	// this checks we are "SupportedClass"
+	if ( ! UExporter::SupportsObject(Object) )
+	{
+		return false;
+	}
+
+	UTexture * Texture = GetExportTexture( Object );
+	if ( Texture == nullptr )
+	{
+		return false;
+	}
+	
+	if ( !Texture->Source.IsValid() )
+	{
+		return false;
+	}
+	
+	return true;
+}
+
+bool UTextureExporterDDS::ExportBinary( UObject* Object, const TCHAR* Type, FArchive& Ar, FFeedbackContext* Warn, int32 FileIndex, uint32 PortFlags )
+{
+	UTexture * Texture = GetExportTexture( Object );
+	check( Texture != nullptr );
+	
+	// FileIndex for layers for VT :
+	
+	int NumBlocks = Texture->Source.GetNumBlocks();
+	int NumLayers = Texture->Source.GetNumLayers();
+
+	int NumFiles = NumBlocks * NumLayers;
+	check( FileIndex < NumFiles );
+	int BlockIndex = FileIndex / NumLayers;
+	int LayerIndex = FileIndex % NumLayers; 
+
+	UE_LOG(LogEditorFactories, Display, TEXT("Exporting DDS from Layer %d/%d Block %d/%d"), LayerIndex,NumLayers, BlockIndex,NumBlocks);
+
+	// Type == Image extension
+	
+	TArray64<uint8> OutArray;
+	if ( ! FImageUtils::ExportTextureSourceToDDS(OutArray,Texture,BlockIndex,LayerIndex) )
+	{
+		Warn->Logf(ELogVerbosity::Error, TEXT("ExportTextureSourceToDDS failed"));
+		return false;
+	}
+
+	Ar.Serialize(OutArray.GetData(),OutArray.Num());
+
+	return true;
+}
+
+UVirtualTextureBuilderExporterDDS::UVirtualTextureBuilderExporterDDS(const FObjectInitializer& ObjectInitializer)
+	: UTextureExporterDDS(ObjectInitializer)
+{
+	SupportedClass = UVirtualTextureBuilder::StaticClass();
+}
+
+/*------------------------------------------------------------------------------
+	UTextureExporterGeneric implementation.
+
+	UTextureExporterGeneric base class for all standard 2d texture exporters
+------------------------------------------------------------------------------*/
+
 static bool ExportTexture2DGeneric( UObject* Object,FArchive& Ar, const TCHAR * ImageFormat, FFeedbackContext* Warn, int FileIndex=0)
 {
 	UTexture2D* Texture = CastChecked<UTexture2D>( Object );
@@ -4699,13 +4802,11 @@ static bool ExportTexture2DGeneric( UObject* Object,FArchive& Ar, const TCHAR *
 	
 	int NumBlocks = Texture->Source.GetNumBlocks();
 	int NumLayers = Texture->Source.GetNumLayers();
-	int BlockIndex = 0;
-	int LayerIndex = 0;
 
 	int NumFiles = NumBlocks * NumLayers;
 	check( FileIndex < NumFiles );
-	BlockIndex = FileIndex / NumLayers;
-	LayerIndex = FileIndex % NumLayers; 
+	int BlockIndex = FileIndex / NumLayers;
+	int LayerIndex = FileIndex % NumLayers; 
 
 	UE_LOG(LogEditorFactories, Display, TEXT("Exporting Texture as %s Layer %d/%d Block %d/%d"), ImageFormat, LayerIndex,NumLayers, BlockIndex,NumBlocks);
 
@@ -4728,14 +4829,6 @@ static bool ExportTexture2DGeneric( UObject* Object,FArchive& Ar, const TCHAR *
 	return true;
 }
 
-/*------------------------------------------------------------------------------
-	UTextureExporterGeneric implementation.
-
-	UTextureExporterGeneric base class for all standard 2d texture exporters
-------------------------------------------------------------------------------*/
-
-// note that 
-
 UTextureExporterGeneric::UTextureExporterGeneric(const FObjectInitializer& ObjectInitializer)
 	: Super(ObjectInitializer)
 {
@@ -4743,7 +4836,7 @@ UTextureExporterGeneric::UTextureExporterGeneric(const FObjectInitializer& Objec
 	PreferredFormatIndex = 0;
 }
 
-UTexture2D* UTextureExporterGeneric::GetExportTexture(UObject* Object) const
+UTexture* UTextureExporterGeneric::GetExportTexture(UObject* Object) const
 {
 	// for UVirtualTextureBuilderExporter
 	UVirtualTextureBuilder* VirtualTextureBuilder = Cast<UVirtualTextureBuilder>(Object);
@@ -4751,7 +4844,7 @@ UTexture2D* UTextureExporterGeneric::GetExportTexture(UObject* Object) const
 	{
 		return VirtualTextureBuilder->Texture;
 	}
-	return Cast<UTexture2D>(Object);
+	return Cast<UTexture>(Object);
 }
 
 // note that UTextureExporterGeneric will actually be registered as an exporter itself
@@ -4762,8 +4855,10 @@ bool UTextureExporterGeneric::SupportsObject(UObject* Object) const
 	bool bSupportsObject = false;
 	if (Super::SupportsObject(Object))
 	{
-		UTexture2D* Texture = GetExportTexture(Object);
-		if ( Texture )
+		UTexture* Texture = GetExportTexture(Object);
+		// only Texture2D's for generic :
+		UTexture2D * Texture2D = Cast<UTexture2D>(Texture);
+		if ( Texture2D != nullptr )
 		{
 			if ( ! Texture->Source.IsValid() )
 			{
@@ -4780,7 +4875,7 @@ bool UTextureExporterGeneric::SupportsObject(UObject* Object) const
 
 int32 UTextureExporterGeneric::GetFileCount(UObject* Object) const
 {
-	UTexture2D* Texture = GetExportTexture(Object);
+	UTexture* Texture = GetExportTexture(Object);
 	check(Texture != nullptr);
 
 	// standard textures will have NumBlocks == NumLayers == 1
@@ -4811,21 +4906,23 @@ FString UTextureExporterGeneric::GetUniqueFilename(const TCHAR* Filename, int32
 		// @todo: change GetUniqueFilename API to pass the Object
 		//
 		// it's possible you could even be able to map back to the original UDIM file names (xxx.1001 etc.)
-		//	using the TextureSource asset import info
+		//	using the TextureSource asset import info (AssetImportData)
 		return FString::Printf(TEXT("%s%d%s"), *FPaths::GetBaseFilename(Filename, false), FileIndex, *FPaths::GetExtension(Filename, true));
 	}
 }
 
 bool UTextureExporterGeneric::ExportBinary( UObject* Object, const TCHAR* Type, FArchive& Ar, FFeedbackContext* Warn, int32 FileIndex, uint32 PortFlags )
 {
-	UTexture2D* Texture = GetExportTexture(Object);
+	UTexture* Texture = GetExportTexture(Object);
 	check( Texture != nullptr );
+		
+	UTexture2D * Texture2D = Cast<UTexture2D>(Texture);
+	check( Texture2D != nullptr );
 	check( SupportsTexture(Texture) );
 	
-	// FileIndex for layers for VT :
-
+	// FileIndex is for layers for VT
 	// Type == Image extension
-	return ExportTexture2DGeneric(Texture,Ar,Type,Warn,FileIndex);
+	return ExportTexture2DGeneric(Texture2D,Ar,Type,Warn,FileIndex);
 }
 
 //===========================================================
@@ -4840,7 +4937,7 @@ UTextureExporterBMP::UTextureExporterBMP(const FObjectInitializer& ObjectInitial
 	FormatDescription.Add(TEXT("Windows Bitmap"));
 }
 
-bool UTextureExporterBMP::SupportsTexture(UTexture2D* Texture) const
+bool UTextureExporterBMP::SupportsTexture(UTexture* Texture) const
 {
 	// U8 formats only :
 	return Texture->Source.GetFormat() == TSF_BGRA8 || Texture->Source.GetFormat() == TSF_G8;
@@ -4861,7 +4958,7 @@ UTextureExporterHDR::UTextureExporterHDR(const FObjectInitializer& ObjectInitial
 	FormatDescription.Add(TEXT("HDR Radiance RGBE texture"));
 }
 
-bool UTextureExporterHDR::SupportsTexture(UTexture2D* Texture) const
+bool UTextureExporterHDR::SupportsTexture(UTexture* Texture) const
 {
 	// only BGRE is allowed to go out to HDR now (lossless)
 	bool bSupportsObject = Texture->Source.GetFormat() == TSF_BGRE8;
@@ -4883,7 +4980,7 @@ UTextureExporterPNG::UTextureExporterPNG(const FObjectInitializer& ObjectInitial
 	FormatDescription.Add(TEXT("PNG"));
 }
 
-bool UTextureExporterPNG::SupportsTexture(UTexture2D* Texture) const
+bool UTextureExporterPNG::SupportsTexture(UTexture* Texture) const
 {
 	ETextureSourceFormat TSF = Texture->Source.GetFormat();
 	ERawImageFormat::Type RawFormat = FImageCoreUtils::ConvertToRawImageFormat(TSF);
@@ -4906,7 +5003,7 @@ UTextureExporterEXR::UTextureExporterEXR(const FObjectInitializer& ObjectInitial
 	FormatDescription.Add(TEXT("EXR HDR float texture"));
 }
 
-bool UTextureExporterEXR::SupportsTexture(UTexture2D* Texture) const
+bool UTextureExporterEXR::SupportsTexture(UTexture* Texture) const
 {
 	ETextureSourceFormat TSF = Texture->Source.GetFormat();
 	ERawImageFormat::Type RawFormat = FImageCoreUtils::ConvertToRawImageFormat(TSF);
@@ -5028,9 +5125,6 @@ UTextureCubeExporterHDR::UTextureCubeExporterHDR(const FObjectInitializer& Objec
 	
 	// this is the only exporter registered for UTextureCube
 	// this does GenerateLongLatUnwrap
-
-	// @@!! todo : register a DDS cube exporter that can write the faces
-	//		how can user choose between the two? -> by extension
 }
 
 bool UTextureCubeExporterHDR::ExportBinary(UObject* Object, const TCHAR* Type, FArchive& Ar, FFeedbackContext* Warn, int32 FileIndex, uint32 PortFlags)
diff --git a/Engine/Source/Runtime/Engine/Classes/Engine/Texture.h b/Engine/Source/Runtime/Engine/Classes/Engine/Texture.h
--- a/Engine/Source/Runtime/Engine/Classes/Engine/Texture.h
+++ b/Engine/Source/Runtime/Engine/Classes/Engine/Texture.h
@@ -192,6 +192,13 @@ struct FTextureSource
 		ETextureSourceFormat NewFormat,
 		const uint8* NewData = NULL
 		);
+	
+	/**
+	* Init and copy in texture bits from Image
+	* 
+	* @param Image -  Image to initialize with
+	*/
+	ENGINE_API void Init(const FImageView & Image);
 
 	/**
 	 * Initialize the source data with the given size, number of mips, and format.
@@ -330,9 +337,13 @@ struct FTextureSource
 
 	/** Returns the compression format of the source data in enum format. */
 	FORCEINLINE ETextureSourceCompressionFormat GetSourceCompression() const { return CompressionFormat; }
-	
+	FORCEINLINE bool IsSourceCompressed() const { return GetSourceCompression() != ETextureSourceCompressionFormat::TSCF_None; }
+
 	/** Get GammaSpace for this Format */
-	ENGINE_API EGammaSpace GetGammaSpace() const;
+	ENGINE_API EGammaSpace GetGammaSpace(int LayerIndex) const;
+	
+	/** Get mipped number of slices (volumes mip but arrays don't) */
+	ENGINE_API int GetMippedNumSlices(int NumSlices,int MipIndex) const;
 
 	/** Support for copy/paste */
 	void ExportCustomProperties(FOutputDevice& Out, uint32 Indent);
@@ -348,13 +359,22 @@ struct FTextureSource
 	FORCEINLINE int32 GetNumLayers() const { return NumLayers; }
 	FORCEINLINE int32 GetNumBlocks() const { return Blocks.Num() + 1; }
 	FORCEINLINE ETextureSourceFormat GetFormat(int32 LayerIndex = 0) const { return (LayerIndex == 0) ? Format : LayerFormat[LayerIndex]; }
+	
+	UE_DEPRECATED(5.1, "Use GetSourceCompression instead")
 	FORCEINLINE bool IsPNGCompressed() const { return GetSourceCompression() == ETextureSourceCompressionFormat::TSCF_PNG; }
+	
+	// Warning: IsLongLatCubemap() is not correct.  LongLat Cubemaps have bLongLatCubemap == false
+	// the correct check is to see if the texture class is CubeMap and (NumSlices%6) == 0
+	// bLongLatCubemap is sometimes set to true for cube arrays to disambiguate the case of 6 longlat cubemaps in an array
+	// @@!! consider putting the NumSlices check in IsLongLatCubemap()
 	FORCEINLINE bool IsLongLatCubemap() const { return bLongLatCubemap; }
+
 	FORCEINLINE int64 GetSizeOnDisk() const { return BulkData.GetPayloadSize(); }
 	inline bool HasPayloadData() const { return BulkData.HasPayloadData(); }
 	/** Returns true if the texture's bulkdata payload is either already in memory or if the payload is 0 bytes in length. It will return false if the payload needs to load from disk */
 	FORCEINLINE bool IsBulkDataLoaded() const { return BulkData.DoesPayloadNeedLoading(); }
 
+	// Apply a visitor to the bulkdata :
 	ENGINE_API void OperateOnLoadedBulkData(TFunctionRef<void (const FSharedBuffer& BulkDataBuffer)> Operation);
 
 	UE_DEPRECATED(5.00, "There is no longer a need to call LoadBulkDataWithFileReader, FTextureSource::BulkData can now load the data on demand without it.")
@@ -1499,6 +1519,7 @@ public:
 
 	ENGINE_API EGammaSpace GetGammaSpace() const
 	{
+		// note: does not validate that the Format respects gamma (TextureSource::GetGammaSpace does)
 		return SRGB ? ( bUseLegacyGamma ? EGammaSpace::Pow22 : EGammaSpace::sRGB ) : EGammaSpace::Linear;
 	}
 #endif
diff --git a/Engine/Source/Runtime/Engine/Private/ImageUtils.cpp b/Engine/Source/Runtime/Engine/Private/ImageUtils.cpp
--- a/Engine/Source/Runtime/Engine/Private/ImageUtils.cpp
+++ b/Engine/Source/Runtime/Engine/Private/ImageUtils.cpp
@@ -7,9 +7,10 @@ ImageUtils.cpp: Image utility functions.
 #include "ImageUtils.h"
 
 #include "CubemapUnwrapUtils.h"
-#include "DDSLoader.h"
 #include "Engine/Texture2D.h"
 #include "Engine/TextureCube.h"
+#include "Engine/TextureCubeArray.h"
+#include "Engine/VolumeTexture.h"
 #include "Engine/TextureRenderTarget2D.h"
 #include "IImageWrapper.h"
 #include "IImageWrapperModule.h"
@@ -19,6 +20,7 @@ ImageUtils.cpp: Image utility functions.
 #include "Misc/Paths.h"
 #include "Modules/ModuleManager.h"
 #include "ImageCoreUtils.h"
+#include "DDSFile.h"
 
 DEFINE_LOG_CATEGORY_STATIC(LogImageUtils, Log, All);
 
@@ -152,6 +154,164 @@ bool FImageUtils::CompressImage(TArray64<uint8> & OutData, const TCHAR * ToForma
 	return ImageWrapperModule->CompressImage(OutData,ToFormat,InImage,Quality);
 }
 
+bool FImageUtils::ExportTextureSourceToDDS(TArray64<uint8> & OutData, UTexture * Texture, int BlockIndex, int LayerIndex)
+{
+#if WITH_EDITORONLY_DATA	
+	UE::DDS::EDDSError Error;
+
+	UTexture2D * Texture2D = Cast<UTexture2D>(Texture);
+	UTextureCube * TextureCube = Cast<UTextureCube>(Texture);
+	UTextureCubeArray * TextureCubeArray = Cast<UTextureCubeArray>(Texture);
+	UVolumeTexture * TextureVolume = Cast<UVolumeTexture>(Texture);
+	const FTextureSource & Source = Texture->Source;
+	
+	check( BlockIndex < Source.GetNumBlocks() );
+	check( LayerIndex < Source.GetNumLayers() );
+	
+	FTextureSourceBlock Block;
+	Source.GetBlock(BlockIndex, Block);
+
+	int32 SizeX = Block.SizeX;
+	int32 SizeY = Block.SizeY;
+	int32 NumSlices = Block.NumSlices;
+	int32 NumMips = Block.NumMips;
+
+	int32 Dimension;
+	int32 SizeZ;
+	int32 ArraySize;
+	uint32 CreateFlags = 0;
+	if ( Texture2D != nullptr )
+	{
+		Dimension = 2;
+		SizeZ = 1;
+		ArraySize = NumSlices;
+	}
+	else if ( TextureCube != nullptr )
+	{
+		Dimension = 2;
+		SizeZ = 1;
+
+		// IsLongLatCubemap() is not right
+		// just guess from NumSlices
+
+		//bool bIsLongLatCubemap = Source.IsLongLatCubemap();
+		bool bIsLongLatCubemap = ( NumSlices == 1 ); 
+
+		if ( bIsLongLatCubemap )
+		{
+			// just export as 2d texture
+			check( NumSlices == 1 );
+			ArraySize = 1;
+		}
+		else
+		{
+			check( NumSlices == 6 );
+			ArraySize = 6;
+			CreateFlags = UE::DDS::FDDSFile::CREATE_FLAG_CUBEMAP;
+		}
+	}
+	else if ( TextureCubeArray != nullptr )
+	{
+		Dimension = 2;
+		SizeZ = 1;
+
+		// IsLongLatCubemap() should be right for CubeArrays
+		bool bIsLongLatCubemap = Source.IsLongLatCubemap() || (NumSlices%6) != 0; 
+
+		if ( bIsLongLatCubemap )
+		{
+			// just export as 2d texture array
+			ArraySize = NumSlices;
+		}
+		else
+		{
+			check( (NumSlices%6) == 0 );
+			ArraySize = NumSlices;
+			CreateFlags = UE::DDS::FDDSFile::CREATE_FLAG_CUBEMAP;
+		}
+	}
+	else if ( TextureVolume != nullptr )
+	{
+		Dimension = 3;
+		SizeZ = NumSlices;
+		ArraySize = 1;
+
+		// @@!! if NumMips != 1 will fail later in FillMip right now due to Unreal getting Volume mips wrong
+	}
+	else
+	{
+		check(0);
+		return false;
+	}
+	
+	ERawImageFormat::Type RawFormat = FImageCoreUtils::ConvertToRawImageFormat(Source.GetFormat(LayerIndex));
+	EGammaSpace GammaSpace = Source.GetGammaSpace(LayerIndex);
+	UE::DDS::EDXGIFormat DXGIFormat = UE::DDS::DXGIFormatFromRawFormat(RawFormat,GammaSpace);
+	
+	if ( RawFormat == ERawImageFormat::BGRE8 )
+	{
+		// convert BGRE to Linear float? or just export the bytes as they are in BGRA ?
+		UE_LOG(LogImageUtils,Warning,TEXT("DDS export to BGRE8 will write the bytes as they were in BGRA, not converted to float linear"));
+	}
+
+	UE_LOG(LogImageUtils,Display,TEXT("Exporting DDS Dimension=%d SizeX=%d SizeY=%d SizeZ=%d NumMips=%d ArraySize=%d"),Dimension, SizeX,SizeY,SizeZ,NumMips,ArraySize);	
+
+	UE::DDS::FDDSFile * DDS = UE::DDS::FDDSFile::CreateEmpty(Dimension, SizeX,SizeY,SizeZ,NumMips,ArraySize, DXGIFormat,CreateFlags, &Error);
+	if ( DDS == nullptr || Error != UE::DDS::EDDSError::OK )
+	{
+		UE_LOG(LogImageUtils,Warning,TEXT("FDDSFile::CreateEmpty (Error=%d)"),(int)Error);			
+		return false;
+	}
+	
+	// delete DDS at scope exit :
+	TUniquePtr<UE::DDS::FDDSFile> DDSPtr(DDS);
+
+	check( DDS->Validate() == UE::DDS::EDDSError::OK );
+	
+	// blit into the mips:
+	
+	for(int MipIndex=0;MipIndex<NumMips;MipIndex++)
+	{
+		FTextureSource::FMipLock MipLock(FTextureSource::ELockState::ReadOnly,const_cast<FTextureSource *>(&Source),BlockIndex,LayerIndex,MipIndex);
+
+		if ( DDS->Dimension == 3 )
+		{
+			check( DDS->Mips.Num() == DDS->MipCount );
+			check( DDS->Mips[MipIndex].Depth == MipLock.Image.NumSlices );
+
+			DDS->FillMip( MipLock.Image, MipIndex );
+		}
+		else
+		{
+			// DDS->Mips[] contains both mips and arrays
+			check( DDS->Mips.Num() == DDS->MipCount * NumSlices );
+
+			for(int SliceIndex=0;SliceIndex<NumSlices;SliceIndex++)
+			{
+				FImageView MipSlice = MipLock.Image.GetSlice(SliceIndex);
+				
+				// DDS Mips[] array has whole mip chain of each slice, then next slice
+				// we have the opposite (all slices of top mip first, then next mip)
+				int DDSMipIndex = SliceIndex * DDS->MipCount + MipIndex;
+				
+				DDS->FillMip( MipSlice, DDSMipIndex );
+			}
+		}
+	}
+
+	Error = DDS->WriteDDS(OutData);
+	if ( Error != UE::DDS::EDDSError::OK )
+	{
+		UE_LOG(LogImageUtils,Warning,TEXT("FDDSFile::WriteDDS failed (Error=%d)"),(int)Error);			
+		return false;
+	}
+
+	return true;
+#else
+	return false;
+#endif
+}
+
 /**
  * Returns data containing the pixmap of the passed in rendertarget.
  * @param TexRT - The 2D rendertarget from which to read pixmap data.
@@ -1113,8 +1273,6 @@ UTexture2D* FImageUtils::ImportFileAsTexture2D(const FString& Filename)
 
 UTexture2D* FImageUtils::ImportBufferAsTexture2D(TArrayView64<const uint8> Buffer)
 {
-	//@@!! test me
-
 	FImage Image;
 	if ( ! DecompressImage(Buffer.GetData(),Buffer.Num(), Image) )
 	{
diff --git a/Engine/Source/Runtime/Engine/Private/LightMap.cpp b/Engine/Source/Runtime/Engine/Private/LightMap.cpp
--- a/Engine/Source/Runtime/Engine/Private/LightMap.cpp
+++ b/Engine/Source/Runtime/Engine/Private/LightMap.cpp
@@ -160,13 +160,13 @@ FString ULightMapTexture2D::GetDesc()
 static void DumpLightmapSizeOnDisk()
 {
 	UE_LOG(LogLightMap,Log,TEXT("Lightmap size on disk"));
-	UE_LOG(LogLightMap,Log,TEXT("Source (KB),Source is PNG,Platform Data (KB),Lightmap"));
+	UE_LOG(LogLightMap,Log,TEXT("Source (KB),Source is Compressed,Platform Data (KB),Lightmap"));
 	for (TObjectIterator<ULightMapTexture2D> It; It; ++It)
 	{
 		ULightMapTexture2D* Lightmap = *It;
 		UE_LOG(LogLightMap,Log,TEXT("%f,%d,%f,%s"),
 			Lightmap->Source.GetSizeOnDisk() / 1024.0f,
-			Lightmap->Source.IsPNGCompressed(),
+			Lightmap->Source.IsSourceCompressed(),
 			Lightmap->CalcTextureMemorySizeEnum(TMC_AllMips) / 1024.0f,
 			*Lightmap->GetPathName()
 			);
diff --git a/Engine/Source/Runtime/Engine/Private/Texture.cpp b/Engine/Source/Runtime/Engine/Private/Texture.cpp
--- a/Engine/Source/Runtime/Engine/Private/Texture.cpp
+++ b/Engine/Source/Runtime/Engine/Private/Texture.cpp
@@ -14,6 +14,7 @@
 #include "UObject/ObjectSaveContext.h"
 #include "TextureResource.h"
 #include "Engine/Texture2D.h"
+#include "Engine/VolumeTexture.h"
 #include "Streaming/TextureMipDataProvider.h"
 #include "Engine/TextureMipDataProviderFactory.h"
 #include "ContentStreaming.h"
@@ -664,8 +665,19 @@ void UTexture::Serialize(FArchive& Ar)
 	
 	if (Ar.IsLoading())
 	{
-		// For a while, we could end up writing textures that were flagged as !bPNGCompressed, had their
-		// compression format set to PNG, but did not actually contain compressed data. Fix these up.
+		// bPNGCompressed is now deprecated and CompressionFormat should be used to detect PNG compression
+		// update old assets that did not have CompressionFormat set
+		// 
+		// - In old versions, CompressionFormat did not exist (so it will load in as None), and bPNGCompressed is used
+		//   as the source to set CompressionFormat
+		// - In new versions, bPNGCompressed is deprecated, never written (so will load as false), and CompressionFormat
+		//   is the authoritative source on whether something is a PNG or not.
+		// - In between, for a while after CompressionFormat was introduced, a bug meant that textures that were flagged 
+		//   as !bPNGCompressed, had their compression format set to PNG, but did not actually contain compressed data. Fix these up.
+		//
+		// Now, the separate bPNGCompressed is gone (to avoid further desyncs like this) and we make sure
+		// that CompressionFormat always matches the contents.	
+
 		if (Ar.CustomVer(FUE5MainStreamObjectVersion::GUID) < FUE5MainStreamObjectVersion::TextureSourceVirtualization
 			&& !Source.bPNGCompressed_DEPRECATED
 			&& Source.CompressionFormat == TSCF_PNG)
@@ -674,24 +686,16 @@ void UTexture::Serialize(FArchive& Ar)
 			Source.CompressionFormat = TSCF_None;
 		}
 
-		// Along with the above, the rules implemented here are this:
-		// - In older versions, bPNGCompressed is the authoritative source of PNG-ness. In those versions,
-		//   if the compression format is set to PNG but bPNGCompressed isn't set, bPNGCompressed wins.
-		// - In new versions, bPNGCompressed is deprecated, never written as true, and CompressionFormat
-		//   is the authoritative source on whether something is a PNG or not.
-		// - Finally, in between, some versions would in certain cases write textures with bPNGCompressed set
-		//   and the CompressionFormat being None, and those are actually PNGs.
-		//
-		// Now, the separate bPNGCompressed is gone (to avoid further desyncs like this) and we make sure
-		// that CompressionFormat always matches the contents.
-		if ( Source.bPNGCompressed_DEPRECATED )
+		if ( Source.bPNGCompressed_DEPRECATED && Source.CompressionFormat != TSCF_PNG )
 		{
-			// loaded with deprecated "bPNGCompressed"
+			// loaded with deprecated "bPNGCompressed" (but not the newer CompressionFormat)
 			// change to CompressionFormat PNG
-			check( Source.CompressionFormat == TSCF_None || Source.CompressionFormat == TSCF_PNG );
+			// this is expected on assets older than the CompressionFormat field
+			check( Source.CompressionFormat == TSCF_None );
 			Source.CompressionFormat = TSCF_PNG;
-			Source.bPNGCompressed_DEPRECATED = false;
 		}
+		
+		// bPNGCompressed_DEPRECATED is not kept in sync with CompressionFormat any more, do not check it after this point
 
 		if ( Source.GetFormat() == TSF_RGBA8_DEPRECATED
 			|| Source.GetFormat() == TSF_RGBE8_DEPRECATED )
@@ -1365,9 +1369,7 @@ void FTextureSource::InitLayered(
 		TotalBytes += CalcLayerSize(0, i);
 	}
 
-	// TODO: Allocating an empty buffer if there is no data to copy from seems like an odd choice but the 
-	// code logic has been doing this for almost a decade so I don't want to change it until I am sure that
-	// it serves no purpose. Given a choice I'd assert on NewData == nullptr instead.
+	// Init with NewData == null is used to allocate space, which is then filled with LockMip
 	if (NewData != nullptr)
 	{
 		BulkData.UpdatePayload(FSharedBuffer::Clone(NewData, TotalBytes), Owner);
@@ -1414,6 +1416,16 @@ void FTextureSource::Init(
 	InitLayered(NewSizeX, NewSizeY, NewNumSlices, 1, NewNumMips, &NewFormat, NewData);
 }
 
+void FTextureSource::Init(const FImageView & Image)
+{
+	ETextureSourceFormat SourceFormat = FImageCoreUtils::ConvertToTextureSourceFormat(Image.Format);
+
+	// FImageView has gamma information too that is lost
+	// TextureSource does not store gamma information (it's in the owning Texture)
+
+	Init(Image.SizeX,Image.SizeY,Image.NumSlices,1,SourceFormat,(const uint8 *)Image.RawData);
+}
+
 void FTextureSource::Init(
 	int32 NewSizeX,
 	int32 NewSizeY,
@@ -1596,9 +1608,9 @@ FTextureSource::FMipLock::FMipLock(ELockState InLockState,FTextureSource * InTex
 		Image.RawData = Locked;
 		Image.SizeX = FMath::Max(Block.SizeX >> MipIndex, 1);
 		Image.SizeY = FMath::Max(Block.SizeY >> MipIndex, 1);
-		Image.NumSlices = Block.NumSlices;
-		Image.Format = FImageCoreUtils::ConvertToRawImageFormat(TextureSource->Format);
-		Image.GammaSpace = TextureSource->GetGammaSpace();
+		Image.NumSlices = TextureSource->GetMippedNumSlices(Block.NumSlices,MipIndex);
+		Image.Format = FImageCoreUtils::ConvertToRawImageFormat(TextureSource->GetFormat(LayerIndex));
+		Image.GammaSpace = TextureSource->GetGammaSpace(LayerIndex);
 		
 		const int64 MipSizeBytes = TextureSource->CalcMipSize(BlockIndex, LayerIndex, MipIndex);
 
@@ -1746,9 +1758,9 @@ bool FTextureSource::GetMipImage(FImage & OutImage, int32 BlockIndex, int32 Laye
 	OutImage.RawData = MoveTemp(MipData);
 	OutImage.SizeX = FMath::Max(Block.SizeX >> MipIndex, 1);
 	OutImage.SizeY = FMath::Max(Block.SizeY >> MipIndex, 1);
-	OutImage.NumSlices = Block.NumSlices;
-	OutImage.Format = FImageCoreUtils::ConvertToRawImageFormat(Format);
-	OutImage.GammaSpace = GetGammaSpace();
+	OutImage.NumSlices = GetMippedNumSlices(Block.NumSlices,MipIndex);
+	OutImage.Format = FImageCoreUtils::ConvertToRawImageFormat(GetFormat(LayerIndex));
+	OutImage.GammaSpace = GetGammaSpace(LayerIndex);
 
 	check( OutImage.GetImageSizeBytes() == MipSizeBytes );
 	return true;
@@ -1823,8 +1835,10 @@ int64 FTextureSource::CalcMipSize(int32 BlockIndex, int32 LayerIndex, int32 MipI
 
 	const int64 MipSizeX = FMath::Max(Block.SizeX >> MipIndex, 1);
 	const int64 MipSizeY = FMath::Max(Block.SizeY >> MipIndex, 1);
+	const int64 MipSlices = GetMippedNumSlices(Block.NumSlices,MipIndex);
+
 	const int64 BytesPerPixel = GetBytesPerPixel(LayerIndex);
-	return MipSizeX * MipSizeY * Block.NumSlices * BytesPerPixel;
+	return MipSizeX * MipSizeY * MipSlices * BytesPerPixel;
 }
 
 int32 FTextureSource::GetBytesPerPixel(int32 LayerIndex) const
@@ -1907,8 +1921,47 @@ FString FTextureSource::GetIdString() const
 	return GuidString;
 }
 
-EGammaSpace FTextureSource::GetGammaSpace() const
+int FTextureSource::GetMippedNumSlices(int InNumSlices,int InMipIndex) const
 {
+	// @@!! test me!  before was always returning InNumSlices
+	//  not sure you can change this, does it break old content?
+	//  but I think any old content that had TextureSource mips for volumes was broken anyway
+	//	so maybe it's okay?
+
+	#if 1
+
+	// match old behavior
+	return InNumSlices;
+
+	#else
+	
+	// what to do with NumSlices on the mip ?
+	// if this is an Array, it should stay the same
+	// if this is a Volume, it should mip down
+
+	// TextureSource does not know if it's a volume or not
+	// need to check type of Owner Texture
+
+	bool bIsVolume = Owner ? Owner->IsA(UVolumeTexture::StaticClass()) : false;
+	if ( bIsVolume )
+	{
+		return FMath::Max(InNumSlices >> InMipIndex, 1);
+	}
+	else
+	{
+		return InNumSlices;
+	}	
+
+	#endif
+}
+
+EGammaSpace FTextureSource::GetGammaSpace(int LayerIndex) const
+{
+	if ( ! ERawImageFormat::GetFormatNeedsGammaSpace( FImageCoreUtils::ConvertToRawImageFormat(GetFormat(LayerIndex)) ) )
+	{
+		return EGammaSpace::Linear;
+	}
+
 	// TextureSource does not know its own gamma, but its owning Texture does :
 	if ( Owner != nullptr )
 	{
@@ -1917,7 +1970,8 @@ EGammaSpace FTextureSource::GetGammaSpace() const
 	else
 	{
 		// TextureSource does not have gamma information, just guess it ?
-		return ERawImageFormat::GetDefaultGammaSpace( FImageCoreUtils::ConvertToRawImageFormat(Format) );
+		// we know format GetFormatNeedsGammaSpace is true, so default is SRGB
+		return EGammaSpace::sRGB;
 	}
 }
 
@@ -1948,7 +2002,7 @@ FSharedBuffer FTextureSource::TryDecompressData() const
 			// this shouldn't ever happen currently
 			UE_LOG(LogTexture, Warning, TEXT("TryDecompressData unexpected format conversion?"));
 
-			Image.ChangeFormat(RawFormat, GetGammaSpace());
+			Image.ChangeFormat(RawFormat, GetGammaSpace(0));
 		}
 			
 		if ( CompressionFormat == TSCF_PNG && Image.Format == ERawImageFormat::BGRA8 )
@@ -2140,24 +2194,24 @@ int64 FTextureSource::CalcBlockSize(const FTextureSourceBlock& Block) const
 int64 FTextureSource::CalcLayerSize(const FTextureSourceBlock& Block, int32 LayerIndex) const
 {
 	int64 BytesPerPixel = GetBytesPerPixel(LayerIndex);
-	int64 MipSizeX = Block.SizeX;
-	int64 MipSizeY = Block.SizeY;
 
 	int64 TotalSize = 0;
 	for (int32 MipIndex = 0; MipIndex < Block.NumMips; ++MipIndex)
 	{
-		TotalSize += MipSizeX * MipSizeY * BytesPerPixel * Block.NumSlices;
-		MipSizeX = FMath::Max<int64>(MipSizeX >> 1, 1);
-		MipSizeY = FMath::Max<int64>(MipSizeY >> 1, 1);
+		int32 MipSizeX = FMath::Max<int32>(Block.SizeX >> MipIndex, 1);
+		int32 MipSizeY = FMath::Max<int32>(Block.SizeY >> MipIndex, 1);
+		int32 MipSizeZ = GetMippedNumSlices(Block.NumSlices,MipIndex);
+
+		TotalSize += MipSizeX * MipSizeY * MipSizeZ * BytesPerPixel;
 	}
 	return TotalSize;
 }
 
-int64 FTextureSource::CalcMipOffset(int32 BlockIndex, int32 LayerIndex, int32 MipIndex) const
+int64 FTextureSource::CalcMipOffset(int32 BlockIndex, int32 LayerIndex, int32 OffsetToMipIndex) const
 {
 	FTextureSourceBlock Block;
 	GetBlock(BlockIndex, Block);
-	check(MipIndex < Block.NumMips);
+	check(OffsetToMipIndex < Block.NumMips);
 
 	int64 MipOffset = BlockDataOffsets[BlockIndex];
 
@@ -2168,14 +2222,14 @@ int64 FTextureSource::CalcMipOffset(int32 BlockIndex, int32 LayerIndex, int32 Mi
 	}
 
 	int64 BytesPerPixel = GetBytesPerPixel(LayerIndex);
-	int64 MipSizeX = Block.SizeX;
-	int64 MipSizeY = Block.SizeY;
 
-	while (MipIndex-- > 0)
+	for (int32 MipIndex = 0; MipIndex < OffsetToMipIndex; ++MipIndex)
 	{
-		MipOffset += MipSizeX * MipSizeY * BytesPerPixel * Block.NumSlices;
-		MipSizeX = FMath::Max<int64>(MipSizeX >> 1, 1);
-		MipSizeY = FMath::Max<int64>(MipSizeY >> 1, 1);
+		int32 MipSizeX = FMath::Max<int32>(Block.SizeX >> MipIndex, 1);
+		int32 MipSizeY = FMath::Max<int32>(Block.SizeY >> MipIndex, 1);
+		int32 MipSizeZ = GetMippedNumSlices(Block.NumSlices,MipIndex);
+
+		MipOffset += MipSizeX * MipSizeY * MipSizeZ * BytesPerPixel;
 	}
 
 	return MipOffset;
diff --git a/Engine/Source/Runtime/Engine/Private/TextureDerivedData.cpp b/Engine/Source/Runtime/Engine/Private/TextureDerivedData.cpp
--- a/Engine/Source/Runtime/Engine/Private/TextureDerivedData.cpp
+++ b/Engine/Source/Runtime/Engine/Private/TextureDerivedData.cpp
@@ -737,8 +737,8 @@ static void GetTextureBuildSettings(
 	{
 		OutBuildSettings.bCubemap = true;
 		OutBuildSettings.DiffuseConvolveMipLevel = GDiffuseConvolveMipLevel;
-		const UTextureCube* Cube = CastChecked<UTextureCube>(&Texture);
-		OutBuildSettings.bLongLatSource = (Cube->Source.GetNumSlices() == 1) || Cube->Source.IsLongLatCubemap();
+		check( Texture.Source.GetNumSlices() == 1 || Texture.Source.GetNumSlices() == 6 );
+		OutBuildSettings.bLongLatSource = (Texture.Source.GetNumSlices() == 1) || Texture.Source.IsLongLatCubemap();
 	}
 	else if (Texture.IsA(UTexture2DArray::StaticClass()))
 	{
@@ -749,6 +749,7 @@ static void GetTextureBuildSettings(
 		OutBuildSettings.bCubemap = true;
 		OutBuildSettings.bTextureArray = true;
 		OutBuildSettings.bLongLatSource = Texture.Source.IsLongLatCubemap();
+		check( ((Texture.Source.GetNumSlices()%6)==0) || OutBuildSettings.bLongLatSource );
 	}
 	else if (Texture.IsA(UVolumeTexture::StaticClass()))
 	{
diff --git a/Engine/Source/Runtime/Engine/Private/TextureDerivedDataBuildUtils.cpp b/Engine/Source/Runtime/Engine/Private/TextureDerivedDataBuildUtils.cpp
--- a/Engine/Source/Runtime/Engine/Private/TextureDerivedDataBuildUtils.cpp
+++ b/Engine/Source/Runtime/Engine/Private/TextureDerivedDataBuildUtils.cpp
@@ -205,14 +205,7 @@ static void WriteSource(FCbWriter& Writer, const UTexture& Texture, int32 LayerI
 
 	Writer.BeginObject();
 
-	ETextureSourceCompressionFormat CompressionFormat = Source.GetSourceCompression();
-	if ((CompressionFormat == ETextureSourceCompressionFormat::TSCF_PNG) && !Source.IsPNGCompressed())
-	{
-		// CompressionFormat might mismatch with IsPNGCompressed.  In that case IsPNGCompressed is authoritative.
-		// This behavior matches FTextureSource::Decompress
-		CompressionFormat = ETextureSourceCompressionFormat::TSCF_None;
-	}
-	Writer.AddInteger("CompressionFormat", CompressionFormat);
+	Writer.AddInteger("CompressionFormat", Source.GetSourceCompression());
 	Writer.AddInteger("SourceFormat", Source.GetFormat(LayerIndex));
 	Writer.AddInteger("GammaSpace", static_cast<uint8>(GammaSpace));
 	Writer.AddInteger("NumSlices", (BuildSettings.bCubemap || BuildSettings.bTextureArray || BuildSettings.bVolume) ? Source.GetNumSlices() : 1);
diff --git a/Engine/Source/Runtime/Engine/Private/TextureDerivedDataTask.cpp b/Engine/Source/Runtime/Engine/Private/TextureDerivedDataTask.cpp
--- a/Engine/Source/Runtime/Engine/Private/TextureDerivedDataTask.cpp
+++ b/Engine/Source/Runtime/Engine/Private/TextureDerivedDataTask.cpp
@@ -84,7 +84,7 @@ public:
 	}
 };
 
-static FText ComposeTextureBuildText(const FString& TexturePathName, int32 SizeX, int32 SizeY, int32 NumBlocks, int32 NumLayers, const FTextureBuildSettings& BuildSettings, ETextureEncodeSpeed InEncodeSpeed, int64 RequiredMemoryEstimate, bool bIsVT)
+static FText ComposeTextureBuildText(const FString& TexturePathName, int32 SizeX, int32 SizeY, int32 NumSlices, int32 NumBlocks, int32 NumLayers, const FTextureBuildSettings& BuildSettings, ETextureEncodeSpeed InEncodeSpeed, int64 RequiredMemoryEstimate, bool bIsVT)
 {
 	FFormatNamedArguments Args;
 	Args.Add(TEXT("TextureName"), FText::FromString(TexturePathName));
@@ -94,6 +94,7 @@ static FText ComposeTextureBuildText(const FString& TexturePathName, int32 SizeX
 	Args.Add(TEXT("TextureResolutionY"), FText::FromString(FString::FromInt(SizeY)));
 	Args.Add(TEXT("NumBlocks"), FText::FromString(FString::FromInt(NumBlocks)));
 	Args.Add(TEXT("NumLayers"), FText::FromString(FString::FromInt(NumLayers)));
+	Args.Add(TEXT("NumSlices"), FText::FromString(FString::FromInt(NumSlices)));
 	Args.Add(TEXT("EstimatedMemory"), FText::FromString(FString::SanitizeFloat(double(RequiredMemoryEstimate) / (1024.0*1024.0), 3)));
 	
 	const TCHAR* SpeedText = TEXT("");
@@ -107,19 +108,20 @@ static FText ComposeTextureBuildText(const FString& TexturePathName, int32 SizeX
 	Args.Add(TEXT("Speed"), FText::FromString(FString(SpeedText)));
 
 	return FText::Format(
-		NSLOCTEXT("Engine", "BuildTextureStatus", "Building textures: {TextureName} ({TextureFormatName}{IsVT}, {TextureResolutionX}X{TextureResolutionY} X{NumBlocks}X{NumLayers}) (Required Memory Estimate: {EstimatedMemory} MB), EncodeSpeed: {Speed}"), 
+		NSLOCTEXT("Engine", "BuildTextureStatus", "Building textures: {TextureName} ({TextureFormatName}{IsVT}, {TextureResolutionX}X{TextureResolutionY} X{NumSlices}X{NumLayers}X{NumBlocks}) (Required Memory Estimate: {EstimatedMemory} MB), EncodeSpeed: {Speed}"), 
 		Args
 	);
 }
 
 static FText ComposeTextureBuildText(const FString& TexturePathName, const FTextureSourceData& TextureData, const FTextureBuildSettings& BuildSettings, ETextureEncodeSpeed InEncodeSpeed, int64 RequiredMemoryEstimate, bool bIsVT)
 {
-	return ComposeTextureBuildText(TexturePathName, TextureData.Blocks[0].MipsPerLayer[0][0].SizeX, TextureData.Blocks[0].MipsPerLayer[0][0].SizeY, TextureData.Blocks.Num(), TextureData.Layers.Num(), BuildSettings, InEncodeSpeed, RequiredMemoryEstimate, bIsVT);
+	const FImage & MipImage = TextureData.Blocks[0].MipsPerLayer[0][0];
+	return ComposeTextureBuildText(TexturePathName, MipImage.SizeX, MipImage.SizeY, MipImage.NumSlices, TextureData.Blocks.Num(), TextureData.Layers.Num(), BuildSettings, InEncodeSpeed, RequiredMemoryEstimate, bIsVT);
 }
 
 static FText ComposeTextureBuildText(const UTexture& Texture, const FTextureBuildSettings& BuildSettings, ETextureEncodeSpeed InEncodeSpeed, int64 RequiredMemoryEstimate, bool bIsVT)
 {
-	return ComposeTextureBuildText(Texture.GetPathName(), Texture.Source.GetSizeX(), Texture.Source.GetSizeY(), Texture.Source.GetNumBlocks(), Texture.Source.GetNumLayers(), BuildSettings, InEncodeSpeed, RequiredMemoryEstimate, bIsVT);
+	return ComposeTextureBuildText(Texture.GetPathName(), Texture.Source.GetSizeX(), Texture.Source.GetSizeY(), Texture.Source.GetNumSlices(), Texture.Source.GetNumBlocks(), Texture.Source.GetNumLayers(), BuildSettings, InEncodeSpeed, RequiredMemoryEstimate, bIsVT);
 }
 
 static bool ValidateTexture2DPlatformData(const FTexturePlatformData& TextureData, const UTexture2D& Texture, bool bFromDDC)
diff --git a/Engine/Source/Runtime/Engine/Public/ImageUtils.h b/Engine/Source/Runtime/Engine/Public/ImageUtils.h
--- a/Engine/Source/Runtime/Engine/Public/ImageUtils.h
+++ b/Engine/Source/Runtime/Engine/Public/ImageUtils.h
@@ -122,6 +122,18 @@ public:
 	 */
 	ENGINE_API static bool DecompressImage(const void* InCompressedData, int64 InCompressedSize, FImage & OutImage);
 	
+	/**
+	*  Export Texture (2D,Array,Cube,CubeArray,Volume) to DDS
+	*  reads from texture source, not platform data
+	* 
+	* @param OutData    Filled with DDS file format
+	* @param Texture    Texture (2D,Array,Cube,CubeArray,Volume) to export
+	* @param BlockIndex Block to export (optional)
+	* @param LayerIndex Layer to export (optional)
+	* 
+	*/
+	ENGINE_API static bool ExportTextureSourceToDDS(TArray64<uint8> & OutData, UTexture * Texture, int BlockIndex=0, int LayerIndex=0);
+
 	/**
 	 * if Texture source is available, get it as an FImage
 	 *
diff --git a/Engine/Source/Runtime/ImageCore/Private/DDSFile.cpp b/Engine/Source/Runtime/ImageCore/Private/DDSFile.cpp
--- a/Engine/Source/Runtime/ImageCore/Private/DDSFile.cpp
+++ b/Engine/Source/Runtime/ImageCore/Private/DDSFile.cpp
@@ -130,6 +130,11 @@ static const EDXGIFormat DXGIFormatRGBATable[] =
 	EDXGIFormat::R8G8B8A8_TYPELESS,		EDXGIFormat::B8G8R8A8_TYPELESS,
 	EDXGIFormat::R8G8B8A8_UNORM,		EDXGIFormat::B8G8R8A8_UNORM,
 	EDXGIFormat::R8G8B8A8_UNORM_SRGB,	EDXGIFormat::B8G8R8A8_UNORM_SRGB,
+	EDXGIFormat::R8G8B8X8,				EDXGIFormat::B8G8R8X8_TYPELESS,
+// not mapped :
+//	RGBFMT(R8G8B8A8_UINT,				30, 4) 
+//	RGBFMT(R8G8B8A8_SNORM,				31, 4) 
+//	RGBFMT(R8G8B8A8_SINT,				32, 4) 
 };
 
 static int DXGIFormatFindInTableImpl(const EDXGIFormat* InTable, int InCount, EDXGIFormat InFormat)
@@ -157,6 +162,12 @@ bool DXGIFormatIsSRGB(EDXGIFormat Format)
 	return idx >= 0 && ((idx & 1) == 1);
 }
 
+bool DXGIFormatHasLinearAndSRGBForm(EDXGIFormat Format) 
+{
+	int idx = DXGIFormatFindInTable(DXGIFormatSRGBTable, Format);
+	return idx >= 0;
+}
+
 EDXGIFormat DXGIFormatRemoveSRGB(EDXGIFormat fmt) 
 {
 	int idx = DXGIFormatFindInTable(DXGIFormatSRGBTable, fmt);
@@ -226,6 +237,7 @@ static const FDXGIFormatInfo SupportedFormatList[] =
 // the other ones we read from D3D9 .DDS files but will only write when FormatVersion is D3D9,
 // and otherwise prefer D3D10 mode, because they're not widely supported in apps that consume
 // D3D9 .DDS files.
+//
 static const FBitmaskToDXGI BitmaskToDXGITable[] = 
 {
 	//flags				bits	r			g			b			a			dxgi							AutoD3d9
@@ -246,8 +258,13 @@ static const FBitmaskToDXGI BitmaskToDXGITable[] =
 	{ DDPF_BUMPDUDV,	32,		0x000000ff, 0x0000ff00, 0x00ff0000, 0xff000000, EDXGIFormat::R8G8B8A8_SNORM,	0 },
 	{ DDPF_BUMPDUDV,	2,		0x0000ffff, 0xffff0000, 0x00000000, 0x00000000, EDXGIFormat::R16G16_SNORM,		0 },
 	{ DDPF_BUMPDUDV,	16,		0x00ff,		0xff00,		0x0000,		0x0000,		EDXGIFormat::R8G8_SNORM,		0 },
+	{ DDPF_RGB,			32,		0x000000ff, 0x0000ff00, 0x00ff0000, 0x00000000, EDXGIFormat::R8G8B8X8,	1 },
+	{ DDPF_RGB,			24,		0x000000ff, 0x0000ff00, 0x00ff0000, 0x00000000, EDXGIFormat::R8G8B8,	1 },
+	{ DDPF_RGB,			24,		0x00ff0000, 0x0000ff00, 0x000000ff, 0x00000000, EDXGIFormat::B8G8R8,	1 },
 };
 
+
+
 // This is following MS DDSTextureLoader11.
 // When multiple FOURCCs map to the same DXGI format, put the preferred FOURCC first.
 // AutoD3d9 works as above.
@@ -358,7 +375,7 @@ EDDSError FDDSFile::Validate() const
 	const FDXGIFormatInfo* FormatInfo = DXGIFormatGetInfo(DXGIFormat);
 	if (!FormatInfo)
 	{
-		UE_LOG(LogDDSFile, Error, TEXT("Unsupported format %d (%s)"), DXGIFormat, DXGIFormatGetName(DXGIFormat));
+		UE_LOG(LogDDSFile, Warning, TEXT("Unsupported format %d (%s)"), DXGIFormat, DXGIFormatGetName(DXGIFormat));
 		return EDDSError::BadPixelFormat;
 	}
 
@@ -368,7 +385,7 @@ EDDSError FDDSFile::Validate() const
 	case 1:
 		if (Height != 1 || Depth != 1)
 		{
-			UE_LOG(LogDDSFile, Error, TEXT("1D textures must have height and depth of 1."));
+			UE_LOG(LogDDSFile, Warning, TEXT("1D textures must have height and depth of 1."));
 			return EDDSError::BadImageDimension;
 		}
 		break;
@@ -376,7 +393,7 @@ EDDSError FDDSFile::Validate() const
 	case 2:
 		if (Depth != 1)
 		{
-			UE_LOG(LogDDSFile, Error, TEXT("2D textures must have depth of 1."));
+			UE_LOG(LogDDSFile, Warning, TEXT("2D textures must have depth of 1."));
 			return EDDSError::BadImageDimension;
 		}
 		break;
@@ -384,19 +401,20 @@ EDDSError FDDSFile::Validate() const
 	case 3:
 		if (ArraySize != 1)
 		{
-			UE_LOG(LogDDSFile, Error, TEXT("3D textures must have array size of 1."));
+			UE_LOG(LogDDSFile, Warning, TEXT("3D textures must have array size of 1."));
 			return EDDSError::BadImageDimension;
 		}
+		break;
 
 	default:
-		UE_LOG(LogDDSFile, Error, TEXT("DDS textures must be 1D, 2D or 3D."));
+		UE_LOG(LogDDSFile, Warning, TEXT("DDS textures must be 1D, 2D or 3D."));
 		return EDDSError::BadResourceDimension;
 	}
 
 	// All dimensions must be non-zero
 	if (Width == 0 || Height == 0 || Depth == 0 || MipCount == 0 || ArraySize == 0)
 	{
-		UE_LOG(LogDDSFile, Error, TEXT("One or more image dimensions are zero."));
+		UE_LOG(LogDDSFile, Warning, TEXT("One or more image dimensions are zero."));
 		return EDDSError::BadImageDimension;
 	}
 
@@ -404,14 +422,14 @@ EDDSError FDDSFile::Validate() const
 	const uint32 MaxDimension = (1 << MAX_MIPS_SUPPORTED) - 1; // 1<<k is when we tip over into needing k+1 mip levels, (1<<k)-1 needs just k.
 	if (Width > MaxDimension || Height > MaxDimension || Depth > MaxDimension)
 	{
-		UE_LOG(LogDDSFile, Error, TEXT("Image dimensions %ux%ux%u of DDS exceed maximum of %u."), Width, Height, Depth, MaxDimension);
+		UE_LOG(LogDDSFile, Warning, TEXT("Image dimensions %ux%ux%u of DDS exceed maximum of %u."), Width, Height, Depth, MaxDimension);
 		return EDDSError::BadImageDimension;
 	}
 
 	// Check that mipmap count is supported and makes sense for the image dimensions.
 	if (MipCount > MAX_MIPS_SUPPORTED)
 	{
-		UE_LOG(LogDDSFile, Error, TEXT("Invalid mipmap count of %u."), MipCount);
+		UE_LOG(LogDDSFile, Warning, TEXT("Invalid mipmap count of %u."), MipCount);
 		return EDDSError::BadMipmapCount;
 	}
 
@@ -423,7 +441,7 @@ EDDSError FDDSFile::Validate() const
 		(Height >> FinalMip) == 0 &&
 		(Depth >> FinalMip) == 0)
 	{
-		UE_LOG(LogDDSFile, Error, TEXT("Invalid mipmap count of %u for %ux%ux%u image."), MipCount, Width, Height, Depth);
+		UE_LOG(LogDDSFile, Warning, TEXT("Invalid mipmap count of %u for %ux%ux%u image."), MipCount, Width, Height, Depth);
 		return EDDSError::BadMipmapCount;
 	}
 
@@ -432,13 +450,13 @@ EDDSError FDDSFile::Validate() const
 	{
 		if (Width != Height || Depth != 1)
 		{
-			UE_LOG(LogDDSFile, Error, TEXT("Cubemap has non-square faces or non-1 depth!"));
+			UE_LOG(LogDDSFile, Warning, TEXT("Cubemap has non-square faces or non-1 depth!"));
 			return EDDSError::BadCubemap;
 		}
 
 		if ((ArraySize % 6) != 0)
 		{
-			UE_LOG(LogDDSFile, Error, TEXT("Cubemap or cubemap array doesn't have a multiple of 6 faces."));
+			UE_LOG(LogDDSFile, Warning, TEXT("Cubemap or cubemap array doesn't have a multiple of 6 faces."));
 			return EDDSError::BadCubemap;
 		}
 	}
@@ -446,6 +464,74 @@ EDDSError FDDSFile::Validate() const
 	return EDDSError::OK;
 }
 
+bool FDDSFile::IsValidTexture2D() const
+{
+	if ( Validate() != EDDSError::OK )
+	{
+		return false;
+	}
+
+	if ( Dimension == 1 || Dimension == 2 )
+	{
+		if ( ArraySize == 1 )
+		{
+			return true;
+		}
+	}
+
+	return false;
+}
+
+bool FDDSFile::IsValidTextureCube() const
+{
+	if ( Validate() != EDDSError::OK )
+	{
+		return false;
+	}
+	
+	if ( Dimension == 2 && ArraySize >= 6 && (ArraySize%6) == 0 )
+	{
+		return true;
+	}
+
+	return false;
+}
+
+bool FDDSFile::IsValidTextureArray() const
+{
+	if ( Validate() != EDDSError::OK )
+	{
+		return false;
+	}
+	
+	if ( Dimension == 1 || Dimension == 2 )
+	{
+		if ( ArraySize > 1 )
+		{
+			return true;
+		}
+	}
+
+	return false;
+}
+
+bool FDDSFile::IsValidTextureVolume() const
+{
+	if ( Validate() != EDDSError::OK )
+	{
+		return false;
+	}
+	
+	if ( Dimension == 3 )
+	{
+		// checked in Validate :
+		check( ArraySize == 1 );
+		return true;
+	}
+
+	return false;
+}
+
 static EDDSError AllocateMips(FDDSFile* InDDS, const FDXGIFormatInfo* InFormatInfo, uint32 InCreateFlags)
 {
 	InDDS->Mips.Empty();
@@ -552,7 +638,7 @@ static EDDSError ParseHeader(FDDSFile* InDDS, FDDSHeaderWithMagic const* InHeade
 		}
 		else
 		{
-			UE_LOG(LogDDSFile, Error, TEXT("D3D10 resource dimension in DDS is neither 1D, 2D, nor 3D."));
+			UE_LOG(LogDDSFile, Warning, TEXT("D3D10 resource dimension in DDS is neither 1D, 2D, nor 3D."));
 			return EDDSError::BadResourceDimension;
 		}
 		InDDS->DXGIFormat = (EDXGIFormat)InDX10Header->dxgi_format;
@@ -606,7 +692,7 @@ static EDDSError ReadPayload(FDDSFile* InDDS, const uint8* InReadCursor, const u
 	{
 		if (InReadEnd - InReadCursor < Mip.DataSize)
 		{
-			UE_LOG(LogDDSFile, Error, TEXT("Error reading texture data."));
+			UE_LOG(LogDDSFile, Warning, TEXT("Error reading texture data."));
 			return EDDSError::IoError;
 		}
 
@@ -634,7 +720,7 @@ static EDDSError ReadPayload(FDDSFile* InDDS, const uint8* InReadCursor, const u
 	// definitely a bad file.
 	if (InDDSSize < sizeof(DDSHeader))
 	{
-		UE_LOG(LogDDSFile, Error, TEXT("Failed to read DDS header"));
+		//UE_LOG(LogDDSFile, Display, TEXT("Too few bytes to read DDS header"));
 		*OutError = EDDSError::IoError;
 		return nullptr;
 	}
@@ -650,7 +736,7 @@ static EDDSError ReadPayload(FDDSFile* InDDS, const uint8* InReadCursor, const u
 
 	if (DDSHeader.Magic != DDS_MAGIC)
 	{
-		UE_LOG(LogDDSFile, Error, TEXT("Not a DDS file."));
+		//UE_LOG(LogDDSFile, Display, TEXT("Not a DDS file."));
 		*OutError = EDDSError::NotADds;
 		return nullptr;
 	}
@@ -661,7 +747,7 @@ static EDDSError ReadPayload(FDDSFile* InDDS, const uint8* InReadCursor, const u
 	{
 		if (ReadEnd - ReadCursor < sizeof(DDS10Header))
 		{
-			UE_LOG(LogDDSFile, Error, TEXT("Failed to read DX10 DDS header"));
+			UE_LOG(LogDDSFile, Display, TEXT("Failed to read DX10 DDS header"));
 			*OutError = EDDSError::IoError;
 			return nullptr;
 		}
@@ -686,6 +772,48 @@ static EDDSError ReadPayload(FDDSFile* InDDS, const uint8* InReadCursor, const u
 
 	return DDS;
  }
+ 
+bool FDDSFile::IsADDS(const uint8* InDDS, int64 InDDSSize)
+{
+	FDDSHeaderWithMagic DDSHeader = {};
+	FDDSHeaderDX10 DDS10Header = {};
+
+	// If we don't even have enough bytes for this to contain a valid header,
+	// definitely a bad file.
+	if (InDDSSize < sizeof(DDSHeader))
+	{
+		return false;
+	}
+	
+	// We've now established that InDDSSize >= sizeof(DDSHeader),
+	// so we can read that much for sure.
+	// It's slightly more convenient to work with an end pointer here.
+	const uint8* ReadEnd = InDDS + InDDSSize;
+	const uint8* ReadCursor = InDDS;
+
+	memcpy(&DDSHeader, ReadCursor, sizeof(DDSHeader));
+	ReadCursor += sizeof(DDSHeader);
+
+	if (DDSHeader.Magic != DDS_MAGIC)
+	{
+		return false;
+	}
+
+	// Do we need to read a dx10 header?
+	const FDDSPixelFormat& ddpf = DDSHeader.ddspf;
+	if ((ddpf.flags & DDPF_FOURCC) && ddpf.fourCC == DX10_MAGIC)
+	{
+		if (ReadEnd - ReadCursor < sizeof(DDS10Header))
+		{
+			return false;
+		}
+
+		memcpy(&DDS10Header, ReadCursor, sizeof(DDS10Header));
+		ReadCursor += sizeof(DDS10Header);
+	}
+
+	return true;
+}
 
 //
 // Write to an archive (i.e. file)
@@ -773,13 +901,13 @@ EDDSError FDDSFile::WriteDDS(TArray64<uint8>& OutDDS, EDDSFormatVersion InFormat
 		// If we found neither a bitmask format nor FourCC, we don't know how to save this pixel format for D3D9.
 		if (!D3d9BitmaskFormat && !D3d9FourCC)
 		{
-			UE_LOG(LogDDSFile, Error, TEXT("Unsupported pixel format %s for D3D9 DDS."), DXGIFormatGetName(EffectiveFormat));
+			UE_LOG(LogDDSFile, Warning, TEXT("Unsupported pixel format %s for D3D9 DDS."), DXGIFormatGetName(EffectiveFormat));
 			return EDDSError::BadPixelFormat;
 		}
 
 		if (ArraySize != 1)
 		{
-			UE_LOG(LogDDSFile, Error, TEXT("D3D9 .DDS does not support arrays."));
+			UE_LOG(LogDDSFile, Warning, TEXT("D3D9 .DDS does not support arrays."));
 			return EDDSError::BadImageDimension;
 		}
 
@@ -885,6 +1013,44 @@ EDDSError FDDSFile::WriteDDS(TArray64<uint8>& OutDDS, EDDSFormatVersion InFormat
 	return EDDSError::OK;
 }
 
+void FDDSFile::ConvertRGBXtoRGBA()
+{
+	// change X8 formats to A8 :
+	if ( DXGIFormat == EDXGIFormat::R8G8B8X8 )
+	{
+		DXGIFormat = EDXGIFormat::R8G8B8A8_UNORM_SRGB;
+	}
+	else if ( DXGIFormat == EDXGIFormat::B8G8R8X8_UNORM )
+	{
+		DXGIFormat = EDXGIFormat::B8G8R8A8_UNORM;
+	}
+	else if ( DXGIFormat == EDXGIFormat::B8G8R8X8_UNORM_SRGB )
+	{
+		DXGIFormat = EDXGIFormat::B8G8R8A8_UNORM_SRGB;
+	}
+	else if ( DXGIFormat == EDXGIFormat::B8G8R8X8_TYPELESS )
+	{
+		DXGIFormat = EDXGIFormat::B8G8R8A8_TYPELESS;
+	}
+	else
+	{
+		// not an X8 format
+		return;
+	}
+
+	// stuff 255 in A :
+	for (auto& Mip : Mips)
+	{
+		// Loop over pixels. Data in DDS mips is densely packed!
+		uint8* PixelData = Mip.Data;
+		const int64 DataSize = Mip.DataSize;
+		for (int64 i = 0; i < DataSize; i += 4)
+		{
+			PixelData[3] = 0xFF;
+		}
+	}
+}
+
 void FDDSFile::ConvertChannelOrder(EChannelOrder InTargetOrder)
 {
 	// Is this one of the few RGBA/BGRA formats we can convert?
@@ -924,4 +1090,248 @@ void FDDSFile::ConvertChannelOrder(EChannelOrder InTargetOrder)
 	}
 }
 
+void FDDSFile::FillMip(const FImageView & FromImage, int MipIndex)
+{
+	check( MipIndex < Mips.Num() );
+	const FDDSMip & ToMip = Mips[MipIndex];
+
+	check( FromImage.GetNumPixels() == ToMip.Width * ToMip.Height * (int64) ToMip.Depth );
+	check( FromImage.GetImageSizeBytes() == ToMip.DataSize );
+		
+	memcpy(ToMip.Data,FromImage.RawData,ToMip.DataSize);
+}
+
+// Blit pixels from DDS mip to Image
+// return false if no pixel format conversion is supported
+bool FDDSFile::GetMipImage(const FImageView & ToImage, int MipIndex) const
+{
+	check( MipIndex < Mips.Num() );
+	const FDDSMip & FromMip = Mips[MipIndex];
+
+	check( ToImage.GetNumPixels() == FromMip.Width * FromMip.Height * (int64) FromMip.Depth );
+
+	// before BlitMip you should have done ConvertChannelOrder to BGRA and ConvertRGBXtoRGBA
+	// so we should only see BGRA (no RGBA or BGRX) here
+
+	if ( ToImage.Format == ERawImageFormat::R16F && DXGIFormat == EDXGIFormat::R32_FLOAT )
+	{
+		// no ERawImageFormat::R32F currently , have to convert to F16
+		
+		const float * FmPtr = (const float *) FromMip.Data;
+		uint16 * ToPtr = (uint16 *)ToImage.RawData;
+		int64 Count = ToImage.GetNumPixels();
+		check( ToImage.GetImageSizeBytes() == Count*sizeof(*ToPtr) );
+		check( FromMip.DataSize == Count*sizeof(*FmPtr) );
+		while(Count--)
+		{
+			FPlatformMath::StoreHalf(ToPtr, *FmPtr++);
+			ToPtr++;
+		}
+		return true;
+	}
+	else if ( ToImage.Format == ERawImageFormat::BGRA8 && DXGIFormat == EDXGIFormat::B8G8R8 )
+	{
+		const uint8 * FmPtr = FromMip.Data;
+		uint8 * ToPtr = (uint8 *)ToImage.RawData;
+		int64 Count = ToImage.GetNumPixels();
+		check( ToImage.GetImageSizeBytes() == Count*4 );
+		check( FromMip.DataSize == Count*3 );
+		while(Count--)
+		{
+			*ToPtr++ = *FmPtr++;
+			*ToPtr++ = *FmPtr++;
+			*ToPtr++ = *FmPtr++;
+			*ToPtr++ = 0xFF;
+		}
+		return true;
+	}
+	else if ( ToImage.Format == ERawImageFormat::BGRA8 && DXGIFormat == EDXGIFormat::R8G8B8 )
+	{
+		const uint8 * FmPtr = FromMip.Data;
+		uint8 * ToPtr = (uint8 *)ToImage.RawData;
+		int64 Count = ToImage.GetNumPixels();
+		check( ToImage.GetImageSizeBytes() == Count*4 );
+		check( FromMip.DataSize == Count*3 );
+		while(Count--)
+		{
+			*ToPtr++ = FmPtr[2];
+			*ToPtr++ = FmPtr[1];
+			*ToPtr++ = FmPtr[0];
+			*ToPtr++ = 0xFF;
+			FmPtr += 3;
+		}
+		return true;
+	}
+	else if ( ToImage.Format == ERawImageFormat::BGRA8 && (
+		DXGIFormat == EDXGIFormat::R8G8_TYPELESS ||
+		DXGIFormat == EDXGIFormat::R8G8_UNORM ||
+		DXGIFormat == EDXGIFormat::R8G8_UINT ||
+		DXGIFormat == EDXGIFormat::R8G8_SNORM ||
+		DXGIFormat == EDXGIFormat::R8G8_SINT ) )
+	{
+		const uint8 * FmPtr = FromMip.Data;
+		uint8 * ToPtr = (uint8 *)ToImage.RawData;
+		int64 Count = ToImage.GetNumPixels();
+		check( ToImage.GetImageSizeBytes() == Count*4 );
+		check( FromMip.DataSize == Count*2 );
+		while(Count--)
+		{
+			*ToPtr++ = 0;
+			*ToPtr++ = FmPtr[1];
+			*ToPtr++ = FmPtr[0];
+			*ToPtr++ = 0xFF;
+			FmPtr += 2;
+		}
+		check( ToPtr == (uint8 *)ToImage.RawData + ToImage.GetImageSizeBytes() );
+		check( FmPtr == FromMip.Data + FromMip.DataSize );
+		return true;
+	}
+	else if ( ToImage.GetImageSizeBytes() == FromMip.DataSize )
+	{	
+		bool bIsExactMatch;
+		ERawImageFormat::Type DXGIToRawFormat = DXGIFormatGetClosestRawFormat(DXGIFormat,&bIsExactMatch);
+
+		if ( DXGIToRawFormat == ToImage.Format )
+		{
+			if ( ! bIsExactMatch )
+			{
+				UE_LOG(LogDDSFile, Warning, TEXT("DDS BlitMip DXGIFormat isn't exact match fmt %d=(%s) to (%s)"), \
+					DXGIFormat, DXGIFormatGetName(DXGIFormat),ERawImageFormat::GetName(ToImage.Format));
+				// but do it anyway
+			}
+		
+			memcpy(ToImage.RawData,FromMip.Data,FromMip.DataSize);
+			return true;
+		}
+		else 
+		{
+			// formats mismatch but are same size
+			// very speculative, try it anyway, but log :
+		
+			UE_LOG(LogDDSFile, Warning, TEXT("DDS BlitMip DXGIFormat incompatible but same size! may be junk!  fmt %d=(%s) to (%s)"), \
+				DXGIFormat, DXGIFormatGetName(DXGIFormat),ERawImageFormat::GetName(ToImage.Format));
+
+			memcpy(ToImage.RawData,FromMip.Data,FromMip.DataSize);
+			return true;
+		}
+	}
+	else
+	{
+		// no supported conversion
+		
+		UE_LOG(LogDDSFile, Warning, TEXT("DDS BlitMip DXGIFormat cannot blit!  %d (%s) to (%s)"), \
+			DXGIFormat, DXGIFormatGetName(DXGIFormat),ERawImageFormat::GetName(ToImage.Format));
+
+		return false;
+	}
+}
+
+struct FDXGIFormatRawFormatMapping
+{
+	EDXGIFormat DXGIFormat;
+	ERawImageFormat::Type RawFormat;
+	bool bIsExactMatch;
+};
+
+FDXGIFormatRawFormatMapping DXGIRawFormatMap[] =
+{
+	{ EDXGIFormat::R32G32B32A32_FLOAT, ERawImageFormat::RGBA32F, true },
+	{ EDXGIFormat::R16G16B16A16_FLOAT, ERawImageFormat::RGBA16F, true },
+	{ EDXGIFormat::R32_FLOAT,  ERawImageFormat::R16F, false },
+	{ EDXGIFormat::R16_FLOAT,  ERawImageFormat::R16F, true },
+
+	// R32G32B32_FLOAT not supported
+	// R32G32_FLOAT not supported
+	// R16G16_FLOAT not supported
+
+	{ EDXGIFormat::R16G16B16A16_UNORM, ERawImageFormat::RGBA16, true },
+	{ EDXGIFormat::R16G16B16A16_UINT,  ERawImageFormat::RGBA16, true },
+	{ EDXGIFormat::R16G16B16A16_SNORM, ERawImageFormat::RGBA16, false },
+	{ EDXGIFormat::R16G16B16A16_SINT,  ERawImageFormat::RGBA16, false },
+
+	{ EDXGIFormat::R8G8B8A8_TYPELESS,  ERawImageFormat::BGRA8, false },
+	{ EDXGIFormat::R8G8B8A8_UNORM,  ERawImageFormat::BGRA8, false },
+	{ EDXGIFormat::R8G8B8A8_UNORM_SRGB,  ERawImageFormat::BGRA8, false },
+	{ EDXGIFormat::R8G8B8A8_UINT,  ERawImageFormat::BGRA8, false },
+	{ EDXGIFormat::R8G8B8A8_SNORM,  ERawImageFormat::BGRA8, false },
+	{ EDXGIFormat::R8G8B8A8_SINT,  ERawImageFormat::BGRA8, false },
+	
+	{ EDXGIFormat::B8G8R8A8_UNORM,  ERawImageFormat::BGRA8, true },
+	{ EDXGIFormat::B8G8R8X8_UNORM,  ERawImageFormat::BGRA8, false },
+	
+	{ EDXGIFormat::B8G8R8A8_TYPELESS,  ERawImageFormat::BGRA8, true },
+	{ EDXGIFormat::B8G8R8A8_UNORM_SRGB,  ERawImageFormat::BGRA8, true },
+	{ EDXGIFormat::B8G8R8X8_TYPELESS,  ERawImageFormat::BGRA8, false },
+	{ EDXGIFormat::B8G8R8X8_UNORM_SRGB,  ERawImageFormat::BGRA8, false },
+
+	{ EDXGIFormat::R32_UINT,  ERawImageFormat::G16, false },
+	{ EDXGIFormat::R32_SINT,  ERawImageFormat::G16, false },
+	
+	{ EDXGIFormat::R16_TYPELESS,  ERawImageFormat::G16, true },
+	{ EDXGIFormat::R16_UNORM,  ERawImageFormat::G16, true },
+	{ EDXGIFormat::R16_UINT,  ERawImageFormat::G16, true },
+	{ EDXGIFormat::R16_SNORM,  ERawImageFormat::G16, false },
+	{ EDXGIFormat::R16_SINT,  ERawImageFormat::G16, false },
+	
+	{ EDXGIFormat::R8_TYPELESS,  ERawImageFormat::G8, true },
+	{ EDXGIFormat::R8_UNORM,  ERawImageFormat::G8, true },
+	{ EDXGIFormat::R8_UINT,  ERawImageFormat::G8, true },
+	{ EDXGIFormat::R8_SNORM,  ERawImageFormat::G8, false },
+	{ EDXGIFormat::R8_SINT,  ERawImageFormat::G8, false },
+	{ EDXGIFormat::A8_UNORM,  ERawImageFormat::G8, false },
+
+	{ EDXGIFormat::B8G8R8, ERawImageFormat::BGRA8, false },
+	{ EDXGIFormat::R8G8B8, ERawImageFormat::BGRA8, false },
+	{ EDXGIFormat::R8G8B8X8, ERawImageFormat::BGRA8, false },
+	
+	{ EDXGIFormat::R8G8_TYPELESS, ERawImageFormat::BGRA8, false },
+	{ EDXGIFormat::R8G8_UNORM, ERawImageFormat::BGRA8, false },
+	{ EDXGIFormat::R8G8_UINT, ERawImageFormat::BGRA8, false },
+	{ EDXGIFormat::R8G8_SNORM, ERawImageFormat::BGRA8, false },
+	{ EDXGIFormat::R8G8_SINT, ERawImageFormat::BGRA8, false },
+
+	// terminates list :
+	{ EDXGIFormat::UNKNOWN, ERawImageFormat::Invalid, false}
+};
+
+ERawImageFormat::Type DXGIFormatGetClosestRawFormat(EDXGIFormat fmt, bool * pIsExactMatch)
+{
+	for(const FDXGIFormatRawFormatMapping * Mapping = DXGIRawFormatMap;Mapping->DXGIFormat != EDXGIFormat::UNKNOWN;++Mapping)
+	{
+		if ( Mapping->DXGIFormat == fmt )
+		{
+			if ( pIsExactMatch )
+			{
+				*pIsExactMatch = Mapping->bIsExactMatch;
+			}
+			return Mapping->RawFormat;
+		}
+	}
+
+	return ERawImageFormat::Invalid;
+}
+
+EDXGIFormat DXGIFormatFromRawFormat(ERawImageFormat::Type RawFormat,EGammaSpace GammaSpace)
+{
+	switch(RawFormat)
+	{
+	case ERawImageFormat::G8:		return EDXGIFormat::R8_UNORM; // would like to encode Gamma, but no way to do so
+	case ERawImageFormat::BGRA8:
+		{
+		if ( GammaSpace == EGammaSpace::Linear)
+			return EDXGIFormat::B8G8R8A8_UNORM;
+		else
+			return EDXGIFormat::B8G8R8A8_UNORM_SRGB;
+		}
+	case ERawImageFormat::BGRE8:	return EDXGIFormat::B8G8R8A8_UNORM; // no way to indicate this is BGRE
+	case ERawImageFormat::RGBA16:	return EDXGIFormat::R16G16B16A16_UNORM;
+	case ERawImageFormat::RGBA16F:	return EDXGIFormat::R16G16B16A16_FLOAT;
+	case ERawImageFormat::RGBA32F:	return EDXGIFormat::R32G32B32A32_FLOAT;
+	case ERawImageFormat::G16:		return EDXGIFormat::R16_UNORM;
+	case ERawImageFormat::R16F:		return EDXGIFormat::R16_FLOAT;
+	default:
+		return EDXGIFormat::UNKNOWN;
+	}
+}
+
 } } // end UE::DDS namespace
\ No newline at end of file
diff --git a/Engine/Source/Runtime/ImageCore/Public/DDSFile.h b/Engine/Source/Runtime/ImageCore/Public/DDSFile.h
--- a/Engine/Source/Runtime/ImageCore/Public/DDSFile.h
+++ b/Engine/Source/Runtime/ImageCore/Public/DDSFile.h
@@ -2,6 +2,7 @@
 #pragma once
 #include "CoreTypes.h"
 #include "Containers/Array.h"
+#include "ImageCore.h"
 
 // This is supposed to be, as close as possible, the complete list of DXGI formats
 // per those docs - the ODDFMTs we don't support for doing anything other than name, but we keep
@@ -131,6 +132,9 @@
 	ODDFMT(P208,						130) \
 	ODDFMT(V208,						131) \
 	ODDFMT(V408,						132) \
+	RGBFMT(B8G8R8,						1001,3) \
+	RGBFMT(R8G8B8,						1002,3) \
+	RGBFMT(R8G8B8X8,					1003,4) \
 	/* end */
 
 namespace UE { namespace DDS
@@ -159,8 +163,10 @@ namespace UE { namespace DDS
 	enum class EDDSError
 	{
 		OK,						// No error
+		Ok = OK,
 		OutOfMemory,			// Out of memory/allocation failure while reading
 		NotADds,				// No valid DDS header found
+		NotADDS = NotADds,
 		BadResourceDimension,	// Resource not 1D, 2D or 3D
 		BadPixelFormat,			// Malformed or unsupported pixel format
 		BadImageDimension,		// Image dimensions (width, height, depth, array size) outside supported range
@@ -272,11 +278,20 @@ namespace UE { namespace DDS
 		// Sanity-check that all members make sense and return an error code.
 		IMAGECORE_API EDDSError Validate() const;
 
+		// Check if this DDS is usable in Unreal Texture types:
+		IMAGECORE_API bool IsValidTexture2D() const;
+		IMAGECORE_API bool IsValidTextureCube() const; // or cube array
+		IMAGECORE_API bool IsValidTextureArray() const;
+		IMAGECORE_API bool IsValidTextureVolume() const;
+
 		// Used for loading from a DDS file image in memory.
 		// 
 		// On error, returns nullptr. If a non-null OutError is supplied, error information
 		// is written there.
 		IMAGECORE_API static FDDSFile* CreateFromDDSInMemory(const uint8* InDDS, int64 InDDSSize, EDDSError* OutError=nullptr);
+		
+		// is this buffer a DDS ?
+		IMAGECORE_API static bool IsADDS(const uint8* InDDS, int64 InDDSSize);
 
 		// For 8-bit UNorm formats, both RGBA and BGRA channel order variants exist and
 		// are popular. This function does the trivial conversion between the two, updating
@@ -285,6 +300,18 @@ namespace UE { namespace DDS
 		// and BGRA variants, this function does nothing.
 		IMAGECORE_API void ConvertChannelOrder(EChannelOrder InTargetOrder);
 
+		// convert X alpha to A  and fill 255
+		IMAGECORE_API void ConvertRGBXtoRGBA();		
+
+		// Blit pixels from DDS mip to Image
+		// return false if no pixel format conversion is supported
+		// ToImage is not allocated, must be pre-allocated and sized correctly
+		IMAGECORE_API bool GetMipImage(const FImageView & ToImage, int MipIndex) const;
+
+		// Fill DDS mip from Image
+		// Image must be of format DXGIFormatFromRawFormat
+		IMAGECORE_API void FillMip(const FImageView & FromImage, int MipIndex);
+
 		// Bit flags
 		static constexpr uint32 CREATE_FLAG_NONE = 0;
 		static constexpr uint32 CREATE_FLAG_CUBEMAP = 1;
@@ -297,6 +324,7 @@ namespace UE { namespace DDS
 
 	// Returns whether a given pixel format is sRGB
 	IMAGECORE_API bool DXGIFormatIsSRGB(EDXGIFormat InFormat);
+	IMAGECORE_API bool DXGIFormatHasLinearAndSRGBForm(EDXGIFormat InFormat);
 
 	// Return the corresponding non-sRGB version of a pixel format if there is one
 	IMAGECORE_API EDXGIFormat DXGIFormatRemoveSRGB(EDXGIFormat InFormat);
@@ -304,4 +332,7 @@ namespace UE { namespace DDS
 	// Return the corresponding sRGB version of a pixel format if there is one
 	IMAGECORE_API EDXGIFormat DXGIFormatAddSRGB(EDXGIFormat InFormat);
 
+	IMAGECORE_API EDXGIFormat DXGIFormatFromRawFormat(ERawImageFormat::Type RawFormat,EGammaSpace GammaSpace);
+
+	IMAGECORE_API ERawImageFormat::Type DXGIFormatGetClosestRawFormat(EDXGIFormat fmt, bool * pIsExactMatch = nullptr);
 } }
diff --git a/Engine/Source/Runtime/ImageCore/Public/ImageCore.h b/Engine/Source/Runtime/ImageCore/Public/ImageCore.h
--- a/Engine/Source/Runtime/ImageCore/Public/ImageCore.h
+++ b/Engine/Source/Runtime/ImageCore/Public/ImageCore.h
@@ -273,7 +273,7 @@ public:
 	 */
 	IMAGECORE_API FImage(int32 InSizeX, int32 InSizeY, int32 InNumSlices, ERawImageFormat::Type InFormat, EGammaSpace InGammaSpace);
 
-	// @@!! danger? changed default GammaSpace from Linear to Guess
+	// note: changed default GammaSpace from Linear to GetDefaultGammaSpace
 	FORCEINLINE FImage(int32 InSizeX, int32 InSizeY, int32 InNumSlices, ERawImageFormat::Type InFormat) :
 		FImage(InSizeX,InSizeY,InNumSlices,InFormat,ERawImageFormat::GetDefaultGammaSpace(InFormat))
 	{
@@ -292,7 +292,7 @@ public:
 	{
 	}
 
-	// @@!! danger? changed default GammaSpace from Linear to Guess
+	// note: changed default GammaSpace from Linear to GetDefaultGammaSpace
 	FORCEINLINE FImage(int32 InSizeX, int32 InSizeY, ERawImageFormat::Type InFormat) :
 		FImage(InSizeX,InSizeY,1,InFormat,ERawImageFormat::GetDefaultGammaSpace(InFormat))
 	{
