Optimize AssetTools AdvancedCopyPackages and add batch versions of ForceReplaceReferences and ConsolidateObjects

331s -> 0.28s to consolidate objects when cloning a small shooting range map

diff --git a/Engine/Source/Developer/AssetTools/Private/AssetTools.cpp b/Engine/Source/Developer/AssetTools/Private/AssetTools.cpp
--- a/Engine/Source/Developer/AssetTools/Private/AssetTools.cpp
+++ b/Engine/Source/Developer/AssetTools/Private/AssetTools.cpp
@@ -991,7 +991,6 @@ bool UAssetToolsImpl::AdvancedCopyPackages(const TMap<FString, FString>& SourceA
 		ObjectTools::GatherSubObjectsForReferenceReplacement(NewObjectSet, ExistingObjectSet, ObjectsAndSubObjectsToReplaceWithin);
 
 		TArray<FName> Dependencies;
-		TArray<UObject*> ObjectsToReplace;
 		for (FName SuccessfullyCopiedPackage : SuccessfullyCopiedSourcePackages)
 		{
 			Dependencies.Reset();
@@ -1001,17 +1000,27 @@ bool UAssetToolsImpl::AdvancedCopyPackages(const TMap<FString, FString>& SourceA
 				const int32 DependencyIndex = SuccessfullyCopiedSourcePackages.IndexOfByKey(Dependency);
 				if (DependencyIndex != INDEX_NONE)
 				{
-					for (const auto& DuplicatedObjectPair : DuplicatedObjectsForEachPackage[DependencyIndex])
+					TMap<UObject*, TArray<UObject*, TInlineAllocator<1>>> Consolidations;
+					Consolidations.Reserve(DuplicatedObjectsForEachPackage[DependencyIndex].Num());
+
+					for (const TPair<TSoftObjectPtr<UObject>, TSoftObjectPtr<UObject>>& Duplication : DuplicatedObjectsForEachPackage[DependencyIndex])
 					{
-						UObject* SourceObject = DuplicatedObjectPair.Key.Get();
-						UObject* NewObject = DuplicatedObjectPair.Value.Get();
+						UObject* SourceObject = Duplication.Key.Get();
+						UObject* NewObject = Duplication.Value.Get();
 						if (SourceObject && NewObject)
 						{
-							ObjectsToReplace.Reset();
-							ObjectsToReplace.Add(SourceObject);
-							ObjectTools::ConsolidateObjects(NewObject, ObjectsToReplace, ObjectsAndSubObjectsToReplaceWithin, ExistingObjectSet, false);
+							Consolidations.FindOrAdd(NewObject).Add(SourceObject);
 						}
 					}
+
+					TArray<ObjectTools::FReplaceRequest> Requests;
+					Requests.Reserve(Consolidations.Num());
+					for (TPair<UObject*, TArray<UObject*, TInlineAllocator<1>>>& Consolidation : Consolidations)
+					{
+						Requests.Add(ObjectTools::FReplaceRequest{Consolidation.Key, Consolidation.Value});
+					}
+
+					ObjectTools::ConsolidateObjects(Requests, ObjectsAndSubObjectsToReplaceWithin, ExistingObjectSet, false);
 				}
 			}
 		}
diff --git a/Engine/Source/Editor/UnrealEd/Private/ObjectTools.cpp b/Engine/Source/Editor/UnrealEd/Private/ObjectTools.cpp
--- a/Engine/Source/Editor/UnrealEd/Private/ObjectTools.cpp
+++ b/Engine/Source/Editor/UnrealEd/Private/ObjectTools.cpp
@@ -967,39 +967,35 @@ namespace ObjectTools
 			}
 		}
 	};
-
-	/**
-	 * Forcefully replaces references to passed in objects
-	 *
-	 * @param ObjectToReplaceWith	Any references found to 'ObjectsToReplace' will be replaced with this object.  If the object is NULL references will be nulled.
-	 * @param ObjectsToReplace		An array of objects that should be replaced with 'ObjectToReplaceWith'
-	 * @param OutInfo				FForceReplaceInfo struct containing useful information about the result of the call to this function
-	 * @param bWarnAboutRootSet		If True a message will be displayed to a user asking them if they would like to remove the rootset flag from objects which have it set.
-									If False, the message will not be displayed and rootset is automatically removed
-	 */
-	void ForceReplaceReferences( UObject* ObjectToReplaceWith, TArray<UObject*>& ObjectsToReplace, TSet<UObject*>& ObjectsToReplaceWithin, FForceReplaceInfo& OutInfo, bool bWarnAboutRootSet = true)
+	
+	void ForceReplaceReferences(TArrayView<FReplaceRequest> Requests, TSet<UObject*>& ObjectsToReplaceWithin, FForceReplaceInfo& OutInfo, bool bWarnAboutRootSet = true)
 	{
 		TRACE_CPUPROFILER_EVENT_SCOPE(ObjectTools::ForceReplaceReferences);
 
+		bool bOnlyNullingOut = true;
 		FPropertyEditorModule& PropertyEditorModule = FModuleManager::LoadModuleChecked<FPropertyEditorModule>("PropertyEditor");
-		PropertyEditorModule.RemoveDeletedObjects( ObjectsToReplace );
+		TArray<UObject*> AllOld;
+		for (FReplaceRequest Request : Requests)
+		{
+			AllOld.Append(Request.Old.GetData(), Request.Old.Num());
+			bOnlyNullingOut &= !Request.New;
+		}
+		PropertyEditorModule.RemoveDeletedObjects(AllOld);
+
 		TSet<UObject*> RootSetObjects;
 
 		GWarn->StatusUpdate( 0, 0, NSLOCTEXT("UnrealEd", "ConsolidateAssetsUpdate_RootSetCheck", "Checking Assets for Root Set...") );
 
 		// Iterate through all the objects to replace and see if they are in the root set.  If they are, offer to remove them from the root set.
-		for ( TArray<UObject*>::TConstIterator ReplaceItr( ObjectsToReplace ); ReplaceItr; ++ReplaceItr )
+		for (UObject* CurObjToReplace : AllOld)
 		{
-			UObject* CurObjToReplace = *ReplaceItr;
-			if ( CurObjToReplace )
+			checkf(CurObjToReplace != nullptr, TEXT("Cannot replace null references"));
+			if (CurObjToReplace->IsRooted())
 			{
-				const bool bFlaggedRootSet = CurObjToReplace->IsRooted();
-				if ( bFlaggedRootSet )
-				{
-					RootSetObjects.Add( CurObjToReplace );
-				}
+				RootSetObjects.Add( CurObjToReplace );
 			}
 		}
+		
 		if ( RootSetObjects.Num() )
 		{
 			if( bWarnAboutRootSet )
@@ -1063,28 +1059,33 @@ namespace ObjectTools
 		//ResetLoaders(nullptr);
 
 		TMap<UObject*, int32> ObjToNumRefsMap;
-		if( ObjectToReplaceWith != NULL )
+		if (!bOnlyNullingOut)
 		{
 			GWarn->StatusUpdate( 0, 0, NSLOCTEXT("UnrealEd", "ConsolidateAssetsUpdate_CheckAssetValidity", "Determining Validity of Assets...") );
-			// Determine if the "object to replace with" has any references to any of the "objects to replace," if so, we don't
-			// want to allow those objects to be replaced, as the object would end up referring to itself!
-			// We can skip this check if "object to replace with" is NULL since it is not useful to check for null references
-			FFindReferencersArchive FindRefsAr( ObjectToReplaceWith, ObjectsToReplace );
-			FindRefsAr.GetReferenceCounts( ObjToNumRefsMap );
+			for (FReplaceRequest Request : Requests)
+			{
+				if (UObject* ObjectToReplaceWith = Request.New)
+				{
+					// Determine if the "object to replace with" has any references to any of the "objects to replace," if so, we don't
+					// want to allow those objects to be replaced, as the object would end up referring to itself!
+					// We can skip this check if "object to replace with" is NULL since it is not useful to check for null references
+					FFindReferencersArchive FindRefsAr( ObjectToReplaceWith, Request.Old );
+					FindRefsAr.AppendReferenceCounts( ObjToNumRefsMap );
+				}
+			}
 		}
 
 		// Objects already loaded and in memory have to have any of their references to the objects to replace swapped with a reference to
 		// the "object to replace with". FArchiveReplaceObjectAndStructPropertyRef can serve this purpose, but it expects a TMap of object to replace : object to replace with.
 		// Therefore, populate a map with all of the valid objects to replace as keys, with the object to replace with as the value for each one.
 		TMap<UObject*, UObject*> ReplacementMap;
-		for ( TArray<UObject*>::TConstIterator ReplaceItr( ObjectsToReplace ); ReplaceItr; ++ReplaceItr )
+		for (FReplaceRequest Request : Requests)
 		{
-			UObject* CurObjToReplace = *ReplaceItr;
-			if ( CurObjToReplace )
+			UObject* ObjectToReplaceWith = Request.New;
+			for (UObject* CurObjToReplace : Request.Old)
 			{
 				// If any of the objects to replace are marked RF_RootSet at this point, an error has occurred
-				const bool bFlaggedRootSet = CurObjToReplace->IsRooted();
-				check( !bFlaggedRootSet );
+				check( !CurObjToReplace->IsRooted() );
 
 				// Exclude root packages from being replaced
 				const bool bRootPackage = ( CurObjToReplace->GetClass() == UPackage::StaticClass() ) && !( CurObjToReplace->GetOuter() );
@@ -1122,47 +1123,50 @@ namespace ObjectTools
 		TArray<UObject*> ReferencingPropertiesMapKeys;
 		TArray<PropertyArrayType> ReferencingPropertiesMapValues;
 
-		// Find the referencers of the objects to be replaced
-		FFindReferencersArchive FindRefsArchive( nullptr, OutInfo.ReplaceableObjects );
-
-		for ( FThreadSafeObjectIterator ObjIter; ObjIter; ++ObjIter )
+		if (!bOnlyNullingOut)
 		{
-			UObject* CurObject = *ObjIter;
-
-			// Don't bother replacing in objects that are about to be garbage collected
-			if ((ObjectsToReplaceWithin.Num() > 0 && !ObjectsToReplaceWithin.Contains(CurObject)) || !IsValidChecked(CurObject) || CurObject->IsUnreachable())
+			// Find the referencers of the objects to be replaced
+			FFindReferencersArchive FindRefsArchive( nullptr, OutInfo.ReplaceableObjects );
+			TMap<UObject*, int32> CurNumReferencesMap;
+			TMultiMap<UObject*, FProperty*> CurReferencingPropertiesMMap;
+			PropertyArrayType CurReferencedProperties;
+			for ( FThreadSafeObjectIterator ObjIter; ObjIter; ++ObjIter )
 			{
-				CurObject = nullptr;
-			}
+				UObject* CurObject = *ObjIter;
 
-			// Unless the "object to replace with" is null, ignore any of the objects to replace to themselves
-			if (CurObject && (ObjectToReplaceWith == NULL || !ReplacementMap.Find( CurObject ) ))
-			{
-				FindRefsArchive.ResetPotentialReferencer(CurObject);
+				// Don't bother replacing in objects that are about to be garbage collected
+				if ((ObjectsToReplaceWithin.Num() > 0 && !ObjectsToReplaceWithin.Contains(CurObject)) || !IsValidChecked(CurObject) || CurObject->IsUnreachable())
+				{
+					continue;
+				}
 
-				// Inform the object referencing any of the objects to be replaced about the properties that are being forcefully
-				// changed, and store both the object doing the referencing as well as the properties that were changed in a map (so that
-				// we can correctly call PostEditChange later)
-				TMap<UObject*, int32> CurNumReferencesMap;
-				TMultiMap<UObject*, FProperty*> CurReferencingPropertiesMMap;
-				if ( FindRefsArchive.GetReferenceCounts( CurNumReferencesMap, CurReferencingPropertiesMMap ) > 0  )
+				UObject** ObjectToReplaceWithPtr = ReplacementMap.Find(CurObject);
+				if (ObjectToReplaceWithPtr == nullptr || *ObjectToReplaceWithPtr == nullptr)
 				{
-					PropertyArrayType CurReferencedProperties;
-					CurReferencingPropertiesMMap.GenerateValueArray( CurReferencedProperties );
+					FindRefsArchive.ResetPotentialReferencer(CurObject);
 
-					ReferencingPropertiesMapKeys.Add(CurObject);
-					ReferencingPropertiesMapValues.Add(CurReferencedProperties);
+					// Inform the object referencing any of the objects to be replaced about the properties that are being forcefully
+					// changed, and store both the object doing the referencing as well as the properties that were changed in a map (so that
+					// we can correctly call PostEditChange later)
 
-					if ( CurReferencedProperties.Num() > 0)
+					if ( FindRefsArchive.GetReferenceCounts( CurNumReferencesMap, CurReferencingPropertiesMMap ) > 0  )
 					{
-						for ( PropertyArrayType::TConstIterator RefPropIter( CurReferencedProperties ); RefPropIter; ++RefPropIter )
+						CurReferencingPropertiesMMap.GenerateValueArray( CurReferencedProperties );
+
+						ReferencingPropertiesMapKeys.Add(CurObject);
+						ReferencingPropertiesMapValues.Add(CurReferencedProperties);
+
+						if ( CurReferencedProperties.Num() > 0)
 						{
-							CurObject->PreEditChange( *RefPropIter );
+							for ( PropertyArrayType::TConstIterator RefPropIter( CurReferencedProperties ); RefPropIter; ++RefPropIter )
+							{
+								CurObject->PreEditChange( *RefPropIter );
+							}
+						}
+						else
+						{
+							CurObject->PreEditChange(nullptr);
 						}
-					}
-					else
-					{
-						CurObject->PreEditChange(nullptr);
 					}
 				}
 			}
@@ -1258,6 +1262,12 @@ namespace ObjectTools
 		}
 	}
 
+	void ForceReplaceReferences( UObject* ObjectToReplaceWith, TArray<UObject*>& ObjectsToReplace, TSet<UObject*>& ObjectsToReplaceWithin, FForceReplaceInfo& OutInfo, bool bWarnAboutRootSet = true)
+	{
+		FReplaceRequest Request = {ObjectToReplaceWith, ObjectsToReplace};
+		ForceReplaceReferences(MakeArrayView(&Request, 1), ObjectsToReplaceWithin, OutInfo, bWarnAboutRootSet);
+	}
+
 	/**
 	 * Forcefully replaces references to passed in objects
 	 *
@@ -1286,136 +1296,175 @@ namespace ObjectTools
 		ForceReplaceReferences(ObjectToReplaceWith, ObjectsToReplace, ObjectsToReplaceWithin, ReplaceInfo, false);
 	}
 
-	FConsolidationResults ConsolidateObjects(UObject* ObjectToConsolidateTo, TArray<UObject*>& ObjectsToConsolidate, TSet<UObject*>& ObjectsToConsolidateWithin, TSet<UObject*>& ObjectsToNotConsolidateWithin, bool bShouldDeleteAfterConsolidate, bool bWarnAboutRootSet)
+	FConsolidationResults ConsolidateObjects(TArrayView<FReplaceRequest> Requests, TSet<UObject*>& ObjectsToConsolidateWithin, TSet<UObject*>& ObjectsToNotConsolidateWithin, bool bShouldDeleteAfterConsolidate, bool bWarnAboutRootSet)
 	{
+		for (FReplaceRequest Request : Requests)
+		{
+			checkf(Request.New != nullptr, TEXT("Can't consolidate objects into null objects"));
+		}
+		
 		FConsolidationResults ConsolidationResults;
 		const bool bShouldShowDialogs = !IsRunningCommandlet();
 		const bool bShouldHandleEditorUIChanges = !IsRunningCommandlet();
-		// Ensure the consolidation is headed toward a valid object and this isn't occurring in game
-		if ( ObjectToConsolidateTo )
+
+		if (bShouldHandleEditorUIChanges)
 		{
-			if (bShouldHandleEditorUIChanges)
+			// Close all editors to avoid changing references to temporary objects used by the editor
+			if (!GEditor->GetEditorSubsystem<UAssetEditorSubsystem>()->CloseAllAssetEditors())
 			{
-				// Close all editors to avoid changing references to temporary objects used by the editor
-				if (!GEditor->GetEditorSubsystem<UAssetEditorSubsystem>()->CloseAllAssetEditors())
-				{
-					// Failed to close at least one editor. It is possible that this editor has in-memory object references
-					// which are not prepared to be changed dynamically so it is not safe to continue
-					return ConsolidationResults;
-				}
+				// Failed to close at least one editor. It is possible that this editor has in-memory object references
+				// which are not prepared to be changed dynamically so it is not safe to continue
+				return ConsolidationResults;
+			}
 
-				// Clear audio components to allow previewed sounds to be consolidated
-				GEditor->ClearPreviewComponents();
+			// Clear audio components to allow previewed sounds to be consolidated
+			GEditor->ClearPreviewComponents();
 
-				// Make sure none of the objects are referenced by the editor's USelection
-				GEditor->GetSelectedObjects()->Deselect(ObjectToConsolidateTo);
-				for (int32 ObjectIdx = 0; ObjectIdx < ObjectsToConsolidate.Num(); ++ObjectIdx)
+			// Make sure none of the objects are referenced by the editor's USelection
+			for (FReplaceRequest Request : Requests)
+			{
+				GEditor->GetSelectedObjects()->Deselect(Request.New);
+				for (UObject* Old : Request.Old)
 				{
-					GEditor->GetSelectedObjects()->Deselect(ObjectsToConsolidate[ObjectIdx]);
+					GEditor->GetSelectedObjects()->Deselect(Old);
 				}
 			}
+		}
 
-			GWarn->BeginSlowTask(NSLOCTEXT("UnrealEd", "ConsolidateAssetsUpdate_Consolidating", "Consolidating Assets..."), true);
-			// Keep track of which objects, if any, cannot be consolidated, in order to notify the user later
-			TArray<UObject*> UnconsolidatableObjects;
+		GWarn->BeginSlowTask(NSLOCTEXT("UnrealEd", "ConsolidateAssetsUpdate_Consolidating", "Consolidating Assets..."), true);
+		// Keep track of which objects, if any, cannot be consolidated, in order to notify the user later
+		TArray<UObject*> UnconsolidatableObjects;
 
-			// Keep track of objects which became partially consolidated but couldn't be deleted for some reason;
-			// these are critical failures, and the user needs to be alerted
-			TArray<UObject*> CriticalFailureObjects;
+		// Keep track of objects which became partially consolidated but couldn't be deleted for some reason;
+		// these are critical failures, and the user needs to be alerted
+		TArray<UObject*> CriticalFailureObjects;
 
-			// Keep track of which packages the consolidate operation has dirtied so the user can be alerted to them
-			// during a critical failure
-			TArray<UPackage*> DirtiedPackages;
+		// Keep track of which packages the consolidate operation has dirtied so the user can be alerted to them
+		// during a critical failure
+		TArray<UPackage*> DirtiedPackages;
+
+		// Keep track of root set objects so the user can be prompted about stripping the flag from them
+		TSet<UObject*> RootSetObjects;
 
-			// Keep track of root set objects so the user can be prompted about stripping the flag from them
-			TSet<UObject*> RootSetObjects;
+		// List of objects successfully deleted
+		TArray<UObject*> ConsolidatedObjects;
 
-			// List of objects successfully deleted
-			TArray<UObject*> ConsolidatedObjects;
+		// A list of names for object redirectors created during the delete process
+		// This is needed because the redirectors may not have the same name as the
+		// objects they are replacing until the objects are garbage collected
+		TMap<UObjectRedirector*, FName> RedirectorToObjectNameMap;
 
-			// A list of names for object redirectors created during the delete process
-			// This is needed because the redirectors may not have the same name as the
-			// objects they are replacing until the objects are garbage collected
-			TMap<UObjectRedirector*, FName> RedirectorToObjectNameMap;
+		// Temporaries used after ReloadEditorWorldForReferenceReplacementIfNecessary
+		TArray<FReplaceRequest, TInlineAllocator<1>> PostReloadRequests;
+		PostReloadRequests.Reserve(Requests.Num());
+		TArray<TArray<UObject*>> PostReloadUpdatedOld;
 
+		{
+			// Note reloading the world via ReloadEditorWorldForReferenceReplacementIfNecessary will cause a garbage collect and potentially cause entries in the ObjectsToConsolidate list to become invalid
+			// We refresh the list here after reloading the editor world
+			TArray< TWeakObjectPtr<UObject> > ObjectsToConsolidateWeakList;
+			ObjectsToConsolidateWeakList.Reserve(Requests.Num());
+
+			for (FReplaceRequest Request : Requests)
 			{
-				// Note reloading the world via ReloadEditorWorldForReferenceReplacementIfNecessary will cause a garbage collect and potentially cause entries in the ObjectsToConsolidate list to become invalid
-				// We refresh the list here after reloading the editor world
-				TArray< TWeakObjectPtr<UObject> > ObjectsToConsolidateWeakList;
-				ObjectsToConsolidateWeakList.Reserve(ObjectsToConsolidate.Num());
-				for(UObject* Object : ObjectsToConsolidate)
+				for (UObject* Old : Request.Old)
 				{
-					ObjectsToConsolidateWeakList.Add(Object);
+					ObjectsToConsolidateWeakList.Add(Old);
 				}
+			}
 
-				ObjectsToConsolidate.Reset();
+			// If the current editor world is in this list, transition to a new map and reload the world to finish the delete
+			ReloadEditorWorldForReferenceReplacementIfNecessary(ObjectsToConsolidateWeakList);
 
-				// If the current editor world is in this list, transition to a new map and reload the world to finish the delete
-				ReloadEditorWorldForReferenceReplacementIfNecessary(ObjectsToConsolidateWeakList);
+			// Make new PostReloadRequests where all Old objects are valid
+			int32 WeakIndex = 0;
+			for (FReplaceRequest Request : Requests)
+			{
+				int32 NumValid = 0;
+				for (int32 OldIndex = 0; OldIndex < Request.Old.Num(); ++OldIndex)
+				{
+					NumValid += ObjectsToConsolidateWeakList[WeakIndex + OldIndex].IsValid();
+				}
 
-				for(TWeakObjectPtr<UObject> WeakObject : ObjectsToConsolidateWeakList)
+				if (NumValid == Request.Old.Num())
+				{
+					PostReloadRequests.Add(Request);
+				}
+				else if (NumValid > 0)
 				{
-					if( WeakObject.IsValid() )
+					TArray<UObject*>& UpdatedOld = PostReloadUpdatedOld.AddDefaulted_GetRef();
+					for (int32 OldIndex = 0; OldIndex < Request.Old.Num(); ++OldIndex)
 					{
-						ObjectsToConsolidate.Add(WeakObject.Get());
+						if (UObject* ValidOld = ObjectsToConsolidateWeakList[WeakIndex + OldIndex].Get())
+						{
+							checkf(ValidOld == Request.Old[OldIndex], TEXT("Indexing bug?"));
+							UpdatedOld.Add(ValidOld);
+						}
 					}
+					checkf(NumValid == UpdatedOld.Num(), TEXT("Weak pointer validity changed"));
+
+					PostReloadRequests.Add({Request.New, MakeArrayView(UpdatedOld)});
 				}
+
+				WeakIndex += Request.Old.Num();
 			}
+			checkf(WeakIndex == ObjectsToConsolidateWeakList.Num(), TEXT("Tested wrong number of weak pointers"));
+
+			Requests = PostReloadRequests;
+		}
 
-			FForceReplaceInfo ReplaceInfo;
-			FForceReplaceInfo GeneratedClassReplaceInfo;
+		FForceReplaceInfo ReplaceInfo;
+
+		bool bNeedsGarbageCollection = false;
 
-			bool bNeedsGarbageCollection = false;
+		// Scope the reregister context below to complete after object deletion and before garbage collection
+		{
+			// Replacing references inside already loaded objects could cause rendering issues, so globally detach all components from their scenes for now
+			FGlobalComponentRecreateRenderStateContext ReregisterContext;
 
-			// Scope the reregister context below to complete after object deletion and before garbage collection
+			for (FReplaceRequest Request : Requests)
 			{
-				// Replacing references inside already loaded objects could cause rendering issues, so globally detach all components from their scenes for now
-				FGlobalComponentRecreateRenderStateContext ReregisterContext;
+				UObject* ObjectToConsolidateTo = Request.New;
+				TArrayView<UObject*> ObjectsToConsolidate = Request.Old;
 
 				// First, make sure that the class we're consolidating to has its hierarchy fixed so
 				// that we don't create a cycle (e.g. directly or indirectly parent it to itself):
 				UClass* ClassToConsolidateTo = nullptr;
-				if (ObjectToConsolidateTo)
+				if (UBlueprint* BlueprintObject = Cast<UBlueprint>(ObjectToConsolidateTo))
 				{
-					if (UBlueprint* BlueprintObject = Cast<UBlueprint>(ObjectToConsolidateTo))
+					ClassToConsolidateTo = BlueprintObject->GeneratedClass;
+					if (!ClassToConsolidateTo)
 					{
-						ClassToConsolidateTo = BlueprintObject->GeneratedClass;
-
-						if (!ClassToConsolidateTo)
-						{
-							ClassToConsolidateTo = UObject::StaticClass();
-						}
-
-						// Don't parent a blueprint to itself, instead fall back to the part of the
-						// hierarchy that is not being consolidated. Worst case, fall back to
-						// UObject::StaticClass():
-						UClass* NewParent = BlueprintObject->ParentClass;
-						UClass* OldParent = BlueprintObject->ParentClass;
-						UClass* ParentIter = NewParent;
-						while (ParentIter)
-						{
-							if (ObjectsToConsolidate.Contains(ParentIter->ClassGeneratedBy))
-							{
-								NewParent = ParentIter->GetSuperClass();
-							}
-							ParentIter = ParentIter->GetSuperClass();
-						}
+						ClassToConsolidateTo = UObject::StaticClass();
+					}
 
-						if (!NewParent || ObjectsToConsolidate.Contains(NewParent->ClassGeneratedBy))
+					// Don't parent a blueprint to itself, instead fall back to the part of the
+					// hierarchy that is not being consolidated. Worst case, fall back to
+					// UObject::StaticClass():
+					UClass* NewParent = BlueprintObject->ParentClass;
+					UClass* OldParent = BlueprintObject->ParentClass;
+					UClass* ParentIter = NewParent;
+					while (ParentIter)
+					{
+						if (ObjectsToConsolidate.Contains(ParentIter->ClassGeneratedBy))
 						{
-							NewParent = UObject::StaticClass();
+							NewParent = ParentIter->GetSuperClass();
 						}
+						ParentIter = ParentIter->GetSuperClass();
+					}
 
-						if (OldParent != NewParent)
-						{
-							BlueprintObject->ParentClass = NewParent;
+					if (!NewParent || ObjectsToConsolidate.Contains(NewParent->ClassGeneratedBy))
+					{
+						NewParent = UObject::StaticClass();
+					}
 
-							// Recompile the child blueprint to fix up the generated class
-							FKismetEditorUtilities::CompileBlueprint(BlueprintObject, EBlueprintCompileOptions::SkipGarbageCollection);
-						}
+					if (OldParent != NewParent)
+					{
+						BlueprintObject->ParentClass = NewParent;
+						FKismetEditorUtilities::CompileBlueprint(BlueprintObject, EBlueprintCompileOptions::SkipGarbageCollection);
 					}
 				}
-
+			
 				// Then reparent any direct children to the class we're consolidating to:
 				for (UObject* Object : ObjectsToConsolidate)
 				{
@@ -1445,8 +1494,6 @@ namespace ObjectTools
 										}
 
 										ChildBlueprint->ParentClass = NewParent;
-
-										// Recompile the child blueprint to fix up the generated class
 										FKismetEditorUtilities::CompileBlueprint(ChildBlueprint, EBlueprintCompileOptions::SkipGarbageCollection);
 
 										// Defer garbage collection until after we're done processing the list of objects
@@ -1457,10 +1504,16 @@ namespace ObjectTools
 						}
 					}
 				}
+			
+			}
 
-				ForceReplaceReferences(ObjectToConsolidateTo, ObjectsToConsolidate, ObjectsToConsolidateWithin, ReplaceInfo, bWarnAboutRootSet);
+			
+			ForceReplaceReferences(Requests, ObjectsToConsolidateWithin, ReplaceInfo, bWarnAboutRootSet);
 
-				if (UBlueprint* ObjectToConsolidateTo_BP = Cast<UBlueprint>(ObjectToConsolidateTo))
+			for (FReplaceRequest Request : Requests)
+			{
+				TArrayView<UObject*> ObjectsToConsolidate = Request.Old;
+				if (UBlueprint* ObjectToConsolidateTo_BP = Cast<UBlueprint>(Request.New))
 				{
 					// Replace all UClass/TSubClassOf properties of generated class.
 					TArray<UObject*> ObjectsToConsolidate_BP;
@@ -1481,7 +1534,8 @@ namespace ObjectTools
 							OldChildClassToOldParentClass.Add(OldChildClass, OldGeneratedClass);
 						}
 					}
-
+					
+					FForceReplaceInfo GeneratedClassReplaceInfo;
 					ForceReplaceReferences(ObjectToConsolidateTo_BP->GeneratedClass, ObjectsToConsolidate_BP, ObjectsToConsolidateWithin, GeneratedClassReplaceInfo, bWarnAboutRootSet);
 
 					// Repair the references of GeneratedClass on the object being consolidated so they can be properly disposed of upon deletion.
@@ -1506,12 +1560,16 @@ namespace ObjectTools
 						ObjectToConsolidateTo_BP->CachedDependents.Add(DependentBP);
 					}
 				}
-				DirtiedPackages.Append( ReplaceInfo.DirtiedPackages );
-				UnconsolidatableObjects.Append( ReplaceInfo.UnreplaceableObjects );
 			}
 
+			DirtiedPackages.Append( ReplaceInfo.DirtiedPackages );
+			UnconsolidatableObjects.Append( ReplaceInfo.UnreplaceableObjects );
+		}
+
+		for (FReplaceRequest Request : Requests)
+		{
 			// See if this is a blueprint consolidate and replace instances of the generated class
-			UBlueprint* BlueprintToConsolidateTo = Cast<UBlueprint>(ObjectToConsolidateTo);
+			UBlueprint* BlueprintToConsolidateTo = Cast<UBlueprint>(Request.New);
 			if (BlueprintToConsolidateTo != NULL && BlueprintToConsolidateTo->GeneratedClass)
 			{
 				for ( TArray<UObject*>::TConstIterator ConsolIter( ReplaceInfo.ReplaceableObjects ); ConsolIter; ++ConsolIter )
@@ -1539,196 +1597,222 @@ namespace ObjectTools
 
 				bNeedsGarbageCollection = true;
 			}
+		}
 
-			if (bNeedsGarbageCollection)
-			{
-				CollectGarbage(GARBAGE_COLLECTION_KEEPFLAGS);
-			}
+		if (bNeedsGarbageCollection)
+		{
+			CollectGarbage(GARBAGE_COLLECTION_KEEPFLAGS);
+		}
 
-			FEditorDelegates::OnAssetsPreDelete.Broadcast(ReplaceInfo.ReplaceableObjects);
+		FEditorDelegates::OnAssetsPreDelete.Broadcast(ReplaceInfo.ReplaceableObjects);
 
-			TSet<FString> AlreadyMappedObjectPaths;
+		TSet<FString> AlreadyMappedObjectPaths;
 
-			if (bShouldDeleteAfterConsolidate)
+		if (bShouldDeleteAfterConsolidate)
+		{
+			// Bit wasteful, rebuild same map that existed temporarily inside ForceReplaceReferences 
+			TMap<UObject*, UObject*> ReplacementMap;
+			ReplacementMap.Reserve(ReplaceInfo.ReplaceableObjects.Num());
+			for (FReplaceRequest Request : Requests)
 			{
-				// With all references to the objects to consolidate to eliminated from objects that are currently loaded, it should now be safe to delete
-				// the objects to be consolidated themselves, leaving behind a redirector in their place to fix up objects that were not currently loaded at the time
-				// of this operation.
-				for ( TArray<UObject*>::TConstIterator ConsolIter( ReplaceInfo.ReplaceableObjects ); ConsolIter; ++ConsolIter )
+				for (UObject* Old : Request.Old)
 				{
-					GWarn->StatusUpdate( ConsolIter.GetIndex(), ReplaceInfo.ReplaceableObjects.Num(), NSLOCTEXT("UnrealEd", "ConsolidateAssetsUpdate_DeletingObjects", "Deleting Assets...") );
+					ReplacementMap.Add(Old, Request.New);
+				}
+			}
 
-					UObject* CurObjToConsolidate = *ConsolIter;
-					UObject* CurObjOuter = CurObjToConsolidate->GetOuter();
-					UPackage* CurObjPackage = CurObjToConsolidate->GetOutermost();
-					const FName CurObjName = CurObjToConsolidate->GetFName();
-					const FString CurObjPath = CurObjToConsolidate->GetPathName();
-					UBlueprint* BlueprintToConsolidate = Cast<UBlueprint>(CurObjToConsolidate);
+			// With all references to the objects to consolidate to eliminated from objects that are currently loaded, it should now be safe to delete
+			// the objects to be consolidated themselves, leaving behind a redirector in their place to fix up objects that were not currently loaded at the time
+			// of this operation.
+			for ( TArray<UObject*>::TConstIterator ConsolIter( ReplaceInfo.ReplaceableObjects ); ConsolIter; ++ConsolIter )
+			{
+				GWarn->StatusUpdate( ConsolIter.GetIndex(), ReplaceInfo.ReplaceableObjects.Num(), NSLOCTEXT("UnrealEd", "ConsolidateAssetsUpdate_DeletingObjects", "Deleting Assets...") );
 
-					// Attempt to delete the object that was consolidated
-					if ( DeleteSingleObject( CurObjToConsolidate ) )
-					{
-						// DONT GC YET!!! we still need these objects around to notify other tools that they are gone and to create redirectors
-						ConsolidatedObjects.Add(CurObjToConsolidate);
+				UObject* CurObjToConsolidate = *ConsolIter;
+				UObject* CurObjOuter = CurObjToConsolidate->GetOuter();
+				UPackage* CurObjPackage = CurObjToConsolidate->GetOutermost();
+				const FName CurObjName = CurObjToConsolidate->GetFName();
+				const FString CurObjPath = CurObjToConsolidate->GetPathName();
+				UBlueprint* BlueprintToConsolidate = Cast<UBlueprint>(CurObjToConsolidate);
 
-						if ( AlreadyMappedObjectPaths.Contains(CurObjPath) )
-						{
-							continue;
-						}
+				// Attempt to delete the object that was consolidated
+				if ( DeleteSingleObject( CurObjToConsolidate ) )
+				{
+					// DONT GC YET!!! we still need these objects around to notify other tools that they are gone and to create redirectors
+					ConsolidatedObjects.Add(CurObjToConsolidate);
 
-						// Create a redirector with a unique name
-						// It will have the same name as the object that was consolidated after the garbage collect
-						UObjectRedirector* Redirector = NewObject<UObjectRedirector>(CurObjOuter, NAME_None, RF_Standalone | RF_Public);
-						check( Redirector );
+					if ( AlreadyMappedObjectPaths.Contains(CurObjPath) )
+					{
+						continue;
+					}
 
-						// Set the redirector to redirect to the object to consolidate to
-						Redirector->DestinationObject = ObjectToConsolidateTo;
+					UObject* ObjectToConsolidateTo = ReplacementMap.FindChecked(CurObjToConsolidate);
+					UBlueprint* BlueprintToConsolidateTo = Cast<UBlueprint>(ObjectToConsolidateTo);
 
-						// Keep track of the object name so we can rename the redirector later
-						RedirectorToObjectNameMap.Add(Redirector, CurObjName);
-						AlreadyMappedObjectPaths.Add(CurObjPath);
+					// Create a redirector with a unique name
+					// It will have the same name as the object that was consolidated after the garbage collect
+					UObjectRedirector* Redirector = NewObject<UObjectRedirector>(CurObjOuter, NAME_None, RF_Standalone | RF_Public);
+					check( Redirector );
 
-						// If consolidating blueprints, make sure redirectors are created for the consolidated blueprint class and CDO
-						if ( BlueprintToConsolidateTo != NULL && BlueprintToConsolidate != NULL )
-						{
-							// One redirector for the class
-							UObjectRedirector* ClassRedirector = NewObject<UObjectRedirector>(CurObjOuter, NAME_None, RF_Standalone | RF_Public);
-							check( ClassRedirector );
-							ClassRedirector->DestinationObject = BlueprintToConsolidateTo->GeneratedClass;
-							RedirectorToObjectNameMap.Add(ClassRedirector, BlueprintToConsolidate->GeneratedClass->GetFName());
-							AlreadyMappedObjectPaths.Add(BlueprintToConsolidate->GeneratedClass->GetPathName());
-
-							// One redirector for the CDO
-							UObjectRedirector* CDORedirector = NewObject<UObjectRedirector>(CurObjOuter, NAME_None, RF_Standalone | RF_Public);
-							check( CDORedirector );
-							CDORedirector->DestinationObject = BlueprintToConsolidateTo->GeneratedClass->GetDefaultObject();
-							RedirectorToObjectNameMap.Add(CDORedirector, BlueprintToConsolidate->GeneratedClass->GetDefaultObject()->GetFName());
-							AlreadyMappedObjectPaths.Add(BlueprintToConsolidate->GeneratedClass->GetDefaultObject()->GetPathName());
-						}
+					// Set the redirector to redirect to the object to consolidate to
+					Redirector->DestinationObject = ObjectToConsolidateTo;
 
-						DirtiedPackages.AddUnique( CurObjPackage );
-					}
-					// If the object couldn't be deleted, store it in the array that will be used to show the user which objects had errors
-					else
+					// Keep track of the object name so we can rename the redirector later
+					RedirectorToObjectNameMap.Add(Redirector, CurObjName);
+					AlreadyMappedObjectPaths.Add(CurObjPath);
+
+					// If consolidating blueprints, make sure redirectors are created for the consolidated blueprint class and CDO
+					if ( BlueprintToConsolidateTo != NULL && BlueprintToConsolidate != NULL )
 					{
-						CriticalFailureObjects.Add( CurObjToConsolidate );
+						// One redirector for the class
+						UObjectRedirector* ClassRedirector = NewObject<UObjectRedirector>(CurObjOuter, NAME_None, RF_Standalone | RF_Public);
+						check( ClassRedirector );
+						ClassRedirector->DestinationObject = BlueprintToConsolidateTo->GeneratedClass;
+						RedirectorToObjectNameMap.Add(ClassRedirector, BlueprintToConsolidate->GeneratedClass->GetFName());
+						AlreadyMappedObjectPaths.Add(BlueprintToConsolidate->GeneratedClass->GetPathName());
+
+						// One redirector for the CDO
+						UObjectRedirector* CDORedirector = NewObject<UObjectRedirector>(CurObjOuter, NAME_None, RF_Standalone | RF_Public);
+						check( CDORedirector );
+						CDORedirector->DestinationObject = BlueprintToConsolidateTo->GeneratedClass->GetDefaultObject();
+						RedirectorToObjectNameMap.Add(CDORedirector, BlueprintToConsolidate->GeneratedClass->GetDefaultObject()->GetFName());
+						AlreadyMappedObjectPaths.Add(BlueprintToConsolidate->GeneratedClass->GetDefaultObject()->GetPathName());
 					}
-				}
 
-				// Prevent newly created redirectors from being GC'ed before we can rename them
-				TArray<TStrongObjectPtr<UObjectRedirector>> Redirectors;
-				Redirectors.Reserve(RedirectorToObjectNameMap.Num());
-				for (TMap<UObjectRedirector*, FName>::TIterator RedirectIt(RedirectorToObjectNameMap); RedirectIt; ++RedirectIt)
-				{
-					UObjectRedirector* Redirector = RedirectIt.Key();
-					Redirectors.Add(TStrongObjectPtr<UObjectRedirector>(Redirector));
+					DirtiedPackages.AddUnique( CurObjPackage );
 				}
-
-				TArray<UPackage*> PotentialPackagesToDelete;
-				for ( int32 ObjIdx = 0; ObjIdx < ConsolidatedObjects.Num(); ++ObjIdx )
+				// If the object couldn't be deleted, store it in the array that will be used to show the user which objects had errors
+				else
 				{
-					PotentialPackagesToDelete.AddUnique(ConsolidatedObjects[ObjIdx]->GetOutermost());
+					CriticalFailureObjects.Add( CurObjToConsolidate );
 				}
+			}
 
-				CleanupAfterSuccessfulDelete(PotentialPackagesToDelete);
+			// Prevent newly created redirectors from being GC'ed before we can rename them
+			TArray<TStrongObjectPtr<UObjectRedirector>> Redirectors;
+			Redirectors.Reserve(RedirectorToObjectNameMap.Num());
+			for (TMap<UObjectRedirector*, FName>::TIterator RedirectIt(RedirectorToObjectNameMap); RedirectIt; ++RedirectIt)
+			{
+				UObjectRedirector* Redirector = RedirectIt.Key();
+				Redirectors.Add(TStrongObjectPtr<UObjectRedirector>(Redirector));
+			}
 
-				// Now that the old objects have been garbage collected, give the redirectors a proper name
-				for (TMap<UObjectRedirector*, FName>::TIterator RedirectIt(RedirectorToObjectNameMap); RedirectIt; ++RedirectIt)
-				{
-					UObjectRedirector* Redirector = RedirectIt.Key();
-					const FName ObjName = RedirectIt.Value();
+			TArray<UPackage*> PotentialPackagesToDelete;
+			for ( int32 ObjIdx = 0; ObjIdx < ConsolidatedObjects.Num(); ++ObjIdx )
+			{
+				PotentialPackagesToDelete.AddUnique(ConsolidatedObjects[ObjIdx]->GetOutermost());
+			}
 
-					if ( Redirector->Rename(*ObjName.ToString(), NULL, REN_Test) )
-					{
-						Redirector->Rename(*ObjName.ToString(), NULL, REN_DontCreateRedirectors | REN_ForceNoResetLoaders | REN_NonTransactional);
-						FAssetRegistryModule::AssetCreated(Redirector);
-					}
-					else
-					{
-						// Could not rename the redirector back to the original object's name. This indicates the original
-						// object could not be garbage collected even though DeleteSingleObject returned true.
-						CriticalFailureObjects.AddUnique(Redirector);
-					}
-				}
+			CleanupAfterSuccessfulDelete(PotentialPackagesToDelete);
 
-				Redirectors.Empty();
+			// Now that the old objects have been garbage collected, give the redirectors a proper name
+			for (TMap<UObjectRedirector*, FName>::TIterator RedirectIt(RedirectorToObjectNameMap); RedirectIt; ++RedirectIt)
+			{
+				UObjectRedirector* Redirector = RedirectIt.Key();
+				const FName ObjName = RedirectIt.Value();
 
+				if ( Redirector->Rename(*ObjName.ToString(), NULL, REN_Test) )
+				{
+					Redirector->Rename(*ObjName.ToString(), NULL, REN_DontCreateRedirectors | REN_ForceNoResetLoaders | REN_NonTransactional);
+					FAssetRegistryModule::AssetCreated(Redirector);
+				}
+				else
+				{
+					// Could not rename the redirector back to the original object's name. This indicates the original
+					// object could not be garbage collected even though DeleteSingleObject returned true.
+					CriticalFailureObjects.AddUnique(Redirector);
+				}
 			}
 
-			// Empty the provided array so it's not full of pointers to deleted objects
-			ObjectsToConsolidate.Empty();
-			ConsolidatedObjects.Empty();
+			Redirectors.Empty();
 
-			GWarn->EndSlowTask();
+		}
+
+		ConsolidatedObjects.Empty();
+
+		GWarn->EndSlowTask();
 
-			ConsolidationResults.DirtiedPackages = DirtiedPackages;
-			ConsolidationResults.FailedConsolidationObjs = CriticalFailureObjects;
-			ConsolidationResults.InvalidConsolidationObjs = UnconsolidatableObjects;
+		ConsolidationResults.DirtiedPackages = DirtiedPackages;
+		ConsolidationResults.FailedConsolidationObjs = CriticalFailureObjects;
+		ConsolidationResults.InvalidConsolidationObjs = UnconsolidatableObjects;
 
-			// If some objects failed to consolidate, notify the user of the failed objects
-			if ( UnconsolidatableObjects.Num() > 0 )
+		// If some objects failed to consolidate, notify the user of the failed objects
+		if ( UnconsolidatableObjects.Num() > 0 )
+		{
+			FString FailedObjectNames;
+			for ( TArray<UObject*>::TConstIterator FailedIter( UnconsolidatableObjects ); FailedIter; ++FailedIter )
 			{
-				FString FailedObjectNames;
-				for ( TArray<UObject*>::TConstIterator FailedIter( UnconsolidatableObjects ); FailedIter; ++FailedIter )
-				{
-					UObject* CurFailedObject = *FailedIter;
-					FailedObjectNames += CurFailedObject->GetName() + TEXT("\n");
-				}
+				UObject* CurFailedObject = *FailedIter;
+				FailedObjectNames += CurFailedObject->GetName() + TEXT("\n");
+			}
 
-				FFormatNamedArguments Arguments;
-				Arguments.Add(TEXT("Objects"), FText::FromString( FailedObjectNames ));
-				FText MessageFormatting = NSLOCTEXT("ObjectTools", "ConsolidateAssetsFailureDlgMFormattings", "The assets below were unable to be consolidated. This is likely because they are referenced by the object to consolidate to.\n\n{Objects}");
-				FText Message = FText::Format( MessageFormatting, Arguments );
-				FText Title = NSLOCTEXT("ObjectTools", "ConsolidateAssetsFailureDlg_Title", "Failed to Consolidate Assets");
+			FFormatNamedArguments Arguments;
+			Arguments.Add(TEXT("Objects"), FText::FromString( FailedObjectNames ));
+			FText MessageFormatting = NSLOCTEXT("ObjectTools", "ConsolidateAssetsFailureDlgMFormattings", "The assets below were unable to be consolidated. This is likely because they are referenced by the object to consolidate to.\n\n{Objects}");
+			FText Message = FText::Format( MessageFormatting, Arguments );
+			FText Title = NSLOCTEXT("ObjectTools", "ConsolidateAssetsFailureDlg_Title", "Failed to Consolidate Assets");
 
-				if (bShouldShowDialogs)
-				{
-					FMessageDialog::Open(EAppMsgType::Ok, Message, &Title);
-				}
-				else
-				{
-					UE_LOG(LogObjectTools, Warning, TEXT("Failed to consolidate assets: %s"), *Message.ToString());
-				}
+			if (bShouldShowDialogs)
+			{
+				FMessageDialog::Open(EAppMsgType::Ok, Message, &Title);
 			}
+			else
+			{
+				UE_LOG(LogObjectTools, Warning, TEXT("Failed to consolidate assets: %s"), *Message.ToString());
+			}
+		}
 
-			// Alert the user to critical object failure
-			if ( CriticalFailureObjects.Num() > 0 )
+		// Alert the user to critical object failure
+		if ( CriticalFailureObjects.Num() > 0 )
+		{
+			FString CriticalFailedObjectNames;
+			for ( TArray<UObject*>::TConstIterator FailedIter( CriticalFailureObjects ); FailedIter; ++FailedIter )
 			{
-				FString CriticalFailedObjectNames;
-				for ( TArray<UObject*>::TConstIterator FailedIter( CriticalFailureObjects ); FailedIter; ++FailedIter )
-				{
-					const UObject* CurFailedObject = *FailedIter;
-					CriticalFailedObjectNames += CurFailedObject->GetName() + TEXT("\n");
-				}
+				const UObject* CurFailedObject = *FailedIter;
+				CriticalFailedObjectNames += CurFailedObject->GetName() + TEXT("\n");
+			}
 
-				FString DirtiedPackageNames;
-				for ( TArray<UPackage*>::TConstIterator DirtyPkgIter( DirtiedPackages ); DirtyPkgIter; ++DirtyPkgIter )
-				{
-					const UPackage* CurDirtyPkg = *DirtyPkgIter;
-					DirtiedPackageNames += CurDirtyPkg->GetName() + TEXT("\n");
-				}
+			FString DirtiedPackageNames;
+			for ( TArray<UPackage*>::TConstIterator DirtyPkgIter( DirtiedPackages ); DirtyPkgIter; ++DirtyPkgIter )
+			{
+				const UPackage* CurDirtyPkg = *DirtyPkgIter;
+				DirtiedPackageNames += CurDirtyPkg->GetName() + TEXT("\n");
+			}
 
-				FFormatNamedArguments Arguments;
-				Arguments.Add(TEXT("Assets"), FText::FromString( CriticalFailedObjectNames ));
-				Arguments.Add(TEXT("Packages"), FText::FromString( DirtiedPackageNames ));
-				FText MessageFormatting = NSLOCTEXT("ObjectTools", "ConsolidateAssetsCriticalFailureDlgMsgFormatting", "CRITICAL FAILURE:\nOne or more assets were partially consolidated, yet still cannot be deleted for some reason. It is highly recommended that you restart the editor without saving any of the assets or packages.\n\nAffected Assets:\n{Assets}\n\nPotentially Affected Packages:\n{Packages}");
-				FText Message = FText::Format( MessageFormatting, Arguments );
-				FText Title = NSLOCTEXT("ObjectTools", "ConsolidateAssetsCriticalFailureDlg_Title", "Critical Failure to Consolidate Assets");
+			FFormatNamedArguments Arguments;
+			Arguments.Add(TEXT("Assets"), FText::FromString( CriticalFailedObjectNames ));
+			Arguments.Add(TEXT("Packages"), FText::FromString( DirtiedPackageNames ));
+			FText MessageFormatting = NSLOCTEXT("ObjectTools", "ConsolidateAssetsCriticalFailureDlgMsgFormatting", "CRITICAL FAILURE:\nOne or more assets were partially consolidated, yet still cannot be deleted for some reason. It is highly recommended that you restart the editor without saving any of the assets or packages.\n\nAffected Assets:\n{Assets}\n\nPotentially Affected Packages:\n{Packages}");
+			FText Message = FText::Format( MessageFormatting, Arguments );
+			FText Title = NSLOCTEXT("ObjectTools", "ConsolidateAssetsCriticalFailureDlg_Title", "Critical Failure to Consolidate Assets");
 
-				if (bShouldShowDialogs)
-				{
-					FMessageDialog::Open(EAppMsgType::Ok, Message, &Title);
-				}
-				else
-				{
-					UE_LOG(LogObjectTools, Warning, TEXT("Failed to consolidate assets: %s"), *Message.ToString());
-				}
+			if (bShouldShowDialogs)
+			{
+				FMessageDialog::Open(EAppMsgType::Ok, Message, &Title);
+			}
+			else
+			{
+				UE_LOG(LogObjectTools, Warning, TEXT("Failed to consolidate assets: %s"), *Message.ToString());
 			}
 		}
 
 		return ConsolidationResults;
 	}
 
+	FConsolidationResults ConsolidateObjects(UObject* ObjectToConsolidateTo, TArray<UObject*>& ObjectsToConsolidate, TSet<UObject*>& ObjectsToConsolidateWithin, TSet<UObject*>& ObjectsToNotConsolidateWithin, bool bShouldDeleteAfterConsolidate, bool bWarnAboutRootSet)
+	{
+		if (!ObjectToConsolidateTo)
+		{
+			return FConsolidationResults();
+		}
+
+		// Empty the provided array so it's not full of pointers to deleted objects
+		ON_SCOPE_EXIT{ ObjectsToConsolidate.Empty(); };
+
+		FReplaceRequest Request = {ObjectToConsolidateTo, ObjectsToConsolidate};
+		return ConsolidateObjects(MakeArrayView(&Request, 1), ObjectsToConsolidateWithin, ObjectsToNotConsolidateWithin, bShouldDeleteAfterConsolidate, bWarnAboutRootSet);
+	}
+
 	FConsolidationResults ConsolidateObjects(UObject* ObjectToConsolidateTo, TArray<UObject*>& ObjectsToConsolidate, bool bShowDeleteConfirmation)
 	{
 		FConsolidationResults ConsolidationResults;
diff --git a/Engine/Source/Editor/UnrealEd/Public/ObjectTools.h b/Engine/Source/Editor/UnrealEd/Public/ObjectTools.h
--- a/Engine/Source/Editor/UnrealEd/Public/ObjectTools.h
+++ b/Engine/Source/Editor/UnrealEd/Public/ObjectTools.h
@@ -190,6 +190,13 @@ namespace ObjectTools
 		TArray<UObject*>	FailedConsolidationObjs;
 	};
 
+	/** Helper struct for batch replacements where Old references get replaced with New */
+	struct FReplaceRequest
+	{
+		UObject* New = nullptr;
+		TArrayView<UObject*> Old; 
+	};
+
 	/**
 	 * Consolidates objects by replacing all references/uses of the provided "objects to consolidate" with references to the "object to consolidate to." This is
 	 * useful for situations such as when a particular asset is duplicated in multiple places and it would be handy to allow all uses to point to one particular copy
@@ -198,14 +205,18 @@ namespace ObjectTools
 	 *
 	 * @param	ObjectToConsolidateTo	Object to which all references of the "objects to consolidate" will instead refer to after this operation completes
 	 * @param	ObjectsToConsolidate	Objects which all references of which will be replaced with references to the "object to consolidate to"; each will also be deleted
+	 * @param	Requests				Batch of consolidations. All objects consilidated to, i.e. FReplaceRequest::New, must be non-null.
 	 *
 	 * @note	This function performs NO type checking, by design. It is potentially dangerous to replace references of one type with another, so utilize caution.
 	 * @note	The "objects to consolidate" are DELETED by this function.
 	 *
 	 * @return	Structure of consolidation results, specifying which packages were dirtied, which objects failed consolidation (if any), etc.
 	 */
-	UNREALED_API FConsolidationResults ConsolidateObjects( UObject* ObjectToConsolidateTo, TArray<UObject*>& ObjectsToConsolidate, bool bShowDeleteConfirmation = true );
+	UNREALED_API FConsolidationResults ConsolidateObjects(UObject* ObjectToConsolidateTo, TArray<UObject*>& ObjectsToConsolidate, bool bShowDeleteConfirmation = true );
 	UNREALED_API FConsolidationResults ConsolidateObjects(UObject* ObjectToConsolidateTo, TArray<UObject*>& ObjectsToConsolidate, TSet<UObject*>& ObjectsToConsolidateWithin, TSet<UObject*>& ObjectsToNotConsolidateWithin, bool bShouldDeleteAfterConsolidate, bool bWarnAboutRootSet = true);
+	UNREALED_API FConsolidationResults ConsolidateObjects(TArrayView<FReplaceRequest> Requests, TSet<UObject*>& ObjectsToConsolidateWithin, TSet<UObject*>& ObjectsToNotConsolidateWithin, bool bShouldDeleteAfterConsolidate, bool bWarnAboutRootSet = true);
+
+	
 	UNREALED_API void CompileBlueprintsAfterRefUpdate(const TArray<UObject*>& ObjectsConsolidatedWithin);
 	/**
 	 * Copies references for selected generic browser objects to the clipboard.
@@ -329,9 +340,11 @@ namespace ObjectTools
 	 *
 	 * @param ObjectToReplaceWith	Any references found to 'ObjectsToReplace' will be replaced with this object.  If the object is nullptr references will be nulled.
 	 * @param ObjectsToReplace		An array of objects that should be replaced with 'ObjectToReplaceWith'
+	 * @param Requests				Batch of replacements where all FReplaceRequest::Old are replaced with FReplaceRequest::New for each request
 	 */
 	UNREALED_API void ForceReplaceReferences(UObject* ObjectToReplaceWith, TArray<UObject*>& ObjectsToReplace);
 	UNREALED_API void ForceReplaceReferences(UObject* ObjectToReplaceWith, TArray<UObject*>& ObjectsToReplace, TSet<UObject*>& ObjectsToReplaceWithin);
+	UNREALED_API void ForceReplaceReferences(TArrayView<FReplaceRequest> Requests, TSet<UObject*>& ObjectsToReplaceWithin);
 
 	/**
 	 * Gathers additional objects to delete such as map built data
diff --git a/Engine/Source/Runtime/CoreUObject/Private/Serialization/FindReferencersArchive.cpp b/Engine/Source/Runtime/CoreUObject/Private/Serialization/FindReferencersArchive.cpp
--- a/Engine/Source/Runtime/CoreUObject/Private/Serialization/FindReferencersArchive.cpp
+++ b/Engine/Source/Runtime/CoreUObject/Private/Serialization/FindReferencersArchive.cpp
@@ -14,7 +14,7 @@
  * @param	InTargetObjects			array of objects to search for references to
  * @param	bFindAlsoWeakReferences should we also look into weak references?
  */
-FFindReferencersArchive::FFindReferencersArchive(UObject* InPotentialReferencer, const TArray<UObject*>& InTargetObjects, bool bFindAlsoWeakReferences)
+FFindReferencersArchive::FFindReferencersArchive(UObject* InPotentialReferencer, TArrayView<UObject*> InTargetObjects, bool bFindAlsoWeakReferences)
 {
 	// use the optimized RefLink to skip over properties which don't contain object references
 	ArIsObjectReferenceCollector = true;
@@ -137,9 +137,14 @@ int32 FFindReferencersArchive::GetReferenceCount( UObject* TargetObject, TArray<
 int32 FFindReferencersArchive::GetReferenceCounts( TMap<class UObject*, int32>& out_ReferenceCounts ) const
 {
 	out_ReferenceCounts.Empty();
+	return AppendReferenceCounts(out_ReferenceCounts);
+}
+
+int32 FFindReferencersArchive::AppendReferenceCounts(TMap<class UObject*, int32>& out_ReferenceCounts) const
+{
 	for (int32 Index = 0; Index < TargetObjects.RefCountNum(); ++Index)
 	{
-		if ( TargetObjects.GetRefCount(Index) > 0 && TargetObjects.GetObject(Index) != PotentialReferencer )
+		if (TargetObjects.GetRefCount(Index) > 0 && TargetObjects.GetObject(Index) != PotentialReferencer)
 		{
 			out_ReferenceCounts.Add(TargetObjects.GetObject(Index), TargetObjects.GetRefCount(Index));
 		}
diff --git a/Engine/Source/Runtime/CoreUObject/Public/Serialization/FindReferencersArchive.h b/Engine/Source/Runtime/CoreUObject/Public/Serialization/FindReferencersArchive.h
--- a/Engine/Source/Runtime/CoreUObject/Public/Serialization/FindReferencersArchive.h
+++ b/Engine/Source/Runtime/CoreUObject/Public/Serialization/FindReferencersArchive.h
@@ -21,7 +21,7 @@ public:
 	 * @param	InTargetObjects			array of objects to search for references to
 	 * @param	bFindAlsoWeakReferences should we also look into weak references?
 	 */
-	COREUOBJECT_API FFindReferencersArchive(class UObject* PotentialReferencer, const TArray<class UObject*>& InTargetObjects, bool bFindAlsoWeakReferences = false);
+	COREUOBJECT_API FFindReferencersArchive(class UObject* PotentialReferencer, TArrayView<class UObject*> InTargetObjects, bool bFindAlsoWeakReferences = false);
 
 	/**
 	 * Retrieves the number of references from PotentialReferencer to the object specified.
@@ -44,6 +44,9 @@ public:
 	 */
 	COREUOBJECT_API int32 GetReferenceCounts( TMap<class UObject*, int32>& out_ReferenceCounts ) const;
 
+	/** GetReferenceCounts without empyting out map */
+	COREUOBJECT_API int32 AppendReferenceCounts( TMap<class UObject*, int32>& out_ReferenceCounts ) const;
+
 	/**
 	 * Retrieves the number of references from PotentialReferencer list of TargetObjects
 	 *
