Export material functions from DumpMaterialExpressionsComandlet

diff --git a/Engine/Source/Editor/UnrealEd/Private/Commandlets/DumpMaterialExpressionsCommandlet.cpp b/Engine/Source/Editor/UnrealEd/Private/Commandlets/DumpMaterialExpressionsCommandlet.cpp
--- a/Engine/Source/Editor/UnrealEd/Private/Commandlets/DumpMaterialExpressionsCommandlet.cpp
+++ b/Engine/Source/Editor/UnrealEd/Private/Commandlets/DumpMaterialExpressionsCommandlet.cpp
@@ -12,37 +12,52 @@
 #include "Materials/MaterialExpressionFunctionInput.h"
 #include "Materials/MaterialExpressionFunctionOutput.h"
 #include "Materials/MaterialExpressionComposite.h"
+#include "Materials/MaterialFunction.h"
 #include "UObject/UObjectIterator.h"
 #include "HAL/FileManager.h"
+#include "AssetRegistryModule.h"
 
 DEFINE_LOG_CATEGORY_STATIC(LogDumpMaterialExpressionsCommandlet, Log, All);
 
-UDumpMaterialExpressionsCommandlet::UDumpMaterialExpressionsCommandlet(const FObjectInitializer& ObjectInitializer)
-	: Super(ObjectInitializer)
+FString GetFormattedText(const FString& InText)
 {
-}
+	FString OutText = InText;
+	OutText = InText.Replace(TEXT("\n"), TEXT(" "));
+	if (OutText.IsEmpty())
+	{
+		OutText = TEXT("N/A");
+	}
+	return OutText;
+};
 
-int32 UDumpMaterialExpressionsCommandlet::Main(const FString& Params)
+FString GenerateSpacePadding(int32 MaxLen, int32 TextLen)
 {
-	TArray<FString> Tokens;
-	TArray<FString> Switches;
-	TMap<FString, FString> ParamVals;
-	UCommandlet::ParseCommandLine(*Params, Tokens, Switches, ParamVals);
+	FString Padding;
+	for (int i = 0; i < MaxLen - TextLen; ++i)
+	{
+		Padding += TEXT(" ");
+	}
+	return Padding;
+};
 
-	if (Switches.Contains("help"))
+void WriteLine(FArchive* FileWriter, const TArray<FString>& FieldNames, const TArray<uint32>& MaxFieldLengths)
+{
+	check(FieldNames.Num() >= MaxFieldLengths.Num());
+
+	FString OutputLine;
+	for (int32 i = 0; i < FieldNames.Num(); ++i)
 	{
-		UE_LOG(LogDumpMaterialExpressionsCommandlet, Log, TEXT("DumpMaterialExpressions"));
-		UE_LOG(LogDumpMaterialExpressionsCommandlet, Log, TEXT("This commandlet will dump to a plain text file an info table of all material expressions in the engine and the plugins enabled on the project."));
-		UE_LOG(LogDumpMaterialExpressionsCommandlet, Log, TEXT("The output fields include:"));
-		UE_LOG(LogDumpMaterialExpressionsCommandlet, Log, TEXT("Name - The class name of the material expression"));
-		UE_LOG(LogDumpMaterialExpressionsCommandlet, Log, TEXT("Type - ControlFlow | HLSLGenerator | CLASS_Deprecated"));
-		UE_LOG(LogDumpMaterialExpressionsCommandlet, Log, TEXT("ShowInCreateMenu - If the expression appears in the create node dropdown menu"));
-		UE_LOG(LogDumpMaterialExpressionsCommandlet, Log, TEXT("CreationName - The name displayed in the create node dropdown menu to add an expression"));
-		UE_LOG(LogDumpMaterialExpressionsCommandlet, Log, TEXT("CreationDescription - The tooltip displayed on the CreationName in the create node dropdown menu"));
-		UE_LOG(LogDumpMaterialExpressionsCommandlet, Log, TEXT("Caption - The caption displayed on the material expression node"));
-		UE_LOG(LogDumpMaterialExpressionsCommandlet, Log, TEXT("Tooltip - The tooltip displayed on the material expression node"));
-		return 0;
+		// The last field doesn't need space padding, it changes to a new line.
+		FString Padding = (i + 1 < FieldNames.Num()) ? GenerateSpacePadding(MaxFieldLengths[i], FieldNames[i].Len()) : TEXT("\n");
+		OutputLine += (FieldNames[i] + Padding);
 	}
+	FileWriter->Serialize(const_cast<ANSICHAR*>(StringCast<ANSICHAR>(*OutputLine).Get()), OutputLine.Len());
+};
+
+void WriteOutMaterialExpressions(FArchive* FileWriter, const FString& OutputFilePath)
+{
+	FString OutputLine = TEXT("[MATERIAL EXPRESSIONS]\n");
+	FileWriter->Serialize(const_cast<ANSICHAR*>(StringCast<ANSICHAR>(*OutputLine).Get()), OutputLine.Len());
 
 	struct FMaterialExpressionInfo
 	{
@@ -77,6 +92,7 @@ int32 UDumpMaterialExpressionsCommandlet::Main(const FString& Params)
 	int32 MaxCreationDescriptionLength = CreationDescriptionField.Len();
 	int32 MaxCaptionLength = CaptionField.Len();
 	int32 MaxDescriptionLength = DescriptionField.Len();
+	int32 MaxTooltipLength = TooltipField.Len();
 
 	// Collect all default material expression objects
 	for (TObjectIterator<UClass> It; It; ++It)
@@ -97,23 +113,23 @@ int32 UDumpMaterialExpressionsCommandlet::Main(const FString& Params)
 				//    FMaterialEditorUtilities::AddMaterialExpressionCategory() 
 				//    MaterialExpressions.cpp IsAllowedExpressionType()
 				bool bShowInCreateMenu = !(bClassDeprecated
-									  || Class == UMaterialExpressionComment::StaticClass()
-									  || Class == UMaterialExpressionMaterialLayerOutput::StaticClass()
-									  || Class == UMaterialExpressionParameter::StaticClass()
-									  || Class == UMaterialExpressionNamedRerouteUsage::StaticClass()
-									  || Class == UMaterialExpressionExecBegin::StaticClass()
-									  || Class == UMaterialExpressionExecEnd::StaticClass()
-									  || Class == UMaterialExpressionPinBase::StaticClass()
-									  || Class == UMaterialExpressionFunctionInput::StaticClass()
-									  || Class == UMaterialExpressionFunctionOutput::StaticClass()
-									  || Class == UMaterialExpressionComposite::StaticClass());
+					|| Class == UMaterialExpressionComment::StaticClass()
+					|| Class == UMaterialExpressionMaterialLayerOutput::StaticClass()
+					|| Class == UMaterialExpressionParameter::StaticClass()
+					|| Class == UMaterialExpressionNamedRerouteUsage::StaticClass()
+					|| Class == UMaterialExpressionExecBegin::StaticClass()
+					|| Class == UMaterialExpressionExecEnd::StaticClass()
+					|| Class == UMaterialExpressionPinBase::StaticClass()
+					|| Class == UMaterialExpressionFunctionInput::StaticClass()
+					|| Class == UMaterialExpressionFunctionOutput::StaticClass()
+					|| Class == UMaterialExpressionComposite::StaticClass());
 
 				FString ExpressionType;
-				if (bControlFlow)									{ ExpressionType = TEXT("ControlFlow"); }
-				if (!ExpressionType.IsEmpty() && bNewHLSLGenerator)	{ ExpressionType += TEXT("|"); }
-				if (bNewHLSLGenerator)								{ ExpressionType += TEXT("HLSLGenerator"); }
-				if (!ExpressionType.IsEmpty() && bClassDeprecated)	{ ExpressionType += TEXT("|"); }
-				if (bClassDeprecated)								{ ExpressionType += TEXT("CLASS_Deprecated"); }
+				if (bControlFlow) { ExpressionType = TEXT("ControlFlow"); }
+				if (!ExpressionType.IsEmpty() && bNewHLSLGenerator) { ExpressionType += TEXT("|"); }
+				if (bNewHLSLGenerator) { ExpressionType += TEXT("HLSLGenerator"); }
+				if (!ExpressionType.IsEmpty() && bClassDeprecated) { ExpressionType += TEXT("|"); }
+				if (bClassDeprecated) { ExpressionType += TEXT("CLASS_Deprecated"); }
 
 				TArray<FString> MultilineCaption;
 				DefaultExpression->GetCaption(MultilineCaption);
@@ -155,6 +171,7 @@ int32 UDumpMaterialExpressionsCommandlet::Main(const FString& Params)
 				MaxCreationDescriptionLength = FMath::Max(MaxCreationDescriptionLength, ExpressionInfo.CreationDescription.Len());
 				MaxCaptionLength = FMath::Max(MaxCaptionLength, ExpressionInfo.Caption.Len());
 				MaxDescriptionLength = FMath::Max(MaxDescriptionLength, ExpressionInfo.Description.Len());
+				MaxTooltipLength = FMath::Max(MaxTooltipLength, ExpressionInfo.Tooltip.Len());
 			}
 		}
 	}
@@ -169,71 +186,176 @@ int32 UDumpMaterialExpressionsCommandlet::Main(const FString& Params)
 	MaxCreationDescriptionLength += AdditionalPadding;
 	MaxCaptionLength += AdditionalPadding;
 	MaxDescriptionLength += AdditionalPadding;
+	MaxTooltipLength += AdditionalPadding;
+
+	TArray<uint32> MaxFieldLength;
+	MaxFieldLength.Add(MaxNameLength);
+	MaxFieldLength.Add(MaxTypeLength);
+	MaxFieldLength.Add(MaxShowInCreateMenuLength);
+	MaxFieldLength.Add(MaxKeywordsLength);
+	MaxFieldLength.Add(MaxCreationNameLength);
+	MaxFieldLength.Add(MaxCreationDescriptionLength);
+	MaxFieldLength.Add(MaxCaptionLength);
+	MaxFieldLength.Add(MaxDescriptionLength);
+	MaxFieldLength.Add(MaxTooltipLength);
 
-	auto GetFormattedText = [](const FString& InText) -> FString
+	// Write the material expression list to a text file
+	TArray<FString> FieldNames;
+	FieldNames.Add(NameField);
+	FieldNames.Add(TypeField);
+	FieldNames.Add(ShowInCreateMenuField);
+	FieldNames.Add(KeywordsField);
+	FieldNames.Add(CreationNameField);
+	FieldNames.Add(CreationDescriptionField);
+	FieldNames.Add(CaptionField);
+	FieldNames.Add(DescriptionField);
+	FieldNames.Add(TooltipField);
+	WriteLine(FileWriter, FieldNames, MaxFieldLength);
+
+	for (FMaterialExpressionInfo& ExpressionInfo : MaterialExpressionInfos)
 	{
-		FString OutText = InText;
-		OutText = InText.Replace(TEXT("\n"), TEXT(" "));
-		if (OutText.IsEmpty())
-		{
-			OutText = TEXT("N/A");
-		}
-		return OutText;
+		FieldNames.Reset();
+		FieldNames.Add(GetFormattedText(ExpressionInfo.Name));
+		FieldNames.Add(GetFormattedText(ExpressionInfo.Type));
+		FieldNames.Add(GetFormattedText(ExpressionInfo.ShowInCreateMenu));
+		FieldNames.Add(GetFormattedText(ExpressionInfo.Keywords));
+		FieldNames.Add(GetFormattedText(ExpressionInfo.CreationName));
+		FieldNames.Add(GetFormattedText(ExpressionInfo.CreationDescription));
+		FieldNames.Add(GetFormattedText(ExpressionInfo.Caption));
+		FieldNames.Add(GetFormattedText(ExpressionInfo.Description));
+		FieldNames.Add(GetFormattedText(ExpressionInfo.Tooltip));
+		WriteLine(FileWriter, FieldNames, MaxFieldLength);
+	}
+
+	OutputLine = FString::Printf(TEXT("\nTotal %d material expressions found."), MaterialExpressionInfos.Num());
+	FileWriter->Serialize(const_cast<ANSICHAR*>(StringCast<ANSICHAR>(*OutputLine).Get()), OutputLine.Len());
+
+	UE_LOG(LogDumpMaterialExpressionsCommandlet, Log, TEXT("Total %d material expressions are written to %s"), MaterialExpressionInfos.Num(), *OutputFilePath);
+}
+
+void WriteOutMaterialFunctions(FArchive* FileWriter, const FString& OutputFilePath)
+{
+	FString OutputLine = TEXT("[MATERIAL FUNCTIONS]\n");
+	FileWriter->Serialize(const_cast<ANSICHAR*>(StringCast<ANSICHAR>(*OutputLine).Get()), OutputLine.Len());
+
+	struct FMaterialFunctionInfo
+	{
+		FString Name;
+		FString Description;
+		FString Path;
 	};
+	TArray<FMaterialFunctionInfo> MaterialFunctionInfos;
+
+	const FString NameField = TEXT("NAME");
+	const FString DescriptionField = TEXT("DESCRIPTION");
+	const FString PathField = TEXT("PATH");
 
-	auto GenerateSpacePadding = [](int32 MaxLen, int32 TextLen) -> FString
+	int32 MaxNameLength = NameField.Len();
+	int32 MaxDescriptionLength = DescriptionField.Len();
+	int32 MaxPathLength = PathField.Len();
+
+	TArray<FAssetData> AssetDataList;
+	FAssetRegistryModule& AssetRegistryModule = FModuleManager::LoadModuleChecked<FAssetRegistryModule>(TEXT("AssetRegistry"));
+	AssetRegistryModule.Get().GetAssetsByClass(UMaterialFunction::StaticClass()->GetFName(), AssetDataList);
+	for (const FAssetData& AssetData : AssetDataList)
 	{
-		FString Padding;
-		for (int i = 0; i < MaxLen - TextLen; ++i)
+		// If this is a function that is selected to be exposed to the library
+		if (AssetData.GetTagValueRef<bool>("bExposeToLibrary"))
 		{
-			Padding += TEXT(" ");
+			const FString FunctionPathName = AssetData.ObjectPath.ToString();
+			const FString Description = AssetData.GetTagValueRef<FText>("Description").ToString();
+
+			FString FunctionName = FunctionPathName;
+			int32 PeriodIndex = FunctionPathName.Find(TEXT("."), ESearchCase::CaseSensitive, ESearchDir::FromEnd);
+			if (PeriodIndex != INDEX_NONE)
+			{
+				FunctionName = FunctionPathName.Right(FunctionPathName.Len() - PeriodIndex - 1);
+			}
+
+			FMaterialFunctionInfo FunctionInfo;
+			FunctionInfo.Name = FunctionName;
+			FunctionInfo.Description = Description;
+			FunctionInfo.Path = FunctionPathName;
+			MaterialFunctionInfos.Add(FunctionInfo);
+
+			MaxNameLength = FMath::Max(MaxNameLength, FunctionInfo.Name.Len());
+			MaxDescriptionLength = FMath::Max(MaxDescriptionLength, FunctionInfo.Description.Len());
+			MaxPathLength = FMath::Max(MaxPathLength, FunctionInfo.Path.Len());
 		}
-		return Padding;
-	};
+	}
+
+	// Additional padding for spacing
+	const int32 AdditionalPadding = 3;
+	MaxNameLength += AdditionalPadding;
+	MaxDescriptionLength += AdditionalPadding;
+	MaxPathLength += AdditionalPadding;
+
+	TArray<uint32> MaxFieldLength;
+	MaxFieldLength.Add(MaxNameLength);
+	MaxFieldLength.Add(MaxDescriptionLength);
+	MaxFieldLength.Add(MaxPathLength);
 
 	// Write the material expression list to a text file
-	const FString OutputFilePath = FPaths::Combine(*FPaths::ProjectSavedDir(), TEXT("MaterialEditor"), TEXT("MaterialExpressions.txt"));
-	FArchive* FileWriter = IFileManager::Get().CreateFileWriter(*OutputFilePath);
+	TArray<FString> FieldNames;
+	FieldNames.Add(NameField);
+	FieldNames.Add(DescriptionField);
+	FieldNames.Add(PathField);
+	WriteLine(FileWriter, FieldNames, MaxFieldLength);
 
-	auto WriteLine = [FileWriter, GenerateSpacePadding, MaxNameLength, MaxTypeLength, MaxShowInCreateMenuLength, MaxKeywordsLength, MaxCreationNameLength, MaxCreationDescriptionLength, MaxCaptionLength, MaxDescriptionLength]
-		(const FString& Name, const FString& Type, const FString& ShowInCreateMenu, const FString& Keywords, const FString& CreationName, const FString& CreationDescription, const FString& Caption, const FString& Description, const FString& Tooltip)
+	for (FMaterialFunctionInfo& FunctionInfo : MaterialFunctionInfos)
 	{
-		FString NamePadding = GenerateSpacePadding(MaxNameLength, Name.Len());
-		FString TypePadding = GenerateSpacePadding(MaxTypeLength, Type.Len());
-		FString ShowInCreateMenuPadding = GenerateSpacePadding(MaxShowInCreateMenuLength, ShowInCreateMenu.Len());
-		FString KeywordsPadding = GenerateSpacePadding(MaxKeywordsLength, Keywords.Len());
-		FString CreationNamePadding = GenerateSpacePadding(MaxCreationNameLength, CreationName.Len());
-		FString CreationDescriptionPadding = GenerateSpacePadding(MaxCreationDescriptionLength, CreationDescription.Len());
-		FString CaptionPadding = GenerateSpacePadding(MaxCaptionLength, Caption.Len());
-		FString DescriptionPadding = GenerateSpacePadding(MaxDescriptionLength, Description.Len());
-
-		FString OutputLine = (Name + NamePadding);
-		OutputLine += (Type + TypePadding);
-		OutputLine += (ShowInCreateMenu + ShowInCreateMenuPadding);
-		OutputLine += (Keywords + KeywordsPadding);
-		OutputLine += (CreationName + CreationNamePadding);
-		OutputLine += (CreationDescription + CreationDescriptionPadding);
-		OutputLine += (Caption + CaptionPadding);
-		OutputLine += (Description + DescriptionPadding);
-		OutputLine += (Tooltip + TEXT("\n"));
-		FileWriter->Serialize(const_cast<ANSICHAR*>(StringCast<ANSICHAR>(*OutputLine).Get()), OutputLine.Len());
-	};
+		FieldNames.Reset();
+		FieldNames.Add(GetFormattedText(FunctionInfo.Name));
+		FieldNames.Add(GetFormattedText(FunctionInfo.Description));
+		FieldNames.Add(GetFormattedText(FunctionInfo.Path));
+		WriteLine(FileWriter, FieldNames, MaxFieldLength);
+	}
+	
+	OutputLine = FString::Printf(TEXT("\nTotal %d (bExposeToLibrary=true) material functions found."), MaterialFunctionInfos.Num());
+	FileWriter->Serialize(const_cast<ANSICHAR*>(StringCast<ANSICHAR>(*OutputLine).Get()), OutputLine.Len());
 
-	WriteLine(NameField, TypeField, ShowInCreateMenuField, KeywordsField, CreationNameField, CreationDescriptionField, CaptionField, DescriptionField, TooltipField);
-	for (FMaterialExpressionInfo& ExpressionInfo : MaterialExpressionInfos)
+	UE_LOG(LogDumpMaterialExpressionsCommandlet, Log, TEXT("Total %d material functions are written to %s"), MaterialFunctionInfos.Num(), *OutputFilePath);
+}
+
+UDumpMaterialExpressionsCommandlet::UDumpMaterialExpressionsCommandlet(const FObjectInitializer& ObjectInitializer)
+	: Super(ObjectInitializer)
+{
+}
+
+int32 UDumpMaterialExpressionsCommandlet::Main(const FString& Params)
+{
+	TArray<FString> Tokens;
+	TArray<FString> Switches;
+	TMap<FString, FString> ParamVals;
+	UCommandlet::ParseCommandLine(*Params, Tokens, Switches, ParamVals);
+
+	if (Switches.Contains("help"))
 	{
-		WriteLine(GetFormattedText(ExpressionInfo.Name), GetFormattedText(ExpressionInfo.Type), GetFormattedText(ExpressionInfo.ShowInCreateMenu), GetFormattedText(ExpressionInfo.Keywords),
-					GetFormattedText(ExpressionInfo.CreationName), GetFormattedText(ExpressionInfo.CreationDescription), GetFormattedText(ExpressionInfo.Caption), 
-					GetFormattedText(ExpressionInfo.Description), GetFormattedText(ExpressionInfo.Tooltip));
+		UE_LOG(LogDumpMaterialExpressionsCommandlet, Log, TEXT("DumpMaterialExpressions"));
+		UE_LOG(LogDumpMaterialExpressionsCommandlet, Log, TEXT("This commandlet will dump to a plain text file an info table of all material expressions in the engine and the plugins enabled on the project."));
+		UE_LOG(LogDumpMaterialExpressionsCommandlet, Log, TEXT("The output fields include:"));
+		UE_LOG(LogDumpMaterialExpressionsCommandlet, Log, TEXT("Name - The class name of the material expression"));
+		UE_LOG(LogDumpMaterialExpressionsCommandlet, Log, TEXT("Type - ControlFlow | HLSLGenerator | CLASS_Deprecated"));
+		UE_LOG(LogDumpMaterialExpressionsCommandlet, Log, TEXT("ShowInCreateMenu - If the expression appears in the create node dropdown menu"));
+		UE_LOG(LogDumpMaterialExpressionsCommandlet, Log, TEXT("CreationName - The name displayed in the create node dropdown menu to add an expression"));
+		UE_LOG(LogDumpMaterialExpressionsCommandlet, Log, TEXT("CreationDescription - The tooltip displayed on the CreationName in the create node dropdown menu"));
+		UE_LOG(LogDumpMaterialExpressionsCommandlet, Log, TEXT("Caption - The caption displayed on the material expression node"));
+		UE_LOG(LogDumpMaterialExpressionsCommandlet, Log, TEXT("Tooltip - The tooltip displayed on the material expression node"));
+		return 0;
 	}
 
-	FString OutputLine = FString::Printf(TEXT("\nTotal %d material expressions found."), MaterialExpressionInfos.Num());
+	const FString OutputFilePath = FPaths::Combine(*FPaths::ProjectSavedDir(), TEXT("MaterialEditor"), TEXT("MaterialExpressions.txt"));
+	FArchive* FileWriter = IFileManager::Get().CreateFileWriter(*OutputFilePath);
+
+	WriteOutMaterialExpressions(FileWriter, OutputFilePath);
+
+	FString OutputLine = FString::Printf(TEXT("\n\n============================================================================\n\n"));
 	FileWriter->Serialize(const_cast<ANSICHAR*>(StringCast<ANSICHAR>(*OutputLine).Get()), OutputLine.Len());
 
+	WriteOutMaterialFunctions(FileWriter, OutputFilePath);
+
 	FileWriter->Close();
 	delete FileWriter;
 
-	UE_LOG(LogDumpMaterialExpressionsCommandlet, Log, TEXT("Total %d material expressions are written to %s"), MaterialExpressionInfos.Num(), *OutputFilePath);
-
 	return 0;
 }
