trivial comment

diff --git a/Engine/Source/Editor/UnrealEd/Private/Commandlets/DumpMaterialExpressionsCommandlet.cpp b/Engine/Source/Editor/UnrealEd/Private/Commandlets/DumpMaterialExpressionsCommandlet.cpp
--- a/Engine/Source/Editor/UnrealEd/Private/Commandlets/DumpMaterialExpressionsCommandlet.cpp
+++ b/Engine/Source/Editor/UnrealEd/Private/Commandlets/DumpMaterialExpressionsCommandlet.cpp
@@ -254,6 +254,7 @@ void WriteOutMaterialFunctions(FArchive* FileWriter, const FString& OutputFilePa
 	int32 MaxDescriptionLength = DescriptionField.Len();
 	int32 MaxPathLength = PathField.Len();
 
+	// See UMaterialGraphSchema::GetMaterialFunctionActions for reference
 	TArray<FAssetData> AssetDataList;
 	FAssetRegistryModule& AssetRegistryModule = FModuleManager::LoadModuleChecked<FAssetRegistryModule>(TEXT("AssetRegistry"));
 	AssetRegistryModule.Get().GetAssetsByClass(UMaterialFunction::StaticClass()->GetFName(), AssetDataList);
