don't expose BGRE export to DDS if HDR is possible

diff --git a/Engine/Source/Editor/UnrealEd/Private/Factories/EditorFactories.cpp b/Engine/Source/Editor/UnrealEd/Private/Factories/EditorFactories.cpp
--- a/Engine/Source/Editor/UnrealEd/Private/Factories/EditorFactories.cpp
+++ b/Engine/Source/Editor/UnrealEd/Private/Factories/EditorFactories.cpp
@@ -4743,6 +4743,17 @@ bool UTextureExporterDDS::SupportsObject(UObject* Object) const
 	{
 		return false;
 	}
+
+	if ( Texture->Source.GetFormat() == TSF_BGRE8 )
+	{
+		// exclude BGRE export from DDS, use HDR files if possible
+		if ( Texture->Source.GetNumSlices() == 1 )
+		{
+			// HDR can handle it, say I won't
+			return false;
+		}
+		// BGRE format, but something HDR can't handle, so I'll do it...
+	}
 	
 	return true;
 }
