World Partition: don't explicitly initialize the editor hash when converting a world as we now properly initalize the world partition object since 19303190.
Also foxed typo in error message when converting an already partitioned world.

#rnx

diff --git a/Engine/Source/Editor/UnrealEd/Private/Commandlets/WorldPartitionConvertCommandlet.cpp b/Engine/Source/Editor/UnrealEd/Private/Commandlets/WorldPartitionConvertCommandlet.cpp
--- a/Engine/Source/Editor/UnrealEd/Private/Commandlets/WorldPartitionConvertCommandlet.cpp
+++ b/Engine/Source/Editor/UnrealEd/Private/Commandlets/WorldPartitionConvertCommandlet.cpp
@@ -250,8 +250,6 @@ UWorldPartition* UWorldPartitionConvertCommandlet::CreateWorldPartition(AWorldSe
 			CurrentHLODLayer = Cast<UHLODLayer>(CurrentHLODLayer->GetParentLayer().Get());
 		}
 	}
-
-	WorldPartition->EditorHash->Initialize();
 	
 	return WorldPartition;
 }
@@ -837,7 +835,7 @@ int32 UWorldPartitionConvertCommandlet::Main(const FString& Params)
 	AWorldSettings* MainWorldSettings = MainWorld->GetWorldSettings();
 	if (MainWorldSettings->IsPartitionedWorld())
 	{
-		UE_LOG(LogWorldPartitionConvertCommandlet, Error, TEXT("Level '%s' is already partitionned"), *Tokens[0]);
+		UE_LOG(LogWorldPartitionConvertCommandlet, Error, TEXT("Level '%s' is already partitioned"), *Tokens[0]);
 		return 1;
 	}
 
