Fix actor placement out of bounds warning being triggered incorrectly for worlds not centered on the origin.

diff --git a/Engine/Source/Editor/UnrealEd/Private/AssetSelection.cpp b/Engine/Source/Editor/UnrealEd/Private/AssetSelection.cpp
--- a/Engine/Source/Editor/UnrealEd/Private/AssetSelection.cpp
+++ b/Engine/Source/Editor/UnrealEd/Private/AssetSelection.cpp
@@ -500,10 +500,7 @@ namespace ActorPlacementUtils
 			{
 				return true;
 			}
-			FVector ExpandedScale = FVector(1.0f + (GetDefault<ULevelEditorMiscSettings>()->PercentageThresholdForPrompt / 100.0f));
-			FTransform ExpandedScaleTransform = FTransform::Identity;
-			ExpandedScaleTransform.SetScale3D(ExpandedScale);
-			CurrentLevelBounds = CurrentLevelBounds.TransformBy(ExpandedScaleTransform);
+			CurrentLevelBounds = CurrentLevelBounds.ExpandBy(BoundsExtent * (GetDefault<ULevelEditorMiscSettings>()->PercentageThresholdForPrompt / 100.0f));
 			for (int32 ActorTransformIndex = 0; ActorTransformIndex < InActorTransforms.Num(); ++ActorTransformIndex)
 			{
 				FTransform ActorTransform = InActorTransforms[ActorTransformIndex];
