Update WragnleContent to use the CollectionManager and GRedirectCollector.
This will allow us to specify Collections that we want to track (useful for EditorOnly and Tool / Utility assets).
This will also help the script catch Soft Object references that aren't initially loaded.

diff --git a/Engine/Source/Editor/UnrealEd/Private/Commandlets/ContentCommandlets.cpp b/Engine/Source/Editor/UnrealEd/Private/Commandlets/ContentCommandlets.cpp
--- a/Engine/Source/Editor/UnrealEd/Private/Commandlets/ContentCommandlets.cpp
+++ b/Engine/Source/Editor/UnrealEd/Private/Commandlets/ContentCommandlets.cpp
@@ -91,6 +91,7 @@ DEFINE_LOG_CATEGORY(LogContentCommandlet);
 #include "UObject/UObjectThreadContext.h"
 #include "Engine/LODActor.h"
 #include "PerQualityLevelProperties.h"
+#include "Misc/RedirectCollector.h"
 
 /**-----------------------------------------------------------------------------
  *	UResavePackages commandlet.
@@ -2433,20 +2434,24 @@ int32 UWrangleContentCommandlet::Main( const FString& Params )
 		FString WrangleContentIniName =	FPaths::SourceConfigDir() + TEXT("WrangleContent.ini");
 
 		// figure out which section to use to get the packages to fully load
-		FString SectionToUse = TEXT("WrangleContent.PackagesToFullyLoad");
-		if( SectionStr.Len() > 0 )
+		FString PackagesToFullyLoadSectionName = TEXT("WrangleContent.PackagesToFullyLoad");
+		FString CollectionsToFullyLoadSectionName = TEXT("WrangleContent.CollectionsToFullyLoad");
+
+		if (SectionStr.Len() > 0)
 		{
-			SectionToUse = FString::Printf( TEXT( "WrangleContent.%sPackagesToFullyLoad" ), *SectionStr );
+			PackagesToFullyLoadSectionName = FString::Printf(TEXT("WrangleContent.%sPackagesToFullyLoad"), *SectionStr);
+			CollectionsToFullyLoadSectionName = FString::Printf(TEXT("WrangleContent.%sCollectionsToFullyLoad"), *SectionStr);
 		}
 
 		// get a list of packages to load
-		const FConfigSection* PackagesToFullyLoadSection = GConfig->GetSectionPrivate( *SectionToUse, 0, 1, *WrangleContentIniName );
-		const FConfigSection* StartupPackages = GConfig->GetSectionPrivate( TEXT("/Script/Engine.StartupPackages"), 0, 1, GEngineIni );
+		const FConfigSection* PackagesToFullyLoadSection = GConfig->GetSectionPrivate(*PackagesToFullyLoadSectionName, 0, 1, *WrangleContentIniName);
+		const FConfigSection* StartupPackages = GConfig->GetSectionPrivate(TEXT("/Script/Engine.StartupPackages"), 0, 1, GEngineIni);
+		const FConfigSection* CollectionsToFullyLoadSection = GConfig->GetSectionPrivate(*CollectionsToFullyLoadSectionName, 0, 1, *WrangleContentIniName);
 
 		// we expect either the .ini to exist, or -allmaps to be specified
-		if (!PackagesToFullyLoadSection && !bShouldLoadAllMaps)
+		if (!PackagesToFullyLoadSection && !bShouldLoadAllMaps && !CollectionsToFullyLoadSection)
 		{
-			UE_LOG(LogContentCommandlet, Error, TEXT("This commandlet needs a WrangleContent.ini in the Config directory with a [WrangleContent.PackagesToFullyLoad] section"));
+			UE_LOG(LogContentCommandlet, Error, TEXT("This commandlet needs a WrangleContent.ini in the Config directory with at least a [WrangleContent.PackagesToFullyLoad] section or a [WragnelContent.CollectionsToFullyLoad] section."));
 			return 1;
 		}
 
@@ -2477,7 +2482,6 @@ int32 UWrangleContentCommandlet::Main( const FString& Params )
 			}
 		}
 
-
 		// read in the per-map packages to cook
 		TMap<FString, TArray<FString> > PerMapCookPackages;
 		GConfig->Parse1ToNSectionOfStrings(TEXT("/Script/Engine.PackagesToForceCookPerMap"), TEXT("Map"), TEXT("Package"), PerMapCookPackages, GEngineIni);
@@ -2515,6 +2519,30 @@ int32 UWrangleContentCommandlet::Main( const FString& Params )
 			}
 		}
 
+		if (CollectionsToFullyLoadSection)
+		{
+			ICollectionManager& CollectionManager = FCollectionManagerModule::GetModule().Get();
+			TArray<FName> CollectionAssets;
+			
+			for (FConfigSectionMap::TConstIterator CollectionIt(*CollectionsToFullyLoadSection); CollectionIt; ++CollectionIt)
+			{
+				CollectionAssets.Reset();
+				FString CollectionName = CollectionIt.Value().GetValue();
+				if (!CollectionManager.GetAssetsInCollection(FName(*CollectionName), ECollectionShareType::CST_All, CollectionAssets))
+				{
+					UE_LOG(LogContentCommandlet, Warning, TEXT("Could not get assets in collection '%s'. Skipping filter."), *CollectionName);
+				}
+				else
+				{
+					//insert all of the collection names into the set for fast filter checks
+					for (const FName &AssetName : CollectionAssets)
+					{
+						PackagesToFullyLoad.Add(TEXT("Package"), FPackageName::ObjectPathToPackageName(AssetName.ToString()));
+					}
+				}
+			}
+		}
+		
 		// go over all the packages that we want to fully load
 		for (FConfigSectionMap::TIterator PackageIt(PackagesToFullyLoad); PackageIt; ++PackageIt)
 		{
@@ -2552,6 +2580,10 @@ int32 UWrangleContentCommandlet::Main( const FString& Params )
 				// load the package fully
 				UPackage* Package = LoadPackage(nullptr, PackagePath, LOAD_None);
 
+				// Now that we've loaded the package, go ahead and also load any soft paths.
+				// This should help capture more references.
+				GRedirectCollector.ResolveAllSoftObjectPaths(NAME_None);
+				
 				FLinkerLoad* Linker = LoadPackageLinker(nullptr, PackagePath, LOAD_Quiet | LOAD_NoWarn | LOAD_NoVerify);
 
 				UWorld* World = UWorld::FindWorldInPackage(Package);
