Require explicit constructors/casts when converting between FVector, FVector3d, and FVector3f.
Editor files

diff --git a/Engine/Source/Editor/UnrealEd/Private/Fbx/FbxMainExport.cpp b/Engine/Source/Editor/UnrealEd/Private/Fbx/FbxMainExport.cpp
--- a/Engine/Source/Editor/UnrealEd/Private/Fbx/FbxMainExport.cpp
+++ b/Engine/Source/Editor/UnrealEd/Private/Fbx/FbxMainExport.cpp
@@ -4884,7 +4884,7 @@ void FFbxExporter::ExportSplineMeshToFbx(const USplineMeshComponent* SplineMeshC
 	{
 		FVector Position = (FVector)RenderMesh.VertexBuffers.PositionVertexBuffer.VertexPosition(VertIndex);
 		const FTransform SliceTransform = SplineMeshComp->CalcSliceTransform(USplineMeshComponent::GetAxisValue(Position, SplineMeshComp->ForwardAxis));
-		FVector Normal = FVector4(RenderMesh.VertexBuffers.StaticMeshVertexBuffer.VertexTangentZ(VertIndex));
+		FVector Normal = FVector(RenderMesh.VertexBuffers.StaticMeshVertexBuffer.VertexTangentZ(VertIndex));
 		Normal = SliceTransform.TransformVector(Normal);
 		FbxVector4& FbxNormal = FbxNormals[VertIndex];
 		FbxNormal = FbxVector4(Normal.X, -Normal.Y, Normal.Z);
