Introduced a Preview cmd line arg to GatherTextCommandlet to allow the gathering of source and asset files without
checking out or writing any localization files. This makes dry runs to print all localization warnings much easier for debugging.
- Introduced the -Preview arg and -GatherType param to UGatherTextCommandlet.
- The -GatherType param can have values of All, Source, Asset or Metadata. Source will only gather source files, Asset will only gather Asset files, Metadata will only gather Metadata and All will gather all 3.
- The -GatherType param is only supported in preview mode right now. This is because manifests are updated every single gather and we don't want the manifest to only reflect partial gathers.
- Added additional logs to display when the commandlet is run with -Preview and -GatherType
- Refactored parts of UGatherTextCommandlet to reduce hard coded switch and param names.
- Introduced a virtual function in UGatherTextCommandletBase to allow commandlets to conditionally be skipped in preview mode.

diff --git a/Engine/Source/Editor/UnrealEd/Classes/Commandlets/GatherTextCommandlet.h b/Engine/Source/Editor/UnrealEd/Classes/Commandlets/GatherTextCommandlet.h
--- a/Engine/Source/Editor/UnrealEd/Classes/Commandlets/GatherTextCommandlet.h
+++ b/Engine/Source/Editor/UnrealEd/Classes/Commandlets/GatherTextCommandlet.h
@@ -30,5 +30,13 @@ public:
 	FText GetChangelistDescription(const TArray<FString>& GatherTextConfigPaths);
 
 	static const FString UsageText;
-
+//~ Begin UGatherTextCommandletBase  Interface
+	virtual bool ShouldRunInPreview(const TArray<FString>& Switches, const TMap<FString, FString>& ParamVals) const override
+	{
+		// This commandlet is the driver for other commandlet. This should always run even in preview
+		return true;
+	}
+	//~ End UGatherTextCommandletBase  Interface
+private:
+	 
 };
diff --git a/Engine/Source/Editor/UnrealEd/Classes/Commandlets/GatherTextCommandletBase.h b/Engine/Source/Editor/UnrealEd/Classes/Commandlets/GatherTextCommandletBase.h
--- a/Engine/Source/Editor/UnrealEd/Classes/Commandlets/GatherTextCommandletBase.h
+++ b/Engine/Source/Editor/UnrealEd/Classes/Commandlets/GatherTextCommandletBase.h
@@ -82,6 +82,16 @@ public:
 	// Utility to get the correct base path (engine or project) for the current environment
 	static const FString& GetProjectBasePath();
 
+	/**
+* Returns true if this commandlet should run during a preview run.
+* Override in child classes to conditionally skip a commandlet from being run.
+* Most commandlets that require source control, write to files etc should be skipped for preview runs
+*/
+	virtual bool ShouldRunInPreview(const TArray<FString>& Switches, const TMap<FString, FString>& ParamVals) const
+	{
+		return false;
+	}
+
 protected:
 	TSharedPtr< FLocTextHelper > GatherManifestHelper;
 
@@ -90,6 +100,13 @@ protected:
 	/** Mapping from platform name to the path marker for that platform */
 	TMap<FName, FString> SplitPlatforms;
 
+	// Common params and switches among all text gathering commadnlets 
+	static const TCHAR* ConfigParam;
+	static const TCHAR* EnableSourceControlSwitch;
+	static const TCHAR* DisableSubmitSwitch;
+	static const TCHAR* PreviewSwitch;
+	static const TCHAR* GatherTypeParam;
+
 private:
 	virtual void CreateCustomEngine(const FString& Params) override ; //Disallow other text commandlets to make their own engine.	
 };
diff --git a/Engine/Source/Editor/UnrealEd/Classes/Commandlets/GatherTextFromAssetsCommandlet.h b/Engine/Source/Editor/UnrealEd/Classes/Commandlets/GatherTextFromAssetsCommandlet.h
--- a/Engine/Source/Editor/UnrealEd/Classes/Commandlets/GatherTextFromAssetsCommandlet.h
+++ b/Engine/Source/Editor/UnrealEd/Classes/Commandlets/GatherTextFromAssetsCommandlet.h
@@ -30,6 +30,9 @@ public:
 	bool ConfigureFromScript(const FString& GatherTextConfigPath, const FString& SectionName);
 
 	//~ End UCommandlet Interface
+	//~ Begin UGatherTextCommandletBase  Interface
+	virtual bool ShouldRunInPreview(const TArray<FString>& Switches, const TMap<FString, FString>& ParamVals) const override;
+	//~ End UGatherTextCommandletBase  Interface
 
 	/** Localization cache states of a package */
 	enum class EPackageLocCacheState : uint8
diff --git a/Engine/Source/Editor/UnrealEd/Classes/Commandlets/GatherTextFromMetadataCommandlet.h b/Engine/Source/Editor/UnrealEd/Classes/Commandlets/GatherTextFromMetadataCommandlet.h
--- a/Engine/Source/Editor/UnrealEd/Classes/Commandlets/GatherTextFromMetadataCommandlet.h
+++ b/Engine/Source/Editor/UnrealEd/Classes/Commandlets/GatherTextFromMetadataCommandlet.h
@@ -20,6 +20,10 @@ public:
 	virtual int32 Main(const FString& Params) override;
 	//~ End UCommandlet Interface
 
+	//~ Begin UGatherTextCommandletBase  Interface
+	virtual bool ShouldRunInPreview(const TArray<FString>& Switches, const TMap<FString, FString>& ParamVals) const override;
+	//~ End UGatherTextCommandletBase  Interface
+
 	struct FGatherParameters
 	{
 		TArray<FString> InputKeys;
diff --git a/Engine/Source/Editor/UnrealEd/Classes/Commandlets/GatherTextFromSourceCommandlet.h b/Engine/Source/Editor/UnrealEd/Classes/Commandlets/GatherTextFromSourceCommandlet.h
--- a/Engine/Source/Editor/UnrealEd/Classes/Commandlets/GatherTextFromSourceCommandlet.h
+++ b/Engine/Source/Editor/UnrealEd/Classes/Commandlets/GatherTextFromSourceCommandlet.h
@@ -498,6 +498,8 @@ public:
 	//~ Begin UCommandlet Interface
 	virtual int32 Main(const FString& Params) override;
 	//~ End UCommandlet Interface
-
+	//~ Begin UGatherTextCommandletBase  Interface
+	virtual bool ShouldRunInPreview(const TArray<FString>& Switches, const TMap<FString, FString>& ParamVals) const override;
+	//~ End UGatherTextCommandletBase  Interface
 #undef LOC_DEFINE_REGION
 };
diff --git a/Engine/Source/Editor/UnrealEd/Private/Commandlets/GatherTextCommandlet.cpp b/Engine/Source/Editor/UnrealEd/Private/Commandlets/GatherTextCommandlet.cpp
--- a/Engine/Source/Editor/UnrealEd/Private/Commandlets/GatherTextCommandlet.cpp
+++ b/Engine/Source/Editor/UnrealEd/Private/Commandlets/GatherTextCommandlet.cpp
@@ -23,9 +23,14 @@ UGatherTextCommandlet::UGatherTextCommandlet(const FObjectInitializer& ObjectIni
 const FString UGatherTextCommandlet::UsageText
 	(
 	TEXT("GatherTextCommandlet usage...\r\n")
-	TEXT("    <GameName> GatherTextCommandlet -Config=<path to config ini file>\r\n")
+	TEXT("    <GameName> GatherTextCommandlet -Config=<path to config ini file> [-Preview -EnableSCC -DisableSCCSubmit -GatherType=<All | Source | Asset | Metadata>]\r\n")
 	TEXT("    \r\n")
 	TEXT("    <path to config ini file> Full path to the .ini config file that defines what gather steps the commandlet will run.\r\n")
+	TEXT("    Preview\t Runs the commandlet and its child commandlets in preview. Some commandlets will not be executed in preview mode. Use this to dump all generated warnings without writing any files. Using this switch implies -DisableSCCSubmit\r\n")
+	TEXT("    EnableSCC\t Enables source control and allows the commandlet to check out files for editing.\r\n")
+	TEXT("    DisableSCCSubmit\t Prevents the commandlet from submitting checked out files in source control that have been edited.\r\n")
+	TEXT("    GatherType\t Only performs a gather on the specified type of file (currently only works in preview mode). Source only runs commandlets that gather source files. Asset only runs commandlets that gather asset files. All runs commandlets that gather both source and asset files. Leaving this param out implies a gather type of All.")
+	TEXT("Metadata only runs commandlets that gather metadata files. All runs commandlets that gather both source and asset files. Leaving this param out implies a gather type of All.\r\n")
 	);
 
 
@@ -35,10 +40,15 @@ int32 UGatherTextCommandlet::Main( const FString& Params )
 	TArray<FString> Switches;
 	TMap<FString, FString> ParamVals;
 	UCommandlet::ParseCommandLine(*Params, Tokens, Switches, ParamVals);
-	
+	if (Switches.Contains(TEXT("help")) || Switches.Contains(TEXT("Help")))
+	{
+		UE_LOG(LogGatherTextCommandlet, Display, TEXT("%s"), *UsageText);
+		return 0;
+	}
+
 	// Build up the complete list of config files to process
 	TArray<FString> GatherTextConfigPaths;
-	if (const FString* ConfigParamPtr = ParamVals.Find(TEXT("config")))
+	if (const FString* ConfigParamPtr = ParamVals.Find(UGatherTextCommandletBase::ConfigParam))
 	{
 		ConfigParamPtr->ParseIntoArray(GatherTextConfigPaths, TEXT(";"));
 
@@ -59,8 +69,25 @@ int32 UGatherTextCommandlet::Main( const FString& Params )
 		return -1;
 	}
 
-	const bool bEnableSourceControl = Switches.Contains(TEXT("EnableSCC"));
-	const bool bDisableSubmit = Switches.Contains(TEXT("DisableSCCSubmit"));
+	const bool bRunningInPreview = Switches.Contains(UGatherTextCommandletBase::PreviewSwitch);
+	if (bRunningInPreview)
+	{
+		UE_LOG(LogGatherTextCommandlet, Display, TEXT("Running commandlet in preview mode. Some child commandlets and steps will be skipped."));
+		// -GatherType is only valid in preview mode right now 
+		if (const FString* GatherTypeParamPtr = ParamVals.Find(UGatherTextCommandletBase::GatherTypeParam))
+		{
+			UE_LOG(LogGatherTextCommandlet, Display, TEXT("Running commandlet with gather type %s. Only %s files will be gathered."), **GatherTypeParamPtr, **GatherTypeParamPtr);
+		}
+		else
+		{
+			// if the -GatherType param is not specified, we default to gathering both source, asset and metadata
+			UE_LOG(LogGatherTextCommandlet, Display, TEXT("-GatherType param not specified. Commandlet will gather source, assset and metadata."));
+		}
+	}
+
+	// If we are only doing a preview run, we will not be writing to any files that will need to be submitted to source control 
+	const bool bEnableSourceControl = Switches.Contains(UGatherTextCommandletBase::EnableSourceControlSwitch) && !bRunningInPreview;
+	const bool bDisableSubmit = Switches.Contains(UGatherTextCommandletBase::DisableSubmitSwitch) || bRunningInPreview;
 
 	TSharedPtr<FLocalizationSCC> CommandletSourceControlInfo;
 	if (bEnableSourceControl)
@@ -75,6 +102,7 @@ int32 UGatherTextCommandlet::Main( const FString& Params )
 		}
 	}
 
+	double AllCommandletExecutionStartTime = FPlatformTime::Seconds();
 	for (const FString& GatherTextConfigPath : GatherTextConfigPaths)
 	{
 		const int32 Result = ProcessGatherConfig(GatherTextConfigPath, CommandletSourceControlInfo, Tokens, Switches, ParamVals);
@@ -101,6 +129,7 @@ int32 UGatherTextCommandlet::Main( const FString& Params )
 			return -1;
 		}
 	}
+	UE_LOG(LogGatherTextCommandlet, Display, TEXT("Completed all steps in %.2f seconds"), FPlatformTime::Seconds() - AllCommandletExecutionStartTime);
 
 	return 0;
 }
@@ -164,7 +193,27 @@ int32 UGatherTextCommandlet::ProcessGatherConfig(const FString& GatherTextConfig
 
 		return NumericalSuffixOne < NumericalSuffixTwo;
 	});
+	// Generate the switches and params to be passed on to child commandlets
+	FString GeneratedParamsAndSwitches;
+	// Add all the command params with the exception of config
+	for (auto ParamIter = ParamVals.CreateConstIterator(); ParamIter; ++ParamIter)
+	{
+		const FString& Key = ParamIter.Key();
+		const FString& Val = ParamIter.Value();
+		if (Key != UGatherTextCommandletBase::ConfigParam)
+		{
+			GeneratedParamsAndSwitches += FString::Printf(TEXT(" -%s=%s"), *Key, *Val);
+		}
+	}
 
+	// Add all the command switches
+	for (auto SwitchIter = Switches.CreateConstIterator(); SwitchIter; ++SwitchIter)
+	{
+		const FString& Switch = *SwitchIter;
+		GeneratedParamsAndSwitches += FString::Printf(TEXT(" -%s"), *Switch);
+	}
+
+	const bool bRunningInPreview = Switches.Contains(TEXT("Preview"));
 	// Execute each step defined in the config file.
 	for (const FString& StepName : StepNames)
 	{
@@ -179,9 +228,16 @@ int32 UGatherTextCommandlet::ProcessGatherConfig(const FString& GatherTextConfig
 
 		UGatherTextCommandletBase* Commandlet = NewObject<UGatherTextCommandletBase>(GetTransientPackage(), CommandletClass);
 		check(Commandlet);
-		
 		FGCObjectScopeGuard CommandletGCGuard(Commandlet);
 		Commandlet->Initialize( CommandletGatherManifestHelper, CommandletSourceControlInfo );
+		// As of now, all params and switches (with the exception of config) is passed on to child commandlets
+		// Instead of parsing in child commandlets, we'll just pass the params and switches along to determine if we need to run the child commandlet 
+		// If we are running in preview mode, then most commandlets should be skipped as ShouldRunInPreview() defaults to false in the base class.
+		if (bRunningInPreview && !Commandlet->ShouldRunInPreview(Switches, ParamVals))
+		{
+			UE_LOG(LogGatherTextCommandlet, Display, TEXT("Should not run %s: %s in preview. Skipping."), *StepName, *CommandletClassName);
+			continue;
+		}
 
 		// Execute the commandlet.
 		double CommandletExecutionStartTime = FPlatformTime::Seconds();
@@ -189,24 +245,7 @@ int32 UGatherTextCommandlet::ProcessGatherConfig(const FString& GatherTextConfig
 		UE_LOG(LogGatherTextCommandlet, Display, TEXT("Executing %s: %s"), *StepName, *CommandletClassName);
 		
 		FString GeneratedCmdLine = FString::Printf(TEXT("-Config=\"%s\" -Section=%s"), *GatherTextConfigPath , *StepName);
-
-		// Add all the command params with the exception of config
-		for(auto ParamIter = ParamVals.CreateConstIterator(); ParamIter; ++ParamIter)
-		{
-			const FString& Key = ParamIter.Key();
-			const FString& Val = ParamIter.Value();
-			if(Key != TEXT("config"))
-			{
-				GeneratedCmdLine += FString::Printf(TEXT(" -%s=%s"), *Key , *Val);
-			}	
-		}
-
-		// Add all the command switches
-		for(auto SwitchIter = Switches.CreateConstIterator(); SwitchIter; ++SwitchIter)
-		{
-			const FString& Switch = *SwitchIter;
-			GeneratedCmdLine += FString::Printf(TEXT(" -%s"), *Switch);
-		}
+		GeneratedCmdLine += GeneratedParamsAndSwitches;
 
 		if( 0 != Commandlet->Main( GeneratedCmdLine ) )
 		{
diff --git a/Engine/Source/Editor/UnrealEd/Private/Commandlets/GatherTextCommandletBase.cpp b/Engine/Source/Editor/UnrealEd/Private/Commandlets/GatherTextCommandletBase.cpp
--- a/Engine/Source/Editor/UnrealEd/Private/Commandlets/GatherTextCommandletBase.cpp
+++ b/Engine/Source/Editor/UnrealEd/Private/Commandlets/GatherTextCommandletBase.cpp
@@ -23,6 +23,12 @@ FGatherTextDelegates::FGetAdditionalGatherPaths FGatherTextDelegates::GetAdditio
 //////////////////////////////////////////////////////////////////////////
 //UGatherTextCommandletBase
 
+const TCHAR* UGatherTextCommandletBase::ConfigParam = TEXT("Config");
+const TCHAR* UGatherTextCommandletBase::EnableSourceControlSwitch = TEXT("EnableSCC");
+const TCHAR* UGatherTextCommandletBase::DisableSubmitSwitch = TEXT("DisableSCCSubmit");
+const TCHAR* UGatherTextCommandletBase::PreviewSwitch = TEXT("Preview");
+const TCHAR* UGatherTextCommandletBase::GatherTypeParam = TEXT("GatherType");
+
 UGatherTextCommandletBase::UGatherTextCommandletBase(const FObjectInitializer& ObjectInitializer)
 	: Super(ObjectInitializer)
 {
diff --git a/Engine/Source/Editor/UnrealEd/Private/Commandlets/GatherTextFromAssetsCommandlet.cpp b/Engine/Source/Editor/UnrealEd/Private/Commandlets/GatherTextFromAssetsCommandlet.cpp
--- a/Engine/Source/Editor/UnrealEd/Private/Commandlets/GatherTextFromAssetsCommandlet.cpp
+++ b/Engine/Source/Editor/UnrealEd/Private/Commandlets/GatherTextFromAssetsCommandlet.cpp
@@ -416,6 +416,13 @@ bool IsGatherableTextDataIdentical(const TArray<FGatherableTextData>& Gatherable
 	return true;
 }
 
+bool UGatherTextFromAssetsCommandlet::ShouldRunInPreview(const TArray<FString>& Switches, const TMap<FString, FString>& ParamVals) const
+{
+	const FString* GatherType = ParamVals.Find(UGatherTextCommandletBase::GatherTypeParam);
+	// If the param is not specified, it is assumed that both source and assets are to be gathered 
+	return !GatherType || *GatherType == TEXT("Asset") || *GatherType == TEXT("All");
+}
+
 int32 UGatherTextFromAssetsCommandlet::Main(const FString& Params)
 {
 	// Parse command line.
diff --git a/Engine/Source/Editor/UnrealEd/Private/Commandlets/GatherTextFromMetadataCommandlet.cpp b/Engine/Source/Editor/UnrealEd/Private/Commandlets/GatherTextFromMetadataCommandlet.cpp
--- a/Engine/Source/Editor/UnrealEd/Private/Commandlets/GatherTextFromMetadataCommandlet.cpp
+++ b/Engine/Source/Editor/UnrealEd/Private/Commandlets/GatherTextFromMetadataCommandlet.cpp
@@ -22,6 +22,13 @@ UGatherTextFromMetaDataCommandlet::UGatherTextFromMetaDataCommandlet(const FObje
 {
 }
 
+bool UGatherTextFromMetaDataCommandlet::ShouldRunInPreview(const TArray<FString>& Switches, const TMap<FString, FString>& ParamVals) const
+{
+	const FString* GatherType = ParamVals.Find(UGatherTextCommandletBase::GatherTypeParam);
+	// If the param is not specified, it is assumed that both source and assets are to be gathered 
+	return !GatherType || *GatherType == TEXT("Metadata") || *GatherType == TEXT("All");
+}
+
 int32 UGatherTextFromMetaDataCommandlet::Main( const FString& Params )
 {
 	// Parse command line - we're interested in the param vals
diff --git a/Engine/Source/Editor/UnrealEd/Private/Commandlets/GatherTextFromSourceCommandlet.cpp b/Engine/Source/Editor/UnrealEd/Private/Commandlets/GatherTextFromSourceCommandlet.cpp
--- a/Engine/Source/Editor/UnrealEd/Private/Commandlets/GatherTextFromSourceCommandlet.cpp
+++ b/Engine/Source/Editor/UnrealEd/Private/Commandlets/GatherTextFromSourceCommandlet.cpp
@@ -33,6 +33,13 @@ const FString UGatherTextFromSourceCommandlet::FPreProcessorDescriptor::IniNames
 const FString UGatherTextFromSourceCommandlet::FMacroDescriptor::TextMacroString(TEXT("TEXT"));
 const FString UGatherTextFromSourceCommandlet::ChangelistName(TEXT("Update Localization"));
 
+bool UGatherTextFromSourceCommandlet::ShouldRunInPreview(const TArray<FString>& Switches, const TMap<FString, FString>& ParamVals) const
+{
+	const FString* GatherType = ParamVals.Find(UGatherTextCommandletBase::GatherTypeParam);
+	// If the param is not specified, it is assumed that both source and assets are to be gathered 
+	return !GatherType || *GatherType == TEXT("Source") || *GatherType == TEXT("All");
+}
+
 int32 UGatherTextFromSourceCommandlet::Main( const FString& Params )
 {
 	// Parse command line - we're interested in the param vals
@@ -40,7 +47,7 @@ int32 UGatherTextFromSourceCommandlet::Main( const FString& Params )
 	TArray<FString> Switches;
 	TMap<FString, FString> ParamVals;
 	UCommandlet::ParseCommandLine(*Params, Tokens, Switches, ParamVals);
-
+	
 	//Set config file
 	const FString* ParamVal = ParamVals.Find(FString(TEXT("Config")));
 	FString GatherTextConfigPath;
