start fixing texture import for very large textures
can now import 32K dim VT
disable broken PNG alpha fill on 16bit

diff --git a/Engine/Source/Editor/UnrealEd/Private/Factories/EditorFactories.cpp b/Engine/Source/Editor/UnrealEd/Private/Factories/EditorFactories.cpp
--- a/Engine/Source/Editor/UnrealEd/Private/Factories/EditorFactories.cpp
+++ b/Engine/Source/Editor/UnrealEd/Private/Factories/EditorFactories.cpp
@@ -2800,9 +2800,9 @@ public:
 
 	void ProcessData()
 	{
-		int32 NumZeroedTopRowsToProcess = 0;
-		int32 FillColorRow = -1;
-		for (int32 Y = 0; Y < TextureHeight; ++Y)
+		int64 NumZeroedTopRowsToProcess = 0;
+		int64 FillColorRow = -1;
+		for (int64 Y = 0; Y < TextureHeight; ++Y)
 		{
 			if (!ProcessHorizontalRow(Y))
 			{
@@ -2824,7 +2824,7 @@ public:
 		// Can only fill upwards if image not fully zeroed
 		if (NumZeroedTopRowsToProcess > 0 && NumZeroedTopRowsToProcess + 1 < TextureHeight)
 		{
-			for (int32 Y = 0; Y <= NumZeroedTopRowsToProcess; ++Y)
+			for (int64 Y = 0; Y <= NumZeroedTopRowsToProcess; ++Y)
 			{
 				FillRowColorPixels(NumZeroedTopRowsToProcess + 1, Y);
 			}
@@ -2832,15 +2832,16 @@ public:
 	}
 
 	/* returns False if requires further processing because entire row is filled with zeroed alpha values */
-	bool ProcessHorizontalRow(int32 Y)
+	bool ProcessHorizontalRow(int64 Y)
 	{
 		// only wipe out colors that are affected by png turning valid colors white if alpha = 0
 		const uint32 WhiteWithZeroAlpha = FColor(255, 255, 255, 0).DWColor();
+		//@@!! this is wrong for 16 bit channels ;ColorDataType is uint64 but we compare to a uin
 
 		// Left -> Right
-		int32 NumLeftmostZerosToProcess = 0;
+		int64 NumLeftmostZerosToProcess = 0;
 		const PixelDataType* FillColor = nullptr;
-		for (int32 X = 0; X < TextureWidth; ++X)
+		for (int64 X = 0; X < TextureWidth; ++X)
 		{
 			PixelDataType* PixelData = SourceData + (Y * TextureWidth + X) * 4;
 			ColorDataType* ColorData = reinterpret_cast<ColorDataType*>(PixelData);
@@ -2884,7 +2885,7 @@ public:
 		FillColor = SourceData + (Y * TextureWidth + NumLeftmostZerosToProcess + 1) * 4;
 
 		// Fill zero pixels found at beginning of row that could not be filled during the Left to Right pass
-		for (int32 X = 0; X <= NumLeftmostZerosToProcess; ++X)
+		for (int64 X = 0; X <= NumLeftmostZerosToProcess; ++X)
 		{
 			PixelDataType* PixelData = SourceData + (Y * TextureWidth + X) * 4;
 			PixelData[RIdx] = FillColor[RIdx];
@@ -2895,9 +2896,9 @@ public:
 		return true;
 	}
 
-	void FillRowColorPixels(int32 FillColorRow, int32 Y)
+	void FillRowColorPixels(int64 FillColorRow, int64 Y)
 	{
-		for (int32 X = 0; X < TextureWidth; ++X)
+		for (int64 X = 0; X < TextureWidth; ++X)
 		{
 			const PixelDataType* FillColor = SourceData + (FillColorRow * TextureWidth + X) * 4;
 			PixelDataType* PixelData = SourceData + (Y * TextureWidth + X) * 4;
@@ -2908,8 +2909,8 @@ public:
 	}
 
 	PixelDataType* SourceData;
-	int32 TextureWidth;
-	int32 TextureHeight;
+	int64 TextureWidth;
+	int64 TextureHeight;
 };
 
 /**
@@ -2934,8 +2935,10 @@ void FillZeroAlphaPNGData( int32 SizeX, int32 SizeY, ETextureSourceFormat Source
 
 		case TSF_RGBA16:
 		{
-			PNGDataFill<uint16, uint64, 0, 1, 2, 3> PNGFill(SizeX, SizeY, SourceData );
-			PNGFill.ProcessData();
+			// @@!!
+			// this is broken so just disable for now
+			//PNGDataFill<uint16, uint64, 0, 1, 2, 3> PNGFill(SizeX, SizeY, SourceData );
+			//PNGFill.ProcessData();
 			break;
 		}
 
@@ -3573,7 +3576,7 @@ UTexture * UTextureFactory::ImportDDS(const uint8* Buffer,int64 Length,UObject*
 
 	if ( MipCount > MAX_TEXTURE_MIP_COUNT )
 	{
-		// unexpected, because IsImportResolutionValid has already limited rsolution
+		// resolution can be above MAX_TEXTURE_MIP_COUNT for VT
 		UE_LOG(LogEditorFactories,Warning,TEXT("DDS exceeds MAX_TEXTURE_MIP_COUNT"));
 		MipCount = MAX_TEXTURE_MIP_COUNT;
 	}
@@ -4515,21 +4518,35 @@ bool UTextureFactory::IsImportResolutionValid(int64 Width, int64 Height, bool bA
 {
 	static const auto CVarVirtualTexturesEnabled = IConsoleManager::Get().FindTConsoleVariableDataInt(TEXT("r.VirtualTextures")); check(CVarVirtualTexturesEnabled);
 
-	// In theory this value could be much higher, but various UE image code currently uses 32bit size/offset values
-	// TODO: in theory we can try making this bigger and fix any int32 overflow issues
-	const int64 MaximumSupportedVirtualTextureResolution = 16 * 1024;
+	//const int64 MaximumSupportedVirtualTextureResolution = 16 * 1024;
+	const int64 MaximumSupportedVirtualTextureResolution = 32768;
 
 	// Calculate the maximum supported resolution utilizing the global max texture mip count
 	// (Note, have to subtract 1 because 1x1 is a valid mip-size; this means a GMaxTextureMipCount of 4 means a max resolution of 8x8, not 2^4 = 16x16)
 	int64 MaximumSupportedResolution = CVarVirtualTexturesEnabled->GetValueOnAnyThread() ? MaximumSupportedVirtualTextureResolution : ((int64)1 << (GMaxTextureMipCount - 1));
 
-	// MaximumSupportedResolution is also limited by Oodle Texture maximum of 16k
+	// MaximumSupportedResolution is also limited by Oodle Texture maximum of 16k (OODLETEX_MAX_SURFACE_DIMENSION)
 	MaximumSupportedResolution = FMath::Min(MaximumSupportedResolution,(int64)16384);
+	
+	// MAX_TEXTURE_MIP_COUNT currently = 15 , so this is 16384
+	// VT can exceed the size limit from MAX_TEXTURE_MIP_COUNT
+	const int64 MaximumSizeFromMipCount = 1 << (MAX_TEXTURE_MIP_COUNT - 1);
+	MaximumSupportedResolution = FMath::Min(MaximumSupportedResolution,(int64)MaximumSizeFromMipCount);
 
-	bool bValid = true;
 	
-	// Check if the texture is above the supported resolution and prompt the user if they wish to continue if it is
+	// No zero-size textures :
+	if (Width == 0 || Height == 0 )
+	{
+		Warn->Log(ELogVerbosity::Error, *FText::Format(
+			NSLOCTEXT("UnrealEd", "Warning_TextureSizeTooLargeOrInvalid", "Texture has zero width or height"),
+			FText::AsNumber(FMath::Square(MaximumSupportedVirtualTextureResolution))
+			).ToString());
+
+		return false;
+	}
 
+	// Dimensions must fit in signed int32
+	//  could be negative here if it was over 2G and int32 was used earlier
 	if (Width < 0 || Height < 0 || Width > MAX_int32 || Height > MAX_int32)
 	{
 		Warn->Log(ELogVerbosity::Error, *FText::Format(
@@ -4537,40 +4554,44 @@ bool UTextureFactory::IsImportResolutionValid(int64 Width, int64 Height, bool bA
 			FText::AsNumber(FMath::Square(MaximumSupportedVirtualTextureResolution))
 			).ToString());
 
-		bValid = false;
+		return false;
 	}
 
 	// pixel count must fit in int32 :
 	//  mip surface could still be larger than 2 GB, that's allowed
 	//	eg. 16k RGBA float = 4 GB
-	if ( bValid && Width * Height > MAX_int32)
+	if ( Width * Height > MAX_int32)
 	{
 		Warn->Log(ELogVerbosity::Error, *FText::Format(
 			NSLOCTEXT("UnrealEd", "Warning_TextureSizeTooLargeOrInvalid", "Texture is too large to import or it has an invalid resolution. The current maximun is {0} pixels"),
 			FText::AsNumber(FMath::Square(MaximumSupportedVirtualTextureResolution))
 			).ToString());
 
-		bValid = false;
+		return false;
 	}
 
 	if ( Width > MaximumSupportedResolution || Height > MaximumSupportedResolution )
 	{
-
-		if ( bValid && (Width * Height) > FMath::Square(MaximumSupportedVirtualTextureResolution))
+		if ( (Width * Height) > FMath::Square(MaximumSupportedVirtualTextureResolution))
 		{
 			Warn->Log(ELogVerbosity::Error, *FText::Format(
 				NSLOCTEXT("UnrealEd", "Warning_TextureSizeTooLarge", "Texture is too large to import. The current maximun is {0} pixels"),
 				FText::AsNumber(FMath::Square(MaximumSupportedVirtualTextureResolution))
 				).ToString());
 
-			bValid = false;
+			return false;
 		}
 
-		if ( bValid && EAppReturnType::Yes != FMessageDialog::Open( EAppMsgType::YesNo, FText::Format(
-				NSLOCTEXT("UnrealEd", "Warning_LargeTextureImport", "Attempting to import {0} x {1} texture, proceed?\nLargest supported texture size: {2} x {3}"),
+		// we're larger than MaximumSupportedResolution
+		// but not larger than MaximumSupportedVirtualTextureResolution
+		// so this texture can still work, but only as VT
+		// prompt about this :
+
+		if ( EAppReturnType::Yes != FMessageDialog::Open( EAppMsgType::YesNo, FText::Format(
+				NSLOCTEXT("UnrealEd", "Warning_LargeTextureImport", "Attempting to import {0} x {1} texture, proceed?\nLargest supported non-VT texture size: {2} x {3}"),
 				FText::AsNumber(Width), FText::AsNumber(Height), FText::AsNumber(MaximumSupportedResolution), FText::AsNumber(MaximumSupportedResolution)) ) )
 		{
-			bValid = false;
+			return false;
 		}
 	}
 
@@ -4579,10 +4600,10 @@ bool UTextureFactory::IsImportResolutionValid(int64 Width, int64 Height, bool bA
 	if ( !bAllowNonPowerOfTwo && !bIsPowerOfTwo )
 	{
 		Warn->Log(ELogVerbosity::Error, *NSLOCTEXT("UnrealEd", "Warning_TextureNotAPowerOfTwo", "Cannot import texture with non-power of two dimensions").ToString() );
-		bValid = false;
+		return false;
 	}
 	
-	return bValid;
+	return true;
 }
 
 IImportSettingsParser* UTextureFactory::GetImportSettingsParser()
diff --git a/Engine/Source/Editor/UnrealEd/Private/Factories/Factory.cpp b/Engine/Source/Editor/UnrealEd/Private/Factories/Factory.cpp
--- a/Engine/Source/Editor/UnrealEd/Private/Factories/Factory.cpp
+++ b/Engine/Source/Editor/UnrealEd/Private/Factories/Factory.cpp
@@ -99,13 +99,15 @@ UObject* UFactory::FactoryCreateFile(UClass* InClass, UObject* InParent, FName I
 
 	// load as binary
 	{
-		TArray<uint8> Data;
+		TArray64<uint8> Data;
 		if (!FFileHelper::LoadFileToArray(Data, *Filename))
 		{
 			UE_LOG(LogFactory, Error, TEXT("Failed to load file '%s' to array"), *Filename);
 			return nullptr;
 		}
 
+		// adds an extra null byte
+		// data loaders must be able to ignore an unexpected extra zero byte at the end of the file
 		Data.Add(0);
 		ParseParms(Parms);
 		const uint8* Ptr = &Data[0];
diff --git a/Engine/Source/Runtime/TextureUtilitiesCommon/Private/NormalMapIdentification.cpp b/Engine/Source/Runtime/TextureUtilitiesCommon/Private/NormalMapIdentification.cpp
--- a/Engine/Source/Runtime/TextureUtilitiesCommon/Private/NormalMapIdentification.cpp
+++ b/Engine/Source/Runtime/TextureUtilitiesCommon/Private/NormalMapIdentification.cpp
@@ -83,8 +83,8 @@ public:
 	}
 
 	UTexture* SourceTexture;
-	int32 TextureSizeX;
-	int32 TextureSizeY;
+	int64 TextureSizeX;
+	int64 TextureSizeY;
 	const uint8* SourceTextureData;
 };
 
