- Added cooker option, -DlcReevaluateUncookedAssets, which will allow a cooked editor to attempt to cook assets that the base game had attempted to cook and skipped, but left in the assetregistry.
#p4v-preflight-copy 19482053

diff --git a/Engine/Source/Developer/IoStoreUtilities/Private/IoStoreUtilities.cpp b/Engine/Source/Developer/IoStoreUtilities/Private/IoStoreUtilities.cpp
--- a/Engine/Source/Developer/IoStoreUtilities/Private/IoStoreUtilities.cpp
+++ b/Engine/Source/Developer/IoStoreUtilities/Private/IoStoreUtilities.cpp
@@ -5859,8 +5859,17 @@ int32 CreateIoStoreContainerFiles(const TCHAR* CmdLine)
 
 					for (FName PackageName : PackageNames)
 					{
-						Arguments.ReleasedPackages.PackageNames.Add(PackageName);
-						Arguments.ReleasedPackages.PackageIdToName.Add(FPackageId::FromName(PackageName), PackageName);
+						// skip over packages that were not actually saved out, but were added to the AR - the DLC may now have those packages included,
+						// and there will be a conflict later on if the package is in this list and the DLC list. PackageFlags of 0 means it was 
+						// evaluated and skipped.
+						TArrayView<FAssetData const* const> AssetsForPackage = ReleaseAssetRegistry.GetAssetsByPackageName(PackageName);
+						checkf(AssetsForPackage.Num() > 0, TEXT("It is unexpected that no assets were found in DevelopmentAssetRegistry for the package %s. This indicates an invalid AR."), *PackageName.ToString());
+						// just check the first one in the list, they will all have the same flags
+						if (AssetsForPackage[0]->PackageFlags != 0)
+						{
+							Arguments.ReleasedPackages.PackageNames.Add(PackageName);
+							Arguments.ReleasedPackages.PackageIdToName.Add(FPackageId::FromName(PackageName), PackageName);
+						}
 					}
 				}
 			}
diff --git a/Engine/Source/Editor/UnrealEd/Classes/CookOnTheSide/CookOnTheFlyServer.h b/Engine/Source/Editor/UnrealEd/Classes/CookOnTheSide/CookOnTheFlyServer.h
--- a/Engine/Source/Editor/UnrealEd/Classes/CookOnTheSide/CookOnTheFlyServer.h
+++ b/Engine/Source/Editor/UnrealEd/Classes/CookOnTheSide/CookOnTheFlyServer.h
@@ -77,6 +77,7 @@ enum class ECookByTheBookOptions
 	CookAgainstFixedBase =				0x00010000, // If cooking DLC, assume that the base content can not be modified. 
 	DlcLoadMainAssetRegistry =			0x00020000, // If cooking DLC, populate the main game asset registry
 	ZenStore =							0x00040000, // Store cooked data in Zen Store
+	DlcReevaluateUncookedAssets =		0x00080000, // If cooking DLC, ignore assets in the base asset registry that were not cooked, so that this cook has an opportunity to cook the assets
 
 	// Deprecated flags
 	NoSlatePackages UE_DEPRECATED(5.0, "The [UI]ContentDirectories is deprecated. You may use DirectoriesToAlwaysCook in your project settings instead.") = 0x00000400, // don't include slate content
@@ -781,11 +782,12 @@ private:
 	*
 	* @param AssetRegistryPath path of the assetregistry.bin file to read
 	* @param bVerifyPackagesExist whether or not we should verify the packages exist on disk.
+	* @param bSkipUncookedPackages whether or not we should skip over packages that are in the AR but were not cooked/saved
 	* @param OutPackageDatas out list of packagename/filenames for packages contained in the asset registry file
 	* @return true if successfully read false otherwise
 	*/
 	bool GetAllPackageFilenamesFromAssetRegistry( const FString& AssetRegistryPath, bool bVerifyPackagesExist,
-		TArray<UE::Cook::FConstructPackageData>& OutPackageDatas) const;
+		 bool bSkipUncookedPackages, TArray<UE::Cook::FConstructPackageData>& OutPackageDatas) const;
 
 	/**
 	* BuildMapDependencyGraph
diff --git a/Engine/Source/Editor/UnrealEd/Private/Commandlets/CookCommandlet.cpp b/Engine/Source/Editor/UnrealEd/Private/Commandlets/CookCommandlet.cpp
--- a/Engine/Source/Editor/UnrealEd/Private/Commandlets/CookCommandlet.cpp
+++ b/Engine/Source/Editor/UnrealEd/Private/Commandlets/CookCommandlet.cpp
@@ -916,6 +916,7 @@ bool UCookCommandlet::CookByTheBook( const TArray<ITargetPlatform*>& Platforms)
 	CookOptions |= Switches.Contains(TEXT("SkipHardReferences")) ? ECookByTheBookOptions::SkipHardReferences : ECookByTheBookOptions::None;
 	CookOptions |= Switches.Contains(TEXT("CookAgainstFixedBase")) ? ECookByTheBookOptions::CookAgainstFixedBase : ECookByTheBookOptions::None;
 	CookOptions |= (Switches.Contains(TEXT("DlcLoadMainAssetRegistry")) || !bErrorOnEngineContentUse) ? ECookByTheBookOptions::DlcLoadMainAssetRegistry : ECookByTheBookOptions::None;
+	CookOptions |= Switches.Contains(TEXT("DlcReevaluateUncookedAssets")) ? ECookByTheBookOptions::DlcReevaluateUncookedAssets : ECookByTheBookOptions::None;
 
 	if (bCookSinglePackage)
 	{
diff --git a/Engine/Source/Editor/UnrealEd/Private/CookOnTheFlyServer.cpp b/Engine/Source/Editor/UnrealEd/Private/CookOnTheFlyServer.cpp
--- a/Engine/Source/Editor/UnrealEd/Private/CookOnTheFlyServer.cpp
+++ b/Engine/Source/Editor/UnrealEd/Private/CookOnTheFlyServer.cpp
@@ -7956,11 +7956,12 @@ void UCookOnTheFlyServer::StartCookByTheBook( const FCookByTheBookStartupOptions
 	{
 		// If we're cooking against a fixed base, we don't need to verify the packages exist on disk, we simply want to use the Release Data 
 		const bool bVerifyPackagesExist = !IsCookingAgainstFixedBase();
+		const bool bReevaluateUncookedPackages = !!(CookOptions & ECookByTheBookOptions::DlcReevaluateUncookedAssets);
 
 		// if we are cooking dlc we must be based on a release version cook
 		check( !BasedOnReleaseVersion.IsEmpty() );
 
-		auto ReadDevelopmentAssetRegistry = [this, &BasedOnReleaseVersion, bVerifyPackagesExist]
+		auto ReadDevelopmentAssetRegistry = [this, &BasedOnReleaseVersion, bVerifyPackagesExist, bReevaluateUncookedPackages]
 			(TArray<UE::Cook::FConstructPackageData>& OutPackageList, const FString& InPlatformName)
 		{
 			TArray<FString> AttemptedNames;
@@ -7968,12 +7969,12 @@ void UCookOnTheFlyServer::StartCookByTheBook( const FCookByTheBookStartupOptions
 			AttemptedNames.Add(OriginalSandboxRegistryFilename);
 
 			// if this check fails probably because the asset registry can't be found or read
-			bool bSucceeded = GetAllPackageFilenamesFromAssetRegistry(OriginalSandboxRegistryFilename, bVerifyPackagesExist, OutPackageList);
+			bool bSucceeded = GetAllPackageFilenamesFromAssetRegistry(OriginalSandboxRegistryFilename, bVerifyPackagesExist, bReevaluateUncookedPackages, OutPackageList);
 			if (!bSucceeded)
 			{
 				OriginalSandboxRegistryFilename = GetBasedOnReleaseVersionAssetRegistryPath(BasedOnReleaseVersion, InPlatformName) / GetAssetRegistryFilename();
 				AttemptedNames.Add(OriginalSandboxRegistryFilename);
-				bSucceeded = GetAllPackageFilenamesFromAssetRegistry(OriginalSandboxRegistryFilename, bVerifyPackagesExist, OutPackageList);
+				bSucceeded = GetAllPackageFilenamesFromAssetRegistry(OriginalSandboxRegistryFilename, bVerifyPackagesExist, bReevaluateUncookedPackages, OutPackageList);
 			}
 
 			if (!bSucceeded)
@@ -7985,7 +7986,7 @@ void UCookOnTheFlyServer::StartCookByTheBook( const FCookByTheBookStartupOptions
 					{
 						OriginalSandboxRegistryFilename = GetBasedOnReleaseVersionAssetRegistryPath(BasedOnReleaseVersion, PlatformFlavor->Name.ToString()) / GetAssetRegistryFilename();
 						AttemptedNames.Add(OriginalSandboxRegistryFilename);
-						bSucceeded = GetAllPackageFilenamesFromAssetRegistry(OriginalSandboxRegistryFilename, bVerifyPackagesExist, OutPackageList);
+						bSucceeded = GetAllPackageFilenamesFromAssetRegistry(OriginalSandboxRegistryFilename, bVerifyPackagesExist, bReevaluateUncookedPackages, OutPackageList);
 						if (bSucceeded)
 						{
 							break;
@@ -8198,7 +8199,7 @@ void UCookOnTheFlyServer::StartCookByTheBook( const FCookByTheBookStartupOptions
 				FString OriginalAssetRegistryPath = GetBasedOnReleaseVersionAssetRegistryPath( BasedOnReleaseVersion, TargetPlatform->PlatformName() ) / GetAssetRegistryFilename();
 
 				TArray<UE::Cook::FConstructPackageData> BasedOnReleaseDatas;
-				verify( GetAllPackageFilenamesFromAssetRegistry(OriginalAssetRegistryPath, true, BasedOnReleaseDatas) );
+				verify( GetAllPackageFilenamesFromAssetRegistry(OriginalAssetRegistryPath, true, false, BasedOnReleaseDatas) );
 
 				TArray<const ITargetPlatform*, TInlineAllocator<1>> RequestPlatforms;
 				RequestPlatforms.Add(TargetPlatform);
@@ -8432,8 +8433,8 @@ void UCookOnTheFlyServer::GetCookOnTheFlyUnsolicitedFiles(const ITargetPlatform*
 	}
 }
 
-bool UCookOnTheFlyServer::GetAllPackageFilenamesFromAssetRegistry(const FString& AssetRegistryPath,
-	bool bVerifyPackagesExist, TArray<UE::Cook::FConstructPackageData>& OutPackageDatas) const
+bool UCookOnTheFlyServer::GetAllPackageFilenamesFromAssetRegistry(const FString& AssetRegistryPath, bool bVerifyPackagesExist,
+	bool bReevaluateUncookedPackages, TArray<UE::Cook::FConstructPackageData>& OutPackageDatas) const
 {
 	using namespace UE::Cook;
 
@@ -8465,7 +8466,22 @@ bool UCookOnTheFlyServer::GetAllPackageFilenamesFromAssetRegistry(const FString&
 		}
 
 		// possibly use the preloaded state
-		const TMap<FName, const FAssetData*>& RegistryDataMap = GPreloadedARState.GetObjectPathToAssetDataMap();
+		const TMap<FName, const FAssetData*>& RegistryDataMapLoaded = GPreloadedARState.GetObjectPathToAssetDataMap();
+		TMap<FName, const FAssetData*> RegistryDataMapTrimmed;
+		
+		// if we want to reevaluate uncooked pacakges, then remove the uncooked packages from the set of known packages
+		// removing them from the map up here simplifies the logic below 
+		if (bReevaluateUncookedPackages)
+		{
+			for (TPair<FName, const FAssetData*> Pair : RegistryDataMapLoaded)
+			{
+				if (Pair.Value->PackageFlags != 0)
+				{
+					RegistryDataMapTrimmed.Add(Pair.Key, Pair.Value);
+				}
+			}
+		}
+		const TMap<FName, const FAssetData*>& RegistryDataMap = bReevaluateUncookedPackages ? RegistryDataMapTrimmed : RegistryDataMapLoaded;
 
 		check(OutPackageDatas.Num() == 0);
 
diff --git a/Engine/Source/Programs/AutomationTool/CookedEditor/CookedEditor.Automation.cs b/Engine/Source/Programs/AutomationTool/CookedEditor/CookedEditor.Automation.cs
--- a/Engine/Source/Programs/AutomationTool/CookedEditor/CookedEditor.Automation.cs
+++ b/Engine/Source/Programs/AutomationTool/CookedEditor/CookedEditor.Automation.cs
@@ -988,6 +988,11 @@ public class MakeCookedEditor : BuildCommand
 		// cause the issues
 		Params.AdditionalCookerOptions += " -AllowUnsafeBlueprintCalls";
 
+		// if we are cooking the editor in Dlc mode, then we want to attempt to cook assets that the base game may have skipped, but still
+		// put into the AssetRegistry, so by default the DLC cooker will skip them. this will make the DLC cooker reevaluate these pacakges
+		// and choose to cook them or not (generally for editoronly assets)
+		Params.AdditionalCookerOptions += " -DlcReevaluateUncookedAssets";
+
 		// set up cooking against a client, as DLC
 		if (bIsDLC)
 		{
