Fix ForceReplaceReferences nulling out references

#rnx

diff --git a/Engine/Source/Editor/UnrealEd/Private/ObjectTools.cpp b/Engine/Source/Editor/UnrealEd/Private/ObjectTools.cpp
--- a/Engine/Source/Editor/UnrealEd/Private/ObjectTools.cpp
+++ b/Engine/Source/Editor/UnrealEd/Private/ObjectTools.cpp
@@ -1123,7 +1123,6 @@ namespace ObjectTools
 		TArray<UObject*> ReferencingPropertiesMapKeys;
 		TArray<PropertyArrayType> ReferencingPropertiesMapValues;
 
-		if (!bOnlyNullingOut)
 		{
 			// Find the referencers of the objects to be replaced
 			FFindReferencersArchive FindRefsArchive( nullptr, OutInfo.ReplaceableObjects );
@@ -1140,7 +1139,8 @@ namespace ObjectTools
 					continue;
 				}
 
-				UObject** ObjectToReplaceWithPtr = ReplacementMap.Find(CurObject);
+				// Unless the "object to replace with" is null, ignore the objects being replaced
+				UObject** ObjectToReplaceWithPtr = bOnlyNullingOut ? nullptr : ReplacementMap.Find(CurObject);
 				if (ObjectToReplaceWithPtr == nullptr || *ObjectToReplaceWithPtr == nullptr)
 				{
 					FindRefsArchive.ResetPotentialReferencer(CurObject);
