CookCommandlet: Remove commented-out code.
#rnx

diff --git a/Engine/Source/Editor/UnrealEd/Private/Commandlets/CookCommandlet.cpp b/Engine/Source/Editor/UnrealEd/Private/Commandlets/CookCommandlet.cpp
--- a/Engine/Source/Editor/UnrealEd/Private/Commandlets/CookCommandlet.cpp
+++ b/Engine/Source/Editor/UnrealEd/Private/Commandlets/CookCommandlet.cpp
@@ -826,35 +826,6 @@ bool UCookCommandlet::CookByTheBook( const TArray<ITargetPlatform*>& Platforms)
 
 	CookOnTheFlyServer->Initialize(ECookMode::CookByTheBook, CookFlags, OutputDirectoryOverride);
 
-	// Add any map sections specified on command line
-	TArray<FString> AlwaysCookMapList;
-
-	// Add the default map section
-	//GEditor->LoadMapListFromIni(TEXT("AlwaysCookMaps"), AlwaysCookMapList);
-
-	TArray<FString> MapList;
-	// Add any map sections specified on command line
-	/*GEditor->ParseMapSectionIni(*Params, MapList);
-
-	if (MapList.Num() == 0 && !bCookSinglePackage)
-	{
-		// If we didn't find any maps look in the project settings for maps
-
-		UProjectPackagingSettings* PackagingSettings = Cast<UProjectPackagingSettings>(UProjectPackagingSettings::StaticClass()->GetDefaultObject());
-
-		for (const auto& MapToCook : PackagingSettings->MapsToCook)
-		{
-			MapList.Add(MapToCook.FilePath);
-		}
-	}*/
-
-	// Add any map specified on the command line.
-	for (const auto& MapName : CmdLineMapEntries)
-	{
-		MapList.Add(MapName);
-	}
-
-	
 	TArray<FString> MapIniSections;
 	FString SectionStr;
 	if (FParse::Value(*Params, TEXT("MAPINISECTION="), SectionStr))
@@ -874,13 +845,6 @@ bool UCookCommandlet::CookByTheBook( const TArray<ITargetPlatform*>& Platforms)
 		}
 	}
 
-	if (!bCookSinglePackage)
-	{
-		// Put the always cook map list at the front of the map list
-		AlwaysCookMapList.Append(MapList);
-		Swap(MapList, AlwaysCookMapList);
-	}
-
 	// Set the list of cultures to cook as those on the commandline, if specified.
 	// Otherwise, use the project packaging settings.
 	TArray<FString> CookCultures;
@@ -966,7 +930,7 @@ bool UCookCommandlet::CookByTheBook( const TArray<ITargetPlatform*>& Platforms)
 		return false;
 	}
 
-	Swap( StartupOptions.CookMaps, MapList );
+	Swap( StartupOptions.CookMaps, CmdLineMapEntries);
 	Swap( StartupOptions.CookDirectories, CmdLineDirEntries );
 	Swap( StartupOptions.NeverCookDirectories, CmdLineNeverCookDirEntries);
 	Swap( StartupOptions.CookCultures, CookCultures );
