Add special case to FKismetDebugUtilities::FindClassPropertyForPin for reroute nodes

diff --git a/Engine/Source/Editor/UnrealEd/Private/Kismet2/KismetDebugUtilities.cpp b/Engine/Source/Editor/UnrealEd/Private/Kismet2/KismetDebugUtilities.cpp
--- a/Engine/Source/Editor/UnrealEd/Private/Kismet2/KismetDebugUtilities.cpp
+++ b/Engine/Source/Editor/UnrealEd/Private/Kismet2/KismetDebugUtilities.cpp
@@ -1297,6 +1297,27 @@ FProperty* FKismetDebugUtilities::FindClassPropertyForPin(UBlueprint* Blueprint,
 {
 	FProperty* FoundProperty = nullptr;
 
+	if (Pin)
+	{
+		// Input Pins linked to a reroute node should use debug property from the reroute node's input pin
+		if (Pin->Direction == EGPD_Input && Pin->LinkedTo.Num() == 1)
+		{
+			if (const UK2Node_Knot* LinkedKnot = Cast<UK2Node_Knot>(Pin->LinkedTo[0]->GetOwningNode()))
+			{
+				return FindClassPropertyForPin(Blueprint, LinkedKnot->GetInputPin());
+			}
+		}
+
+		// Reroute nodes should always use the input pin, not the output pin
+		if (const UK2Node_Knot* OwningKnot = Cast<UK2Node_Knot>(Pin->GetOwningNode()))
+		{
+			if (Pin->Direction == EGPD_Output)
+			{
+				return FindClassPropertyForPin(Blueprint, OwningKnot->GetInputPin());
+			}
+		}
+	}
+
 	UClass* Class = Blueprint->GeneratedClass;
 	while (UBlueprintGeneratedClass* BlueprintClass = Cast<UBlueprintGeneratedClass>(Class))
 	{
