Level Instance now support Partitioned worlds with disabled streaming.
Newly created Level Instances setups the worlds to be partitioned with disabled streaming.
This is in preparation of adding DataLayer support inside Level Instances.

diff --git a/Engine/Source/Editor/UnrealEd/Private/EditorLevelUtils.cpp b/Engine/Source/Editor/UnrealEd/Private/EditorLevelUtils.cpp
--- a/Engine/Source/Editor/UnrealEd/Private/EditorLevelUtils.cpp
+++ b/Engine/Source/Editor/UnrealEd/Private/EditorLevelUtils.cpp
@@ -663,14 +663,13 @@ ULevel* GetPersistentLevelForNewStreamingLevel(UWorld* NewLevelWorld, UWorld* In
 	return NewLevelWorld ? NewLevelWorld->PersistentLevel : InTemplateWorld->PersistentLevel;
 }
 
-UWorld* GetWorldForNewStreamingLevel(UWorld* InTemplateWorld)
+UWorld* GetWorldForNewStreamingLevel(UWorld* InTemplateWorld, bool bIsPartitioned = false)
 {
 	if (!InTemplateWorld)
 	{
 		// Create a new world
 		UWorldFactory* Factory = NewObject<UWorldFactory>();
-		// streaming levels shouldn't be WP
-		Factory->bCreateWorldPartition = false;
+		Factory->bCreateWorldPartition = bIsPartitioned;
 		Factory->WorldType = EWorldType::Inactive;
 		UPackage* Pkg = CreatePackage( NULL);
 		FName WorldName(TEXT("Untitled"));
@@ -686,7 +685,7 @@ UWorld* GetWorldForNewStreamingLevel(UWorld* InTemplateWorld)
 }
 };
 
-ULevelStreaming* UEditorLevelUtils::CreateNewStreamingLevelForWorld(UWorld& InWorld, TSubclassOf<ULevelStreaming> LevelStreamingClass, const FString& DefaultFilename /* = TEXT( "" ) */, bool bMoveSelectedActorsIntoNewLevel /* = false */, UWorld* InTemplateWorld /* = nullptr */, bool bInUseSaveAs /*= true*/)
+ULevelStreaming* UEditorLevelUtils::CreateNewStreamingLevelForWorld(UWorld& InWorld, TSubclassOf<ULevelStreaming> LevelStreamingClass, const FString& DefaultFilename /* = TEXT( "" ) */, bool bMoveSelectedActorsIntoNewLevel /* = false */, UWorld* InTemplateWorld /* = nullptr */, bool bInUseSaveAs /*= true*/, TFunction<void(ULevel*)> InPreSaveLevelOperation /* = TFunction<void(ULevel*)>()*/)
 {
 	TArray<AActor*> ActorsToMove;
 	if (bMoveSelectedActorsIntoNewLevel)
@@ -701,10 +700,10 @@ ULevelStreaming* UEditorLevelUtils::CreateNewStreamingLevelForWorld(UWorld& InWo
 		}
 	}
 
-	return CreateNewStreamingLevelForWorld(InWorld, LevelStreamingClass, /*bUseExternalActors=*/false, DefaultFilename, &ActorsToMove, InTemplateWorld, bInUseSaveAs);
+	return CreateNewStreamingLevelForWorld(InWorld, LevelStreamingClass, /*bUseExternalActors=*/false, DefaultFilename, &ActorsToMove, InTemplateWorld, bInUseSaveAs, /*bIsPartitioned=*/false, InPreSaveLevelOperation);
 }
 
-ULevelStreaming* UEditorLevelUtils::CreateNewStreamingLevelForWorld(UWorld& InWorld, TSubclassOf<ULevelStreaming> LevelStreamingClass, bool bUseExternalActors, const FString& DefaultFilename /* = TEXT("") */, const TArray<AActor*>* ActorsToMove /* = nullptr */, UWorld* InTemplateWorld /* = nullptr */, bool bInUseSaveAs /*= true*/)
+ULevelStreaming* UEditorLevelUtils::CreateNewStreamingLevelForWorld(UWorld& InWorld, TSubclassOf<ULevelStreaming> LevelStreamingClass, bool bUseExternalActors, const FString& DefaultFilename /* = TEXT("") */, const TArray<AActor*>* ActorsToMove /* = nullptr */, UWorld* InTemplateWorld /* = nullptr */, bool bInUseSaveAs /*= true*/, bool bIsPartitioned /*= false*/, TFunction<void(ULevel*)> InPreSaveLevelOperation /* = TFunction<void(ULevel*)>()*/)
 {
 	// Editor modes cannot be active when any level saving occurs.
 	if (!IsRunningCommandlet())
@@ -717,13 +716,18 @@ ULevelStreaming* UEditorLevelUtils::CreateNewStreamingLevelForWorld(UWorld& InWo
 	UWorld* WorldToAddLevelTo = &InWorld;
 
 	// This is the new streaming level's world not the persistent level world
-	UWorld* NewLevelWorld = GetWorldForNewStreamingLevel(InTemplateWorld);
+	UWorld* NewLevelWorld = GetWorldForNewStreamingLevel(InTemplateWorld, bIsPartitioned);
 	ULevel* PersistentLevel = GetPersistentLevelForNewStreamingLevel(NewLevelWorld, InTemplateWorld);
-	if (bUseExternalActors)
+	// No need to convert actors in partitioned worlds as they are already external
+	if (bUseExternalActors && !bIsPartitioned)
 	{
 		PersistentLevel->ConvertAllActorsToPackaging(true);
 		PersistentLevel->bUseExternalActors = true;
 	}
+	check(bIsPartitioned == PersistentLevel->bIsPartitioned);
+
+	// Call lambda before saving level
+	InPreSaveLevelOperation(PersistentLevel);
 
 	bool bNewWorldSaved = false;
 	FString NewPackageName = DefaultFilename;
diff --git a/Engine/Source/Editor/UnrealEd/Private/FileHelpers.cpp b/Engine/Source/Editor/UnrealEd/Private/FileHelpers.cpp
--- a/Engine/Source/Editor/UnrealEd/Private/FileHelpers.cpp
+++ b/Engine/Source/Editor/UnrealEd/Private/FileHelpers.cpp
@@ -751,7 +751,8 @@ static bool SaveWorld(UWorld* World,
 				RenamedWorldPartition = World->GetWorldPartition();
 
 				// Load all unloaded actors before rename. If this is causing issues (oom or other) map will need to be renamed through a provided builder commandlet
-				if (RenamedWorldPartition)
+				// When creating a Partitioned Level for a Level Instance, the WorldPartition is not initialized, so no need to do this.
+				if (RenamedWorldPartition && RenamedWorldPartition->IsInitialized())
 				{
 					LoadedEditorCells = RenamedWorldPartition->GetUserLoadedEditorGridCells();
 					RenamedWorldPartition->LoadAllActors(ActorReferences);
diff --git a/Engine/Source/Editor/UnrealEd/Public/EditorLevelUtils.h b/Engine/Source/Editor/UnrealEd/Public/EditorLevelUtils.h
--- a/Engine/Source/Editor/UnrealEd/Public/EditorLevelUtils.h
+++ b/Engine/Source/Editor/UnrealEd/Public/EditorLevelUtils.h
@@ -144,10 +144,11 @@ public:
 	 * @param	bMoveSelectedActorsIntoNewLevel		If true, move any selected actors into the new level.
 	 * @param	InTemplateWorld						If valid, the new level will be a copy of the template world.
 	 * @param	bInUseSaveAs						If true, show SaveAs dialog instead of Save with DefaultFilename
+	 * @param	InPreSaveLevelOperation				Optional function to call before saving the created level
 	 * 
 	 * @return	Returns the newly created level, or NULL on failure
 	 */
-	static UNREALED_API ULevelStreaming* CreateNewStreamingLevelForWorld(UWorld& World, TSubclassOf<ULevelStreaming> LevelStreamingClass, const FString& DefaultFilename = TEXT(""), bool bMoveSelectedActorsIntoNewLevel = false, UWorld* InTemplateWorld = nullptr, bool bInUseSaveAs = true);
+	static UNREALED_API ULevelStreaming* CreateNewStreamingLevelForWorld(UWorld& World, TSubclassOf<ULevelStreaming> LevelStreamingClass, const FString& DefaultFilename = TEXT(""), bool bMoveSelectedActorsIntoNewLevel = false, UWorld* InTemplateWorld = nullptr, bool bInUseSaveAs = true, TFunction<void(ULevel*)> InPreSaveLevelOperation = TFunction<void(ULevel*)>());
 
 	/**
 	 * Creates a new streaming level and adds it to a world
@@ -159,10 +160,12 @@ public:
 	 * @param	ActorsToMove						Optional, move provided actors into the new level.
 	 * @param	InTemplateWorld						If valid, the new level will be a copy of the template world.
 	 * @param	bInUseSaveAs						If true, show SaveAs dialog instead of Save with DefaultFilename
+	 * @param	bIsPartitioned						If level should be partitioned (has precedence over bUseExternalActors).
+	 * @param	InPreSaveLevelOperation				Optional function to call before saving the created level
 	 *
 	 * @return	Returns the newly created level, or NULL on failure
 	 */
-	static UNREALED_API ULevelStreaming* CreateNewStreamingLevelForWorld(UWorld& World, TSubclassOf<ULevelStreaming> LevelStreamingClass, bool bUseExternalActors, const FString& DefaultFilename, const TArray<AActor*>* ActorsToMove = nullptr, UWorld* InTemplateWorld = nullptr, bool bInUseSaveAs = true);
+	static UNREALED_API ULevelStreaming* CreateNewStreamingLevelForWorld(UWorld& World, TSubclassOf<ULevelStreaming> LevelStreamingClass, bool bUseExternalActors, const FString& DefaultFilename, const TArray<AActor*>* ActorsToMove = nullptr, UWorld* InTemplateWorld = nullptr, bool bInUseSaveAs = true, bool bIsPartitioned = false, TFunction<void(ULevel*)> InPreSaveLevelOperation = TFunction<void(ULevel*)>());
 
 	/**
 	 * Adds the named level packages to the world.  Does nothing if all the levels already exist in the world.
diff --git a/Engine/Source/Runtime/Engine/Classes/Engine/Level.h b/Engine/Source/Runtime/Engine/Classes/Engine/Level.h
--- a/Engine/Source/Runtime/Engine/Classes/Engine/Level.h
+++ b/Engine/Source/Runtime/Engine/Classes/Engine/Level.h
@@ -777,6 +777,8 @@ public:
 	ENGINE_API static bool GetIsLevelUsingExternalActorsFromPackage(FName LevelPackage);
 	ENGINE_API static bool GetIsUsingActorFoldersFromPackage(FName LevelPackage);
 	ENGINE_API static bool GetIsStreamingDisabledFromPackage(FName LevelPackage);
+	ENGINE_API static bool GetPartitionedLevelCanBeUsedByLevelInstanceFromAsset(const FAssetData& Asset);
+	ENGINE_API static bool GetPartitionedLevelCanBeUsedByLevelInstanceFromPackage(FName LevelPackage);
 
 	ENGINE_API bool GetPromptWhenAddingToLevelOutsideBounds() const;
 	ENGINE_API bool GetPromptWhenAddingToLevelBeforeCheckout() const;
@@ -1201,7 +1203,7 @@ public:
 	* Get a properly formated external actor package instance name for this level package to be used in FLinkerInstancingContext
 	* @return external actor package instance name
 	*/
-	static ENGINE_API FString GetExternalActorPackageInstanceName(const FString& LevelPackageName, const FString& ActorShortPackageName);
+	static ENGINE_API FString GetExternalActorPackageInstanceName(const FString& LevelPackageName, const FString& ActorPackageName);
 
 	/**
 	 * Get the list of (on disk) external actor packages associated with this external actors path
diff --git a/Engine/Source/Runtime/Engine/Private/ExternalPackageHelper.cpp b/Engine/Source/Runtime/Engine/Private/ExternalPackageHelper.cpp
--- a/Engine/Source/Runtime/Engine/Private/ExternalPackageHelper.cpp
+++ b/Engine/Source/Runtime/Engine/Private/ExternalPackageHelper.cpp
@@ -121,9 +121,9 @@ FString FExternalPackageHelper::GetExternalPackageName(UPackage* InOuterPackage,
 	return ObjectPackageName.ToString();
 }
 
-FString FExternalPackageHelper::GetExternalObjectPackageInstanceName(const FString& OuterPackageName, const FString& ObjectShortPackageName)
+FString FExternalPackageHelper::GetExternalObjectPackageInstanceName(const FString& OuterPackageName, const FString& ObjectPackageName)
 {
-	return FString::Printf(TEXT("%s_InstanceOf_%s"), *OuterPackageName, *ObjectShortPackageName);
+	return FLinkerInstancingContext::GetInstancedPackageName(OuterPackageName, ObjectPackageName);
 }
 
 #endif
\ No newline at end of file
diff --git a/Engine/Source/Runtime/Engine/Private/Level.cpp b/Engine/Source/Runtime/Engine/Private/Level.cpp
--- a/Engine/Source/Runtime/Engine/Private/Level.cpp
+++ b/Engine/Source/Runtime/Engine/Private/Level.cpp
@@ -943,8 +943,7 @@ void ULevel::PostLoad()
 			{
 				if (bInstanced)
 				{
-					const FString ActorShortPackageName = FPackageName::GetShortName(ActorPackageName);
-					const FString InstancedName = GetExternalActorPackageInstanceName(LevelPackage->GetName(), ActorShortPackageName);
+					const FString InstancedName = GetExternalActorPackageInstanceName(LevelPackage->GetName(), ActorPackageName);
 					InstancePackageNames.Add(InstancedName);
 
 					InstancingContext.AddMapping(FName(*ActorPackageName), FName(*InstancedName));
@@ -2224,6 +2223,26 @@ bool ULevel::GetIsStreamingDisabledFromPackage(FName LevelPackage)
 	});
 }
 
+bool ULevel::GetPartitionedLevelCanBeUsedByLevelInstanceFromAsset(const FAssetData& Asset)
+{
+	FString PartitionedLevelCanBeUsedByLevelInstanceStr;
+	static const FName NAME_PartitionedLevelCanBeUsedByLevelInstance(TEXT("PartitionedLevelCanBeUsedByLevelInstance"));
+	if (Asset.GetTagValue(NAME_PartitionedLevelCanBeUsedByLevelInstance, PartitionedLevelCanBeUsedByLevelInstanceStr))
+	{
+		check(PartitionedLevelCanBeUsedByLevelInstanceStr == TEXT("1"));
+		return true;
+	}
+	return false;
+}
+
+bool ULevel::GetPartitionedLevelCanBeUsedByLevelInstanceFromPackage(FName LevelPackage)
+{
+	return LevelAssetRegistryHelper::GetLevelInfoFromAssetRegistry(LevelPackage, [](const FAssetData& Asset)
+	{
+		return GetPartitionedLevelCanBeUsedByLevelInstanceFromAsset(Asset);
+	});
+}
+
 bool ULevel::GetLevelBoundsFromAsset(const FAssetData& Asset, FBox& OutLevelBounds)
 {
 	FString LevelBoundsLocationStr;
@@ -2449,13 +2468,22 @@ void ULevel::OnLevelLoaded()
 	FixupActorFolders();
 #endif
 
+	auto IsValidLevelInstanceWorldPartition = [](UWorldPartition* InWorldPartition)
+	{
+#if WITH_EDITOR
+		return InWorldPartition->CanBeUsedByLevelInstance();
+#else
+		return false;
+#endif
+	};
+
 	// 1. Cook commandlet does it's own UWorldPartition::Initialize call in FWorldPartitionCookPackageSplitter::GetGenerateList
 	// 2. Do not Initialize if World doesn't have a UWorldPartitionSubsystem (Known case is when WorldType == EWorldType::Inactive)
 	if (!IsRunningCookCommandlet() && GetWorld()->HasSubsystem<UWorldPartitionSubsystem>())
 	{
-		if (IsPersistentLevel() || GetWorld()->IsGameWorld())
+		if (UWorldPartition* WorldPartition = GetWorldPartition())
 		{
-			if (UWorldPartition* WorldPartition = GetWorldPartition())
+			if (IsPersistentLevel() || GetWorld()->IsGameWorld() || IsValidLevelInstanceWorldPartition(WorldPartition))
 			{
 				FTransform Transform = FTransform::Identity;
 				if (ULevelStreaming* LevelStreaming = FLevelUtils::FindStreamingLevel(this))
@@ -3222,9 +3250,9 @@ void ULevel::ConvertAllActorsToPackaging(bool bExternal)
 	}
 }
 
-FString ULevel::GetExternalActorPackageInstanceName(const FString& LevelPackageName, const FString& ActorShortPackageName)
+FString ULevel::GetExternalActorPackageInstanceName(const FString& LevelPackageName, const FString& ActorPackageName)
 {
-	return FString::Printf(TEXT("%s_InstanceOf_%s"), *LevelPackageName, *ActorShortPackageName);
+	return FLinkerInstancingContext::GetInstancedPackageName(LevelPackageName, ActorPackageName);
 }
 
 TArray<FString> ULevel::GetOnDiskExternalActorPackages(const FString& ExternalActorsPath)
diff --git a/Engine/Source/Runtime/Engine/Private/LevelInstance/LevelInstanceActor.cpp b/Engine/Source/Runtime/Engine/Private/LevelInstance/LevelInstanceActor.cpp
--- a/Engine/Source/Runtime/Engine/Private/LevelInstance/LevelInstanceActor.cpp
+++ b/Engine/Source/Runtime/Engine/Private/LevelInstance/LevelInstanceActor.cpp
@@ -105,9 +105,9 @@ void ALevelInstance::PostLoad()
 	Super::PostLoad();
 
 #if WITH_EDITOR
-	if (ULevel::GetIsLevelPartitionedFromPackage(*WorldAsset.GetLongPackageName()))
+	if (!ULevelInstanceSubsystem::CanUsePackage(*WorldAsset.GetLongPackageName()))
 	{
-		UE_LOG(LogLevelInstance, Warning, TEXT("LevelInstance doesn't support partitioned world %s"), *WorldAsset.GetLongPackageName());
+		UE_LOG(LogLevelInstance, Warning, TEXT("LevelInstance doesn't support partitioned world %s, make sure to flag world partition's 'Can be Used by Level Instance'."), *WorldAsset.GetLongPackageName());
 		WorldAsset.Reset();
 	}
 #endif
diff --git a/Engine/Source/Runtime/Engine/Private/LevelInstance/LevelInstanceSubsystem.cpp b/Engine/Source/Runtime/Engine/Private/LevelInstance/LevelInstanceSubsystem.cpp
--- a/Engine/Source/Runtime/Engine/Private/LevelInstance/LevelInstanceSubsystem.cpp
+++ b/Engine/Source/Runtime/Engine/Private/LevelInstance/LevelInstanceSubsystem.cpp
@@ -867,8 +867,24 @@ ILevelInstanceInterface* ULevelInstanceSubsystem::CreateLevelInstanceFrom(const
 
 	ULevelStreamingLevelInstanceEditor* LevelStreaming = nullptr;
 	{
+		const bool bIsPartitioned = GetWorld()->IsPartitionedWorld();
 		LevelStreaming = StaticCast<ULevelStreamingLevelInstanceEditor*>(EditorLevelUtils::CreateNewStreamingLevelForWorld(
-		*GetWorld(), ULevelStreamingLevelInstanceEditor::StaticClass(), CreationParams.UseExternalActors(), LevelFilename, &ActorsToMove, CreationParams.TemplateWorld));
+		*GetWorld(), ULevelStreamingLevelInstanceEditor::StaticClass(), CreationParams.UseExternalActors(), LevelFilename, &ActorsToMove, CreationParams.TemplateWorld, /*bUseSaveAs*/true, bIsPartitioned, [bIsPartitioned](ULevel* InLevel)
+		{
+			if (bIsPartitioned)
+			{
+				UWorldPartition* WorldPartition = InLevel->GetWorldPartition();
+				if (ensure(WorldPartition))
+				{
+					// Flag the world partition that it can be used by a Level Instance
+					WorldPartition->SetCanBeUsedByLevelInstance(true);
+
+					// Validation
+					check(InLevel->IsUsingActorFolders());
+					check(!WorldPartition->IsStreamingEnabled());
+				}
+			}
+		}));
 	}
 
 	if (!LevelStreaming)
@@ -889,9 +905,6 @@ ILevelInstanceInterface* ULevelInstanceSubsystem::CreateLevelInstanceFrom(const
 		}
 	}
 
-	// Enable actor folder objects on level
-	LoadedLevel->SetUseActorFolders(true);
-
 	FActorSpawnParameters SpawnParams;
 	SpawnParams.OverrideLevel = CurrentLevel;
 	AActor* NewLevelInstanceActor = nullptr;
@@ -1499,13 +1512,16 @@ bool ULevelInstanceSubsystem::CanEditLevelInstance(const ILevelInstanceInterface
 
 	if (ULevel* LevelInstanceLevel = GetLevelInstanceLevel(LevelInstance))
 	{
-		if (LevelInstanceLevel->GetWorldPartition())
+		if (UWorldPartition* WorldPartition = LevelInstanceLevel->GetWorldPartition())
 		{
-			if (OutReason)
+			if (!WorldPartition->CanBeUsedByLevelInstance())
 			{
-				*OutReason = FText::Format(LOCTEXT("CanEditPartitionedLevelInstance", "Can't edit partitioned Level Instance ({0})."), FText::FromString(LevelInstance->GetWorldAssetPackage()));
+				if (OutReason)
+				{
+					*OutReason = FText::Format(LOCTEXT("CanEditPartitionedLevelInstance", "LevelInstance doesn't support partitioned world {0}, make sure to flag world partition's 'Can be Used by Level Instance'."), FText::FromString(LevelInstance->GetWorldAssetPackage()));
+				}
+				return false;
 			}
-			return false;
 		}
 	}
 	
@@ -1994,6 +2010,11 @@ bool ULevelInstanceSubsystem::CheckForLoop(const ILevelInstanceInterface* LevelI
 	return bValid;
 }
 
+bool ULevelInstanceSubsystem::CanUsePackage(FName InPackageName)
+{
+	return !ULevel::GetIsLevelPartitionedFromPackage(InPackageName) || ULevel::GetPartitionedLevelCanBeUsedByLevelInstanceFromPackage(InPackageName);
+}
+
 bool ULevelInstanceSubsystem::CanUseWorldAsset(const ILevelInstanceInterface* LevelInstance, TSoftObjectPtr<UWorld> WorldAsset, FString* OutReason)
 {
 	if (WorldAsset.IsNull())
@@ -2014,11 +2035,11 @@ bool ULevelInstanceSubsystem::CanUseWorldAsset(const ILevelInstanceInterface* Le
 	TArray<TPair<FText, TSoftObjectPtr<UWorld>>> LoopInfo;
 	const ILevelInstanceInterface* LoopStart = nullptr;
 
-	if (ULevel::GetIsLevelPartitionedFromPackage(*WorldAsset.GetLongPackageName()))
+	if (!ULevelInstanceSubsystem::CanUsePackage(*WorldAsset.GetLongPackageName()))
 	{
 		if (OutReason)
 		{
-			*OutReason = FString::Format(TEXT("LevelInstance doesn't support partitioned world {0}\n"), { WorldAsset.GetLongPackageName() });
+			*OutReason = FString::Format(TEXT("LevelInstance doesn't support partitioned world {0}, make sure to flag world partition's 'Can be Used by Level Instance'.\n"), { WorldAsset.GetLongPackageName() });
 		}
 
 		return false;
diff --git a/Engine/Source/Runtime/Engine/Private/LevelStreaming.cpp b/Engine/Source/Runtime/Engine/Private/LevelStreaming.cpp
--- a/Engine/Source/Runtime/Engine/Private/LevelStreaming.cpp
+++ b/Engine/Source/Runtime/Engine/Private/LevelStreaming.cpp
@@ -1224,9 +1224,7 @@ bool ULevelStreaming::RequestLevel(UWorld* PersistentWorld, bool bAllowLevelLoad
 				InstancingContextPtr = &InstancingContext;
 				for (const FString& ActorPackageName : ActorPackageNames)
 				{
-					const FString ActorShortPackageName = FPackageName::GetShortName(ActorPackageName);
-					const FString InstancedName = ULevel::GetExternalActorPackageInstanceName(DesiredPackageName.ToString(), ActorShortPackageName);
-				
+					const FString InstancedName = ULevel::GetExternalActorPackageInstanceName(DesiredPackageName.ToString(), ActorPackageName);
 					InstancingContext.AddMapping(FName(*ActorPackageName), FName(*InstancedName));
 				}
 			}
diff --git a/Engine/Source/Runtime/Engine/Private/WorldPartition/DataLayer/WorldDataLayersActorDesc.cpp b/Engine/Source/Runtime/Engine/Private/WorldPartition/DataLayer/WorldDataLayersActorDesc.cpp
--- a/Engine/Source/Runtime/Engine/Private/WorldPartition/DataLayer/WorldDataLayersActorDesc.cpp
+++ b/Engine/Source/Runtime/Engine/Private/WorldPartition/DataLayer/WorldDataLayersActorDesc.cpp
@@ -7,6 +7,7 @@
 #include "WorldPartition/DataLayer/WorldDataLayers.h"
 #include "WorldPartition/DataLayer/DataLayerInstanceWithAsset.h"
 #include "WorldPartition/DataLayer/DataLayerAsset.h"
+#include "WorldPartition/WorldPartitionActorCluster.h"
 #include "UObject/FortniteNCBranchObjectVersion.h"
 
 FArchive& operator<<(FArchive& Ar, FDataLayerInstanceDesc& Desc)
@@ -176,6 +177,11 @@ void FWorldDataLayersActorDesc::Serialize(FArchive& Ar)
 	}
 }
 
+bool FWorldDataLayersActorDesc::IsRuntimeRelevant(const FActorContainerID& InContainerID) const
+{
+	return InContainerID.IsMainContainer();
+}
+
 const FDataLayerInstanceDesc* FWorldDataLayersActorDesc::GetDataLayerInstanceFromInstanceName(FName InDataLayerInstanceName) const
 {
 	for (const FDataLayerInstanceDesc& DataLayerInstance : DataLayerInstances)
diff --git a/Engine/Source/Runtime/Engine/Private/WorldPartition/LevelInstance/LevelInstanceActorDesc.cpp b/Engine/Source/Runtime/Engine/Private/WorldPartition/LevelInstance/LevelInstanceActorDesc.cpp
--- a/Engine/Source/Runtime/Engine/Private/WorldPartition/LevelInstance/LevelInstanceActorDesc.cpp
+++ b/Engine/Source/Runtime/Engine/Private/WorldPartition/LevelInstance/LevelInstanceActorDesc.cpp
@@ -92,7 +92,7 @@ void FLevelInstanceActorDesc::RegisterContainerInstance(UWorld* InWorld)
 	check(!LevelInstanceContainer.IsValid());
 	if (DesiredRuntimeBehavior == ELevelInstanceRuntimeBehavior::Partitioned && !GLevelInstanceDebugForceLevelStreaming)
 	{
-		if (!LevelPackage.IsNone() && ULevel::GetIsLevelUsingExternalActorsFromPackage(LevelPackage) && !ULevel::GetIsLevelPartitionedFromPackage(LevelPackage))
+		if (!LevelPackage.IsNone() && ULevel::GetIsLevelUsingExternalActorsFromPackage(LevelPackage) && ULevelInstanceSubsystem::CanUsePackage(LevelPackage))
 		{
 			LevelInstanceContainer = FLevelInstanceActorDesc::RegisterActorDescContainer(LevelPackage, InWorld);
 			check(LevelInstanceContainer.IsValid());
diff --git a/Engine/Source/Runtime/Engine/Private/WorldPartition/WorldPartition.cpp b/Engine/Source/Runtime/Engine/Private/WorldPartition/WorldPartition.cpp
--- a/Engine/Source/Runtime/Engine/Private/WorldPartition/WorldPartition.cpp
+++ b/Engine/Source/Runtime/Engine/Private/WorldPartition/WorldPartition.cpp
@@ -261,6 +261,7 @@ UWorldPartition::UWorldPartition(const FObjectInitializer& ObjectInitializer)
 	, bStreamingWasEnabled(true)
 	, bShouldEnableStreamingWarned(false)
 	, bShouldCheckEnableStreamingWarning(false)
+	, bCanBeUsedByLevelInstance(false)
 	, bForceGarbageCollection(false)
 	, bForceGarbageCollectionPurge(false)
 	, bIsPIE(false)
@@ -515,6 +516,14 @@ void UWorldPartition::Initialize(UWorld* InWorld, const FTransform& InTransform)
 			GEditor->OnEditorClose().AddUObject(this, &UWorldPartition::SavePerUserSettings);
 		}
 	}
+
+	if (bIsEditor)
+	{
+		// Because WorldPartition applies instance transforms to each actor (see ApplyActorTransform)
+		// We need to flag Level's bAlreadyMovedActors to true so that ULevelStreaming::PrepareLoadedLevel won't reapply the same transform again.
+		// We asssume that non-external actors won't be transformed (which is fine since the only 2 cases are the worldsettings and the default brush).
+		OuterWorld->PersistentLevel->bAlreadyMovedActors = true;
+	}
 #endif //WITH_EDITOR
 
 	InitState = EWorldPartitionInitState::Initialized;
@@ -1177,6 +1186,27 @@ void UWorldPartition::SetEnableStreaming(bool bInEnableStreaming)
 	}
 }
 
+bool UWorldPartition::CanBeUsedByLevelInstance() const
+{
+	return bCanBeUsedByLevelInstance && !IsStreamingEnabled();
+}
+
+void UWorldPartition::SetCanBeUsedByLevelInstance(bool bInCanBeUsedByLevelInstance)
+{
+	// Only allowed to be called when WorldPartition is not initialized
+	check(!IsInitialized());
+
+	if (bCanBeUsedByLevelInstance != bInCanBeUsedByLevelInstance)
+	{
+		Modify();
+		bCanBeUsedByLevelInstance = bInCanBeUsedByLevelInstance;
+		if (bCanBeUsedByLevelInstance)
+		{
+			bEnableStreaming = false;
+		}
+	}
+}
+
 void UWorldPartition::OnEnableStreamingChanged()
 {
 	// Pin the actor handles on the actor to prevent unloading it when unhashing
@@ -1485,6 +1515,12 @@ void UWorldPartition::AppendAssetRegistryTags(TArray<FAssetRegistryTag>& OutTags
 		static const FName NAME_LevelHasStreamingDisabled(TEXT("LevelHasStreamingDisabled"));
 		OutTags.Add(FAssetRegistryTag(NAME_LevelHasStreamingDisabled, TEXT("1"), FAssetRegistryTag::TT_Hidden));
 	}
+
+	if (CanBeUsedByLevelInstance())
+	{
+		static const FName NAME_PartitionedLevelCanBeUsedByLevelInstance(TEXT("PartitionedLevelCanBeUsedByLevelInstance"));
+		OutTags.Add(FAssetRegistryTag(NAME_PartitionedLevelCanBeUsedByLevelInstance, TEXT("1"), FAssetRegistryTag::TT_Hidden));
+	}
 }
 
 uint32 UWorldPartition::GetWantedEditorCellSize() const
diff --git a/Engine/Source/Runtime/Engine/Private/WorldPartition/WorldPartitionStreamingGeneration.cpp b/Engine/Source/Runtime/Engine/Private/WorldPartition/WorldPartitionStreamingGeneration.cpp
--- a/Engine/Source/Runtime/Engine/Private/WorldPartition/WorldPartitionStreamingGeneration.cpp
+++ b/Engine/Source/Runtime/Engine/Private/WorldPartition/WorldPartitionStreamingGeneration.cpp
@@ -69,7 +69,7 @@ class FWorldPartitionStreamingGenerator
 		TMap<FGuid, FGuid> ContainerGuidsRemap;
 		for (FActorDescList::TConstIterator<> ActorDescIt(InContainer); ActorDescIt; ++ActorDescIt)
 		{
-			if (!IsActorEditorOnly(*ActorDescIt))
+			if (!IsActorEditorOnly(*ActorDescIt) && ActorDescIt->IsRuntimeRelevant(InContainerID))
 			{
 				// Handle unsaved actors
 				if (AActor* Actor = ActorDescIt->GetActor())
diff --git a/Engine/Source/Runtime/Engine/Public/ExternalPackageHelper.h b/Engine/Source/Runtime/Engine/Public/ExternalPackageHelper.h
--- a/Engine/Source/Runtime/Engine/Public/ExternalPackageHelper.h
+++ b/Engine/Source/Runtime/Engine/Public/ExternalPackageHelper.h
@@ -68,7 +68,7 @@ private:
 	static FString GetExternalPackageName(UPackage* InOuterPackage, const FString& InObjectPath);
 
 	/** Get the external object package instance name. */
-	static FString GetExternalObjectPackageInstanceName(const FString& OuterPackageName, const FString& ObjectShortPackageName);
+	static FString GetExternalObjectPackageInstanceName(const FString& OuterPackageName, const FString& ObjectPackageName);
 };
 
 template<typename T>
@@ -102,8 +102,7 @@ void FExternalPackageHelper::LoadObjectsFromExternalPackages(UObject* InOuter, T
 		for (const FAssetData& Asset : Assets)
 		{
 			const FString ObjectPackageName = Asset.PackageName.ToString();
-			const FString ShortPackageName = FPackageName::GetShortName(ObjectPackageName);
-			const FString InstancedName = GetExternalObjectPackageInstanceName(OuterPackage->GetName(), ShortPackageName);
+			const FString InstancedName = GetExternalObjectPackageInstanceName(OuterPackage->GetName(), ObjectPackageName);
 			InstancingContext.AddMapping(FName(*ObjectPackageName), FName(*InstancedName));
 
 			// Create instance package
diff --git a/Engine/Source/Runtime/Engine/Public/LevelInstance/LevelInstanceSubsystem.h b/Engine/Source/Runtime/Engine/Public/LevelInstance/LevelInstanceSubsystem.h
--- a/Engine/Source/Runtime/Engine/Public/LevelInstance/LevelInstanceSubsystem.h
+++ b/Engine/Source/Runtime/Engine/Public/LevelInstance/LevelInstanceSubsystem.h
@@ -100,6 +100,7 @@ public:
 	static bool CheckForLoop(const ILevelInstanceInterface* LevelInstance, TArray<TPair<FText, TSoftObjectPtr<UWorld>>>* LoopInfo = nullptr, const ILevelInstanceInterface** LoopStart = nullptr);
 	static bool CheckForLoop(const ILevelInstanceInterface* LevelInstance, TSoftObjectPtr<UWorld> WorldAsset, TArray<TPair<FText, TSoftObjectPtr<UWorld>>>* LoopInfo = nullptr, const ILevelInstanceInterface** LoopStart = nullptr);
 	static bool CanUseWorldAsset(const ILevelInstanceInterface* LevelInstance, TSoftObjectPtr<UWorld> WorldAsset, FString* OutReason);
+	static bool CanUsePackage(FName InPackageName);
 #endif
 
 private:
diff --git a/Engine/Source/Runtime/Engine/Public/WorldPartition/DataLayer/WorldDataLayersActorDesc.h b/Engine/Source/Runtime/Engine/Public/WorldPartition/DataLayer/WorldDataLayersActorDesc.h
--- a/Engine/Source/Runtime/Engine/Public/WorldPartition/DataLayer/WorldDataLayersActorDesc.h
+++ b/Engine/Source/Runtime/Engine/Public/WorldPartition/DataLayer/WorldDataLayersActorDesc.h
@@ -65,6 +65,7 @@ protected:
 	virtual bool Equals(const FWorldPartitionActorDesc* Other) const override;
 	virtual void Serialize(FArchive& Ar) override;
 	virtual bool IsResaveNeeded() const override { return !IsValid(); }
+	virtual bool IsRuntimeRelevant(const FActorContainerID& InContainerID) const;
 
 private:
 	TArray<FDataLayerInstanceDesc> DataLayerInstances;
diff --git a/Engine/Source/Runtime/Engine/Public/WorldPartition/WorldPartition.h b/Engine/Source/Runtime/Engine/Public/WorldPartition/WorldPartition.h
--- a/Engine/Source/Runtime/Engine/Public/WorldPartition/WorldPartition.h
+++ b/Engine/Source/Runtime/Engine/Public/WorldPartition/WorldPartition.h
@@ -104,6 +104,8 @@ public:
 	bool SupportsStreaming() const;
 	bool IsStreamingEnabled() const;
 	void SetEnableStreaming(bool bInEnableStreaming);
+	bool CanBeUsedByLevelInstance() const;
+	void SetCanBeUsedByLevelInstance(bool bInCanBeUsedByLevelInstance);
 	void OnEnableStreamingChanged();
 
 private:
@@ -230,6 +232,10 @@ private:
 
 	/** Used to know if we need to recheck if the user should enable streaming based on world size. */
 	bool bShouldCheckEnableStreamingWarning;
+
+	/** Whether Level Instance can reference this partition. */
+	UPROPERTY(EditAnywhere, Category = "WorldPartition", AdvancedDisplay, meta = (EditConditionHides, HideEditConditionToggle, EditCondition = "!bEnableStreaming"))
+	bool bCanBeUsedByLevelInstance;
 #endif
 
 public:
diff --git a/Engine/Source/Runtime/Engine/Public/WorldPartition/WorldPartitionActorDesc.h b/Engine/Source/Runtime/Engine/Public/WorldPartition/WorldPartitionActorDesc.h
--- a/Engine/Source/Runtime/Engine/Public/WorldPartition/WorldPartitionActorDesc.h
+++ b/Engine/Source/Runtime/Engine/Public/WorldPartition/WorldPartitionActorDesc.h
@@ -20,6 +20,7 @@ struct FWorldPartitionActorDescInitData
 };
 
 class UActorDescContainer;
+struct FActorContainerID;
 
 enum class EContainerClusterMode : uint8
 {
@@ -72,6 +73,7 @@ public:
 	virtual bool GetContainerInstance(const UActorDescContainer*& OutLevelContainer, FTransform& OutLevelTransform, EContainerClusterMode& OutClusterMode) const { return false; }
 	virtual const FGuid& GetSceneOutlinerParent() const { return GetParentActor(); }
 	virtual bool IsResaveNeeded() const { return false; }
+	virtual bool IsRuntimeRelevant(const FActorContainerID& InContainerID) const { return true; }
 
 	bool operator==(const FWorldPartitionActorDesc& Other) const
 	{
