Update package writers to gather package hashes internally, then have CookByTheBookFinished copy them to asset package datas.
Add iostore chunk hashes to the asset registry during cook.
Remove old hashing code, which was extensive. This moves the async writes in to UE::Task from the old AsyncWorkSequence().

diff --git a/Engine/Source/Developer/IoStoreUtilities/Private/ZenStoreWriter.cpp b/Engine/Source/Developer/IoStoreUtilities/Private/ZenStoreWriter.cpp
--- a/Engine/Source/Developer/IoStoreUtilities/Private/ZenStoreWriter.cpp
+++ b/Engine/Source/Developer/IoStoreUtilities/Private/ZenStoreWriter.cpp
@@ -485,6 +485,7 @@ void FZenStoreWriter::Initialize(const FCookInfo& Info)
 void FZenStoreWriter::BeginCook()
 {
 	PackageStoreManifest.Load(*(MetadataDirectoryPath / TEXT("packagestore.manifest")));
+	AllPackageHashes.Empty();
 
 	if (CookMode == ICookedPackageWriter::FCookInfo::CookOnTheFlyMode)
 	{
@@ -578,17 +579,26 @@ bool FZenStoreWriter::IsReservedOplogKey(FUtf8StringView Key)
 		FUtf8StringView(ReservedOplogKeys[Index]).Equals(Key, ESearchCase::IgnoreCase);
 }
 
-TFuture<FMD5Hash> FZenStoreWriter::CommitPackage(FCommitPackageInfo&& Info)
+void FZenStoreWriter::CommitPackage(FCommitPackageInfo&& Info)
 {
-	TFuture<FMD5Hash> PkgHash;
+	UE::SavePackageUtilities::IncrementOutstandingAsyncWrites();
 
+	// If we are computing hashes, we need to allocate where the hashes will go.
+	// Access to this is protected by the above IncrementOutstandingAsyncWrites.
 	if (EnumHasAnyFlags(Info.WriteOptions, EWriteOptions::ComputeHash))
 	{
 		FPendingPackageState& ExistingState = GetPendingPackage(Info.PackageName);
-		PkgHash = ExistingState.HashPromise.GetFuture();
-	}
+		ExistingState.PackageHashes = new FPackageHashes();
 
-	UE::SavePackageUtilities::IncrementOutstandingAsyncWrites();
+		// This looks weird but we're finding the _refcount_, not the hashes. So if it gets
+		// constructed, it's not actually assigned a pointer.
+		TRefCountPtr<FPackageHashes>& ExistingPackageHashes = AllPackageHashes.FindOrAdd(Info.PackageName);
+		if (ExistingPackageHashes.IsValid())
+		{
+			UE_LOG(LogZenStoreWriter, Error, TEXT("FZenStoreWriter commiting the same package twice during a cook! (%s)"), *Info.PackageName.ToString());
+		}
+		ExistingPackageHashes = ExistingState.PackageHashes;
+	}
 
 	if (FPlatformProcess::SupportsMultithreading())
 	{
@@ -598,8 +608,6 @@ TFuture<FMD5Hash> FZenStoreWriter::CommitPackage(FCommitPackageInfo&& Info)
 	{
 		CommitPackageInternal(Forward<FCommitPackageInfo>(Info));
 	}
-	
-	return PkgHash;
 }
 
 void FZenStoreWriter::CommitPackageInternal(FCommitPackageInfo&& CommitInfo)
@@ -640,6 +648,11 @@ void FZenStoreWriter::CommitPackageInternal(FCommitPackageInfo&& CommitInfo)
 		FMD5 PkgHashGen;
 		FCbPackage OplogEntry;
 		
+		if (bComputeHash)
+		{
+			PackageState->PackageHashes->ChunkHashes.Add(PkgData.Info.ChunkId, PkgData.CompressedPayload.Get().GetRawHash());
+		}
+
 		FCbAttachment PkgDataAttachment = FCbAttachment(PkgData.CompressedPayload.Get());
 		PkgHashGen.Update(PkgDataAttachment.GetHash().GetBytes(), sizeof(FIoHash::ByteArray));
 		OplogEntry.AddAttachment(PkgDataAttachment);
@@ -716,6 +729,11 @@ void FZenStoreWriter::CommitPackageInternal(FCommitPackageInfo&& CommitInfo)
 
 			for (FBulkDataEntry& Bulk : PackageState->BulkData)
 			{
+				if (bComputeHash)
+				{
+					PackageState->PackageHashes->ChunkHashes.Add(Bulk.Info.ChunkId, Bulk.CompressedPayload.Get().GetRawHash());
+				}
+
 				FCbAttachment BulkAttachment(Bulk.CompressedPayload.Get());
 				PkgHashGen.Update(BulkAttachment.GetHash().GetBytes(), sizeof(FIoHash::ByteArray));
 				OplogEntry.AddAttachment(BulkAttachment);
@@ -736,6 +754,11 @@ void FZenStoreWriter::CommitPackageInternal(FCommitPackageInfo&& CommitInfo)
 
 			for (FFileDataEntry& File : PackageState->FileData)
 			{
+				if (bComputeHash)
+				{
+					PackageState->PackageHashes->ChunkHashes.Add(File.Info.ChunkId, File.CompressedPayload.Get().GetRawHash());
+				}
+
 				FCbAttachment FileDataAttachment(File.CompressedPayload.Get());
 				PkgHashGen.Update(FileDataAttachment.GetHash().GetBytes(), sizeof(FIoHash::ByteArray));
 				OplogEntry.AddAttachment(FileDataAttachment);
@@ -759,9 +782,10 @@ void FZenStoreWriter::CommitPackageInternal(FCommitPackageInfo&& CommitInfo)
 			OplogEntryDesc.EndArray();
 		}
 
-		FMD5Hash PkgHash;
-		PkgHash.Set(PkgHashGen);
-		PackageState->HashPromise.SetValue(PkgHash);
+		if (bComputeHash)
+		{
+			PackageState->PackageHashes->PackageHash.Set(PkgHashGen);
+		}
 
 		for (int32 Index = 0; Index < NumAttachments; ++Index)
 		{
@@ -786,6 +810,7 @@ void FZenStoreWriter::CommitPackageInternal(FCommitPackageInfo&& CommitInfo)
 		{
 			FCompressedBuffer Payload = PkgData.CompressedPayload.Get();
 			FIoHash IoHash = Payload.GetRawHash();
+			PackageState->PackageHashes->ChunkHashes.Add(PkgData.Info.ChunkId, IoHash);
 			PkgHashGen.Update(IoHash.GetBytes(), sizeof(FIoHash::ByteArray));
 		}
 		
@@ -793,6 +818,7 @@ void FZenStoreWriter::CommitPackageInternal(FCommitPackageInfo&& CommitInfo)
 		{
 			FCompressedBuffer Payload = Bulk.CompressedPayload.Get();
 			FIoHash IoHash = Payload.GetRawHash();
+			PackageState->PackageHashes->ChunkHashes.Add(Bulk.Info.ChunkId, IoHash);
 			PkgHashGen.Update(IoHash.GetBytes(), sizeof(FIoHash::ByteArray));
 		}
 		
@@ -800,16 +826,11 @@ void FZenStoreWriter::CommitPackageInternal(FCommitPackageInfo&& CommitInfo)
 		{
 			FCompressedBuffer Payload = File.CompressedPayload.Get();
 			FIoHash IoHash = Payload.GetRawHash();
+			PackageState->PackageHashes->ChunkHashes.Add(File.Info.ChunkId, IoHash);
 			PkgHashGen.Update(IoHash.GetBytes(), sizeof(FIoHash::ByteArray));
 		}
 		
-		FMD5Hash PkgHash;
-		PkgHash.Set(PkgHashGen);
-		PackageState->HashPromise.SetValue(PkgHash);
-	}
-	else
-	{
-		PackageState->HashPromise.SetValue(FMD5Hash());
+		PackageState->PackageHashes->PackageHash.Set(PkgHashGen);
 	}
 
 	if (bWritePackage)
diff --git a/Engine/Source/Developer/IoStoreUtilities/Public/ZenStoreWriter.h b/Engine/Source/Developer/IoStoreUtilities/Public/ZenStoreWriter.h
--- a/Engine/Source/Developer/IoStoreUtilities/Public/ZenStoreWriter.h
+++ b/Engine/Source/Developer/IoStoreUtilities/Public/ZenStoreWriter.h
@@ -41,7 +41,7 @@ public:
 	IOSTOREUTILITIES_API ~FZenStoreWriter();
 
 	IOSTOREUTILITIES_API virtual void BeginPackage(const FBeginPackageInfo& Info) override;
-	IOSTOREUTILITIES_API virtual TFuture<FMD5Hash> CommitPackage(FCommitPackageInfo&& Info) override;
+	IOSTOREUTILITIES_API virtual void CommitPackage(FCommitPackageInfo&& Info) override;
 
 	IOSTOREUTILITIES_API virtual void WritePackageData(const FPackageInfo& Info, FLargeMemoryWriter& ExportsArchive, const TArray<FFileRegion>& FileRegions) override;
 	IOSTOREUTILITIES_API virtual void WriteAdditionalFile(const FAdditionalFileInfo& Info, const FIoBuffer& FileData) override;
@@ -77,6 +77,10 @@ public:
 	IOSTOREUTILITIES_API virtual void RemoveCookedPackages(TArrayView<const FName> PackageNamesToRemove) override;
 	IOSTOREUTILITIES_API virtual void RemoveCookedPackages() override;
 	IOSTOREUTILITIES_API virtual void MarkPackagesUpToDate(TArrayView<const FName> UpToDatePackages) override;
+	IOSTOREUTILITIES_API virtual TMap<FName, TRefCountPtr<FPackageHashes>>& GetPackageHashes() override
+	{
+		return AllPackageHashes;
+	}
 
 private:
 	struct FBulkDataEntry
@@ -110,7 +114,7 @@ private:
 		FPackageDataEntry PackageData;
 		TArray<FBulkDataEntry> BulkData;
 		TArray<FFileDataEntry> FileData;
-		TPromise<FMD5Hash> HashPromise;
+		TRefCountPtr<FPackageHashes> PackageHashes;
 	};
 
 	FPendingPackageState& GetPendingPackage(const FName& PackageName)
@@ -152,6 +156,7 @@ private:
 	FString								OutputPath;
 	FString								MetadataDirectoryPath;
 	FIoContainerId						ContainerId = FIoContainerId::FromName(TEXT("global"));
+	TMap<FName, TRefCountPtr<FPackageHashes>> AllPackageHashes;
 
 	FPackageStoreManifest				PackageStoreManifest;
 	TUniquePtr<FPackageStoreOptimizer>	PackageStoreOptimizer;
diff --git a/Engine/Source/Editor/UnrealEd/Private/Commandlets/AssetRegistryGenerator.cpp b/Engine/Source/Editor/UnrealEd/Private/Commandlets/AssetRegistryGenerator.cpp
--- a/Engine/Source/Editor/UnrealEd/Private/Commandlets/AssetRegistryGenerator.cpp
+++ b/Engine/Source/Editor/UnrealEd/Private/Commandlets/AssetRegistryGenerator.cpp
@@ -2405,7 +2405,7 @@ const FAssetData* FAssetRegistryGenerator::CreateOrFindAssetData(UObject& Object
 }
 
 void FAssetRegistryGenerator::UpdateAssetRegistryPackageData(const UPackage& Package,
-	FSavePackageResultStruct& SavePackageResult, TFuture<FMD5Hash>& CookedHash,
+	FSavePackageResultStruct& SavePackageResult,
 	FCookTagList&& InArchiveCookTagList
 	)
 {
@@ -2428,22 +2428,8 @@ void FAssetRegistryGenerator::UpdateAssetRegistryPackageData(const UPackage& Pac
 		NewPackageFlags = SavePackageResult.SerializedPackageFlags;
 
 		AssetPackageData->DiskSize = SavePackageResult.TotalFileSize;
-		// If there is no hash (e.g.: when SavePackageResult == ESavePackageResult::ReplaceCompletely), don't attempt to setup a continuation to update
-		// the AssetRegistry entry with it later.  Just leave the asset registry entry with a default constructed FMD5Hash which is marked as invalid.
-		if (CookedHash.IsValid())
-		{
-			CookedHash.Next([AssetPackageData](const FMD5Hash& CookedHash)
-				{
-					// Store the cooked hash in the Asset Registry when it is done computing in another thread.
-					// NOTE: For this to work, we rely on:
-					// 1) UPackage::WaitForAsyncFileWrites needs to be called before any use of the CookedHash.
-					//    It is currently called in CookByTheBookFinished before BuildChunkManifest.
-					// 2) AssetPackageData must continue to be a valid pointer until WaitForAsyncFileWrites is called.
-					//    This is currently true because CreateOrFindAssetDatas allocates FAssetPackageData individually
-					//    and doesn't modify them until pruning, which happens during BuildChunkManifest.
-					AssetPackageData->CookedHash = CookedHash;
-				});
-		}
+
+		// The CookedHash is assigned during CookByTheBookFinished		
 	}
 	else
 	{
diff --git a/Engine/Source/Editor/UnrealEd/Private/CookOnTheFlyServer.cpp b/Engine/Source/Editor/UnrealEd/Private/CookOnTheFlyServer.cpp
--- a/Engine/Source/Editor/UnrealEd/Private/CookOnTheFlyServer.cpp
+++ b/Engine/Source/Editor/UnrealEd/Private/CookOnTheFlyServer.cpp
@@ -4335,7 +4335,6 @@ void FSaveCookedPackageContext::SetupPackage()
 	bool bKeepEditorOnlyPackages = !(COTFS.IsCookByTheBookMode() && !COTFS.IsCookingInEditor());
 	bKeepEditorOnlyPackages |= COTFS.IsCookFlagSet(ECookInitializationFlags::Iterative);
 	SaveFlags |= bKeepEditorOnlyPackages ? SAVE_KeepEditorOnlyCookedPackages : SAVE_None;
-	SaveFlags |= COTFS.CookByTheBookOptions ? SAVE_ComputeHash : SAVE_None;
 
 	// Use SandboxFile to do path conversion to properly handle sandbox paths (outside of standard paths in particular).
 	Filename = COTFS.ConvertToFullSandboxPath(*Filename, true);
@@ -4440,7 +4439,6 @@ void FSaveCookedPackageContext::FinishPlatform()
 	bool bSuccessful = SavePackageResult.IsSuccessful();
 	bool bLocalReferencedOnlyByEditorOnlyData = SavePackageResult == ESavePackageResult::ReferencedOnlyByEditorOnlyData;
 
-	TFuture<FMD5Hash> CookedHash;
 	FAssetRegistryGenerator& Generator = *(COTFS.PlatformManager->GetPlatformData(TargetPlatform)->RegistryGenerator);
 	if (bPlatformSetupSuccessful)
 	{
@@ -4466,18 +4464,18 @@ void FSaveCookedPackageContext::FinishPlatform()
 		// TODO: Reenable BuildDefinitionList once FCbPackage support for empty FCbObjects is in
 		//Info.Attachments.Add({ "BuildDefinitionList", BuildDefinitionList });
 		Info.WriteOptions = IPackageWriter::EWriteOptions::Write;
-		if (!!(SaveFlags & SAVE_ComputeHash))
+		if (COTFS.CookByTheBookOptions != nullptr)
 		{
 			Info.WriteOptions |= IPackageWriter::EWriteOptions::ComputeHash;
 		}
 
-		CookedHash = PackageWriter->CommitPackage(MoveTemp(Info));
+		PackageWriter->CommitPackage(MoveTemp(Info));
 	}
 
 	// Update asset registry
 	if (COTFS.CookByTheBookOptions)
 	{
-		Generator.UpdateAssetRegistryPackageData(*Package, SavePackageResult, CookedHash, MoveTemp(*ArchiveCookContext.GetCookTagList()));
+		Generator.UpdateAssetRegistryPackageData(*Package, SavePackageResult, MoveTemp(*ArchiveCookContext.GetCookTagList()));
 	}
 
 	// In success or failure we want to mark the package as cooked, unless the failure was due to being referenced only
@@ -7073,6 +7071,20 @@ void UCookOnTheFlyServer::CookByTheBookFinished()
 						IgnorePackageNames.Add(UncookedEditorOnlyPackage);
 					}
 				}
+				
+				// Add the package hashes to the relevant AssetPackageDatas.
+				// PackageHashes are gated by requiring UPackage::WaitForAsyncFileWrites(), which is called above.
+				FCookSavePackageContext& SaveContext = FindOrCreateSaveContext(TargetPlatform);
+				TMap<FName, TRefCountPtr<FPackageHashes>>& AllPackageHashes = SaveContext.PackageWriter->GetPackageHashes();
+				for (TPair<FName, TRefCountPtr<FPackageHashes>>& HashSet : AllPackageHashes)
+				{
+					FAssetPackageData* AssetPackageData = Generator.GetAssetPackageData(HashSet.Key);
+					TRefCountPtr<FPackageHashes>& PackageHashes = HashSet.Value;
+
+					AssetPackageData->CookedHash = PackageHashes->PackageHash;
+					Move(AssetPackageData->ChunkHashes, PackageHashes->ChunkHashes);
+				}
+
 				{
 					Generator.PreSave(CookedPackageNames);
 				}
@@ -7648,7 +7660,7 @@ UE::Cook::FCookSavePackageContext* UCookOnTheFlyServer::CreateSaveContext(const
 	{
 		PackageWriter = new FLooseCookedPackageWriter(ResolvedRootPath, ResolvedMetadataPath, TargetPlatform,
 			GetAsyncIODelete(), *PackageDatas, PluginsToRemap);
-		WriterDebugName = TEXT("DirectoryWriter");
+		WriterDebugName = TEXT("LooseCookedPackageWriter");
 	}
 
 	DiffModeHelper->InitializePackageWriter(PackageWriter);
@@ -8633,7 +8645,7 @@ uint32 UCookOnTheFlyServer::FullLoadAndSave(uint32& CookedPackageCount)
 	}
 
 	const bool bSaveConcurrent = FParse::Param(FCommandLine::Get(), TEXT("ConcurrentSave"));
-	uint32 SaveFlags = SAVE_KeepGUID | SAVE_Async | SAVE_ComputeHash | (IsCookFlagSet(ECookInitializationFlags::Unversioned) ? SAVE_Unversioned : 0);
+	uint32 SaveFlags = SAVE_KeepGUID | SAVE_Async | (IsCookFlagSet(ECookInitializationFlags::Unversioned) ? SAVE_Unversioned : 0);
 	if (bSaveConcurrent)
 	{
 		SaveFlags |= SAVE_Concurrent;
@@ -9072,7 +9084,7 @@ uint32 UCookOnTheFlyServer::FullLoadAndSave(uint32& CookedPackageCount)
 						}
 
 						// Update asset registry
-						Generator.UpdateAssetRegistryPackageData(*Package, SaveResult, SaveResult.CookedHash, MoveTemp(*CookContext.GetCookTagList()));
+						Generator.UpdateAssetRegistryPackageData(*Package, SaveResult, MoveTemp(*CookContext.GetCookTagList()));
 						FAssetPackageData* AssetPackageData = Generator.GetAssetPackageData(Package->GetFName());
 						check(AssetPackageData);
 
@@ -9082,7 +9094,7 @@ uint32 UCookOnTheFlyServer::FullLoadAndSave(uint32& CookedPackageCount)
 						PRAGMA_DISABLE_DEPRECATION_WARNINGS;
 						CommitInfo.PackageGuid = AssetPackageData->PackageGuid;
 						PRAGMA_ENABLE_DEPRECATION_WARNINGS;
-						CommitInfo.WriteOptions = IPackageWriter::EWriteOptions::Write;
+						CommitInfo.WriteOptions = IPackageWriter::EWriteOptions::Write | IPackageWriter::EWriteOptions::ComputeHash;
 						SavePackageContext.PackageWriter->CommitPackage(MoveTemp(CommitInfo));
 
 						if (SaveResult.IsSuccessful())
diff --git a/Engine/Source/Editor/UnrealEd/Private/Cooker/DiffPackageWriter.cpp b/Engine/Source/Editor/UnrealEd/Private/Cooker/DiffPackageWriter.cpp
--- a/Engine/Source/Editor/UnrealEd/Private/Cooker/DiffPackageWriter.cpp
+++ b/Engine/Source/Editor/UnrealEd/Private/Cooker/DiffPackageWriter.cpp
@@ -124,7 +124,7 @@ void FDiffPackageWriter::BeginPackage(const FBeginPackageInfo& Info)
 	Inner->BeginPackage(Info);
 }
 
-TFuture<FMD5Hash> FDiffPackageWriter::CommitPackage(FCommitPackageInfo&& Info)
+void FDiffPackageWriter::CommitPackage(FCommitPackageInfo&& Info)
 {
 	if (bDiffCallstack && bSaveForDiff)
 	{
diff --git a/Engine/Source/Editor/UnrealEd/Private/Cooker/DiffPackageWriter.h b/Engine/Source/Editor/UnrealEd/Private/Cooker/DiffPackageWriter.h
--- a/Engine/Source/Editor/UnrealEd/Private/Cooker/DiffPackageWriter.h
+++ b/Engine/Source/Editor/UnrealEd/Private/Cooker/DiffPackageWriter.h
@@ -20,7 +20,7 @@ public:
 		return Result;
 	}
 	virtual void BeginPackage(const FBeginPackageInfo& Info) override;
-	virtual TFuture<FMD5Hash> CommitPackage(FCommitPackageInfo&& Info) override;
+	virtual void CommitPackage(FCommitPackageInfo&& Info) override;
 	virtual void WritePackageData(const FPackageInfo& Info, FLargeMemoryWriter& ExportsArchive,
 		const TArray<FFileRegion>& FileRegions) override;
 	virtual void WriteBulkData(const FBulkDataInfo& Info, const FIoBuffer& BulkData,
@@ -95,7 +95,10 @@ public:
 		Inner->UpdateSaveArguments(SaveArgs);
 	}
 	virtual bool IsAnotherSaveNeeded(FSavePackageResultStruct& PreviousResult, FSavePackageArgs& SaveArgs) override;
-
+	virtual TMap<FName, TRefCountPtr<FPackageHashes>>& GetPackageHashes() override
+	{
+		return Inner->GetPackageHashes();
+	}
 private:
 	void ParseCmds();
 	void ParseDumpObjList(FString InParams);
@@ -133,7 +136,7 @@ public:
 		return Inner->GetCapabilities();
 	}
 	virtual void BeginPackage(const FBeginPackageInfo& Info) override;
-	virtual TFuture<FMD5Hash> CommitPackage(FCommitPackageInfo&& Info) override
+	virtual void CommitPackage(FCommitPackageInfo&& Info) override
 	{
 		return Inner->CommitPackage(MoveTemp(Info));
 	}
@@ -214,7 +217,10 @@ public:
 	}
 	virtual void UpdateSaveArguments(FSavePackageArgs& SaveArgs) override;
 	virtual bool IsAnotherSaveNeeded(FSavePackageResultStruct& PreviousResult, FSavePackageArgs& SaveArgs) override;
-
+	virtual TMap<FName, TRefCountPtr<FPackageHashes>>& GetPackageHashes() override
+	{
+		return Inner->GetPackageHashes();
+	}
 private:
 	enum class EDiffMode : uint8
 	{
diff --git a/Engine/Source/Editor/UnrealEd/Private/Cooker/LooseCookedPackageWriter.cpp b/Engine/Source/Editor/UnrealEd/Private/Cooker/LooseCookedPackageWriter.cpp
--- a/Engine/Source/Editor/UnrealEd/Private/Cooker/LooseCookedPackageWriter.cpp
+++ b/Engine/Source/Editor/UnrealEd/Private/Cooker/LooseCookedPackageWriter.cpp
@@ -26,6 +26,7 @@
 #include "Serialization/ArchiveStackTrace.h"
 #include "Serialization/ArrayReader.h"
 #include "Serialization/LargeMemoryWriter.h"
+#include "Tasks/Task.h"
 #include "UObject/Package.h"
 #include "PackageStoreOptimizer.h"
 
@@ -53,20 +54,18 @@ void FLooseCookedPackageWriter::BeginPackage(const FBeginPackageInfo& Info)
 	PackageStoreManifest.BeginPackage(Info.PackageName);
 }
 
-TFuture<FMD5Hash> FLooseCookedPackageWriter::CommitPackageInternal(FPackageWriterRecords::FPackage&& BaseRecord,
+void FLooseCookedPackageWriter::CommitPackageInternal(FPackageWriterRecords::FPackage&& BaseRecord,
 	const FCommitPackageInfo& Info)
 {
 	FRecord& Record = static_cast<FRecord&>(BaseRecord);
-	TFuture<FMD5Hash> CookedHash;
 	if (Info.bSucceeded)
 	{
-		CookedHash = AsyncSave(Record, Info);
+		AsyncSave(Record, Info);
 	}
 	UpdateManifest(Record);
-	return CookedHash;
 }
 
-TFuture<FMD5Hash> FLooseCookedPackageWriter::AsyncSave(FRecord& Record, const FCommitPackageInfo& Info)
+void FLooseCookedPackageWriter::AsyncSave(FRecord& Record, const FCommitPackageInfo& Info)
 {
 	FCommitContext Context{ Info };
 
@@ -80,7 +79,7 @@ TFuture<FMD5Hash> FLooseCookedPackageWriter::AsyncSave(FRecord& Record, const FC
 	CollectForSaveExportsFooter(Record, Context);
 	CollectForSaveExportsBuffers(Record, Context);
 
-	return AsyncSaveOutputFiles(Record, Context);
+	AsyncSaveOutputFiles(Record, Context);
 }
 
 void FLooseCookedPackageWriter::CompleteExportsArchiveForDiff(const FPackageInfo& Info, FLargeMemoryWriter& ExportsArchive)
@@ -142,6 +141,7 @@ void FLooseCookedPackageWriter::CollectForSaveBulkData(FRecord& Record, FCommitC
 			OutputFile.Regions = MoveTemp(BulkRecord.Regions);
 			OutputFile.bIsSidecar = true;
 			OutputFile.bContributeToHash = BulkRecord.Info.MultiOutputIndex == 0; // Only caculate the main package output hash
+			OutputFile.ChunkId = BulkRecord.Info.ChunkId;
 		}
 	}
 }
@@ -168,6 +168,7 @@ void FLooseCookedPackageWriter::CollectForSaveAdditionalFileRecords(FRecord& Rec
 		OutputFile.Buffer = FCompositeBuffer(AdditionalRecord.Buffer);
 		OutputFile.bIsSidecar = true;
 		OutputFile.bContributeToHash = AdditionalRecord.Info.MultiOutputIndex == 0; // Only calculate the main package output hash
+		OutputFile.ChunkId = AdditionalRecord.Info.ChunkId;
 	}
 }
 
@@ -252,28 +253,48 @@ FPackageWriterRecords::FPackage* FLooseCookedPackageWriter::ConstructRecord()
 	return new FRecord();
 }
 
-TFuture<FMD5Hash> FLooseCookedPackageWriter::AsyncSaveOutputFiles(FRecord& Record, FCommitContext& Context)
+void FLooseCookedPackageWriter::AsyncSaveOutputFiles(FRecord& Record, FCommitContext& Context)
 {
 	if (!EnumHasAnyFlags(Context.Info.WriteOptions, EWriteOptions::Write | EWriteOptions::ComputeHash))
 	{
-		return TFuture<FMD5Hash>();
+		return;
 	}
 
 	UE::SavePackageUtilities::IncrementOutstandingAsyncWrites();
-	return Async(EAsyncExecution::TaskGraph,
-		[OutputFiles = MoveTemp(Context.OutputFiles), WriteOptions = Context.Info.WriteOptions]() mutable
+
+	TRefCountPtr<FPackageHashes> ThisPackageHashes;
+	
+	if (EnumHasAnyFlags(Context.Info.WriteOptions, EWriteOptions::ComputeHash))
 	{
-		FMD5 AccumulatedHash;
-		for (FWriteFileData& OutputFile : OutputFiles)
+		ThisPackageHashes = new FPackageHashes();
+
+		// This looks weird but we're finding the _refcount_, not the hashes. So if it gets
+		// constructed, it's not actually assigned a pointer.
+		TRefCountPtr<FPackageHashes>& ExistingPackageHashes = AllPackageHashes.FindOrAdd(Context.Info.PackageName);
+		if (ExistingPackageHashes.IsValid())
 		{
-			OutputFile.Write(AccumulatedHash, WriteOptions);
+			UE_LOG(LogSavePackage, Error, TEXT("FLooseCookedPackageWriter encountered the same package twice in a cook! (%s)"), *Context.Info.PackageName.ToString());
 		}
+		ExistingPackageHashes = ThisPackageHashes;
+	}
 
-		FMD5Hash OutputHash;
-		OutputHash.Set(AccumulatedHash);
-		UE::SavePackageUtilities::DecrementOutstandingAsyncWrites();
-		return OutputHash;
-	});
+	UE::Tasks::Launch(TEXT("HashAndWriteLooseCookedFile"), [OutputFiles = MoveTemp(Context.OutputFiles), WriteOptions = Context.Info.WriteOptions, ThisPackageHashes = MoveTemp(ThisPackageHashes)]()
+		{
+			FMD5 AccumulatedHash;
+			for (const FWriteFileData& OutputFile : OutputFiles)
+			{
+				OutputFile.HashAndWrite(AccumulatedHash, ThisPackageHashes, WriteOptions);
+			}
+
+			if (EnumHasAnyFlags(WriteOptions, EWriteOptions::ComputeHash))
+			{
+				ThisPackageHashes->PackageHash.Set(AccumulatedHash);
+			}
+
+			// This is used to release the game thread to access the hashes
+			UE::SavePackageUtilities::DecrementOutstandingAsyncWrites();
+		}
+	);
 }
 
 static void WriteToFile(const FString& Filename, const FCompositeBuffer& Buffer)
@@ -336,7 +357,7 @@ static void WriteToFile(const FString& Filename, const FCompositeBuffer& Buffer)
 	UE_LOG(LogSavePackage, Fatal, TEXT("SavePackage Async write %s failed: %s"), *Filename, LastErrorText);
 }
 
-void FLooseCookedPackageWriter::FWriteFileData::Write(FMD5& AccumulatedHash, EWriteOptions WriteOptions) const
+void FLooseCookedPackageWriter::FWriteFileData::HashAndWrite(FMD5& AccumulatedHash, const TRefCountPtr<FPackageHashes>& PackageHashes, EWriteOptions WriteOptions) const
 {
 	//@todo: FH: Should we calculate the hash of both output, currently only the main package output hash is calculated
 	if (EnumHasAnyFlags(WriteOptions, EWriteOptions::ComputeHash) && bContributeToHash)
@@ -345,6 +366,17 @@ void FLooseCookedPackageWriter::FWriteFileData::Write(FMD5& AccumulatedHash, EWr
 		{
 			AccumulatedHash.Update(static_cast<const uint8*>(Segment.GetData()), Segment.GetSize());
 		}
+
+		if (ChunkId.IsValid())
+		{
+			FBlake3 ChunkHash;
+			for (const FSharedBuffer& Segment : Buffer.GetSegments())
+			{
+				ChunkHash.Update(static_cast<const uint8*>(Segment.GetData()), Segment.GetSize());
+			}
+			FIoHash FinalHash(ChunkHash.Finalize());
+			PackageHashes->ChunkHashes.Add(ChunkId, FinalHash);
+		}
 	}
 
 	if ((bIsSidecar && EnumHasAnyFlags(WriteOptions, EWriteOptions::WriteSidecars)) ||
@@ -421,6 +453,7 @@ void FLooseCookedPackageWriter::Initialize(const FCookInfo& Info)
 void FLooseCookedPackageWriter::BeginCook()
 {
 	PackageStoreManifest.Load(*(MetadataDirectoryPath / TEXT("packagestore.manifest")));
+	AllPackageHashes.Empty();
 }
 
 void FLooseCookedPackageWriter::EndCook()
diff --git a/Engine/Source/Editor/UnrealEd/Private/Cooker/LooseCookedPackageWriter.h b/Engine/Source/Editor/UnrealEd/Private/Cooker/LooseCookedPackageWriter.h
--- a/Engine/Source/Editor/UnrealEd/Private/Cooker/LooseCookedPackageWriter.h
+++ b/Engine/Source/Editor/UnrealEd/Private/Cooker/LooseCookedPackageWriter.h
@@ -47,7 +47,7 @@ public:
 	virtual void MarkPackagesUpToDate(TArrayView<const FName> UpToDatePackages) override;
 	virtual bool GetPreviousCookedBytes(const FPackageInfo& Info, FPreviousCookedBytesData& OutData) override;
 	virtual void CompleteExportsArchiveForDiff(const FPackageInfo& Info, FLargeMemoryWriter& ExportsArchive) override;
-	virtual TFuture<FMD5Hash> CommitPackageInternal(FPackageWriterRecords::FPackage&& BaseRecord,
+	virtual void CommitPackageInternal(FPackageWriterRecords::FPackage&& BaseRecord,
 		const FCommitPackageInfo& Info) override;
 	virtual FPackageWriterRecords::FPackage* ConstructRecord() override;
 
@@ -78,8 +78,9 @@ private:
 		TArray<FFileRegion> Regions;
 		bool bIsSidecar;
 		bool bContributeToHash = true;
+		FIoChunkId ChunkId = FIoChunkId::InvalidChunkId;
 
-		void Write(FMD5& AccumulatedHash, EWriteOptions WriteOptions) const;
+		void HashAndWrite(FMD5& AccumulatedHash, const TRefCountPtr<FPackageHashes>& PackageHashes, EWriteOptions WriteOptions) const;
 	};
 
 	/** Stack data for the helper functions of CommitPackageInternal. */
@@ -106,7 +107,7 @@ private:
 		const FString& SandboxProjectDir, const FString& RelativeProjectDir,
 		const FString& CookedPath, FString& OutUncookedPath) const;
 	void RemoveCookedPackagesByUncookedFilename(const TArray<FName>& UncookedFileNamesToRemove);
-	TFuture<FMD5Hash> AsyncSave(FRecord& Record, const FCommitPackageInfo& Info);
+	void AsyncSave(FRecord& Record, const FCommitPackageInfo& Info);
 
 	void CollectForSavePackageData(FRecord& Record, FCommitContext& Context);
 	void CollectForSaveBulkData(FRecord& Record, FCommitContext& Context);
@@ -114,9 +115,17 @@ private:
 	void CollectForSaveAdditionalFileRecords(FRecord& Record, FCommitContext& Context);
 	void CollectForSaveExportsFooter(FRecord& Record, FCommitContext& Context);
 	void CollectForSaveExportsBuffers(FRecord& Record, FCommitContext& Context);
-	TFuture<FMD5Hash> AsyncSaveOutputFiles(FRecord& Record, FCommitContext& Context);
+	void AsyncSaveOutputFiles(FRecord& Record, FCommitContext& Context);
 	void UpdateManifest(FRecord& Record);
 
+	TMap<FName, TRefCountPtr<FPackageHashes>>& GetPackageHashes() override
+	{
+		return AllPackageHashes;
+	}
+
+	// If EWriteOptions::ComputeHash is not set, the package will not get added to this.
+	TMap<FName, TRefCountPtr<FPackageHashes>> AllPackageHashes;
+
 	TMap<FName, FName> UncookedPathToCookedPath;
 	FString OutputPath;
 	FString MetadataDirectoryPath;
diff --git a/Engine/Source/Editor/UnrealEd/Private/EditorDomain/EditorDomainUtils.cpp b/Engine/Source/Editor/UnrealEd/Private/EditorDomain/EditorDomainUtils.cpp
--- a/Engine/Source/Editor/UnrealEd/Private/EditorDomain/EditorDomainUtils.cpp
+++ b/Engine/Source/Editor/UnrealEd/Private/EditorDomain/EditorDomainUtils.cpp
@@ -1434,7 +1434,7 @@ public:
 	}
 
 protected:
-	virtual TFuture<FMD5Hash> CommitPackageInternal(FPackageWriterRecords::FPackage&& Record,
+	virtual void CommitPackageInternal(FPackageWriterRecords::FPackage&& Record,
 		const FCommitPackageInfo& Info) override
 	{
 		// CommitPackage is called below with these options
@@ -1538,8 +1538,6 @@ protected:
 			FileSize += Buffer.GetSize();
 		}
 		WritePackageRecord = MoveTemp(Package);
-
-		return TFuture<FMD5Hash>();
 	}
 
 private:
diff --git a/Engine/Source/Editor/UnrealEd/Public/Commandlets/AssetRegistryGenerator.h b/Engine/Source/Editor/UnrealEd/Public/Commandlets/AssetRegistryGenerator.h
--- a/Engine/Source/Editor/UnrealEd/Public/Commandlets/AssetRegistryGenerator.h
+++ b/Engine/Source/Editor/UnrealEd/Public/Commandlets/AssetRegistryGenerator.h
@@ -196,7 +196,7 @@ public:
 	 * @param Package The package to update info on
 	 * @param SavePackageResult The metadata to associate with the given package name
 	 */
-	void UpdateAssetRegistryPackageData(const UPackage& Package, FSavePackageResultStruct& SavePackageResult, TFuture<FMD5Hash>& CookedHash, FCookTagList&& InArchiveCookTagList);
+	void UpdateAssetRegistryPackageData(const UPackage& Package, FSavePackageResultStruct& SavePackageResult, FCookTagList&& InArchiveCookTagList);
 
 private:
 	/**
diff --git a/Engine/Source/Runtime/AssetRegistry/Private/AssetRegistryState.cpp b/Engine/Source/Runtime/AssetRegistry/Private/AssetRegistryState.cpp
--- a/Engine/Source/Runtime/AssetRegistry/Private/AssetRegistryState.cpp
+++ b/Engine/Source/Runtime/AssetRegistry/Private/AssetRegistryState.cpp
@@ -1037,8 +1037,6 @@ bool FAssetRegistryState::Save(FArchive& OriginalAr, const FAssetRegistrySeriali
 
 	check(!OriginalAr.IsLoading());
 
-	// This is only used for the runtime version of the AssetRegistry
-
 #if !ALLOW_NAME_BATCH_SAVING
 	checkf(false, TEXT("Cannot save cooked AssetRegistryState in this configuration"));
 #else
diff --git a/Engine/Source/Runtime/Core/Public/Serialization/PackageWriter.h b/Engine/Source/Runtime/Core/Public/Serialization/PackageWriter.h
--- a/Engine/Source/Runtime/Core/Public/Serialization/PackageWriter.h
+++ b/Engine/Source/Runtime/Core/Public/Serialization/PackageWriter.h
@@ -84,7 +84,7 @@ public:
 
 	/** Finalize a package started with BeginPackage()
 	  */
-	virtual TFuture<FMD5Hash> CommitPackage(FCommitPackageInfo&& Info) = 0;
+	virtual void CommitPackage(FCommitPackageInfo&& Info) = 0;
 
 	struct FPackageInfo
 	{
@@ -137,7 +137,7 @@ public:
 		/** Output Package Name (Input Package can produce multiple output) */
 		FName		OutputPackageName;
 		FString		Filename;
-		FIoChunkId	ChunkId;
+		FIoChunkId	ChunkId = FIoChunkId::InvalidChunkId;
 		uint32		MultiOutputIndex = 0;
 	};
 
@@ -180,6 +180,17 @@ public:
 
 ENUM_CLASS_FLAGS(IPackageWriter::EWriteOptions);
 
+/** Struct containing hashes computed during cooked package writing. */
+struct FPackageHashes : FRefCountBase
+{
+	// Hashes for each chunk saved by the package.
+	TMap<FIoChunkId, FIoHash> ChunkHashes;
+
+	// This is a hash representing the entire package. Note this is
+	// not consistently computed across PackageWriters!
+	FMD5Hash PackageHash;
+};
+
 /** Interface for cooking that writes cooked packages to storage usable by the runtime game. */
 class ICookedPackageWriter : public IPackageWriter
 {
@@ -247,7 +258,9 @@ public:
 	};
 
 	/**
-	 * Returns an AssetRegistry describing the previous cook results.
+	 * Returns an AssetRegistry describing the previous cook results. This doesn't mean a cook saved off
+	 * to another directory - it means the AssetRegistry that's living in the directory we are about
+	 * to cook in to.
 	 */
 	virtual TUniquePtr<FAssetRegistryState> LoadPreviousAssetRegistry() = 0;
 
@@ -306,6 +319,13 @@ public:
 	{
 		return nullptr;
 	}
+
+	/** 
+	*	Cooked package writers asynchronously hash the chunks for each package after CommitPackage. Once cooking has completed,
+	*	use this to acquire the results. This is synced using void UPackage::WaitForAsyncFileWrites() - do not access
+	*	the results before that completes. Non-const so that the cooking process can Move the map of hashes.
+	*/
+	virtual TMap<FName, TRefCountPtr<FPackageHashes>>& GetPackageHashes() = 0;
 };
 
 static inline const ANSICHAR* LexToString(IPackageWriter::FBulkDataInfo::EType Value)
diff --git a/Engine/Source/Runtime/CoreUObject/Internal/Serialization/PackageWriterToSharedBuffer.h b/Engine/Source/Runtime/CoreUObject/Internal/Serialization/PackageWriterToSharedBuffer.h
--- a/Engine/Source/Runtime/CoreUObject/Internal/Serialization/PackageWriterToSharedBuffer.h
+++ b/Engine/Source/Runtime/CoreUObject/Internal/Serialization/PackageWriterToSharedBuffer.h
@@ -117,16 +117,15 @@ public:
 	{
 		Records.WriteLinkerAdditionalData(Info, Data, FileRegions);
 	}
-	virtual TFuture<FMD5Hash> CommitPackage(IPackageWriter::FCommitPackageInfo&& Info) override
+	virtual void CommitPackage(IPackageWriter::FCommitPackageInfo&& Info) override
 	{
 		TUniquePtr<FPackageRecord> Record = Records.FindAndRemoveRecordChecked(Info.PackageName);
 		ValidateCommit(*Record, Info);
-		TFuture<FMD5Hash> CookedHash = CommitPackageInternal(MoveTemp(*Record), Info);
-		return CookedHash;
+		CommitPackageInternal(MoveTemp(*Record), Info);
 	}
 
 protected:
-	virtual TFuture<FMD5Hash> CommitPackageInternal(FPackageRecord&& Record,
+	virtual void CommitPackageInternal(FPackageRecord&& Record,
 		const IPackageWriter::FCommitPackageInfo& Info) = 0;
 
 protected:
diff --git a/Engine/Source/Runtime/CoreUObject/Private/AssetRegistry/AssetData.cpp b/Engine/Source/Runtime/CoreUObject/Private/AssetRegistry/AssetData.cpp
--- a/Engine/Source/Runtime/CoreUObject/Private/AssetRegistry/AssetData.cpp
+++ b/Engine/Source/Runtime/CoreUObject/Private/AssetRegistry/AssetData.cpp
@@ -261,6 +261,10 @@ void FAssetPackageData::SerializeForCacheInternal(FArchive& Ar, FAssetPackageDat
 	{
 		Ar << PackageData.CookedHash;
 	}
+	if (Version >= FAssetRegistryVersion::AddedChunkHashes)
+	{
+		Ar << PackageData.ChunkHashes;
+	}
 	if (Version >= FAssetRegistryVersion::WorkspaceDomain)
 	{
 		if (Version >= FAssetRegistryVersion::PackageFileSummaryVersionChange)
diff --git a/Engine/Source/Runtime/CoreUObject/Private/Serialization/PackageWriterToSharedBuffer.cpp b/Engine/Source/Runtime/CoreUObject/Private/Serialization/PackageWriterToSharedBuffer.cpp
--- a/Engine/Source/Runtime/CoreUObject/Private/Serialization/PackageWriterToSharedBuffer.cpp
+++ b/Engine/Source/Runtime/CoreUObject/Private/Serialization/PackageWriterToSharedBuffer.cpp
@@ -5,6 +5,8 @@
 #include "Misc/ScopeRWLock.h"
 #include "Serialization/LargeMemoryWriter.h"
 
+DEFINE_LOG_CATEGORY_STATIC(LogPackageWriter, Display, All);
+
 TUniquePtr<FLargeMemoryWriter> IPackageWriter::CreateLinkerArchive(FName PackageName, UObject* Asset)
 {
 	// The LargeMemoryWriter does not need to be persistent; the LinkerSave wraps it and reports Persistent=true
@@ -57,7 +59,16 @@ void FPackageWriterRecords::WriteAdditionalFile(const IPackageWriter::FAdditiona
 	const FIoBuffer& FileData)
 {
 	FPackage& Record = FindRecordChecked(Info.InputPackageName);
-	Record.AdditionalFiles.Add(FAdditionalFile{ Info, IoBufferToSharedBuffer(FileData) });
+	FAdditionalFile& AdditionalFile = Record.AdditionalFiles.Add_GetRef(FAdditionalFile{ Info, IoBufferToSharedBuffer(FileData) });
+	FIoChunkId ChunkId = CreateExternalFileChunkId(AdditionalFile.Info.Filename);
+	if (AdditionalFile.Info.ChunkId.IsValid() && AdditionalFile.Info.ChunkId != ChunkId)
+	{
+		UE_LOG(LogPackageWriter, Warning, TEXT("PackageWriter->WriteAdditionalFile called with an unexpected chunkid: Should match CreateExternalFileChunkId (Package: %s Filename: %s)"), *Info.InputPackageName.ToString(), *AdditionalFile.Info.Filename);
+	}
+	else
+	{
+		AdditionalFile.Info.ChunkId = ChunkId;
+	}
 }
 
 void FPackageWriterRecords::WriteLinkerAdditionalData(const IPackageWriter::FLinkerAdditionalDataInfo& Info,
diff --git a/Engine/Source/Runtime/CoreUObject/Private/UObject/AsyncWorkSequence.h b/Engine/Source/Runtime/CoreUObject/Private/UObject/AsyncWorkSequence.h
deleted file mode 100644
diff --git a/Engine/Source/Runtime/CoreUObject/Private/UObject/SavePackage.cpp b/Engine/Source/Runtime/CoreUObject/Private/UObject/SavePackage.cpp
--- a/Engine/Source/Runtime/CoreUObject/Private/UObject/SavePackage.cpp
+++ b/Engine/Source/Runtime/CoreUObject/Private/UObject/SavePackage.cpp
@@ -67,7 +67,6 @@
 #include "Serialization/PackageWriter.h"
 #include "Serialization/UnversionedPropertySerialization.h"
 #include "Serialization/EditorBulkData.h"
-#include "UObject/AsyncWorkSequence.h"
 #include "Misc/ScopeExit.h"
 #include "Misc/PackageAccessTracking.h"
 #include "Misc/PackageAccessTrackingOps.h"
@@ -1858,7 +1857,6 @@ FSavePackageResultStruct UPackage::Save(UPackage* InOuter, UObject* InAsset, con
 		TRefCountPtr<FUObjectSerializeContext> SaveContext(FUObjectThreadContext::Get().GetSerializeContext());
 
 		const bool bCompareLinker = (SaveFlags & SAVE_CompareLinker) != 0;
-		const bool bComputeHash = (SaveFlags & SAVE_ComputeHash) != 0;
 
 		if (GIsSavingPackage && !bSavingConcurrent)
 		{
@@ -1918,9 +1916,6 @@ FSavePackageResultStruct UPackage::Save(UPackage* InOuter, UObject* InAsset, con
 
 		int64 TotalPackageSizeUncompressed = 0;
 
-		TFuture<FMD5Hash> PackageMD5Destination;
-		TAsyncWorkSequence<FMD5> AsyncWriteAndHashSequence;
-
 		// Make sure package is fully loaded before saving. 
 		// IsFullyLoaded has important byproducts for new packages, so make sure we evaluate it rather than short-circuiting it
 		if (!InOuter->IsFullyLoaded() && !Base)
@@ -3833,7 +3828,7 @@ FSavePackageResultStruct UPackage::Save(UPackage* InOuter, UObject* InAsset, con
 					int64 DataStartOffset = LinkerSize >= 0 ? LinkerSize : Linker->Tell();
 					ESavePackageResult SaveResult =
 						SavePackageUtilities::SaveBulkData(Linker.Get(), DataStartOffset, InOuter, Filename, TargetPlatform, SavePackageContext, SaveFlags, bTextFormat,
-							bComputeHash, AsyncWriteAndHashSequence, TotalPackageSizeUncompressed);
+							TotalPackageSizeUncompressed);
 					if (SaveResult != ESavePackageResult::Success)
 					{
 						return SaveResult;
@@ -4159,16 +4154,11 @@ FSavePackageResultStruct UPackage::Save(UPackage* InOuter, UObject* InAsset, con
 							// Add the uasset file to the list of output files
 							AdditionalOutputFiles.Emplace(NewPath, FLargeMemoryPtr(Writer->ReleaseOwnership()), Linker->FileRegions, DataSize);
 							
-							EAsyncWriteOptions WriteOptions(EAsyncWriteOptions::None);
-							if (bComputeHash)
-							{
-								WriteOptions |= EAsyncWriteOptions::ComputeHash;
-							}
-
+							EAsyncWriteOptions WriteOptions(EAsyncWriteOptions::None);							
 							for (FSavePackageOutputFile& Entry : AdditionalOutputFiles)
 							{
 								TotalPackageSizeUncompressed += Entry.DataSize;
-								SavePackageUtilities::AsyncWriteFile(AsyncWriteAndHashSequence, WriteOptions, Entry);
+								SavePackageUtilities::AsyncWriteFile(WriteOptions, Entry);
 							}	
 						}
 						Linker->CloseAndDestroySaver();
@@ -4195,7 +4185,7 @@ FSavePackageResultStruct UPackage::Save(UPackage* InOuter, UObject* InAsset, con
 						// Add the .uasset file to the list of output files
 						AdditionalOutputFiles.Emplace(NewPath, TempFilename.GetValue(), PackageSize);
 						
-						ESavePackageResult FinalizeResult = SavePackageUtilities::FinalizeTempOutputFiles(TargetPackagePath, AdditionalOutputFiles, bComputeHash, FinalTimeStamp, AsyncWriteAndHashSequence);
+						ESavePackageResult FinalizeResult = SavePackageUtilities::FinalizeTempOutputFiles(TargetPackagePath, AdditionalOutputFiles, FinalTimeStamp);
 						if (FinalizeResult != ESavePackageResult::Success)
 						{
 							bSuccess = false;
@@ -4342,20 +4332,13 @@ FSavePackageResultStruct UPackage::Save(UPackage* InOuter, UObject* InAsset, con
 			}
 #endif
 
-			auto HashCompletionFunc = [](FMD5& State)
-			{
-				FMD5Hash OutputHash;
-				OutputHash.Set(State);
-				return OutputHash;
-			};
-
 			if (bRequestStub)
 			{
-				return FSavePackageResultStruct(ESavePackageResult::GenerateStub, TotalPackageSizeUncompressed, AsyncWriteAndHashSequence.Finalize(EAsyncExecution::TaskGraph, MoveTemp(HashCompletionFunc)), SerializedPackageFlags, bCompareLinker ? MoveTemp(Linker) : nullptr);
+				return FSavePackageResultStruct(ESavePackageResult::GenerateStub, TotalPackageSizeUncompressed, SerializedPackageFlags, bCompareLinker ? MoveTemp(Linker) : nullptr);
 			}
 			else
 			{
-				return FSavePackageResultStruct(ESavePackageResult::Success, TotalPackageSizeUncompressed, AsyncWriteAndHashSequence.Finalize(EAsyncExecution::TaskGraph, MoveTemp(HashCompletionFunc)), SerializedPackageFlags, bCompareLinker ? MoveTemp(Linker) : nullptr);
+				return FSavePackageResultStruct(ESavePackageResult::Success, TotalPackageSizeUncompressed, SerializedPackageFlags, bCompareLinker ? MoveTemp(Linker) : nullptr);
 			}
 		}
 		else
diff --git a/Engine/Source/Runtime/CoreUObject/Private/UObject/SavePackage/SaveContext.h b/Engine/Source/Runtime/CoreUObject/Private/UObject/SavePackage/SaveContext.h
--- a/Engine/Source/Runtime/CoreUObject/Private/UObject/SavePackage/SaveContext.h
+++ b/Engine/Source/Runtime/CoreUObject/Private/UObject/SavePackage/SaveContext.h
@@ -14,7 +14,6 @@
 #include "Serialization/UnversionedPropertySerialization.h"
 #include "Templates/PimplPtr.h"
 #include "Templates/UniquePtr.h"
-#include "UObject/AsyncWorkSequence.h"
 #include "UObject/LinkerSave.h"
 #include "UObject/NameTypes.h"
 #include "UObject/ObjectMacros.h"
@@ -576,11 +575,6 @@ public:
 		return !!(SaveArgs.SaveFlags & SAVE_Optional);
 	}
 
-	bool IsComputeHash() const
-	{
-		return !!(SaveArgs.SaveFlags & SAVE_ComputeHash);
-	}
-
 	bool IsConcurrent() const
 	{
 		return !!(SaveArgs.SaveFlags & SAVE_Concurrent);
@@ -919,13 +913,6 @@ public:
 
 	FSavePackageResultStruct GetFinalResult()
 	{
-		auto HashCompletionFunc = [](FMD5& State)
-		{
-			FMD5Hash OutputHash;
-			OutputHash.Set(State);
-			return OutputHash;
-		};
-
 		if (Result != ESavePackageResult::Success)
 		{
 			return Result;
@@ -933,7 +920,6 @@ public:
 
 		ESavePackageResult FinalResult = IsStubRequested() ? ESavePackageResult::GenerateStub : ESavePackageResult::Success;
 		return FSavePackageResultStruct(FinalResult, TotalPackageSizeUncompressed,
-			AsyncWriteAndHashSequence.Finalize(EAsyncExecution::TaskGraph, MoveTemp(HashCompletionFunc)),
 			SerializedPackageFlags, IsCompareLinker() ? MoveTemp(GetHarvestedRealm().Linker) : nullptr);
 	}
 
@@ -974,7 +960,6 @@ public:
 	int32 OffsetAfterExportMap = 0;
 	int64 OffsetAfterPayloadToc = 0;
 	int32 SerializedPackageFlags = 0;
-	TAsyncWorkSequence<FMD5> AsyncWriteAndHashSequence;
 	TArray<FLargeMemoryWriter, TInlineAllocator<4>> AdditionalFilesFromExports;
 	FSavePackageOutputFileArray AdditionalPackageFiles;
 private:
diff --git a/Engine/Source/Runtime/CoreUObject/Private/UObject/SavePackage/SavePackageUtilities.cpp b/Engine/Source/Runtime/CoreUObject/Private/UObject/SavePackage/SavePackageUtilities.cpp
--- a/Engine/Source/Runtime/CoreUObject/Private/UObject/SavePackage/SavePackageUtilities.cpp
+++ b/Engine/Source/Runtime/CoreUObject/Private/UObject/SavePackage/SavePackageUtilities.cpp
@@ -19,7 +19,7 @@
 #include "Serialization/EditorBulkData.h"
 #include "Serialization/LargeMemoryWriter.h"
 #include "Serialization/PackageWriter.h"
-#include "UObject/AsyncWorkSequence.h"
+#include "Tasks/Task.h"
 #include "UObject/Class.h"
 #include "UObject/GCScopeLock.h"
 #include "UObject/Linker.h"
@@ -616,27 +616,7 @@ void FindMostLikelyCulprit(const TArray<UObject*>& BadObjects, UObject*& MostLik
 	}
 }
 
-void AddFileToHash(FString const& Filename, FMD5& Hash)
-{
-	TArray<uint8> LocalScratch;
-	LocalScratch.SetNumUninitialized(1024 * 64);
-
-	FArchive* Ar = IFileManager::Get().CreateFileReader(*Filename);
-
-	const int64 Size = Ar->TotalSize();
-	int64 Position = 0;
-
-	while (Position < Size)
-	{
-		const auto ReadNum = FMath::Min(Size - Position, (int64)LocalScratch.Num());
-		Ar->Serialize(LocalScratch.GetData(), ReadNum);
-		Hash.Update(LocalScratch.GetData(), ReadNum);
-		Position += ReadNum;
-	}
-	delete Ar;
-}
-
-ESavePackageResult FinalizeTempOutputFiles(const FPackagePath& PackagePath, const FSavePackageOutputFileArray& OutputFiles, const bool bComputeHash, const FDateTime& FinalTimeStamp, TAsyncWorkSequence<FMD5>& AsyncWriteAndHashSequence)
+ESavePackageResult FinalizeTempOutputFiles(const FPackagePath& PackagePath, const FSavePackageOutputFileArray& OutputFiles, const FDateTime& FinalTimeStamp)
 {
 	UE_LOG(LogSavePackage, Log,  TEXT("Moving output files for package: %s"), *PackagePath.GetDebugName());
 
@@ -724,16 +704,6 @@ ESavePackageResult FinalizeTempOutputFiles(const FPackagePath& PackagePath, cons
 			{
 				FileSystem.SetTimeStamp(*File.TargetPath, FinalTimeStamp);
 			}
-
-			if (bComputeHash)
-			{
-				UE::SavePackageUtilities::IncrementOutstandingAsyncWrites();
-				AsyncWriteAndHashSequence.AddWork([NewPath = File.TargetPath](FMD5& State)
-				{
-					SavePackageUtilities::AddFileToHash(NewPath, State);
-					UE::SavePackageUtilities::DecrementOutstandingAsyncWrites();
-				});
-			}
 		}
 	}
 
@@ -767,17 +737,13 @@ void WriteToFile(const FString& Filename, const uint8* InDataPtr, int64 InDataSi
 	UE_LOG(LogSavePackage, Fatal, TEXT("Could not write to %s!"), *Filename);
 }
 
-void AsyncWriteFile(TAsyncWorkSequence<FMD5>& AsyncWriteAndHashSequence, FLargeMemoryPtr Data, const int64 DataSize, const TCHAR* Filename, EAsyncWriteOptions Options, TArrayView<const FFileRegion> InFileRegions)
+void AsyncWriteFile(FLargeMemoryPtr Data, const int64 DataSize, const TCHAR* Filename, EAsyncWriteOptions Options, TArrayView<const FFileRegion> InFileRegions)
 {
 	OutstandingAsyncWrites.Increment();
 	FString OutputFilename(Filename);
-	AsyncWriteAndHashSequence.AddWork([Data = MoveTemp(Data), DataSize, OutputFilename = MoveTemp(OutputFilename), Options, FileRegions = TArray<FFileRegion>(InFileRegions)](FMD5& State) mutable
-	{
-		if (EnumHasAnyFlags(Options, EAsyncWriteOptions::ComputeHash))
-		{
-			State.Update(Data.Get(), DataSize);
-		}
 
+	UE::Tasks::Launch(TEXT("PackageAsyncFileWrite"), [Data = MoveTemp(Data), DataSize, OutputFilename = MoveTemp(OutputFilename), Options, FileRegions = TArray<FFileRegion>(InFileRegions)]() mutable
+	{
 		WriteToFile(OutputFilename, Data.Get(), DataSize);
 
 		if (FileRegions.Num() > 0)
@@ -793,10 +759,10 @@ void AsyncWriteFile(TAsyncWorkSequence<FMD5>& AsyncWriteAndHashSequence, FLargeM
 	});
 }
 
-void AsyncWriteFile(TAsyncWorkSequence<FMD5>& AsyncWriteAndHashSequence, EAsyncWriteOptions Options, FSavePackageOutputFile& File)
+void AsyncWriteFile(EAsyncWriteOptions Options, FSavePackageOutputFile& File)
 {
 	checkf(File.TempFilePath.IsEmpty(), TEXT("AsyncWriteFile does not handle temp files!"));
-	AsyncWriteFile(AsyncWriteAndHashSequence, FLargeMemoryPtr(File.FileMemoryBuffer.Release()), File.DataSize, *File.TargetPath, Options, File.FileRegions);
+	AsyncWriteFile(FLargeMemoryPtr(File.FileMemoryBuffer.Release()), File.DataSize, *File.TargetPath, Options, File.FileRegions);
 }
 
 /** For a CDO get all of the subobjects templates nested inside it or it's class */
@@ -2026,8 +1992,8 @@ ESavePackageResult CreatePayloadSidecarFile(FLinkerSave& Linker, const FPackageP
 }
 
 ESavePackageResult SaveBulkData(FLinkerSave* Linker, int64& InOutStartOffset, const UPackage* InOuter, const TCHAR* Filename, const ITargetPlatform* TargetPlatform,
-				  FSavePackageContext* SavePackageContext, uint32 SaveFlags, const bool bTextFormat, const bool bComputeHash,
-				  TAsyncWorkSequence<FMD5>& AsyncWriteAndHashSequence, int64& TotalPackageSizeUncompressed, bool bIsOptionalRealm)
+				  FSavePackageContext* SavePackageContext, uint32 SaveFlags, const bool bTextFormat, 
+				  int64& TotalPackageSizeUncompressed, bool bIsOptionalRealm)
 {
 	// Now we write all the bulkdata that is supposed to be at the end of the package
 	// and fix up the offset
@@ -2359,11 +2325,7 @@ ESavePackageResult SaveBulkData(FLinkerSave* Linker, int64& InOutStartOffset, co
 					const FString ArchiveFilename = FPaths::ChangeExtension(Filename, BulkFileExtension);
 
 					EAsyncWriteOptions WriteOptions(EAsyncWriteOptions::None);
-					if (bComputeHash)
-					{
-						WriteOptions |= EAsyncWriteOptions::ComputeHash;
-					}
-					SavePackageUtilities::AsyncWriteFile(AsyncWriteAndHashSequence, MoveTemp(DataPtr), DataSize, *ArchiveFilename, WriteOptions, Archive->FileRegions);
+					SavePackageUtilities::AsyncWriteFile(MoveTemp(DataPtr), DataSize, *ArchiveFilename, WriteOptions, Archive->FileRegions);
 				}
 			};
 
diff --git a/Engine/Source/Runtime/CoreUObject/Private/UObject/SavePackage/SavePackageUtilities.h b/Engine/Source/Runtime/CoreUObject/Private/UObject/SavePackage/SavePackageUtilities.h
--- a/Engine/Source/Runtime/CoreUObject/Private/UObject/SavePackage/SavePackageUtilities.h
+++ b/Engine/Source/Runtime/CoreUObject/Private/UObject/SavePackage/SavePackageUtilities.h
@@ -17,12 +17,10 @@
 
 // This file contains private utilities shared by UPackage::Save and UPackage::Save2 
 
-class FMD5;
 class FPackagePath;
 class FSaveContext;
 class FSavePackageContext;
 class IPackageWriter;
-template<typename StateType> class TAsyncWorkSequence;
 
 enum class ESavePackageResult;
 
@@ -50,8 +48,7 @@ typedef TUniquePtr<uint8, FLargeMemoryDelete> FLargeMemoryPtr;
 
 enum class EAsyncWriteOptions
 {
-	None = 0,
-	ComputeHash = 0x01,
+	None = 0
 };
 ENUM_CLASS_FLAGS(EAsyncWriteOptions)
 
@@ -362,6 +359,7 @@ namespace SavePackageUtilities
 	/**
 	 * Save all of the BulkDatas that have registered as not being inline.
 	 * They may be saved to the end of the file, or to an archive for a separate file.
+	 * Kicks off async writes, use UPackage::WaitForAsyncFileWrites(); to ensure they are complete.
 	 * 
 	 * @param Linker The linker containing the exports. Provides metadata for BulkData, and BulkData may write to it as their target archive.
 	 * @param InOutStartOffset In value is the offset in the Linker's archive where the BulkDatas will be put. If SavePackageContext settings direct
@@ -374,14 +372,12 @@ namespace SavePackageUtilities
 	 * @param SaveFlags The flags passed into SavePackage; affects how the bulkdata should be saved.
 	 * @param bTextFormat True if package is saving to text. Bulkdata has special handling for text output.
 	 * @param bDiffing True if the package is only diffing, so bulk data should be recorded but not saved to disk.
-	 * @param bComputeHash True if package hash needs to be computed; bulkdatas contribute to the hash.
-	 * @param AsyncWriteAndHashSequence Output: Collects the writes of the bulkdata to disk for async writing.
 	 * @param TotalPackageSizeUncompressed Output: Bulkdata sizes are added to it.
 	 */
 	ESavePackageResult SaveBulkData(FLinkerSave* Linker, int64& InOutStartOffset, const UPackage* InOuter,
 		const TCHAR* Filename, const ITargetPlatform* TargetPlatform, FSavePackageContext* SavePackageContext,
-		uint32 SaveFlags, const bool bTextFormat, const bool bComputeHash,
-		TAsyncWorkSequence<FMD5>& AsyncWriteAndHashSequence, int64& TotalPackageSizeUncompressed, bool bIsOptionalRealm = false);
+		uint32 SaveFlags, const bool bTextFormat,
+		int64& TotalPackageSizeUncompressed, bool bIsOptionalRealm = false);
 	
 	/**
 	 * Used to append additional data to the end of the package file by invoking callbacks stored in the linker.
@@ -406,19 +402,18 @@ namespace SavePackageUtilities
 	void CheckObjectPriorToSave(FArchiveUObject& Ar, UObject* InObj, UPackage* InSavingPackage);
 	void ConditionallyExcludeObjectForTarget(UObject* Obj, EObjectMark ExcludedObjectMarks, const ITargetPlatform* TargetPlatform);
 	void FindMostLikelyCulprit(const TArray<UObject*>& BadObjects, UObject*& MostLikelyCulprit, FString& OutReferencer, FSaveContext* InOptionalSaveContext = nullptr);
-	void AddFileToHash(FString const& Filename, FMD5& Hash);
-
+	
 	/** 
 	  * Search 'OutputFiles' for output files that were saved to the temp directory and move those files
 	  * to their final location. Output files that were not saved to the temp directory will be ignored.
 	  * 
 	  * If errors are encountered then the original state of the package will be restored and should continue to work.
 	  */
-	ESavePackageResult FinalizeTempOutputFiles(const FPackagePath& PackagePath, const FSavePackageOutputFileArray& OutputFiles, const bool bComputeHash, const FDateTime& FinalTimeStamp, TAsyncWorkSequence<FMD5>& AsyncWriteAndHashSequence);
+	ESavePackageResult FinalizeTempOutputFiles(const FPackagePath& PackagePath, const FSavePackageOutputFileArray& OutputFiles, const FDateTime& FinalTimeStamp);
 
 	void WriteToFile(const FString& Filename, const uint8* InDataPtr, int64 InDataSize);
-	void AsyncWriteFile(TAsyncWorkSequence<FMD5>& AsyncWriteAndHashSequence, FLargeMemoryPtr Data, const int64 DataSize, const TCHAR* Filename, EAsyncWriteOptions Options, TArrayView<const FFileRegion> InFileRegions);
-	void AsyncWriteFile(TAsyncWorkSequence<FMD5>& AsyncWriteAndHashSequence, EAsyncWriteOptions Options, FSavePackageOutputFile& File);
+	void AsyncWriteFile(FLargeMemoryPtr Data, const int64 DataSize, const TCHAR* Filename, EAsyncWriteOptions Options, TArrayView<const FFileRegion> InFileRegions);
+	void AsyncWriteFile(EAsyncWriteOptions Options, FSavePackageOutputFile& File);
 
 	void GetCDOSubobjects(UObject* CDO, TArray<UObject*>& Subobjects);
 }
diff --git a/Engine/Source/Runtime/CoreUObject/Private/UObject/SavePackage2.cpp b/Engine/Source/Runtime/CoreUObject/Private/UObject/SavePackage2.cpp
--- a/Engine/Source/Runtime/CoreUObject/Private/UObject/SavePackage2.cpp
+++ b/Engine/Source/Runtime/CoreUObject/Private/UObject/SavePackage2.cpp
@@ -29,7 +29,6 @@
 #include "Serialization/PackageWriter.h"
 #include "Serialization/PropertyLocalizationDataGathering.h"
 #include "Serialization/UnversionedPropertySerialization.h"
-#include "UObject/AsyncWorkSequence.h"
 #include "UObject/DebugSerializationFlags.h"
 #include "UObject/EditorObjectVersion.h"
 #include "UObject/LinkerSave.h"
@@ -1822,7 +1821,6 @@ ESavePackageResult WriteAdditionalExportFiles(FSaveContext& SaveContext)
 		checkf(SaveContext.GetCurrentHarvestingRealm() != ESaveRealm::Optional , TEXT("Addtional export files is currently unsupported with optional package multi output"));
 		IPackageWriter* PackageWriter = SavePackageContext ? SavePackageContext->PackageWriter : nullptr;
 		checkf(PackageWriter, TEXT("Cooking requires a PackageWriter"));
-		const bool bComputeHash = SaveContext.IsComputeHash();
 		for (FLargeMemoryWriter& Writer : SaveContext.AdditionalFilesFromExports)
 		{
 			const int64 Size = Writer.TotalSize();
@@ -1834,6 +1832,8 @@ ESavePackageResult WriteAdditionalExportFiles(FSaveContext& SaveContext)
 
 			FIoBuffer FileData(FIoBuffer::AssumeOwnership, Writer.ReleaseOwnership(), Size);
 
+			// This might not actuallly write the file, but instead add it to a queue to write later.
+			// (See TPackageWriterToSharedBuffer).
 			PackageWriter->WriteAdditionalFile(FileInfo, FileData);
 		}
 		SaveContext.AdditionalFilesFromExports.Empty();
@@ -1997,10 +1997,6 @@ ESavePackageResult FinalizeFile(FStructuredArchive::FRecord& StructuredArchiveRo
 		{
 			checkf(!SaveContext.IsCooking(), TEXT("Cooking requires a PackageWriter"));
 			EAsyncWriteOptions WriteOptions(EAsyncWriteOptions::None);
-			if (SaveContext.IsComputeHash())
-			{
-				WriteOptions |= EAsyncWriteOptions::ComputeHash;
-			}
 
 			// Add the uasset file to the list of output files
 			int64 DataSize = Writer->TotalSize();
@@ -2009,7 +2005,7 @@ ESavePackageResult FinalizeFile(FStructuredArchive::FRecord& StructuredArchiveRo
 
 			for (FSavePackageOutputFile& Entry : SaveContext.AdditionalPackageFiles)
 			{
-				SavePackageUtilities::AsyncWriteFile(SaveContext.AsyncWriteAndHashSequence, WriteOptions, Entry);
+				SavePackageUtilities::AsyncWriteFile(WriteOptions, Entry);
 			}
 		}
 		SaveContext.CloseLinkerArchives();
@@ -2039,8 +2035,8 @@ ESavePackageResult FinalizeFile(FStructuredArchive::FRecord& StructuredArchiveRo
 		// Add the .uasset file to the list of output files (TODO: Fix the 0 size, it isn't used after this point but needs to be cleaned up)
 		SaveContext.AdditionalPackageFiles.Emplace(SaveContext.GetFilename(), SaveContext.GetTempFilename().GetValue(), 0);
 
-		ESavePackageResult FinalizeResult = SavePackageUtilities::FinalizeTempOutputFiles(SaveContext.GetTargetPackagePath(), SaveContext.AdditionalPackageFiles, SaveContext.IsComputeHash(), 
-			SaveContext.GetFinalTimestamp(), SaveContext.AsyncWriteAndHashSequence);
+		ESavePackageResult FinalizeResult = SavePackageUtilities::FinalizeTempOutputFiles(SaveContext.GetTargetPackagePath(), SaveContext.AdditionalPackageFiles,
+			SaveContext.GetFinalTimestamp());
 		
 		SaveContext.SetTempFilename(TOptional<FString>());
 
@@ -2424,7 +2420,7 @@ ESavePackageResult WriteAdditionalFiles(FSaveContext& SaveContext, FScopedSlowTa
 	SaveContext.Result = SavePackageUtilities::SaveBulkData(SaveContext.GetLinker(), DataStartOffset,
 		SaveContext.GetPackage(), SaveContext.GetFilename(), SaveContext.GetTargetPlatform(),
 		SaveContext.GetSavePackageContext(), SaveContext.GetSaveArgs().SaveFlags, SaveContext.IsTextFormat(),
-		SaveContext.IsComputeHash(), SaveContext.AsyncWriteAndHashSequence, SaveContext.TotalPackageSizeUncompressed,
+		SaveContext.TotalPackageSizeUncompressed,
 		SaveContext.GetCurrentHarvestingRealm() == ESaveRealm::Optional);
 
 	// Add any pending data blobs to the end of the file by invoking the callbacks
diff --git a/Engine/Source/Runtime/CoreUObject/Public/AssetRegistry/AssetData.h b/Engine/Source/Runtime/CoreUObject/Public/AssetRegistry/AssetData.h
--- a/Engine/Source/Runtime/CoreUObject/Public/AssetRegistry/AssetData.h
+++ b/Engine/Source/Runtime/CoreUObject/Public/AssetRegistry/AssetData.h
@@ -9,6 +9,7 @@
 #include "Containers/ArrayView.h"
 #include "Containers/StringView.h"
 #include "Containers/Set.h"
+#include "IO/IoDispatcher.h"
 #include "Misc/PackageName.h"
 #include "Misc/StringBuilder.h"
 #include "Templates/UniquePtr.h"
@@ -55,6 +56,7 @@ struct COREUOBJECT_API FAssetRegistryVersion
 		PackageImportedClasses,				// Added ImportedClasses to AssetPackageData
 		PackageFileSummaryVersionChange,	// A new version number of UE5 was added to FPackageFileSummary
 		ObjectResourceOptionalVersionChange,// Change to linker export/import resource serializationn
+		AddedChunkHashes,					// Added FIoHash for each FIoChunkId in the package to the AssetPackageData.
 
 		// -----<new versions can be added above this line>-------------------------------------------------
 		VersionPlusOne,
@@ -166,6 +168,8 @@ public:
 	/**
 	 * Returns true if this is the main asset in a package, true for maps and assets but false for secondary objects like class redirectors
 	 * Every UAsset is also a TopLevelAsset.
+	 * Specifically, this is just testing whether the asset name matches the package name, and it is possible for legacy content to not return
+	 * true for any asset in a package.
 	 */
 	bool IsUAsset() const
 	{
@@ -431,6 +435,13 @@ public:
 	template <typename ValueType>
 	bool GetTagValue(FName Tag, ValueType& OutValue) const;
 
+	/** Call the given function for each tag on the asset. Function signature is: void Fn(TPair<FName, FAssetTagValueRef>); */
+	template<typename Func>
+	void EnumerateTags(Func Fn) const
+	{
+		TagsAndValues.ForEach(Fn);
+	}
+
 	/** Try and get the value associated with the given tag as a type converted value, or an empty value if it doesn't exist */
 	template <typename ValueType>
 	ValueType GetTagValueRef(const FName Tag) const;
@@ -699,6 +710,13 @@ public:
 	/** MD5 of the cooked package on disk, for tracking nondeterministic changes */
 	FMD5Hash CookedHash;
 
+	/** 
+		The hash of all the chunks written by the package writer for this package. This is only available after cooking.
+		(Note, this only includes the ExportBundleData chunk if using ZenStoreWriter, as that chunk contains both the uasset and uexp
+		data, which is managed separately in LooseCookedPackageWriter.)
+	*/
+	TMap<FIoChunkId, FIoHash> ChunkHashes;
+
 	/** List of classes used by exports in the package. Does not include classes in the same package. */
 	TArray<FName> ImportedClasses;
 
diff --git a/Engine/Source/Runtime/CoreUObject/Public/UObject/ObjectMacros.h b/Engine/Source/Runtime/CoreUObject/Public/UObject/ObjectMacros.h
--- a/Engine/Source/Runtime/CoreUObject/Public/UObject/ObjectMacros.h
+++ b/Engine/Source/Runtime/CoreUObject/Public/UObject/ObjectMacros.h
@@ -93,7 +93,7 @@ enum ESaveFlags
 	SAVE_Concurrent					= 0x00000100,	///< We are save packages in multiple threads at once and should not call non-threadsafe functions or rely on globals. GIsSavingPackage should be set and PreSave/Postsave functions should be called before/after the entire concurrent save.
 	SAVE_DiffOnly UE_DEPRECATED(5.0, "Diffing is now done using FDiffPackageWriter.") = 0x00000200,	
 	SAVE_DiffCallstack UE_DEPRECATED(5.0, "Diffing is now done using FDiffPackageWriter.") = 0x00000400,
-	SAVE_ComputeHash				= 0x00000800,	///< Compute the MD5 hash of the cooked data
+	SAVE_ComputeHash UE_DEPRECATED(5.1, "Add EWriteOptions::ComputeHash directly to CommitPackage options") = 0x00000800,	///< Compute the MD5 hash of the cooked data, as well as the per iostore chunk FIoHash.
 	SAVE_CompareLinker				= 0x00001000,	///< Return the linker save to compare against another
 	SAVE_BulkDataByReference		= 0x00002000,	///< When saving to a different file than the package's LoadedPath, point bulkdata in the new file to be loaded from the original file.
 	SAVE_Unversioned_Properties		= 0x00004000,	///< Properties are saved without property name information, and are saved/loaded in the order of the current binary.
diff --git a/Engine/Source/Runtime/CoreUObject/Public/UObject/Package.h b/Engine/Source/Runtime/CoreUObject/Public/UObject/Package.h
--- a/Engine/Source/Runtime/CoreUObject/Public/UObject/Package.h
+++ b/Engine/Source/Runtime/CoreUObject/Public/UObject/Package.h
@@ -65,6 +65,7 @@ enum class ESavePackageResult
 /**
 * Struct returned from save package, contains the enum as well as extra data about what was written
 */
+PRAGMA_DISABLE_DEPRECATION_WARNINGS // Silence deprecation warnings for deprecated CookedHash member in implicit constructors
 struct FSavePackageResultStruct
 {
 	/** Success/failure of the save operation */
@@ -73,7 +74,7 @@ struct FSavePackageResultStruct
 	/** Total size of all files written out, including bulk data */
 	int64 TotalFileSize;
 
-	/** MD5 hash of the cooked data */
+	UE_DEPRECATED(5.1, "CookedHash is now available through PackageWriter->CommitPackage instead. For waiting on completion in the non-cook case, use UPackage::WaitForAsyncFileWrites.")
 	TFuture<FMD5Hash> CookedHash;
 
 	/** Serialized package flags */
@@ -86,7 +87,7 @@ struct FSavePackageResultStruct
 	FSavePackageResultStruct() : Result(ESavePackageResult::Error), TotalFileSize(0), SerializedPackageFlags(0) {}
 	FSavePackageResultStruct(ESavePackageResult InResult) : Result(InResult), TotalFileSize(0), SerializedPackageFlags(0) {}
 	FSavePackageResultStruct(ESavePackageResult InResult, int64 InTotalFileSize) : Result(InResult), TotalFileSize(InTotalFileSize), SerializedPackageFlags(0) {}
-	FSavePackageResultStruct(ESavePackageResult InResult, int64 InTotalFileSize, TFuture<FMD5Hash>&& InHash, uint32 InSerializedPackageFlags, TPimplPtr<FLinkerSave> Linker = nullptr) : Result(InResult), TotalFileSize(InTotalFileSize), CookedHash(MoveTemp(InHash)), SerializedPackageFlags(InSerializedPackageFlags), LinkerSave(MoveTemp(Linker)) {}
+	FSavePackageResultStruct(ESavePackageResult InResult, int64 InTotalFileSize, uint32 InSerializedPackageFlags, TPimplPtr<FLinkerSave> Linker = nullptr) : Result(InResult), TotalFileSize(InTotalFileSize), SerializedPackageFlags(InSerializedPackageFlags), LinkerSave(MoveTemp(Linker)) {}
 
 	bool operator==(const FSavePackageResultStruct& Other) const
 	{
@@ -107,6 +108,7 @@ struct FSavePackageResultStruct
 			Result == ESavePackageResult::ReplaceCompletely;
 	}
 };
+PRAGMA_ENABLE_DEPRECATION_WARNINGS
 
 /**
 * A package.
