SourceControl: Prevent unnecessary asset saves.

When choosing Check-In on a directory in the Content Browser, PromptForCheckoutAndSave is called with a list of all loaded packages that reside in that directory hierarchy by FAssetFolderContextMenu::ExecuteSCCCheckIn with bCheckDirty=true and bPromptToSave=true. However, if none of the given packages is dirty, the end result would be that all packages would be saved instead of none of the packages being saved. This is unfortunate as source control systems will -unnecessarily- end up considering those assets as changed.

diff --git a/Engine/Source/Editor/UnrealEd/Private/FileHelpers.cpp b/Engine/Source/Editor/UnrealEd/Private/FileHelpers.cpp
--- a/Engine/Source/Editor/UnrealEd/Private/FileHelpers.cpp
+++ b/Engine/Source/Editor/UnrealEd/Private/FileHelpers.cpp
@@ -4121,7 +4121,7 @@ FEditorFileUtils::EPromptReturnCode FEditorFileUtils::PromptForCheckoutAndSave(
 
 		TArray<UPackage*> AddPackageItemsChecked;
 		TArray<UPackage*> AddPackageItemsUnchecked;
-		for ( TArray<UPackage*>::TConstIterator PkgIter(PackagesToSave); PkgIter; ++PkgIter )
+		for ( TArray<UPackage*>::TIterator PkgIter(PackagesToSave); PkgIter; ++PkgIter )
 		{
 			UPackage* CurPackage = *PkgIter;
 			check( CurPackage );
@@ -4154,6 +4154,10 @@ FEditorFileUtils::EPromptReturnCode FEditorFileUtils::PromptForCheckoutAndSave(
 					UE_LOG(LogFileHelpers, Warning, TEXT("PromptForCheckoutAndSave attempted to open the save dialog with the transient package"));
 				}
 			}
+			else
+			{
+				PkgIter.RemoveCurrent();
+			}
 		}
 
 		if ( AddPackageItemsUnchecked.Num() > 0 || AddPackageItemsChecked.Num() > 0 )
@@ -4279,7 +4283,6 @@ FEditorFileUtils::EPromptReturnCode FEditorFileUtils::PromptForCheckoutAndSave(
 			// The user cancelled the checkout dialog, so set the return response accordingly  
 			ReturnResponse = PR_Cancelled;
 		}
-
 	}
 
 	return ReturnResponse;
