StrataLegacyConversion node cannot be copied. When the user tries to, we notify why.

diff --git a/Engine/Source/Editor/UnrealEd/Private/MaterialGraphNode.cpp b/Engine/Source/Editor/UnrealEd/Private/MaterialGraphNode.cpp
--- a/Engine/Source/Editor/UnrealEd/Private/MaterialGraphNode.cpp
+++ b/Engine/Source/Editor/UnrealEd/Private/MaterialGraphNode.cpp
@@ -138,6 +138,22 @@ bool UMaterialGraphNode::CanPasteHere(const UEdGraph* TargetGraph) const
 {
 	if (Super::CanPasteHere(TargetGraph))
 	{
+		if (MaterialExpression->IsA(UMaterialExpressionStrataLegacyConversion::StaticClass()))
+		{
+			// We could have used CanDuplicateNode() returning false to prevent the copy but it is nicer to have a notification about why the copy is not happening.
+			FText NotificationText = NSLOCTEXT("UMaterialGraphNode", "SkippedStrataLegacyConversion", "StrataLegacyConversion node cannot be copied! It is only used to convert legacy material to Strata.");
+			FNotificationInfo Info(NotificationText);
+			Info.ExpireDuration = 5.0f;
+			Info.bUseLargeFont = false;
+			Info.Image = FCoreStyle::Get().GetBrush(TEXT("MessageLog.Warning"));
+			TSharedPtr<SNotificationItem> Notification = FSlateNotificationManager::Get().AddNotification(Info);
+			if (Notification.IsValid())
+			{
+				Notification->SetCompletionState(SNotificationItem::CS_None);
+			}
+			return false; // We do not allow the copy of the StrataLegacyConversion which should only be used to convert pre-strata materials.
+		}
+
 		const UMaterialGraph* MaterialGraph = Cast<const UMaterialGraph>(TargetGraph);
 		if (MaterialGraph)
 		{
