Fix FBlueprintEditorUtils::FindUniqueKismetName that returned a numbered name even if the base name was valid.
#rrb marc.audy

diff --git a/Engine/Source/Editor/UnrealEd/Private/Kismet2/BlueprintEditorUtils.cpp b/Engine/Source/Editor/UnrealEd/Private/Kismet2/BlueprintEditorUtils.cpp
--- a/Engine/Source/Editor/UnrealEd/Private/Kismet2/BlueprintEditorUtils.cpp
+++ b/Engine/Source/Editor/UnrealEd/Private/Kismet2/BlueprintEditorUtils.cpp
@@ -7827,14 +7827,16 @@ FBlueprintMacroCosmeticInfo FBlueprintEditorUtils::GetCosmeticInfoForMacro(UEdGr
 FName FBlueprintEditorUtils::FindUniqueKismetName(const UBlueprint* InBlueprint, const FString& InBaseName, UStruct* InScope/* = nullptr*/)
 {
 	int32 Count = 0;
-	FString KismetName;
 	// If an empty string is given then we need to give a valid backup
 	static const FString BackupKismetName = TEXT("K2Name");
 	FString BaseName = InBaseName.IsEmpty() ? BackupKismetName : InBaseName;
+	FString KismetName = InBaseName;
 	TSharedPtr<FKismetNameValidator> NameValidator = MakeShareable(new FKismetNameValidator(InBlueprint, NAME_None, InScope));
 
+	EValidatorResult Result = NameValidator->IsValid(KismetName);
+
 	// Clean up BaseName to not contain any invalid characters, which will mean we can never find a legal name no matter how many numbers we add
-	if (NameValidator->IsValid(BaseName) == EValidatorResult::ContainsInvalidCharacters)
+	if (Result == EValidatorResult::ContainsInvalidCharacters)
 	{
 		for (TCHAR& TestChar : BaseName)
 		{
@@ -7847,9 +7849,11 @@ FName FBlueprintEditorUtils::FindUniqueKismetName(const UBlueprint* InBlueprint,
 				}
 			}
 		}
+		KismetName = BaseName;
+		Result = NameValidator->IsValid(KismetName);
 	}
 
-	while(NameValidator->IsValid(KismetName) != EValidatorResult::Ok)
+	while(Result != EValidatorResult::Ok)
 	{
 		// Calculate the number of digits in the number, adding 2 (1 extra to correctly count digits, another to account for the '_' that will be added to the name
 		int32 CountLength = Count > 0? (int32)log((double)Count) + 2 : 2;
@@ -7861,6 +7865,7 @@ FName FBlueprintEditorUtils::FindUniqueKismetName(const UBlueprint* InBlueprint,
 		}
 		KismetName = FString::Printf(TEXT("%s_%d"), *BaseName, Count);
 		Count++;
+		Result = NameValidator->IsValid(KismetName);
 	}
 
 	return FName(*KismetName);
