AsyncIODelete: Change to support multiple processes sharing the same deletion root (and the last one cleaning it up), to support upcoming multiprocess cooker.
#rnx

diff --git a/Engine/Source/Editor/UnrealEd/Private/Cooker/AsyncIODelete.cpp b/Engine/Source/Editor/UnrealEd/Private/Cooker/AsyncIODelete.cpp
--- a/Engine/Source/Editor/UnrealEd/Private/Cooker/AsyncIODelete.cpp
+++ b/Engine/Source/Editor/UnrealEd/Private/Cooker/AsyncIODelete.cpp
@@ -29,23 +29,23 @@ FAsyncIODelete::~FAsyncIODelete()
 	SetTempRoot(FStringView());
 }
 
-void FAsyncIODelete::SetTempRoot(const FStringView& InOwnedTempRoot)
+void FAsyncIODelete::SetTempRoot(FStringView InSharedTempRoot)
 {
 	Teardown();
 
 #if WITH_ASYNCIODELETE_DEBUG
-	if (!RequestedTempRoot.IsEmpty())
+	if (!SharedTempRoot.IsEmpty())
 	{
-		RemoveTempRoot(*RequestedTempRoot);
+		RemoveTempRoot(SharedTempRoot);
 	}
 #endif
 
-	RequestedTempRoot = InOwnedTempRoot;
+	SharedTempRoot = InSharedTempRoot;
 
 #if WITH_ASYNCIODELETE_DEBUG
-	if (!RequestedTempRoot.IsEmpty())
+	if (!SharedTempRoot.IsEmpty())
 	{
-		AddTempRoot(*RequestedTempRoot);
+		AddTempRoot(*SharedTempRoot);
 	}
 #endif
 }
@@ -81,8 +81,7 @@ void FAsyncIODelete::Setup()
 	}
 
 	bInitialized = true;
-	TempRoot = RequestedTempRoot;
-	if (RequestedTempRoot.IsEmpty())
+	if (SharedTempRoot.IsEmpty())
 	{
 		checkf(false, TEXT("DeleteDirectory called without having first set a TempRoot"));
 		return;
@@ -90,23 +89,15 @@ void FAsyncIODelete::Setup()
 
 	if (AsyncEnabled())
 	{
-		// Delete the TempRoot directory to clear the results from any previous process using the same TempRoot that did not shut down cleanly
-		FString FoundRoot;
-		if (!TryPurgeTempRootFamily(&FoundRoot))
+		// Create the new root and at the same time clear the results from any previous process using the same
+		// TempRoot that did not shut down cleanly
+		TArray<FDeleteRequest> OrphanedRootsToDelete;
+		if (!TryPurgeOldAndCreateRoot(true, OrphanedRootsToDelete))
 		{
-			// TryPurgeTempRootFamily logged the warning
+			// TryPurgeOldAndCreateRoot logged the warning
 			return;
 		}
 
-		// Create the empty directory to work in
-		if (!IFileManager::Get().MakeDirectory(*FoundRoot, true))
-		{
-			UE_LOG(LogCook, Error, TEXT("Could not create asyncdelete root directory '%s'. LastError: %i. Falling back to synchronous delete."),
-				*TempRoot, FPlatformMisc::GetLastError());
-			return;
-		}
-		TempRoot = FoundRoot;
-
 		// Allocate the task event
 		check(TasksComplete == nullptr);
 		TasksComplete = FPlatformProcess::GetSynchEventFromPool(true /* IsManualReset */);
@@ -116,8 +107,17 @@ void FAsyncIODelete::Setup()
 		// Assert that all other teardown-transient variables were cleared by the constructor or by the previous teardown
 		// TempRoot and bPaused are preserved across setup/teardown and may have any value
 		check(PausedDeletes.Num() == 0);
-		check(DeleteCounter == 0);
+		check(DeleteCounter == OrphanedRootsToDelete.Num()); // We should have started with DeleteCounter==0, and incremented it for each OrphanedRootToDelete request
 		bAsyncInitialized = true;
+
+		// Queue the async deletion of any discovered orphaned roots
+		for (const FDeleteRequest& DeleteRequest : OrphanedRootsToDelete)
+		{
+			if (!DeleteRequest.Path.IsEmpty())
+			{
+				CreateDeleteTask(DeleteRequest.Path, DeleteRequest.PathType);
+			}
+		}
 	}
 }
 
@@ -137,7 +137,9 @@ void FAsyncIODelete::Teardown()
 		TasksComplete = nullptr;
 
 		// Remove the temp directory from disk
-		TryPurgeTempRootFamily(nullptr);
+		TArray<FDeleteRequest> OrphanedRootsToDelete;
+		TryPurgeOldAndCreateRoot(false, OrphanedRootsToDelete);
+		check(OrphanedRootsToDelete.IsEmpty()); // Should not be populated when bCreateRoot is false
 
 		// Clear delete variables; we don't need to run the tasks for the remaining pauseddeletes because synchronously deleting the temp directory above did the work they were going to do
 		PausedDeletes.Empty();
@@ -172,6 +174,17 @@ bool FAsyncIODelete::WaitForAllTasks(float TimeLimitSeconds)
 	return true;
 }
 
+bool FAsyncIODelete::AsyncEnabled()
+{
+#if PLATFORM_LINUX
+	// Temporarily disable Async on Linux until we have fixed TryPurgeOldAndCreateRoot to request exclusive-write behavior explicitly via calls to flock
+	return false;
+#else
+	return FPlatformMisc::SupportsMultithreadedFileHandles();
+#endif
+}
+
+
 bool FAsyncIODelete::Delete(const FStringView& PathToDelete, EPathType ExpectedType)
 {
 	IFileManager& FileManager = IFileManager::Get();
@@ -203,7 +216,8 @@ bool FAsyncIODelete::Delete(const FStringView& PathToDelete, EPathType ExpectedT
 	}
 	Setup();
 	// Prevent the user from trying to delete our temproot or anything inside it
-	if (FPaths::IsUnderDirectory(PathToDeleteStr, TempRoot) || FPaths::IsUnderDirectory(TempRoot, PathToDeleteStr))
+	if (!SharedTempRoot.IsEmpty() &&
+		(FPaths::IsUnderDirectory(PathToDeleteStr, SharedTempRoot) || FPaths::IsUnderDirectory(SharedTempRoot, PathToDeleteStr)))
 	{
 		return false;
 	}
@@ -286,119 +300,233 @@ bool FAsyncIODelete::SynchronousDelete(const TCHAR* InDeletePath, EPathType Path
 	return Result;
 }
 
-bool FAsyncIODelete::TryPurgeTempRootFamily(FString* OutNewTempRoot)
+FStringView FAsyncIODelete::GetLockSuffix()
+{
+	return TEXTVIEW(".lock");
+}
+
+constexpr float MaxWaitSecondsForLockDefault = 5.0f;
+constexpr float SleepSecondsForLock = 0.01f;
+namespace UE::AsyncIODelete::Private
+{
+	float MaxWaitSecondsForLock = MaxWaitSecondsForLockDefault;
+}
+
+void FAsyncIODelete::SetMaxWaitSecondsForLock(float MaxWaitTimeSeconds)
+{
+	if (MaxWaitTimeSeconds < 0)
+	{
+		MaxWaitTimeSeconds = MaxWaitSecondsForLockDefault;
+	}
+	UE::AsyncIODelete::Private::MaxWaitSecondsForLock = MaxWaitTimeSeconds;
+}
+
+bool FAsyncIODelete::TryPurgeOldAndCreateRoot(bool bCreateRoot, TArray<FDeleteRequest>& OutOrphanedRootsToDelete)
 {
+	check(!SharedTempRoot.IsEmpty());
+	check(bCreateRoot == TempRoot.IsEmpty()); // We should only call TryPurgeOldAndCreateRoot(true) during setup and (false) during Teardown
 	IFileManager& FileManager = IFileManager::Get();
-	FString RequestedLeaf = FPaths::GetPathLeaf(RequestedTempRoot);
-	FString ParentDir = FPaths::GetPath(RequestedTempRoot);
 
-	TArray<uint32, TInlineAllocator<2>> ExistingRoots;
-	FileManager.IterateDirectory(*ParentDir,
-		[&ExistingRoots, &RequestedLeaf](const TCHAR* FilenameOrDirectory, bool bIsDirectory)
+	FString ParentDir = FPaths::GetPath(SharedTempRoot);
+	if (!FileManager.DirectoryExists(*ParentDir))
+	{
+		if (bCreateRoot)
 		{
-			// Compare by PathLeaf instead of full path because absolute vs relative paths and junctions
-			// may change the name of the parent directory
-			FStringView ExistingPath(FilenameOrDirectory);
-			FStringView ExistingLeaf = FPathViews::GetPathLeaf(ExistingPath);
-			if (bIsDirectory && ExistingLeaf.StartsWith(FStringView(RequestedLeaf), ESearchCase::IgnoreCase))
+			if (!FileManager.MakeDirectory(*ParentDir, true /* Tree */) && !FileManager.DirectoryExists(*ParentDir))
+			{
+				UE_LOG(LogCook, Error, TEXT("Could not create AsyncIoDelete parent directory %s. LastError: %d. Falling back to synchronous delete."),
+					*ParentDir, FPlatformMisc::GetLastError());
+				return false;
+			}
+		}
+		else
+		{
+			// If parent directory doesn't exist when we are shutting down, do not create it
+			return true;
+		}
+	}
+
+	bool bDirectoryEmpty = true;
+	if (!TempRoot.IsEmpty())
+	{
+		// Our directory is empty because we waited on all of the async delete tasks, so delete it synchronously
+		if (!FileManager.DeleteDirectory(*TempRoot, false /* bRequireExists */, true /* Tree */) &&
+			FileManager.DirectoryExists(*TempRoot))
+		{
+			UE_LOG(LogCook, Display, TEXT("AsyncIoDelete could not clean up its root %s. LastError: %d."),
+				*TempRoot, FPlatformMisc::GetLastError());
+			bDirectoryEmpty = false;
+		}
+	}
+
+	// Temporarily lock the SharedTempRoot while we are querying the directory and need to be inside a machine-wide critical section
+	FStringView LockSuffix = GetLockSuffix();
+	FString SharedTempRootLockFileName = SharedTempRoot + LockSuffix;
+	TUniquePtr<FArchive> SharedTempRootLockFile = nullptr;
+	double StartTimeSeconds = FPlatformTime::Seconds();
+	do
+	{
+		SharedTempRootLockFile.Reset(FileManager.CreateFileWriter(*SharedTempRootLockFileName));
+		if (!SharedTempRootLockFile)
+		{
+			if (FPlatformTime::Seconds() - StartTimeSeconds > UE::AsyncIODelete::Private::MaxWaitSecondsForLock)
 			{
-				FStringView Suffix = ExistingLeaf.RightChop(RequestedLeaf.Len());
-				if (Suffix.Len() == 0)
+				if (bCreateRoot)
 				{
-					ExistingRoots.Add(0);
+					UE_LOG(LogCook, Error, TEXT("AsyncIoDelete could not create LockFile %s. Falling back to synchronous delete."),
+						*SharedTempRootLockFileName);
 				}
 				else
 				{
-					// Suffix is null-terminated because it came from FilenameOrDirectory
-					const TCHAR* SuffixPtr = Suffix.GetData();
-					uint32 IntSuffix = 0;
-					LexFromString(IntSuffix, SuffixPtr + 1);
-					if (IntSuffix > 0)
-					{
-						ExistingRoots.Add(IntSuffix);
-					}
+					UE_LOG(LogCook, Display, TEXT("AsyncIoDelete could not clean up its root %s, because another process has LockFile %s locked."),
+						*TempRoot, *SharedTempRootLockFileName, FPlatformMisc::GetLastError());
+					TempRootLockFile.Reset(); // Drop our lock that prevents the lockfile from being deleted, so the next AsyncIODelete can clean up after us.
 				}
+				return false;
 			}
-			return true;
-		});
-
-
-	auto GetRootPath = [this](uint32 Suffix)
+			FPlatformProcess::Sleep(SleepSecondsForLock);
+		}
+	} while (!SharedTempRootLockFile);
+	ON_SCOPE_EXIT
 	{
-		return Suffix == 0 ? RequestedTempRoot : FString::Printf(TEXT("%s_%d"), *RequestedTempRoot, Suffix);
+		SharedTempRootLockFile.Reset();
+		FileManager.Delete(*SharedTempRootLockFileName, false /* bRequireExists*/, true /* bEvenIfReadOnly */, true /* Quiet */);
 	};
 
-	uint32 LastError = 0;
-	TArray<uint32> Undeletables;
-	for (uint32 Suffix : ExistingRoots)
+	// Delete our TempRoot inside the machine-wide critical section.
+	// Doing it before would remove our right to expect that the SharedTempRoot would not be deleted out from under us.
+	FString TempRootLockFileName;
+	if (!TempRoot.IsEmpty())
 	{
-		FString ExistingRoot = GetRootPath(Suffix);
-
-		// Since we sometimes will be creating the directory again immediately, we need to take precautions against the delayed delete of directories that
-		// occurs on Windows platforms; creating a new file/directory in one that was just deleted can fail.  So we need to move-delete our TempRoot
-		// in addition to move-delete our clients' directories.  Since we don't have a TempRoot to move-delete into, we create a unique sibling directory name.
-		FString UniqueDirectory = FPaths::CreateTempFilename(*ParentDir, TEXT("DeleteTemp"), TEXT(""));
-
-		const bool bReplace = false;
-		const bool bEvenIfReadOnly = true;
-		const bool bMoveAttributes = false;
-		const bool bDoNotRetryOnError = true;
-		const TCHAR* DirectoryToDelete = *UniqueDirectory;
-		const bool bMoveSucceeded = FileManager.Move(DirectoryToDelete, *ExistingRoot, bReplace, bEvenIfReadOnly,
-			bMoveAttributes, bDoNotRetryOnError);
-		if (!bMoveSucceeded)
+		TempRootLockFile.Reset(); // Drop our lock that prevents the lockfile from being deleted
+		TempRootLockFileName = TempRoot + LockSuffix;
+		if (!FileManager.Delete(*TempRootLockFileName, false /* bRequireExists */, true /* bEvenIfReadOnly */, true /* Quiet */))
 		{
-			// Move failed; fallback to inplace delete
-			DirectoryToDelete = *ExistingRoot;
+			UE_LOG(LogCook, Display, TEXT("AsyncIoDelete could not clean up its lock file %s. LastError: %d."),
+				*TempRootLockFileName, FPlatformMisc::GetLastError());
+			bDirectoryEmpty = false;
 		}
+	}
 
-		const bool bRequireExists = false;
-		const bool bTree = true;
-		const bool bDeleteSucceeded = FileManager.DeleteDirectory(DirectoryToDelete, bRequireExists, bTree);
-		if (!bDeleteSucceeded)
+	auto GetCountFromFilename = [](const FString& BaseFileName)
+	{
+		int32 IntValue;
+		LexFromString(IntValue, *BaseFileName);
+		return IntValue; // Will be 0 if invalid
+	};
+
+	int32 FirstUnusedCount = 1;
+	bool bHasUnexpectedFiles = false;
+	TSet<FString> ExistingLockFileLeafs;
+	TSet<FString> ExistingDirLeafs;
+	FileManager.IterateDirectory(*SharedTempRoot,
+		[&ExistingLockFileLeafs, &ExistingDirLeafs, &bHasUnexpectedFiles, LockSuffix,
+		&FirstUnusedCount, &GetCountFromFilename](const TCHAR* FilenameOrDirectory, bool bIsDirectory)
 		{
-			LastError = FPlatformMisc::GetLastError();
-			if (bMoveSucceeded && !bDeleteSucceeded)
+			// Convert FilenameOrDirectory to a relative path because absolute vs relative paths and junctions
+			// may change the name of the parent directory
+			FStringView ExistingPath(FilenameOrDirectory);
+			FString ExistingLeaf = FString(FPathViews::GetPathLeaf(ExistingPath));
+
+			if (!bIsDirectory)
+			{
+				if (FStringView(ExistingLeaf).EndsWith(LockSuffix))
+				{
+					FirstUnusedCount = FMath::Max(FirstUnusedCount, 1+GetCountFromFilename(ExistingLeaf.LeftChop(LockSuffix.Len())));
+					ExistingLockFileLeafs.Add(ExistingLeaf);
+				}
+				else
+				{
+					UE_LOG(LogCook, Warning, TEXT("AsyncIoDelete found unexpected file '%s' in its SharedTempRoot. This file will block cleanup of the SharedTempRoot, please delete it manually."),
+						FilenameOrDirectory);
+					bHasUnexpectedFiles = true;
+				}
+			}
+			else
 			{
-				// Try to move the directory back so that we can try again to delete it next time.
-				FileManager.Move(*ExistingRoot, DirectoryToDelete, bReplace, bEvenIfReadOnly,
-					bMoveAttributes, bDoNotRetryOnError);
+				ExistingDirLeafs.Add(ExistingLeaf);
+				FirstUnusedCount = FMath::Max(FirstUnusedCount, 1+GetCountFromFilename(ExistingLeaf));
 			}
-			Undeletables.Add(Suffix);
+			return true;
+		});
+	for (const FString& ExistingDirLeaf : ExistingDirLeafs)
+	{
+		if (!ExistingLockFileLeafs.Contains(ExistingDirLeaf + LockSuffix))
+		{
+			UE_LOG(LogCook, Warning, TEXT("AsyncIoDelete found unexpected directory '%s' in its SharedTempRoot. This directory will block cleanup of the SharedTempRoot, please delete it manually."),
+				*FPaths::Combine(SharedTempRoot, ExistingDirLeaf));
+			bHasUnexpectedFiles = true;
 		}
 	}
 
-	if (OutNewTempRoot)
+	if (bCreateRoot)
 	{
-		OutNewTempRoot->Reset();
-		uint32 NewSuffix = 0;
-		Undeletables.Sort();
-		if (Undeletables.Num() > 0 && Undeletables[0] == 0)
+		FString TempRootDirName = FPaths::Combine(SharedTempRoot, LexToString(FirstUnusedCount));
+		TempRootLockFileName = TempRootDirName + LockSuffix;
+		if (!FileManager.MakeDirectory(*TempRootDirName, true /* Tree */))
+		{
+			UE_LOG(LogCook, Error, TEXT("AsyncIoDelete could not create its root %s. LastError: %d. Falling back to synchronous delete."),
+				*TempRootDirName, FPlatformMisc::GetLastError());
+			return false;
+		}
+		TempRootLockFile.Reset(FileManager.CreateFileWriter(*TempRootLockFileName));
+		if (!TempRootLockFile)
 		{
-			NewSuffix = 1;
-			for (uint32 Suffix : Undeletables)
+			UE_LOG(LogCook, Error, TEXT("AsyncIoDelete could not create LockFile %s. LastError: %d. Falling back to synchronous delete."),
+				*TempRootLockFileName, FPlatformMisc::GetLastError());
+			FileManager.DeleteDirectory(*TempRootDirName, false /* bRequireExists */, true /* Tree */);
+			return false;
+		}
+		TempRoot = TempRootDirName;
+
+		// We do cleanup of orphaned temp roots only on startup, so that we can move their orphaned directories
+		// into the new directory we are creating, for async delete
+		for (const FString& ExistingDirLeaf : ExistingDirLeafs)
+		{
+			FString SourceDir = FPaths::Combine(SharedTempRoot, ExistingDirLeaf);
+			FString SourceLockFile = SourceDir + LockSuffix;
+			FString DestLockFile = FPaths::Combine(TempRoot, FString::Printf(TEXT("%u"), DeleteCounter));
+
+			// If we can move-delete the lock file, the other process must have dropped its lock, so we can move-delete the directory.
+			if (FileManager.Move(*DestLockFile, *SourceLockFile, true /* bReplace */, true /* EvenIfReadOnly */,
+				false /* Attributes */, true /* bDoNotRetryOnError */))
 			{
-				if (Suffix == 0)
+				++DeleteCounter;
+				FString DestDir = FPaths::Combine(TempRoot, FString::Printf(TEXT("%u"), DeleteCounter));
+				if (FileManager.Move(*DestDir, *SourceDir, true /* bReplace */, true /* EvenIfReadOnly */,
+					false /* Attributes */, true /* bDoNotRetryOnError */))
 				{
-					continue;
+					++DeleteCounter;
+					OutOrphanedRootsToDelete.Add({ DestLockFile, EPathType::File });
+					OutOrphanedRootsToDelete.Add({ DestDir, EPathType::Directory });
 				}
-				if (Suffix != NewSuffix)
+				else
 				{
-					break;
+					// If we move-deleted the lock file but failed to move-delete the directory, put the lock file back so a future
+					// AsyncIODelete can try to delete it
+					FileManager.Move(*SourceLockFile, *DestLockFile, true /* bReplace */, true /* EvenIfReadOnly */,
+						false /* Attributes */, true /* bDoNotRetryOnError */);
+
+					// Add a dummy entry in OutOrphanedRootsToDelete so that the DeleteCounter == OutOrphanedRootsToDelete.Num().
+					OutOrphanedRootsToDelete.Add({ FString(), EPathType::File});
 				}
-				++NewSuffix;
 			}
 		}
-
-		const uint32 MaxHangingTempRoots = 20;
-		if (NewSuffix > MaxHangingTempRoots)
+	}
+	else
+	{
+		// When deleting our TempRoot, also delete the SharedTempRoot if we're the last thing in the directory
+		// Remove our TempRoot from the ExistingDirs if it showed up in the iteration even though we deleted it due to FileManager delete lag
+		ExistingDirLeafs.Remove(FPaths::GetBaseFilename(TempRoot));
+		ExistingLockFileLeafs.Remove(FPaths::GetBaseFilename(TempRootLockFileName));
+		if (!ExistingDirLeafs.IsEmpty() || !ExistingLockFileLeafs.IsEmpty() || bHasUnexpectedFiles)
 		{
-			UE_LOG(LogCook, Error, TEXT("Could not clear %d old asyncdelete root directories '%s'_*.  LastError: %i.")
-				TEXT("\n\tFalling back to synchronous delete. Delete the directories manually to silence this message."),
-				Undeletables.Num(), *RequestedTempRoot, LastError);
-			return false;
+			bDirectoryEmpty = false;
+		}
+		if (bDirectoryEmpty)
+		{
+			FileManager.DeleteDirectory(*SharedTempRoot, false /* bRequireExists */, true /* Tree */);
 		}
-		*OutNewTempRoot = GetRootPath(NewSuffix);
 	}
 	return true;
 }
@@ -409,6 +537,10 @@ void FAsyncIODelete::AddTempRoot(const FStringView& InTempRoot)
 	FString TempRoot(InTempRoot);
 	for (FString& Existing : AllTempRoots)
 	{
+		if (FPaths::IsSamePath(Existing, TempRoot))
+		{
+			continue;
+		}
 		checkf(!FPaths::IsUnderDirectory(Existing, TempRoot), TEXT("New FAsyncIODelete has TempRoot \"%s\" that is a subdirectory of existing TempRoot \"%s\"."), *TempRoot, *Existing);
 		checkf(!FPaths::IsUnderDirectory(TempRoot, Existing), TEXT("New FAsyncIODelete has TempRoot \"%s\" that is a parent directory of existing TempRoot \"%s\"."), *TempRoot, *Existing);
 	}
@@ -419,4 +551,4 @@ void FAsyncIODelete::RemoveTempRoot(const FStringView& InTempRoot)
 {
 	AllTempRoots.Remove(FString(InTempRoot));
 }
-#endif
+#endif
\ No newline at end of file
diff --git a/Engine/Source/Editor/UnrealEd/Private/Cooker/AsyncIODelete.h b/Engine/Source/Editor/UnrealEd/Private/Cooker/AsyncIODelete.h
--- a/Engine/Source/Editor/UnrealEd/Private/Cooker/AsyncIODelete.h
+++ b/Engine/Source/Editor/UnrealEd/Private/Cooker/AsyncIODelete.h
@@ -34,9 +34,9 @@ public:
 	/**
 	 * Quick constructor to create flags; does not take any IO action until Setup or a Delete function is used
 	 *
-	 * @param InOwnedTempRoot - The FAsyncIODelete takes ownership of this directory, and deletes it when done.
-	 *							It is a temporary directory used to hold the deleted directories from elsewhere.  Must not be shared with (or a parent or child of) any other FAsyncDelete's InTempRoot.
-	 *						    The FAsyncIODelete will not be able to delete directories until a valid TempRoot has been set, either in the constructor or in SetTempRoot
+	 * @param InOwnedTempRoot - The FAsyncIODelete system takes ownership of this directory, and deletes it when done.
+	 * See SetTempRoot. The FAsyncIODelete will not be able to delete directories until a valid TempRoot has been set,
+	 * either in the constructor or in SetTempRoot
 	 */
 	explicit FAsyncIODelete(const FStringView& InOwnedTempRoot=FStringView());
 
@@ -44,12 +44,21 @@ public:
 	~FAsyncIODelete();
 
 	/**
-	 * Set the TempRoot directory used to hold the deleted directories from elsewhere.  The FAsyncIODelete takes ownership of this directory, and deletes it when destructed or when it is changed.
+	 * Set a TempRoot directory to hold deleted directories from elsewhere, which is potentially shared with other
+	 * processes running other FAsyncIODelete. The AsyncIoDelete processes take ownership of this directory and will
+	 * delete it; do not pass in a directory shared with other systems, create a subdirectory for it.
+	 * Must not be a parent or child of any other FAsyncDelete's TempRoot, but is allowed to be the same.
+	 * Each new FAsyncIODelete process will negotiate with other proccesses using lock files, to delete the directory
+	 * when all processes are done.
 	 * If the FAsyncIODelete has already been used to delete files, this Set will block until all deletes are finished.
-	 * If the given TempRoot is unavailable, falls back to _digit suffixes.
 	 */
-	void SetTempRoot(const FStringView& InOwnedTempRoot);
-	FStringView GetTempRoot() const { return TempRoot.IsEmpty() ? RequestedTempRoot : TempRoot; }
+	void SetTempRoot(FStringView InSharedRoot);
+	FStringView GetTempRoot() const { return SharedTempRoot; }
+	/**
+	 * Returns the path this AsyncIODelete is using for its uniquely owned deletion root (a subdirectory of TempRoot).
+	 * Will return empty before Setup (or on-demand Setup) or after Teardown.
+	 */
+	FStringView GetDeletionRoot() const { return TempRoot; };
 
 	/**
 	 * Set whether new background deletes are paused.  If paused, paths will be moved immediately but will not be deleted from the temporary location until unpaused (or at the FAsyncIODelete's destruction).
@@ -93,10 +102,11 @@ public:
 	/** Synchronously wait for all deletes to complete.  */
 	bool WaitForAllTasks(float TimeLimitSeconds = 0.0f);
 
-	static bool AsyncEnabled()
-	{
-		return FPlatformMisc::SupportsMultithreadedFileHandles();
-	}
+	static bool AsyncEnabled();
+
+	/** Internal accessors for unittests */
+	static FStringView GetLockSuffix();
+	static void SetMaxWaitSecondsForLock(float WaitSeconds = -1.f);
 
 private:
 	/** Update the ActiveTaskCount and events for a task completing */
@@ -107,6 +117,11 @@ private:
 		File,
 		Directory
 	};
+	struct FDeleteRequest
+	{
+		FString Path;
+		EPathType PathType;
+	};
 	/** Asynchronously delete a file or directory.  We handle both in the same function, but we want the interface to be explicit. */
 	bool Delete(const FStringView& PathToDelete, EPathType ExpectedType);
 
@@ -116,10 +131,11 @@ private:
 	/** Delete the given path synchronously; called from a task or in error fallback cases from the public thread */
 	bool SynchronousDelete(const TCHAR* InDeletePath, EPathType PathType);
 
-	bool TryPurgeTempRootFamily(FString* OutNextTempRoot);
+	bool TryPurgeOldAndCreateRoot(bool bCreateRoot, TArray<FDeleteRequest>& OutHangingRootsToDelete);
 
-	FString	RequestedTempRoot; 
+	FString SharedTempRoot;
 	FString TempRoot;
+	TUniquePtr<FArchive> TempRootLockFile;
 	TArray<FString> PausedDeletes;
 	FCriticalSection CriticalSection; // We use a CriticalSection instead of a TAtomic ActiveTaskCount so that we can atomically { trigger TasksComplete if ActiveTaskCount == 0 }
 	FEvent* TasksComplete = nullptr;
diff --git a/Engine/Source/Editor/UnrealEd/Private/Cooker/AsyncIODeleteTest.cpp b/Engine/Source/Editor/UnrealEd/Private/Cooker/AsyncIODeleteTest.cpp
--- a/Engine/Source/Editor/UnrealEd/Private/Cooker/AsyncIODeleteTest.cpp
+++ b/Engine/Source/Editor/UnrealEd/Private/Cooker/AsyncIODeleteTest.cpp
@@ -24,10 +24,10 @@ bool FAsyncIODeleteTest::RunTest(const FString& Parameters)
 		FileManager.DeleteDirectory(*TestRoot, bRequireExists, bApplyToTreeTrue);
 	};
 
-	FString TempRoot = FPaths::Combine(TestRoot, TEXT("TempRoot"));
-	FString TempRoot2 = FPaths::Combine(TestRoot, TEXT("TempRoot2"));
-	FString TempRoot3 = FPaths::Combine(TestRoot, TEXT("TempRoot3"));
-	FString TempRoot4 = FPaths::Combine(TestRoot, TEXT("TempRoot4"));
+	FString SharedTempRoot = FPaths::Combine(TestRoot, TEXT("TempRoot"));
+	FString SharedTempRoot2 = FPaths::Combine(TestRoot, TEXT("TempRoot2"));
+	FString SharedTempRoot3 = FPaths::Combine(TestRoot, TEXT("TempRoot3"));
+	FString SharedTempRoot4 = FPaths::Combine(TestRoot, TEXT("TempRoot4"));
 	FString TestFile1 = FPaths::Combine(TestRoot, TEXT("TestFile1"));
 	FString TestDir1 = FPaths::Combine(TestRoot, TEXT("TestDir1"));
 	int32 NumFiles = 0;
@@ -74,7 +74,7 @@ bool FAsyncIODeleteTest::RunTest(const FString& Parameters)
 	};
 	{
 		StartSection(TEXT("Constructing first FAsyncIODelete"));
-		FAsyncIODelete AsyncIODelete(TempRoot);
+		FAsyncIODelete AsyncIODelete(SharedTempRoot);
 
 		StartSection(TEXT("Waiting for tasks to complete when none have been launched should succeed"));
 		WaitForAllTasksAndVerify(AsyncIODelete);
@@ -85,12 +85,13 @@ bool FAsyncIODeleteTest::RunTest(const FString& Parameters)
 		AsyncIODelete.DeleteDirectory(TestDir1);
 		TestRequestedPathsDeleted(TEXT("AsyncIODelete::Delete should have moved the deleted paths before returning."));
 
+		FString DeletionRoot = FString(AsyncIODelete.GetDeletionRoot());
 		StartSection(TEXT("Deleting the temporary files/directories should be finished before WaitForAllTasks returns"));
 		WaitForAllTasksAndVerify(AsyncIODelete);
-		TestTempRootCountsEqual(TempRoot, 0, 0, TEXT("AsyncIODelete should have deleted the moved paths before WaitForAllTasks returned."));
+		TestTempRootCountsEqual(DeletionRoot, 0, 0, TEXT("AsyncIODelete should have deleted the moved paths before WaitForAllTasks returned."));
 
 		StartSection(TEXT("Two FAsyncIODelete constructed at once should be legal, as long as they have different TempRoots"));
-		FAsyncIODelete AsyncIODelete2(TempRoot2);
+		FAsyncIODelete AsyncIODelete2(SharedTempRoot2);
 
 		StartSection(TEXT("Use the pause feature to verify that the paths are indeed moved into the TempRoot"));
 		AsyncIODelete2.SetDeletesPaused(true);
@@ -99,36 +100,41 @@ bool FAsyncIODeleteTest::RunTest(const FString& Parameters)
 		AsyncIODelete2.DeleteDirectory(TestDir1);
 		TestRequestedPathsDeleted(TEXT("AsyncIODelete::Delete should have moved the deleted paths before returning even when paused."));
 		WaitForAllTasksAndVerify(AsyncIODelete2);
+		FString DeletionRoot2 = FString(AsyncIODelete2.GetDeletionRoot());
 		if (bAsyncEnabled)
 		{
-			TestTempRootCountsEqual(TempRoot2, 1, 1, TEXT("AsyncIODelete should not have deleted the moved paths because it is paused."));
+			TestTempRootCountsEqual(DeletionRoot2, 1, 1, TEXT("AsyncIODelete should not have deleted the moved paths because it is paused."));
 		}
 		AsyncIODelete2.SetDeletesPaused(false);
 		WaitForAllTasksAndVerify(AsyncIODelete2);
 		if (bAsyncEnabled)
 		{
-			TestTempRootCountsEqual(TempRoot2, 0, 0, TEXT("AsyncIODelete should have deleted the moved paths after unpausing."));
+			TestTempRootCountsEqual(DeletionRoot2, 0, 0, TEXT("AsyncIODelete should have deleted the moved paths after unpausing."));
 		}
 
 		StartSection(TEXT("Verify Teardown() deletes the TempRoot and Setup() creates it"));
 		AsyncIODelete2.Teardown();
 		if (bAsyncEnabled)
 		{
-			TestTrue(TEXT("AsyncIODelete::Teardown should have deleted its TempRoot."), !FileManager.DirectoryExists(*TempRoot2));
+			TestTrue(TEXT("AsyncIODelete::Teardown should have deleted its TempRoot."), !FileManager.DirectoryExists(*SharedTempRoot2));
 		}
 		AsyncIODelete2.Setup();
+		DeletionRoot2 = AsyncIODelete2.GetDeletionRoot();
 		if (bAsyncEnabled)
 		{
-			TestTrue(TEXT("AsyncIODelete::Setup should have created its TempRoot."), FileManager.DirectoryExists(*TempRoot2));
+			TestTrue(TEXT("AsyncIODelete::Setup should have created its TempRoot."),
+				!DeletionRoot2.IsEmpty() && FileManager.DirectoryExists(*SharedTempRoot2) && FileManager.DirectoryExists(*DeletionRoot2));
 		}
 
 		StartSection(TEXT("Manual setup works as long as you call SetTempRoot before Setup"));
 		FAsyncIODelete AsyncIODelete3;
-		AsyncIODelete3.SetTempRoot(TempRoot3);
+		AsyncIODelete3.SetTempRoot(SharedTempRoot3);
 		AsyncIODelete3.Setup();
+		FString DeletionRoot3 = FString(AsyncIODelete3.GetDeletionRoot());
 		if (bAsyncEnabled)
 		{
-			TestTrue(TEXT("Setup should have created the TempRoot."), FileManager.DirectoryExists(*TempRoot3));
+			TestTrue(TEXT("Setup should have created the TempRoot."),
+				!DeletionRoot3.IsEmpty() && FileManager.DirectoryExists(*SharedTempRoot3) && FileManager.DirectoryExists(*DeletionRoot3));
 		}
 
 		StartSection(TEXT("Check that even after Setup, waiting for tasks to complete when none have been launched should succeed"));
@@ -138,38 +144,131 @@ bool FAsyncIODeleteTest::RunTest(const FString& Parameters)
 		CreateTestPathsToDelete();
 		AsyncIODelete3.DeleteFile(TestFile1);
 		AsyncIODelete3.DeleteDirectory(TestDir1);
-		AsyncIODelete3.SetTempRoot(TempRoot4);
+		AsyncIODelete3.SetTempRoot(SharedTempRoot4);
 		if (bAsyncEnabled)
 		{
-			TestTrue(TEXT("SetTempRoot should have deleted the old TempRoot."), !FileManager.DirectoryExists(*TempRoot3));
+			TestTrue(TEXT("SetTempRoot should have deleted the old TempRoot."), !FileManager.DirectoryExists(*SharedTempRoot3));
 		}
 
 		CreateTestPathsToDelete();
 		AsyncIODelete3.DeleteFile(TestFile1);
 		AsyncIODelete3.DeleteDirectory(TestDir1);
 		TestRequestedPathsDeleted(TEXT("AsyncIODelete::Delete should have worked after changing the TempRoot."));
+		FString DeletionRoot4 = FString(AsyncIODelete3.GetDeletionRoot());
 		if (bAsyncEnabled)
 		{
-			TestTrue(TEXT("Delete should have created the new TempRoot after SetTempRoot."), FileManager.DirectoryExists(*TempRoot4));
+			TestTrue(TEXT("Delete should have created the new TempRoot after SetTempRoot."),
+				!DeletionRoot4.IsEmpty() && FileManager.DirectoryExists(*SharedTempRoot4) && FileManager.DirectoryExists(*DeletionRoot4));
 		}
 		WaitForAllTasksAndVerify(AsyncIODelete3);
 		if (bAsyncEnabled)
 		{
-			TestTempRootCountsEqual(TempRoot4, 0, 0, TEXT("AsyncIODelete::Delete should have created the tasks to delete moved files after changing the TempRoot."));
+			TestTempRootCountsEqual(DeletionRoot4, 0, 0, TEXT("AsyncIODelete::Delete should have created the tasks to delete moved files after changing the TempRoot."));
 		}
 
 		StartSection(TEXT("Attempting to delete a parent directory of the temproot, the temproot itself, or a child inside of it fails"));
-		FString SubDirInTempRoot4 = FPaths::Combine(TempRoot4, TEXT("SubDir"));
+		FString SubDirInTempRoot4 = FPaths::Combine(DeletionRoot4, TEXT("SubDir"));
 		FileManager.MakeDirectory(*SubDirInTempRoot4, bApplyToTreeTrue); // Note it's illegal to add files into TempRoot, but we're not currently checking for it and we're not colliding with the DeleteN paths AsyncIODelete uses, so this breaking of the rule will not cause problems
 		TestFalse(TEXT("AsyncIODelete should refuse to delete a parent of its TempRoot."), AsyncIODelete3.DeleteDirectory(TestRoot));
-		TestFalse(TEXT("AsyncIODelete should refuse to delete its TempRoot."), AsyncIODelete3.DeleteDirectory(TempRoot4));
+		TestFalse(TEXT("AsyncIODelete should refuse to delete its TempRoot."), AsyncIODelete3.DeleteDirectory(DeletionRoot4));
+		TestFalse(TEXT("AsyncIODelete should refuse to delete its SharedTempRoot."), AsyncIODelete3.DeleteDirectory(SharedTempRoot4));
 		TestFalse(TEXT("AsyncIODelete should refuse to delete a child of its TempRoot."), AsyncIODelete3.DeleteDirectory(SubDirInTempRoot4));
 	}
+	if (bAsyncEnabled)
+	{
+		TestTrue(TEXT("AsyncIODelete destructor should have deleted its TempRoot."),
+			!FileManager.DirectoryExists(*SharedTempRoot) && !FileManager.DirectoryExists(*SharedTempRoot2)
+			&& !FileManager.DirectoryExists(*SharedTempRoot3) && !FileManager.DirectoryExists(*SharedTempRoot4));
+	}
+
+
+	FString SharedTempRoot5 = FPaths::Combine(TestRoot, TEXT("TempRoot5"));
+	if (bAsyncEnabled)
+	{
+		StartSection(TEXT("Two AsyncIODeletes in the same SharedTempRoot do not delete each other's TempRoots."));
+		FAsyncIODelete Async1(SharedTempRoot5);
+		Async1.Setup();
+		FString TempRoot1 = FString(Async1.GetDeletionRoot());
+		TestTrue(TEXT("Async1 creates its root."), !TempRoot1.IsEmpty() && FileManager.DirectoryExists(*TempRoot1));
+		FString DeletionRoot2;
+		{
+			FAsyncIODelete Async2(SharedTempRoot5);
+			Async2.Setup();
+			DeletionRoot2 = FString(Async2.GetDeletionRoot());
+			TestTrue(TEXT("Async2 creates its root."), !DeletionRoot2.IsEmpty() && FileManager.DirectoryExists(*DeletionRoot2));
+			TestTrue(TEXT("Async1's root still exists after Async2 Setup."), !TempRoot1.IsEmpty() && FileManager.DirectoryExists(*TempRoot1));
+		}
+		TestTrue(TEXT("Async2 deletes its root."), DeletionRoot2.IsEmpty() || !FileManager.DirectoryExists(*DeletionRoot2));
+		TestTrue(TEXT("Async1's root still exists after Async2 Teardown."), !TempRoot1.IsEmpty() && FileManager.DirectoryExists(*TempRoot1));
 
+		CreateTestPathsToDelete();
+		Async1.DeleteFile(TestFile1);
+		FString DeletionRoot3;
+		{
+			FAsyncIODelete Async3(SharedTempRoot5);
+			Async3.DeleteDirectory(TestDir1);
+			DeletionRoot3 = FString(Async3.GetDeletionRoot());
+			TestTrue(TEXT("Async3 creates its root."), !DeletionRoot3.IsEmpty() && FileManager.DirectoryExists(*DeletionRoot3));
+			TestTrue(TEXT("Async1's root still exists after Async3 Setup."), !TempRoot1.IsEmpty() && FileManager.DirectoryExists(*TempRoot1));
+		}
+		TestTrue(TEXT("Async3 deletes its root."), DeletionRoot3.IsEmpty() || !FileManager.DirectoryExists(*DeletionRoot3));
+		TestTrue(TEXT("Async1's root still exists after Async3 Teardown."), !TempRoot1.IsEmpty() && FileManager.DirectoryExists(*TempRoot1));
+	}
 	if (bAsyncEnabled)
 	{
 		TestTrue(TEXT("AsyncIODelete destructor should have deleted its TempRoot."),
-			!FileManager.DirectoryExists(*TempRoot) && !FileManager.DirectoryExists(*TempRoot2) && !FileManager.DirectoryExists(*TempRoot3) && !FileManager.DirectoryExists(*TempRoot4));
+			!FileManager.DirectoryExists(*SharedTempRoot5));
+	}
+
+	FString SharedTempRoot6 = FPaths::Combine(TestRoot, TEXT("TempRoot6"));
+	if (bAsyncEnabled)
+	{
+		StartSection(TEXT("AsyncIODelete cleans up a hanging directory from a crashed process."));
+		FString HangingTempRoot = FPaths::Combine(SharedTempRoot6, TEXT("5"));
+		FString HangingTempRootLockFile = HangingTempRoot + FAsyncIODelete::GetLockSuffix();
+		FileManager.MakeDirectory(*HangingTempRoot, true /* Tree */);
+		TUniquePtr<FArchive> LockFile(FileManager.CreateFileWriter(*HangingTempRootLockFile));
+		LockFile.Reset();
+		TestTrue(TEXT("Hanging directory was successfully created."), FileManager.DirectoryExists(*HangingTempRoot) && FileManager.FileExists(*HangingTempRootLockFile));
+
+		{
+			FAsyncIODelete Async(SharedTempRoot6);
+			Async.Setup();
+			FString DeletionRoot = FString(Async.GetDeletionRoot());
+			TestTrue(TEXT("Async creates its root."), !DeletionRoot.IsEmpty() && FileManager.DirectoryExists(*DeletionRoot));
+		}
+		TestTrue(TEXT("Async deletes the hanging directory."), !FileManager.DirectoryExists(*SharedTempRoot6));
+	}
+	if (bAsyncEnabled)
+	{
+		TestTrue(TEXT("AsyncIODelete destructor should have deleted its TempRoot."),
+			!FileManager.DirectoryExists(*SharedTempRoot6));
+	}
+
+	FString SharedTempRoot7 = FPaths::Combine(TestRoot, TEXT("TempRoot7"));
+	if (bAsyncEnabled)
+	{
+		StartSection(TEXT("A lock on the shared root prevents AsyncIODeletes from working with it."));
+		FAsyncIODelete::SetMaxWaitSecondsForLock(0.02f);
+		FileManager.MakeDirectory(*SharedTempRoot7, true /* Tree */);
+		FString SharedTempRoot7LockFileName = SharedTempRoot7 + FAsyncIODelete::GetLockSuffix();
+		TUniquePtr<FArchive> LockFile(FileManager.CreateFileWriter(*SharedTempRoot7LockFileName));
+		{
+			FAsyncIODelete AsyncFail(SharedTempRoot7);
+			AddExpectedError(SharedTempRoot7, EAutomationExpectedErrorFlags::Contains);
+			AsyncFail.Setup();
+			TestTrue(TEXT("AsyncIODelete on a locked root is unable to create its deletion root."), AsyncFail.GetDeletionRoot().IsEmpty());
+			CreateTestPathsToDelete();
+			AsyncFail.DeleteFile(TestFile1);
+			AsyncFail.DeleteDirectory(TestDir1);
+			TestRequestedPathsDeleted(TEXT("AsyncIODelete that failed due to a locked root should delete its requested paths synchronously."));
+			TestTrue(TEXT("AsyncIODelete on a locked root is unable to create its deletion root."), AsyncFail.GetDeletionRoot().IsEmpty());
+			TestTempRootCountsEqual(SharedTempRoot7, 0, 0, TEXT("AsyncIODelete that failed due to a locked root should not have created a TempRoot in the locked root."));
+		}
+		FAsyncIODelete::SetMaxWaitSecondsForLock(-1.f);
+		LockFile.Reset();
+		FileManager.DeleteDirectory(*SharedTempRoot7, false /* RequireEixsts */, true /* Tree */);
+		FileManager.Delete(*SharedTempRoot7LockFileName, false /* RequireEixsts */, true/* EvenReadOnly */, true /* Quiet */);
 	}
 
 	return true;
