EditorDomain: Add placeholder records in the EditorDomain for packages that we tried to save but could not store in the EditorDomain. Add SaveStorageResult to describe why the save failed.
#rnx

diff --git a/Engine/Source/Editor/UnrealEd/Private/EditorDomain/EditorDomain.cpp b/Engine/Source/Editor/UnrealEd/Private/EditorDomain/EditorDomain.cpp
--- a/Engine/Source/Editor/UnrealEd/Private/EditorDomain/EditorDomain.cpp
+++ b/Engine/Source/Editor/UnrealEd/Private/EditorDomain/EditorDomain.cpp
@@ -263,15 +263,35 @@ TRefCountPtr<FEditorDomain::FPackageSource> FEditorDomain::FindPackageSource(con
 	return TRefCountPtr<FPackageSource>();
 }
 
-void FEditorDomain::MarkNeedsLoadFromWorkspace(const FPackagePath& PackagePath, TRefCountPtr<FPackageSource>& PackageSource)
+void FEditorDomain::MarkLoadedFromWorkspaceDomain(const FPackagePath& PackagePath, TRefCountPtr<FPackageSource>& PackageSource,
+	bool bHasRecordInEditorDomain)
 {
+	if (PackageSource->Source == FEditorDomain::EPackageSource::Workspace)
+	{
+		return;
+	}
+
 	PackageSource->Source = FEditorDomain::EPackageSource::Workspace;
+	PackageSource->bHasRecordInEditorDomain |= bHasRecordInEditorDomain;
 	if (bExternalSave)
 	{
-		SaveClient->RequestSave(PackagePath);
+		if (PackageSource->NeedsEditorDomainSave(*this))
+		{
+			SaveClient->RequestSave(PackagePath);
+		}
 	}
-	// Otherwise, we will note the need for save in OnEndLoadPackage
+	// Otherwise, we will check whether it needs to save in OnEndLoadPackage
+}
 
+void FEditorDomain::MarkLoadedFromEditorDomain(const FPackagePath& PackagePath, TRefCountPtr<FPackageSource>& PackageSource)
+{
+	if (PackageSource->Source == FEditorDomain::EPackageSource::Editor)
+	{
+		return;
+	}
+
+	PackageSource->Source = FEditorDomain::EPackageSource::Editor;
+	PackageSource->bHasRecordInEditorDomain = true;
 }
 
 int64 FEditorDomain::FileSize(const FPackagePath& PackagePath, EPackageSegment PackageSegment,
@@ -306,32 +326,49 @@ int64 FEditorDomain::FileSize(const FPackagePath& PackagePath, EPackageSegment P
 			[&FileSize, &PackageSource, &PackagePath, PackageSegment, Locks=this->Locks, OutUpdatedPath]
 			(UE::DerivedData::FCacheGetResponse&& Response)
 		{
-			FScopeLock ScopeLock(&Locks->Lock);
-			if ((PackageSource->Source == FEditorDomain::EPackageSource::Undecided || PackageSource->Source == FEditorDomain::EPackageSource::Editor) &&
-				Response.Status == UE::DerivedData::EStatus::Ok)
+			bool bLoadFromWorkspace = false;
+			bool bHasRecordInEditorDomain = false;
+			if (Response.Status != UE::DerivedData::EStatus::Ok)
 			{
-				const FCbObject& MetaData = Response.Record.GetMeta();
-				FileSize = MetaData["FileSize"].AsInt64();
-				PackageSource->Source = EPackageSource::Editor;
+				if (PackageSource->Source == FEditorDomain::EPackageSource::Editor)
+				{
+					UE_LOG(LogEditorDomain, Error, TEXT("%s was previously loaded from the EditorDomain but now is unavailable. This may cause failures during serialization due to changed FileSize and Format."),
+						*PackagePath.GetDebugName());
+				}
+				bLoadFromWorkspace = true;
 			}
 			else
 			{
-				checkf(PackageSource->Source == EPackageSource::Undecided || PackageSource->Source == EPackageSource::Workspace,
-					TEXT("%s was previously loaded from the EditorDomain but now is unavailable."),
-					*PackagePath.GetDebugName());
-				if (Locks->Owner)
+				bHasRecordInEditorDomain = true;
+				const FCbObject& MetaData = Response.Record.GetMeta();
+				bool bStorageValid = MetaData["Valid"].AsBool(false);
+				if (!bStorageValid)
 				{
-					FEditorDomain& EditorDomain = *Locks->Owner;
-					EditorDomain.MarkNeedsLoadFromWorkspace(PackagePath, PackageSource);
-					FileSize = EditorDomain.Workspace->FileSize(PackagePath, PackageSegment, OutUpdatedPath);
+					bLoadFromWorkspace = true;
 				}
 				else
 				{
-					UE_LOG(LogEditorDomain, Warning, TEXT("%s size read after EditorDomain shutdown. Returning -1."),
-						*PackagePath.GetDebugName());
-					FileSize = -1;
+					FileSize = MetaData["FileSize"].AsInt64();
 				}
 			}
+
+			FScopeLock ScopeLock(&Locks->Lock);
+			FEditorDomain* EditorDomain = Locks->Owner;
+			if (!EditorDomain)
+			{
+				UE_LOG(LogEditorDomain, Warning, TEXT("%s size read after EditorDomain shutdown. Returning -1."),
+					*PackagePath.GetDebugName());
+				FileSize = -1;
+			}
+			else if (PackageSource->Source == FEditorDomain::EPackageSource::Workspace || bLoadFromWorkspace)
+			{
+				EditorDomain->MarkLoadedFromWorkspaceDomain(PackagePath, PackageSource, bHasRecordInEditorDomain);
+				FileSize = EditorDomain->Workspace->FileSize(PackagePath, PackageSegment, OutUpdatedPath);
+			}
+			else
+			{
+				EditorDomain->MarkLoadedFromEditorDomain(PackagePath, PackageSource);
+			}
 		};
 		// Fetch meta-data only
 		ECachePolicy SkipFlags = ECachePolicy::SkipData & ~ECachePolicy::SkipMeta;
@@ -603,7 +640,7 @@ void FEditorDomain::FilterKeepPackagesToSave(TArray<UPackage*>& InOutPackagesToS
 
 bool FEditorDomain::FPackageSource::NeedsEditorDomainSave(FEditorDomain& EditorDomain) const
 {
-	return !bHasSaved && Source == EPackageSource::Workspace &&
+	return !bHasSaved && !bHasRecordInEditorDomain &&
 		(!EditorDomain.bSkipSavesUntilCatalogLoaded || bLoadedAfterCatalogLoaded);
 }
 
diff --git a/Engine/Source/Editor/UnrealEd/Private/EditorDomain/EditorDomain.h b/Engine/Source/Editor/UnrealEd/Private/EditorDomain/EditorDomain.h
--- a/Engine/Source/Editor/UnrealEd/Private/EditorDomain/EditorDomain.h
+++ b/Engine/Source/Editor/UnrealEd/Private/EditorDomain/EditorDomain.h
@@ -192,7 +192,8 @@ private:
 	struct FPackageSource : public FThreadSafeRefCountedObject
 	{
 		FPackageSource()
-			:bHasSaved(false), bHasLoaded(false), bHasQueriedCatalog(false), bLoadedAfterCatalogLoaded(false)
+			: bHasSaved(false), bHasLoaded(false), bHasQueriedCatalog(false), bLoadedAfterCatalogLoaded(false)
+			, bHasRecordInEditorDomain(false)
 		{
 		}
 
@@ -207,6 +208,7 @@ private:
 		bool bHasLoaded : 1;
 		bool bHasQueriedCatalog : 1;
 		bool bLoadedAfterCatalogLoaded : 1;
+		bool bHasRecordInEditorDomain : 1;
 	};
 
 	/** Disallow copy constructors */
@@ -218,8 +220,10 @@ private:
 		UE::EditorDomain::FPackageDigest* OutErrorDigest=nullptr);
 	/** Return the PackageSource data in PackageSources, if it exists */
 	TRefCountPtr<FPackageSource> FindPackageSource(const FPackagePath& PackagePath);
-	/** Mark that we had to load the Package from the workspace domain, and schedule its save into the EditorDomain. */
-	void MarkNeedsLoadFromWorkspace(const FPackagePath& PackagePath, TRefCountPtr<FPackageSource>& PackageSource);
+	/** Mark that we had to load the Package from the workspace domain, and schedule its save into the EditorDomain if allowed. */
+	void MarkLoadedFromWorkspaceDomain(const FPackagePath& PackagePath, TRefCountPtr<FPackageSource>& PackageSource, bool bHasRecordInEditorDomain);
+	/** Mark that we were able to load from the EditorDomain */
+	void MarkLoadedFromEditorDomain(const FPackagePath& PackagePath, TRefCountPtr<FPackageSource>& PackageSource);
 	/** Callback for PostEngineInit, to handle saving of packages which we could not save before then. */
 	void OnPostEngineInit();
 	/** EndLoad callback to handle saving the EditorDomain version of the package. */
diff --git a/Engine/Source/Editor/UnrealEd/Private/EditorDomain/EditorDomainArchive.cpp b/Engine/Source/Editor/UnrealEd/Private/EditorDomain/EditorDomainArchive.cpp
--- a/Engine/Source/Editor/UnrealEd/Private/EditorDomain/EditorDomainArchive.cpp
+++ b/Engine/Source/Editor/UnrealEd/Private/EditorDomain/EditorDomainArchive.cpp
@@ -278,7 +278,17 @@ void FEditorDomainPackageSegments::OnRecordRequestComplete(UE::DerivedData::FCac
 
 	ESource NewAsyncSource = ESource::Uninitialized;
 	int32 InitialRequestIndex = -1;
+	bool bExistsInEditorDomain = false;
+	bool bStorageValid = false;
+
 	if (Response.Status == EStatus::Ok)
+	{
+		const FCbObject& MetaData = Response.Record.GetMeta();
+		bExistsInEditorDomain = true;
+		bStorageValid = MetaData["Valid"].AsBool(false);
+	}
+
+	if (bStorageValid)
 	{
 		const FCacheRecord& Record = Response.Record;
 		const uint64 FileSize = Record.GetMeta()["FileSize"].AsUInt64();
@@ -317,7 +327,7 @@ void FEditorDomainPackageSegments::OnRecordRequestComplete(UE::DerivedData::FCac
 					bEditorDomainAvailable = true;
 					if (PackageSource->Source == FEditorDomain::EPackageSource::Undecided || PackageSource->Source == FEditorDomain::EPackageSource::Editor)
 					{
-						PackageSource->Source = FEditorDomain::EPackageSource::Editor;
+						EditorDomainLocks->Owner->MarkLoadedFromEditorDomain(PackagePath, PackageSource);
 						bRegisterSuccessful = true;
 					}
 				}
@@ -353,13 +363,16 @@ void FEditorDomainPackageSegments::OnRecordRequestComplete(UE::DerivedData::FCac
 		FScopeLock DomainScopeLock(&EditorDomainLocks->Lock);
 		if (EditorDomainLocks->Owner)
 		{
-			checkf(PackageSource->Source == FEditorDomain::EPackageSource::Undecided || PackageSource->Source == FEditorDomain::EPackageSource::Workspace,
-				TEXT("%s was previously loaded from the EditorDomain but now is unavailable."), *PackagePath.GetDebugName());
+			if (PackageSource->Source == FEditorDomain::EPackageSource::Editor)
+			{
+				UE_LOG(LogEditorDomain, Error, TEXT("%s was previously loaded from the EditorDomain but now is unavailable. This may cause failures during serialization due to changed FileSize and Format."),
+					*PackagePath.GetDebugName());
+			}
 			FEditorDomain& EditorDomain(*EditorDomainLocks->Owner);
 			bool bSucceeded = TryCreateFallbackData(EditorDomain);
 			if (bSucceeded)
 			{
-				EditorDomain.MarkNeedsLoadFromWorkspace(PackagePath, PackageSource);
+				EditorDomain.MarkLoadedFromWorkspaceDomain(PackagePath, PackageSource, bExistsInEditorDomain);
 				NewAsyncSource = ESource::Fallback;
 			}
 			else
diff --git a/Engine/Source/Editor/UnrealEd/Private/EditorDomain/EditorDomainUtils.cpp b/Engine/Source/Editor/UnrealEd/Private/EditorDomain/EditorDomainUtils.cpp
--- a/Engine/Source/Editor/UnrealEd/Private/EditorDomain/EditorDomainUtils.cpp
+++ b/Engine/Source/Editor/UnrealEd/Private/EditorDomain/EditorDomainUtils.cpp
@@ -90,6 +90,46 @@ static ValueType MapFindRef(const TMap<KeyType, ValueType, SetAllocator, KeyFunc
 namespace UE::EditorDomain
 {
 
+struct FSaveStoreResultToText
+{
+	ESaveStorageResult Code;
+	FUtf8StringView Text;
+};
+
+const FUtf8StringView SaveStoreResultInvalidText = UTF8TEXTVIEW("InvalidCode");
+const FSaveStoreResultToText SaveStoreResultToText[]
+{
+	{ ESaveStorageResult::Valid,					UTF8TEXTVIEW("Valid") },
+	{ ESaveStorageResult::UnexpectedClass,			UTF8TEXTVIEW("UnexpectedClass") },
+	{ ESaveStorageResult::UnexpectedCustomVersion,	UTF8TEXTVIEW("UnexpectedCustomVersion") },
+	{ ESaveStorageResult::BulkDataTooLarge,			UTF8TEXTVIEW("BulkDataTooLarge") },
+	{ ESaveStorageResult::InvalidCode,				SaveStoreResultInvalidText },
+};
+
+FUtf8StringView LexToUtf8(ESaveStorageResult Result)
+{
+	for (const FSaveStoreResultToText& ResultToText : SaveStoreResultToText)
+	{
+		if (Result == ResultToText.Code)
+		{
+			return ResultToText.Text;
+		}
+	}
+	return SaveStoreResultInvalidText;
+}
+
+ESaveStorageResult SaveStorageResultFromString(FUtf8StringView Text)
+{
+	for (const FSaveStoreResultToText& ResultToText : SaveStoreResultToText)
+	{
+		if (Text == ResultToText.Text)
+		{
+			return ResultToText.Code;
+		}
+	}
+	return ESaveStorageResult::InvalidCode;
+}
+
 TArray<FGuid> GetCustomVersions(UClass& Class);
 TMap<FGuid, UObject*> FindCustomVersionCulprits(TConstArrayView<FGuid> UnknownGuids, UPackage* Package);
 void InitializeGlobalConstructClasses();
@@ -111,7 +151,7 @@ bool bGGlobalConstructClassesInitialized = false;
 int64 GMaxBulkDataSize = -1;
 
 // Change to a new guid when EditorDomain needs to be invalidated
-const TCHAR* EditorDomainVersion = TEXT("B0ED553FD35F4185B095EF52ED119C7F");
+const TCHAR* EditorDomainVersion = TEXT("A5D5DB1F5951446EA0E2F3E017314DFC");
 
 // Identifier of the CacheBuckets for EditorDomain tables
 const TCHAR* EditorDomainPackageBucketName = TEXT("EditorDomainPackage");
@@ -1699,6 +1739,8 @@ bool TrySavePackage(UPackage* Package)
 		}
 	}
 
+	ESaveStorageResult StorageResult = ESaveStorageResult::Valid;
+
 	TSet<FName> KnownImportedClasses(PackageDigest.ImportedClasses);
 	for (FName ImportedClass : SavedImportedClasses)
 	{
@@ -1712,7 +1754,7 @@ bool TrySavePackage(UPackage* Package)
 				TEXT("Find the class that added %s and add ")
 				TEXT("Editor.ini:[EditorDomain]:+PostLoadCanConstructClasses=<ConstructingClass>,%s"),
 				*Package->GetName(), *ImportedClass.ToString(), *ImportedClass.ToString());
-			return false;
+			StorageResult = ESaveStorageResult::UnexpectedClass;
 		}
 	}
 
@@ -1731,7 +1773,7 @@ bool TrySavePackage(UPackage* Package)
 			UnknownGuids.Add(CustomVersion.Key);
 		}
 	}
-	if (!UnknownGuids.IsEmpty())
+	if (StorageResult == ESaveStorageResult::Valid && !UnknownGuids.IsEmpty())
 	{
 		TMap<FGuid, UObject*> Culprits = FindCustomVersionCulprits(UnknownGuids, Package);
 		TStringBuilder<256> VersionsText;
@@ -1751,10 +1793,11 @@ bool TrySavePackage(UPackage* Package)
 			TEXT("Optimized loading and iterative cooking will be disabled for this package. ")
 			TEXT("Modify the classes/structs to call Ar.UsingCustomVersion(Guid) in Serialize or DeclareCustomVersions.%s"),
 			*Package->GetName(), VersionsText.ToString());
-		return false;
+		StorageResult = ESaveStorageResult::UnexpectedCustomVersion;
 	}
 
-	if (GMaxBulkDataSize >= 0 && PackageWriter->GetBulkDataSize() > static_cast<uint64>(GMaxBulkDataSize))
+	if (StorageResult == ESaveStorageResult::Valid && 
+		GMaxBulkDataSize >= 0 && PackageWriter->GetBulkDataSize() > static_cast<uint64>(GMaxBulkDataSize))
 	{
 		// TODO EditorDomain MeshDescription: We suppress the warning if the package includes a MeshDescription;
 		// MeshDescriptions do not yet detect the need for a bulkdata resave until save serialization.
@@ -1776,38 +1819,53 @@ bool TrySavePackage(UPackage* Package)
 		{
 			UE_LOG(LogEditorDomain, Verbose, TEXT("%s"), *Message);
 		}
-		return false;
+		StorageResult = ESaveStorageResult::BulkDataTooLarge;
 	}
 
 	ICache& Cache = GetCache();
 
+	bool bStorageResultValid = StorageResult == ESaveStorageResult::Valid;
 	TCbWriter<16> MetaData;
 	MetaData.BeginObject();
 	uint64 FileSize = PackageWriter->GetFileSize();
-	MetaData << "FileSize" << FileSize;
+	MetaData << "Valid" << bStorageResultValid;
+	if (bStorageResultValid)
+	{
+		MetaData << "FileSize" << FileSize;
+	}
+	else
+	{
+		MetaData << "StorageResult" << LexToUtf8(StorageResult);
+	}
 	MetaData.EndObject();
 
 	FCacheRecordBuilder RecordBuilder(GetEditorDomainPackageKey(PackageDigest.Hash));
-	for (const FEditorDomainPackageWriter::FAttachment& Attachment : PackageWriter->GetAttachments())
+	if (bStorageResultValid)
 	{
-		RecordBuilder.AddValue(Attachment.ValueId, Attachment.Buffer);
+		for (const FEditorDomainPackageWriter::FAttachment& Attachment : PackageWriter->GetAttachments())
+		{
+			RecordBuilder.AddValue(Attachment.ValueId, Attachment.Buffer);
+		}
 	}
 	RecordBuilder.SetMeta(MetaData.Save().AsObject());
 	FRequestOwner Owner(EPriority::Normal);
-	Cache.Put({{{Package->GetName()}, RecordBuilder.Build()}}, Owner);
+	Cache.Put({ {{Package->GetName()}, RecordBuilder.Build()} }, Owner);
 	Owner.KeepAlive();
 
-	// TODO_BuildDefinitionList: Calculate and store BuildDefinitionList on the PackageData, or collect it here from some other source.
-	TArray<UE::DerivedData::FBuildDefinition> BuildDefinitions;
-	FCbObject BuildDefinitionList = UE::TargetDomain::BuildDefinitionListToObject(BuildDefinitions);
-	FCbObject TargetDomainDependencies = UE::TargetDomain::CollectDependenciesObject(Package, nullptr, nullptr);
-	if (TargetDomainDependencies)
-	{
-		TArray<IPackageWriter::FCommitAttachmentInfo, TInlineAllocator<2>> Attachments;
-		Attachments.Add({ "Dependencies", TargetDomainDependencies });
-		// TODO: Reenable BuildDefinitionList once FCbPackage support for empty FCbObjects is in
-		//Attachments.Add({ "BuildDefinitionList", BuildDefinitionList });
-		UE::TargetDomain::CommitEditorDomainCookAttachments(Package->GetFName(), Attachments);
+	if (bStorageResultValid)
+	{
+		// TODO_BuildDefinitionList: Calculate and store BuildDefinitionList on the PackageData, or collect it here from some other source.
+		TArray<UE::DerivedData::FBuildDefinition> BuildDefinitions;
+		FCbObject BuildDefinitionList = UE::TargetDomain::BuildDefinitionListToObject(BuildDefinitions);
+		FCbObject TargetDomainDependencies = UE::TargetDomain::CollectDependenciesObject(Package, nullptr, nullptr);
+		if (TargetDomainDependencies)
+		{
+			TArray<IPackageWriter::FCommitAttachmentInfo, TInlineAllocator<2>> Attachments;
+			Attachments.Add({ "Dependencies", TargetDomainDependencies });
+			// TODO: Reenable BuildDefinitionList once FCbPackage support for empty FCbObjects is in
+			//Attachments.Add({ "BuildDefinitionList", BuildDefinitionList });
+			UE::TargetDomain::CommitEditorDomainCookAttachments(Package->GetFName(), Attachments);
+		}
 	}
 	return true;
 }
diff --git a/Engine/Source/Editor/UnrealEd/Private/EditorDomain/EditorDomainUtils.h b/Engine/Source/Editor/UnrealEd/Private/EditorDomain/EditorDomainUtils.h
--- a/Engine/Source/Editor/UnrealEd/Private/EditorDomain/EditorDomainUtils.h
+++ b/Engine/Source/Editor/UnrealEd/Private/EditorDomain/EditorDomainUtils.h
@@ -59,6 +59,23 @@ struct FClassDigestMap
 	FRWLock Lock;
 };
 
+/**
+ * A result code indicating whether a EditorDomain package could be saved into the editor domain, and if not why not
+ * EditorDomain records in DDC have a Valid flag in their metadata. If the Valid flag is invalid, the record does not
+ * store a EditorDomain package and it instead has a StorageResult in the metadata explaining why it failed.
+ */
+enum class ESaveStorageResult
+{
+	Valid,
+	UnexpectedClass,
+	UnexpectedCustomVersion,
+	BulkDataTooLarge,
+	InvalidCode,
+	// When adding values, add them to SaveStoreResultToText as well
+};
+FUtf8StringView LexToUtf8(UE::EditorDomain::ESaveStorageResult Result);
+ESaveStorageResult SaveStorageResultFromString(FUtf8StringView Text);
+
 /**
  * Calculate the PackageDigest for the given packagePath.
  * Reads information from the AssetRegistry to compute the digest.
@@ -95,4 +112,4 @@ FClassDigestMap& GetClassDigests();
 /** Initializes some global config-driven values used by the EditorDomain and TargetDomain. */
 void UtilsInitialize();
 
-}
+}
\ No newline at end of file
