CookOnTheFlyServer: Remove DefaultAsyncIODeletePlatformName and use just Saved\Cooked\_Del for the SharedTempRoot. This is possible now that AsyncIODelete supports simultaneous access to a SharedTempRoot from multiple processes.
#rnx

diff --git a/Engine/Source/Editor/UnrealEd/Classes/CookOnTheSide/CookOnTheFlyServer.h b/Engine/Source/Editor/UnrealEd/Classes/CookOnTheSide/CookOnTheFlyServer.h
--- a/Engine/Source/Editor/UnrealEd/Classes/CookOnTheSide/CookOnTheFlyServer.h
+++ b/Engine/Source/Editor/UnrealEd/Classes/CookOnTheSide/CookOnTheFlyServer.h
@@ -234,7 +234,6 @@ private:
 	ECookInitializationFlags CookFlags = ECookInitializationFlags::None;
 	TUniquePtr<class FSandboxPlatformFile> SandboxFile;
 	TUniquePtr<FAsyncIODelete> AsyncIODelete; // Helper for deleting the old cook directory asynchronously
-	FString DefaultAsyncIODeletePlatformName; // Default platform name to be used by GetAsyncIODelete()
 	bool bIsSavingPackage = false; // used to stop recursive mark package dirty functions
 	/** Set to true during CookOnTheFly if a plugin is calling RequestPackage and we should therefore not make assumptions about when platforms are done cooking */
 	bool bCookOnTheFlyExternalRequests = false;
@@ -983,9 +982,6 @@ private:
 	*/
 	FString GetSandboxDirectory( const FString& PlatformName ) const;
 
-	/** If not already done, set the default platform name to be used by GetAsyncIODelete(). */
-	void TrySetDefaultAsyncIODeletePlatform(const FString& PlatformName);
-
 	/* Create the delete-old-cooked-directory helper.*/
 	FAsyncIODelete& GetAsyncIODelete();
 
diff --git a/Engine/Source/Editor/UnrealEd/Private/CookOnTheFlyServer.cpp b/Engine/Source/Editor/UnrealEd/Private/CookOnTheFlyServer.cpp
--- a/Engine/Source/Editor/UnrealEd/Private/CookOnTheFlyServer.cpp
+++ b/Engine/Source/Editor/UnrealEd/Private/CookOnTheFlyServer.cpp
@@ -564,7 +564,6 @@ const ITargetPlatform* UCookOnTheFlyServer::AddCookOnTheFlyPlatform(const FName&
 		// Platform has already been added by this function or by StartCookByTheBook
 		return PlatformData->TargetPlatform;
 	}
-	TrySetDefaultAsyncIODeletePlatform(PlatformData->TargetPlatform->PlatformName());
 
 	if (IsInGameThread())
 	{
@@ -5699,14 +5698,6 @@ bool UCookOnTheFlyServer::SaveCurrentIniSettings(const ITargetPlatform* TargetPl
 
 }
 
-void UCookOnTheFlyServer::TrySetDefaultAsyncIODeletePlatform(const FString& PlatformName)
-{
-	if (DefaultAsyncIODeletePlatformName.IsEmpty())
-	{
-		DefaultAsyncIODeletePlatformName = PlatformName;
-	}
-}
-
 FAsyncIODelete& UCookOnTheFlyServer::GetAsyncIODelete()
 {
 	if (AsyncIODelete)
@@ -5714,16 +5705,9 @@ FAsyncIODelete& UCookOnTheFlyServer::GetAsyncIODelete()
 		return *AsyncIODelete;
 	}
 
-	// The PlatformName is used construct a directory that we can be sure no other process is using (because a sandbox can only be cooked by one process at a time)
-	// The TempRoot we will delete into is a sibling of the the Platform-specific sandbox directory, with name [PlatformDir]AsyncDelete
-	// Note that two UnrealEd-Cmd processes cooking to the same sandbox at the same time will therefore cause an error since FAsyncIODelete doesn't handle multiple processes sharing TempRoots.
-	// That simultaneous-cook behavior is also not supported in other assumptions throughout the cooker.
-	check(!DefaultAsyncIODeletePlatformName.IsEmpty());
-	FString DeleteDirectory = GetSandboxDirectory(DefaultAsyncIODeletePlatformName);
-	FPaths::NormalizeDirectoryName(DeleteDirectory);
-	DeleteDirectory += TEXT("_Del");
-
-	AsyncIODelete = MakeUnique<FAsyncIODelete>(DeleteDirectory);
+	FString SharedDeleteRoot = GetSandboxDirectory(TEXT("_Del"));
+	FPaths::NormalizeDirectoryName(SharedDeleteRoot);
+	AsyncIODelete = MakeUnique<FAsyncIODelete>(SharedDeleteRoot);
 	return *AsyncIODelete;
 }
 
@@ -7842,11 +7826,6 @@ void UCookOnTheFlyServer::StartCookByTheBook( const FCookByTheBookStartupOptions
 	bPackageFilterDirty = true;
 	check(PlatformManager->GetSessionPlatforms().Num() == TargetPlatforms.Num());
 
-	if (ensure(TargetPlatforms.Num() > 0))
-	{
-		TrySetDefaultAsyncIODeletePlatform(TargetPlatforms[0]->PlatformName());
-	}
-
 	// We want to set bRunning = true as early as possible, but it implies that session platforms have been selected so this is the earliest point we can set it
 	CookByTheBookOptions->bRunning = true;
 
