ZenCOTF: Fix cook server deadlock on connect with multiple game clients by not holding the FIoStorecookOnTheFlyRequestManager::ContextsCriticalSection lock when doing CookOnTheFlyServer.AddPlatform(PlatformName).
Make FCookOnTheFlyServerInterface::GetPackageWriter thread safe using the same ExternalRequests->AddCallback mechanism as AddPlatform.
Move UCookOnTheFlyServer::AddCookOnTheFlyPlatform logic to FCookOnTheFlyServerInterface::AddPlatformInternal.
Reduce stats spam in cook on the fly mode.

#rnx

diff --git a/Engine/Source/Editor/UnrealEd/Classes/CookOnTheSide/CookOnTheFlyServer.h b/Engine/Source/Editor/UnrealEd/Classes/CookOnTheSide/CookOnTheFlyServer.h
--- a/Engine/Source/Editor/UnrealEd/Classes/CookOnTheSide/CookOnTheFlyServer.h
+++ b/Engine/Source/Editor/UnrealEd/Classes/CookOnTheSide/CookOnTheFlyServer.h
@@ -301,9 +301,7 @@ private:
 	*/
 	void OnTargetPlatformChangedSupportedFormats(const ITargetPlatform* TargetPlatform);
 
-	/* Version of AddCookOnTheFlyPlatform that takes the Platform name instead of an ITargetPlatform*.  Returns the Platform if found */
-	const ITargetPlatform* AddCookOnTheFlyPlatform(const FName& PlatformName);
-	/* Internal helper for AddCookOnTheFlyPlatform.  Initializing Platforms must be done on the tickloop thread; Platform data is read only on other threads */
+	/* Initializing Platforms must be done on the tickloop thread; Platform data is read only on other threads */
 	void AddCookOnTheFlyPlatformFromGameThread(ITargetPlatform* TargetPlatform);
 
 	/* Callback to recalculate all ITargetPlatform* pointers when they change due to modules reloading */
diff --git a/Engine/Source/Editor/UnrealEd/Private/CookOnTheFlyServer.cpp b/Engine/Source/Editor/UnrealEd/Private/CookOnTheFlyServer.cpp
--- a/Engine/Source/Editor/UnrealEd/Private/CookOnTheFlyServer.cpp
+++ b/Engine/Source/Editor/UnrealEd/Private/CookOnTheFlyServer.cpp
@@ -286,20 +286,17 @@ public:
 
 	virtual const ITargetPlatform* AddPlatform(const FName& PlatformName) override
 	{
-		const ITargetPlatform* Result;
+		UE::Cook::FPlatformManager::FReadScopeLock PlatformScopeLock(Cooker.PlatformManager->ReadLockPlatforms());
+		const ITargetPlatform* TargetPlatform = AddPlatformInternal(PlatformName);
+		if (!TargetPlatform)
 		{
-			UE::Cook::FPlatformManager::FReadScopeLock PlatformScopeLock(Cooker.PlatformManager->ReadLockPlatforms());
-			Result = Cooker.AddCookOnTheFlyPlatform(PlatformName);
-			if (!Result)
-			{
-				UE_LOG(LogCook, Warning, TEXT("Trying to add invalid platform '%s' on the fly"), *PlatformName.ToString());
-				return nullptr;
-			}
-
-			Cooker.PlatformManager->AddRefCookOnTheFlyPlatform(PlatformName, Cooker);
+			UE_LOG(LogCook, Warning, TEXT("Trying to add invalid platform '%s' on the fly"), *PlatformName.ToString());
+			return nullptr;
 		}
 
-		return Result;
+		Cooker.PlatformManager->AddRefCookOnTheFlyPlatform(PlatformName, Cooker);
+
+		return TargetPlatform;
 	}
 
 	virtual void RemovePlatform(const FName& PlatformName) override
@@ -310,52 +307,42 @@ public:
 
 	virtual void GetUnsolicitedFiles(const FName& PlatformName, const FString& Filename, const bool bIsCookable, TArray<FString>& OutUnsolicitedFiles) override
 	{
-		while (true)
+		UE::Cook::FPlatformManager::FReadScopeLock PlatformsScopeLock(Cooker.PlatformManager->ReadLockPlatforms());
+		const ITargetPlatform* TargetPlatform = AddPlatformInternal(PlatformName);
+		if (!TargetPlatform)
 		{
-			{
-				UE::Cook::FPlatformManager::FReadScopeLock PlatformsScopeLock(Cooker.PlatformManager->ReadLockPlatforms());
-				const ITargetPlatform* TargetPlatform = Cooker.AddCookOnTheFlyPlatform(PlatformName);
-				if (!TargetPlatform)
-				{
-					break;
-				}
-				if (Cooker.PlatformManager->IsPlatformInitialized(TargetPlatform))
-				{
-					Cooker.GetCookOnTheFlyUnsolicitedFiles(TargetPlatform, PlatformName.ToString(), OutUnsolicitedFiles, Filename, bIsCookable);
-					break;
-				}
-			}
-			// Wait for the Platform to be added if this is the first time; it is not legal to call GetCookOnTheFlyUnsolicitedFiles until after the platform has been added
-			FPlatformProcess::Sleep(0.001f);
+			UE_LOG(LogCook, Warning, TEXT("Trying to get unsolicited files on the fly for an invalid platform '%s'"),
+					*PlatformName.ToString());
+			return;
 		}
+		Cooker.GetCookOnTheFlyUnsolicitedFiles(TargetPlatform, PlatformName.ToString(), OutUnsolicitedFiles, Filename, bIsCookable);
 	}
 
 	virtual bool EnqueueCookRequest(UE::Cook::FCookPackageRequest CookPackageRequest) override
 	{
 		using namespace UE::Cook;
-		FPlatformManager::FReadScopeLock PlatformsScopeLock(Cooker.PlatformManager->ReadLockPlatforms());
 
-		if (const ITargetPlatform* TargetPlatform = Cooker.AddCookOnTheFlyPlatform(CookPackageRequest.PlatformName))
+		UE::Cook::FPlatformManager::FReadScopeLock PlatformsScopeLock(Cooker.PlatformManager->ReadLockPlatforms());
+		const ITargetPlatform* TargetPlatform = AddPlatformInternal(CookPackageRequest.PlatformName);
+		if (!TargetPlatform)
 		{
-			Cooker.CookOnTheFlyDeferredInitialize();
-			FName StandardFileName(*FPaths::CreateStandardFilename(CookPackageRequest.Filename));
-			UE_LOG(LogCook, Verbose, TEXT("Enqueing cook request, Filename='%s', Platform='%s'"), *CookPackageRequest.Filename, *CookPackageRequest.PlatformName.ToString());
-			Cooker.ExternalRequests->EnqueueUnique(
-				FFilePlatformRequest(StandardFileName, EInstigator::CookOnTheFly, TargetPlatform, MoveTemp(CookPackageRequest.CompletionCallback)),
-				true);
-			if (Cooker.ExternalRequests->CookRequestEvent)
-			{
-				Cooker.ExternalRequests->CookRequestEvent->Trigger();
-			}
-
-			return true;
+			UE_LOG(LogCook, Warning, TEXT("Trying to cook package on the fly for invalid platform '%s'"),
+					*CookPackageRequest.PlatformName.ToString());
+			return false;
 		}
-		else
-		{
-			UE_LOG(LogCook, Warning, TEXT("Trying to cook package on the fly for invalid platform '%s'"), *CookPackageRequest.PlatformName.ToString());
 
-			return false;
+		Cooker.CookOnTheFlyDeferredInitialize();
+		FName StandardFileName(*FPaths::CreateStandardFilename(CookPackageRequest.Filename));
+		UE_LOG(LogCook, Verbose, TEXT("Enqueing cook request, Filename='%s', Platform='%s'"), *CookPackageRequest.Filename, *CookPackageRequest.PlatformName.ToString());
+		Cooker.ExternalRequests->EnqueueUnique(
+			FFilePlatformRequest(StandardFileName, EInstigator::CookOnTheFly, TargetPlatform, MoveTemp(CookPackageRequest.CompletionCallback)),
+			true);
+		if (Cooker.ExternalRequests->CookRequestEvent)
+		{
+			Cooker.ExternalRequests->CookRequestEvent->Trigger();
 		}
+
+		return true;
 	};
 
 	virtual void MarkPackageDirty(const FName& PackageName) override
@@ -391,10 +378,75 @@ public:
 
 	virtual ICookedPackageWriter& GetPackageWriter(const ITargetPlatform* TargetPlatform) override
 	{
-		return Cooker.FindOrCreatePackageWriter(TargetPlatform);
+		return GetPackageWriterInternal(TargetPlatform);
 	}
 
 private:
+
+	const ITargetPlatform* AddPlatformInternal(const FName& PlatformName)
+	{
+		using namespace UE::Cook;
+
+		const UE::Cook::FPlatformData* PlatformData = Cooker.PlatformManager->GetPlatformDataByName(PlatformName);
+		if (!PlatformData)
+		{
+			UE_LOG(LogCook, Warning, TEXT("Target platform %s wasn't found."), *PlatformName.ToString());
+			return nullptr;
+		}
+
+		ITargetPlatform* TargetPlatform = PlatformData->TargetPlatform;
+
+		if (PlatformData->bIsSandboxInitialized)
+		{
+			return TargetPlatform;
+		}
+
+		if (IsInGameThread())
+		{
+			Cooker.AddCookOnTheFlyPlatformFromGameThread(TargetPlatform);
+			return TargetPlatform;
+		}
+
+		FEventRef Event;
+		Cooker.ExternalRequests->AddCallback([this, &Event, &TargetPlatform]()
+			{
+				Cooker.AddCookOnTheFlyPlatformFromGameThread(TargetPlatform);
+				Event->Trigger();
+			});
+
+		if (Cooker.ExternalRequests->CookRequestEvent)
+		{
+			Cooker.ExternalRequests->CookRequestEvent->Trigger();
+		}
+
+		Event->Wait();
+		return TargetPlatform;
+	}
+
+	ICookedPackageWriter& GetPackageWriterInternal(const ITargetPlatform* TargetPlatform)
+	{
+		if (IsInGameThread())
+		{
+			return Cooker.FindOrCreatePackageWriter(TargetPlatform);
+		}
+
+		FEventRef Event;
+		ICookedPackageWriter* PackageWriter = nullptr;
+		Cooker.ExternalRequests->AddCallback([this, &Event, &TargetPlatform, &PackageWriter]()
+			{
+				PackageWriter = &Cooker.FindOrCreatePackageWriter(TargetPlatform);
+				check(PackageWriter);
+				Event->Trigger();
+			});
+
+		if (Cooker.ExternalRequests->CookRequestEvent)
+		{
+			Cooker.ExternalRequests->CookRequestEvent->Trigger();
+		}
+
+		Event->Wait();
+		return *PackageWriter;
+	}
 	UCookOnTheFlyServer& Cooker;
 };
 
@@ -507,12 +559,9 @@ bool UCookOnTheFlyServer::StartCookOnTheFly(FCookOnTheFlyOptions InCookOnTheFlyO
 
 	GRedirectCollector.OnStartupPackageLoadComplete();
 
+	for (ITargetPlatform* TargetPlatform : CookOnTheFlyOptions.TargetPlatforms)
 	{
-		UE::Cook::FPlatformManager::FReadScopeLock PlatformScopeLock(PlatformManager->ReadLockPlatforms());
-		for (ITargetPlatform* TargetPlatform : CookOnTheFlyOptions.TargetPlatforms)
-		{
-			AddCookOnTheFlyPlatform(FName(*TargetPlatform->PlatformName()));
-		}
+		AddCookOnTheFlyPlatformFromGameThread(TargetPlatform);
 	}
 
 	UE_LOG(LogCook, Display, TEXT("Starting '%s' cook-on-the-fly server"),
@@ -550,49 +599,6 @@ void UCookOnTheFlyServer::CookOnTheFlyDeferredInitialize()
 	PackageDatas->BeginCook();
 }
 
-const ITargetPlatform* UCookOnTheFlyServer::AddCookOnTheFlyPlatform(const FName& PlatformName)
-{
-	const UE::Cook::FPlatformData* PlatformData = PlatformManager->GetPlatformDataByName(PlatformName);
-	if (!PlatformData)
-	{
-		UE_LOG(LogCook, Warning, TEXT("Target platform %s wasn't found."), *PlatformName.ToString());
-		return nullptr;
-	}
-
-	if (PlatformData->bIsSandboxInitialized)
-	{
-		// Platform has already been added by this function or by StartCookByTheBook
-		return PlatformData->TargetPlatform;
-	}
-
-	if (IsInGameThread())
-	{
-		AddCookOnTheFlyPlatformFromGameThread(PlatformData->TargetPlatform);
-	}
-	else
-	{
-		FEvent* PlatformInitializedEvent = FPlatformProcess::GetSynchEventFromPool();
-		
-		// Registering a new platform is not thread safe; queue the command for TickCookOnTheSide to execute
-		ExternalRequests->AddCallback([this, PlatformName, PlatformInitializedEvent]()
-			{
-				const UE::Cook::FPlatformData* PlatformData = PlatformManager->GetPlatformDataByName(PlatformName);
-				check(PlatformData);
-				AddCookOnTheFlyPlatformFromGameThread(PlatformData->TargetPlatform);
-			PlatformInitializedEvent->Trigger();
-			});
-
-		if (ExternalRequests->CookRequestEvent)
-		{
-			ExternalRequests->CookRequestEvent->Trigger();
-		}
-		
-		PlatformInitializedEvent->Wait();
-		FPlatformProcess::ReturnSynchEventToPool(PlatformInitializedEvent);
-	}
-	return PlatformData->TargetPlatform;
-}
-
 void UCookOnTheFlyServer::InitializeShadersForCookOnTheFly(const TArrayView<ITargetPlatform* const>& NewTargetPlatforms)
 {
 	for (const ITargetPlatform* TargetPlatform : NewTargetPlatforms)
@@ -1270,12 +1276,17 @@ uint32 UCookOnTheFlyServer::TickCookOnTheSide(const float TimeSlice, uint32 &Coo
 
 	if (IsCookOnTheFlyMode() && (IsCookingInEditor() == false))
 	{
+		static uint64 LastNumLoadedAndSaved = 0;
 		static int32 TickCounter = 0;
 		++TickCounter;
+		// dump stats every 50 ticks or so, but only if there is new data
 		if (TickCounter > 50)
 		{
-			// dump stats every 50 ticks or so
-			DumpStats();
+			if (StatLoadedPackageCount + StatSavedPackageCount != LastNumLoadedAndSaved)
+			{
+				LastNumLoadedAndSaved = StatLoadedPackageCount + StatSavedPackageCount;
+				DumpStats();
+			}
 			TickCounter = 0;
 		}
 	}
diff --git a/Engine/Source/Editor/UnrealEd/Private/Cooker/IoStoreCookOnTheFlyRequestManager.cpp b/Engine/Source/Editor/UnrealEd/Private/Cooker/IoStoreCookOnTheFlyRequestManager.cpp
--- a/Engine/Source/Editor/UnrealEd/Private/Cooker/IoStoreCookOnTheFlyRequestManager.cpp
+++ b/Engine/Source/Editor/UnrealEd/Private/Cooker/IoStoreCookOnTheFlyRequestManager.cpp
@@ -27,7 +27,6 @@
 class FIoStoreCookOnTheFlyNetworkServer
 {
 	static constexpr uint32 ServerSenderId = ~uint32(0);
-	static constexpr double HeartbeatTimeoutInSeconds = 60 * 5;
 
 public:
 	/**
@@ -750,8 +749,6 @@ private:
 private:
 	bool HandleClientConnection(const FName& PlatformName, bool bIsSingleThreaded, FIoStoreCookOnTheFlyNetworkServer::EConnectionStatus ConnectionStatus)
 	{
-		FScopeLock _(&ContextsCriticalSection);
-
 		if (PlatformName.IsNone())
 		{
 			return true;
@@ -760,36 +757,39 @@ private:
 		if (ConnectionStatus == FIoStoreCookOnTheFlyNetworkServer::EConnectionStatus::Connected)
 		{
 			const ITargetPlatform* TargetPlatform = CookOnTheFlyServer.AddPlatform(PlatformName);
-			if (TargetPlatform)
+			if (!TargetPlatform)
 			{
-				if (!PlatformContexts.Contains(PlatformName))
-				{
-					IPackageStoreWriter* PackageWriter = CookOnTheFlyServer.GetPackageWriter(TargetPlatform).AsPackageStoreWriter();
-					check(PackageWriter); // This class should not be used except when COTFS is using an IPackageStoreWriter
-					TUniquePtr<FPlatformContext>& Context = PlatformContexts.Add(PlatformName, MakeUnique<FPlatformContext>(PlatformName, PackageWriter));
-					
-					PackageWriter->GetEntries([&Context](TArrayView<const FPackageStoreEntryResource> Entries,
-						TArrayView<const IPackageStoreWriter::FOplogCookInfo> CookInfos)
-					{
-						Context->AddExistingPackages(Entries, CookInfos);
-					});
+				return false;
+			}
 
-					PackageWriter->OnEntryCreated().AddRaw(this, &FIoStoreCookOnTheFlyRequestManager::OnPackageStoreEntryCreated);
-					PackageWriter->OnCommit().AddRaw(this, &FIoStoreCookOnTheFlyRequestManager::OnPackageCooked);
-					PackageWriter->OnMarkUpToDate().AddRaw(this, &FIoStoreCookOnTheFlyRequestManager::OnPackagesMarkedUpToDate);
-				}
+			IPackageStoreWriter* PackageWriter = CookOnTheFlyServer.GetPackageWriter(TargetPlatform).AsPackageStoreWriter();
+			check(PackageWriter); // This class should not be used except when COTFS is using an IPackageStoreWriter
+
+			FScopeLock _(&ContextsCriticalSection);
 
-				if (bIsSingleThreaded)
+			if (!PlatformContexts.Contains(PlatformName))
+			{
+				TUniquePtr<FPlatformContext>& Context = PlatformContexts.Add(PlatformName, MakeUnique<FPlatformContext>(PlatformName, PackageWriter));
+					
+				PackageWriter->GetEntries([&Context](TArrayView<const FPackageStoreEntryResource> Entries,
+					TArrayView<const IPackageStoreWriter::FOplogCookInfo> CookInfos)
 				{
-					FPlatformContext& Context = GetContext(PlatformName);
-					FScopeLock __(&Context.GetLock());
-					Context.AddSingleThreadedClient();
-				}
+					Context->AddExistingPackages(Entries, CookInfos);
+				});
 
-				return true;
+				PackageWriter->OnEntryCreated().AddRaw(this, &FIoStoreCookOnTheFlyRequestManager::OnPackageStoreEntryCreated);
+				PackageWriter->OnCommit().AddRaw(this, &FIoStoreCookOnTheFlyRequestManager::OnPackageCooked);
+				PackageWriter->OnMarkUpToDate().AddRaw(this, &FIoStoreCookOnTheFlyRequestManager::OnPackagesMarkedUpToDate);
 			}
 
-			return false;
+			if (bIsSingleThreaded)
+			{
+				FPlatformContext& Context = GetContext(PlatformName);
+				FScopeLock __(&Context.GetLock());
+				Context.AddSingleThreadedClient();
+			}
+
+			return true;
 		}
 		else
 		{
