Removed unnecessary warnings from CIS localization preview.
#rnx

diff --git a/Engine/Source/Editor/UnrealEd/Private/Commandlets/GatherTextCommandletBase.cpp b/Engine/Source/Editor/UnrealEd/Private/Commandlets/GatherTextCommandletBase.cpp
--- a/Engine/Source/Editor/UnrealEd/Private/Commandlets/GatherTextCommandletBase.cpp
+++ b/Engine/Source/Editor/UnrealEd/Private/Commandlets/GatherTextCommandletBase.cpp
@@ -32,6 +32,7 @@ const TCHAR* UGatherTextCommandletBase::GatherTypeParam = TEXT("GatherType");
 UGatherTextCommandletBase::UGatherTextCommandletBase(const FObjectInitializer& ObjectInitializer)
 	: Super(ObjectInitializer)
 {
+	ShowErrorCount = false;
 }
 
 void UGatherTextCommandletBase::Initialize( const TSharedRef< FLocTextHelper >& InGatherManifestHelper, const TSharedPtr< FLocalizationSCC >& InSourceControlInfo )
