Fix for non-unity build.

diff --git a/Engine/Source/Editor/UnrealEd/Private/MaterialGraphNode.cpp b/Engine/Source/Editor/UnrealEd/Private/MaterialGraphNode.cpp
--- a/Engine/Source/Editor/UnrealEd/Private/MaterialGraphNode.cpp
+++ b/Engine/Source/Editor/UnrealEd/Private/MaterialGraphNode.cpp
@@ -5,6 +5,8 @@
 =============================================================================*/
 
 #include "MaterialGraph/MaterialGraphNode.h"
+#include "Widgets/Notifications/SNotificationList.h"
+#include "Framework/Notifications/NotificationManager.h"
 #include "Framework/MultiBox/MultiBoxBuilder.h"
 #include "ToolMenus.h"
 #include "MaterialGraph/MaterialGraphSchema.h"
@@ -48,8 +50,6 @@
 #include "Framework/Commands/GenericCommands.h"
 #include "ScopedTransaction.h"
 
-
-
 #define LOCTEXT_NAMESPACE "MaterialGraphNode"
 
 static FText SpaceText = LOCTEXT("Space", " ");
