SavePackage: Add more information in the fatal failed to write message.
#rnx

diff --git a/Engine/Source/Editor/UnrealEd/Private/Cooker/LooseCookedPackageWriter.cpp b/Engine/Source/Editor/UnrealEd/Private/Cooker/LooseCookedPackageWriter.cpp
--- a/Engine/Source/Editor/UnrealEd/Private/Cooker/LooseCookedPackageWriter.cpp
+++ b/Engine/Source/Editor/UnrealEd/Private/Cooker/LooseCookedPackageWriter.cpp
@@ -305,6 +305,9 @@ static void WriteToFile(const FString& Filename, const FCompositeBuffer& Buffer)
 	{
 		uint32 LastErrorCode = 0;
 		bool bSizeMatchFailed = false;
+		int64 ExpectedSize = 0;
+		int64 ActualSize = 0;
+		bool bArchiveError = false;
 	};
 	TOptional<FFailureReason> FailureReason;
 
@@ -327,13 +330,15 @@ static void WriteToFile(const FString& Filename, const FCompositeBuffer& Buffer)
 			Ar->Serialize(const_cast<void*>(Segment.GetData()), SegmentSize);
 			DataSize += SegmentSize;
 		}
+		bool bArchiveError = Ar->IsError();
 		delete Ar;
 
-		if (FileManager.FileSize(*Filename) != DataSize)
+		int64 ActualSize = FileManager.FileSize(*Filename);
+		if (ActualSize != DataSize)
 		{
 			if (!FailureReason)
 			{
-				FailureReason = FFailureReason{ 0, true };
+				FailureReason = FFailureReason{ 0, true, DataSize, ActualSize, bArchiveError };
 			}
 			FileManager.Delete(*Filename);
 			continue;
@@ -341,20 +346,24 @@ static void WriteToFile(const FString& Filename, const FCompositeBuffer& Buffer)
 		return;
 	}
 
-	TCHAR LastErrorText[1024];
+	FString ReasonText;
 	if (FailureReason && FailureReason->bSizeMatchFailed)
 	{
-		FCString::Strcpy(LastErrorText, TEXT("Unexpected file size. Another operation is modifying the file, or the write operation failed to write completely."));
+		ReasonText = FString::Printf(TEXT("Unexpected file size. Tried to write %" INT64_FMT " but resultant size was %" INT64_FMT ".%s")
+			TEXT(" Another operation is modifying the file, or the write operation failed to write completely."),
+			FailureReason->ExpectedSize, FailureReason->ActualSize, FailureReason->bArchiveError ? TEXT(" Ar->Serialize failed.") : TEXT(""));
 	}
 	else if (FailureReason && FailureReason->LastErrorCode != 0)
 	{
+		TCHAR LastErrorText[1024];
 		FPlatformMisc::GetSystemErrorMessage(LastErrorText, UE_ARRAY_COUNT(LastErrorText), FailureReason->LastErrorCode);
+		ReasonText = LastErrorText;
 	}
 	else
 	{
-		FCString::Strcpy(LastErrorText, TEXT("Unknown failure reason."));
+		ReasonText = TEXT("Unknown failure reason.");
 	}
-	UE_LOG(LogSavePackage, Fatal, TEXT("SavePackage Async write %s failed: %s"), *Filename, LastErrorText);
+	UE_LOG(LogSavePackage, Fatal, TEXT("SavePackage Async write %s failed: %s"), *Filename, *ReasonText);
 }
 
 void FLooseCookedPackageWriter::FWriteFileData::HashAndWrite(FMD5& AccumulatedHash, const TRefCountPtr<FPackageHashes>& PackageHashes, EWriteOptions WriteOptions) const
diff --git a/Engine/Source/Runtime/CoreUObject/Private/UObject/SavePackage/SavePackageUtilities.cpp b/Engine/Source/Runtime/CoreUObject/Private/UObject/SavePackage/SavePackageUtilities.cpp
--- a/Engine/Source/Runtime/CoreUObject/Private/UObject/SavePackage/SavePackageUtilities.cpp
+++ b/Engine/Source/Runtime/CoreUObject/Private/UObject/SavePackage/SavePackageUtilities.cpp
@@ -722,13 +722,16 @@ void WriteToFile(const FString& Filename, const uint8* InDataPtr, int64 InDataSi
 		if (FArchive* Ar = FileManager.CreateFileWriter(*Filename))
 		{
 			Ar->Serialize(/* grrr */ const_cast<uint8*>(InDataPtr), InDataSize);
+			bool bArchiveError = Ar->IsError();
 			delete Ar;
 
-			if (FileManager.FileSize(*Filename) != InDataSize)
+			int64 ActualSize = FileManager.FileSize(*Filename);
+			if (ActualSize != InDataSize)
 			{
 				FileManager.Delete(*Filename);
 
-				UE_LOG(LogSavePackage, Fatal, TEXT("Could not save to %s!"), *Filename);
+				UE_LOG(LogSavePackage, Fatal, TEXT("Could not save to %s! Tried to write %" INT64_FMT " bytes but resultant size was %" INT64_FMT ".%s"),
+					*Filename, InDataSize, ActualSize, bArchiveError ? TEXT(" Ar->Serialize failed.") : TEXT(""));
 			}
 			return;
 		}
