[WP navmesh] Navmesh settings -> Addition of a list of data layers to be included in the base navmesh

diff --git a/Engine/Source/Editor/UnrealEd/Private/WorldPartition/WorldPartitionNavigationDataBuilder.cpp b/Engine/Source/Editor/UnrealEd/Private/WorldPartition/WorldPartitionNavigationDataBuilder.cpp
--- a/Engine/Source/Editor/UnrealEd/Private/WorldPartition/WorldPartitionNavigationDataBuilder.cpp
+++ b/Engine/Source/Editor/UnrealEd/Private/WorldPartition/WorldPartitionNavigationDataBuilder.cpp
@@ -14,6 +14,7 @@
 
 #include "WorldPartition/WorldPartition.h"
 #include "WorldPartition/WorldPartitionSubsystem.h"
+#include "WorldPartition/DataLayer/DataLayerAsset.h"
 #include "WorldPartition/DataLayer/WorldDataLayers.h"
 #include "WorldPartition/NavigationData/NavigationDataChunkActor.h"
 
@@ -26,18 +27,24 @@ UWorldPartitionNavigationDataBuilder::UWorldPartitionNavigationDataBuilder(const
 
 bool UWorldPartitionNavigationDataBuilder::PreRun(UWorld* World, FPackageSourceControlHelper& PackageHelper)
 {
-	// Set runtime data layer to be included in the generation.
+	// Set runtime data layer to be included in the base navmesh generation.
 	if (const AWorldDataLayers* WorldDataLayers = World->GetWorldDataLayers())
 	{
-		WorldDataLayers->ForEachDataLayer([this](const UDataLayerInstance* DataLayer)
+		for(const TObjectPtr<UDataLayerAsset> DataLayer : World->GetWorldSettings()->BaseNavmeshDataLayers)
 		{
-			if (DataLayer->IsRuntime())
+			if (DataLayer != nullptr)
 			{
-				DataLayerShortNames.Add(FName(DataLayer->GetDataLayerShortName()));
+				const UDataLayerInstance* DataLayerInstance = WorldDataLayers->GetDataLayerInstance(DataLayer);
+				if(DataLayerInstance == nullptr)
+				{
+					UE_LOG(LogWorldPartitionNavigationDataBuilder, Error, TEXT("Missing UDataLayerInstance for %s."), *DataLayer->GetName());	
+				}
+				else if (DataLayerInstance->IsRuntime())
+				{
+					DataLayerShortNames.Add(FName(DataLayerInstance->GetDataLayerShortName()));
+				}
 			}
-
-			return true;
-		});
+		}
 	}
 	
 	const TSubclassOf<APartitionActor>& NavigationDataActorClass = ANavigationDataChunkActor::StaticClass();
diff --git a/Engine/Source/Runtime/Engine/Classes/GameFramework/WorldSettings.h b/Engine/Source/Runtime/Engine/Classes/GameFramework/WorldSettings.h
--- a/Engine/Source/Runtime/Engine/Classes/GameFramework/WorldSettings.h
+++ b/Engine/Source/Runtime/Engine/Classes/GameFramework/WorldSettings.h
@@ -572,6 +572,13 @@ public:
 	UPROPERTY(EditAnywhere, Category = Navigation)
 	uint32 NavigationDataChunkGridSize;
 
+	/**
+	 * A list of runtime data layers that should be included in the base navmesh.
+	 * Editor data layers and actors outside data layers will be included.
+	 */
+	UPROPERTY(EditAnywhere, Category = Navigation)
+	TArray<TObjectPtr<UDataLayerAsset>> BaseNavmeshDataLayers;
+	
 	/**
 	 * Loading cell size used when building navigation data iteratively.
 	 * The actual cell size used will be rounded using the NavigationDataChunkGridSize.
