EditorFactories cleanup some todo marks

diff --git a/Engine/Source/Editor/UnrealEd/Private/Factories/EditorFactories.cpp b/Engine/Source/Editor/UnrealEd/Private/Factories/EditorFactories.cpp
--- a/Engine/Source/Editor/UnrealEd/Private/Factories/EditorFactories.cpp
+++ b/Engine/Source/Editor/UnrealEd/Private/Factories/EditorFactories.cpp
@@ -198,7 +198,6 @@
 #include "IAssetTools.h"
 
 #include "DDSFile.h"
-#include "Factories/TIFFLoader.h"
 #include "IESConverter.h"
 #include "IImageWrapper.h"
 #include "IImageWrapperModule.h"
@@ -1227,8 +1226,7 @@ UObject* ULevelFactory::FactoryCreateText
 		else // This actor is new, but rejected to import its properties, so just delete...
 		{
 			Actor->Destroy();
-			// @@!! Actor pointer continues to be used below?
-			// continue;
+			continue;
 		}
 
 		// If this is a newly imported brush, validate it.  If it's a newly imported dynamic brush, rebuild it first.
@@ -2708,7 +2706,7 @@ UTextureFactory::UTextureFactory(const FObjectInitializer& ObjectInitializer)
 	Formats.Add( TEXT( "bmp;Texture" ) );
 	Formats.Add( TEXT( "pcx;Texture" ) );
 	Formats.Add( TEXT( "tga;Texture" ) );
-	Formats.Add( TEXT( "float;Texture" ) ); // @@!! what's this ? seems to be listed but have no loader?
+	//Formats.Add( TEXT( "float;Texture" ) ); // no such thing?
 	Formats.Add( TEXT( "psd;Texture" ) );
 	Formats.Add( TEXT( "dds;Texture (Cubemap or 2D)" ) );
 	Formats.Add( TEXT( "hdr;Texture (HDR) (LongLat unwrap or 2D)" ) );
@@ -2836,8 +2834,17 @@ public:
 	bool ProcessHorizontalRow(int64 Y)
 	{
 		// only wipe out colors that are affected by png turning valid colors white if alpha = 0
-		const uint32 WhiteWithZeroAlpha = FColor(255, 255, 255, 0).DWColor();
-		//@@!! this is wrong for 16 bit channels ;ColorDataType is uint64 but we compare to a uint32
+		ColorDataType WhiteWithZeroAlpha;
+		if ( sizeof(ColorDataType) == 4 )
+		{
+			WhiteWithZeroAlpha = FColor(255, 255, 255, 0).DWColor();
+		}
+		else
+		{
+			check( sizeof(ColorDataType) == 8 );
+			uint16 RGBA[4] = { 0xFFFF,0xFFFF,0xFFFF, 0 };
+			memcpy(&WhiteWithZeroAlpha,RGBA,8);
+		}
 
 		// Left -> Right
 		int64 NumLeftmostZerosToProcess = 0;
@@ -2936,9 +2943,8 @@ void FillZeroAlphaPNGData( int32 SizeX, int32 SizeY, ETextureSourceFormat Source
 
 		case TSF_RGBA16:
 		{
-			// @@!! this is broken so just disable for now
-			//PNGDataFill<uint16, uint64, 0, 1, 2, 3> PNGFill(SizeX, SizeY, SourceData );
-			//PNGFill.ProcessData();
+			PNGDataFill<uint16, uint64, 0, 1, 2, 3> PNGFill(SizeX, SizeY, SourceData );
+			PNGFill.ProcessData();
 			break;
 		}
 
@@ -3357,28 +3363,6 @@ bool UTextureFactory::ImportImage(const uint8* Buffer, int64 Length, FFeedbackCo
 		return true;
 	}
 	
-	//
-	// Legacy TIFF import (for the platforms that doesn't have libtiff)
-	// @@!!  this is not used?; all the Editor platforms have libtiff?  see libtiff.build.cs
-	//		remove me?
-	FTiffLoadHelper TiffLoaderHelper;
-	if (TiffLoaderHelper.IsValid())
-	{
-		if (TiffLoaderHelper.Load(Buffer, Length))
-		{
-			OutImage.Init2DWithOneMip(
-				TiffLoaderHelper.Width,
-				TiffLoaderHelper.Height,
-				TiffLoaderHelper.TextureSourceFormat,
-				TiffLoaderHelper.RawData.GetData()
-			);
-
-			OutImage.SRGB = TiffLoaderHelper.bSRGB;
-			OutImage.CompressionSettings = TiffLoaderHelper.CompressionSettings;
-			return true;
-		}
-	}
-
 	return false;
 }
 
@@ -3441,9 +3425,8 @@ UTexture* UTextureFactory::ImportTextureUDIM(UClass* Class, UObject* InParent, F
 		}
 	}
 
-	if (SourceImages.Num() < 2)
+	if (SourceImages.Num() < 1)
 	{
-		// @@!! seems like it should continue if Count == 1
 		return nullptr;
 	}
 
@@ -4198,7 +4181,7 @@ UObject* UTextureFactory::FactoryCreateBinary
 	// Start with the value that the loader suggests.
 	CompressionSettings = Texture->CompressionSettings;
 
-	// @@!! LODGroup is a member var on TextureFactory not pulled from Texture->LODGroup ?
+	// note: LODGroup is a member var on TextureFactory not pulled from Texture->LODGroup ?
 	//	notice above that CompressionSettings is primed from Texture->CompressionSettings
 	//	but LODGroup is not pulled from Texture->LODGroup
 
diff --git a/Engine/Source/Editor/UnrealEd/Private/Factories/TIFFLoader.h b/Engine/Source/Editor/UnrealEd/Private/Factories/TIFFLoader.h
--- a/Engine/Source/Editor/UnrealEd/Private/Factories/TIFFLoader.h
+++ b/Engine/Source/Editor/UnrealEd/Private/Factories/TIFFLoader.h
@@ -5,6 +5,10 @@
 #include "CoreMinimal.h"
 #include "Engine/Texture.h"
 
+// DEPRECATED
+// use Tiff ImageWrapper instead
+// delete me in UE 5.3
+
 class UTexture2D;
 
 struct FIBITMAP;
diff --git a/Engine/Source/Runtime/ImageWrapper/Private/Formats/TiffImageWrapper.cpp b/Engine/Source/Runtime/ImageWrapper/Private/Formats/TiffImageWrapper.cpp
--- a/Engine/Source/Runtime/ImageWrapper/Private/Formats/TiffImageWrapper.cpp
+++ b/Engine/Source/Runtime/ImageWrapper/Private/Formats/TiffImageWrapper.cpp
@@ -2,7 +2,7 @@
 
 #include "Formats/TiffImageWrapper.h"
 
-#if WITH_LIBTIFF
+#if WITH_LIBTIFF // set from libtiff.build.cs , available on all tool platforms
 
 #include "Async/ParallelFor.h"
 #include "Async/TaskGraphInterfaces.h"
