EditorFactories export UDIM VT with good file names
also allow UDIM import of mismatched formats

diff --git a/Engine/Source/Editor/UnrealEd/Classes/Exporters/SoundSurroundExporterWAV.h b/Engine/Source/Editor/UnrealEd/Classes/Exporters/SoundSurroundExporterWAV.h
--- a/Engine/Source/Editor/UnrealEd/Classes/Exporters/SoundSurroundExporterWAV.h
+++ b/Engine/Source/Editor/UnrealEd/Classes/Exporters/SoundSurroundExporterWAV.h
@@ -23,7 +23,7 @@ public:
 
 	//~ Begin UExporter Interface
 	virtual int32 GetFileCount( UObject* Object ) const override;
-	virtual FString GetUniqueFilename( const TCHAR* Filename, int32 FileIndex, int32 FileCount ) override;
+	virtual FString GetUniqueFilename( UObject* Object, const TCHAR* Filename, int32 FileIndex, int32 FileCount ) const override;
 	virtual bool ExportBinary( UObject* Object, const TCHAR* Type, FArchive& Ar, FFeedbackContext* Warn, int32 FileIndex = 0, uint32 PortFlags=0 ) override;
 	virtual bool SupportsObject(UObject* Object) const override;
 	//~ End UExporter Interface
diff --git a/Engine/Source/Editor/UnrealEd/Classes/Exporters/TextureExporterGeneric.h b/Engine/Source/Editor/UnrealEd/Classes/Exporters/TextureExporterGeneric.h
--- a/Engine/Source/Editor/UnrealEd/Classes/Exporters/TextureExporterGeneric.h
+++ b/Engine/Source/Editor/UnrealEd/Classes/Exporters/TextureExporterGeneric.h
@@ -19,7 +19,7 @@ class UTextureExporterGeneric : public UExporter
 
 	virtual bool SupportsObject(UObject* Object) const override;
 	virtual int32 GetFileCount(UObject* Object) const override;
-	virtual FString GetUniqueFilename(const TCHAR* Filename, int32 FileIndex, int32 FileCount) override;
+	virtual FString GetUniqueFilename(UObject* Object, const TCHAR* Filename, int32 FileIndex, int32 FileCount) const override;
 	virtual bool ExportBinary(UObject* Object, const TCHAR* Type, FArchive& Ar, FFeedbackContext* Warn, int32 FileIndex = 0, uint32 PortFlags = 0) override;
 	
 	// derived class exporter must implement :
diff --git a/Engine/Source/Editor/UnrealEd/Private/EditorExporters.cpp b/Engine/Source/Editor/UnrealEd/Private/EditorExporters.cpp
--- a/Engine/Source/Editor/UnrealEd/Private/EditorExporters.cpp
+++ b/Engine/Source/Editor/UnrealEd/Private/EditorExporters.cpp
@@ -229,7 +229,7 @@ int32 USoundSurroundExporterWAV::GetFileCount( UObject* Object ) const
 	return( SPEAKER_Count );
 }
 
-FString USoundSurroundExporterWAV::GetUniqueFilename( const TCHAR* Filename, int32 FileIndex, int32 FileCount )
+FString USoundSurroundExporterWAV::GetUniqueFilename( UObject* Object, const TCHAR* Filename, int32 FileIndex, int32 FileCount ) const
 {
 	static FString SpeakerLocations[SPEAKER_Count] =
 	{
diff --git a/Engine/Source/Editor/UnrealEd/Private/Factories/EditorFactories.cpp b/Engine/Source/Editor/UnrealEd/Private/Factories/EditorFactories.cpp
--- a/Engine/Source/Editor/UnrealEd/Private/Factories/EditorFactories.cpp
+++ b/Engine/Source/Editor/UnrealEd/Private/Factories/EditorFactories.cpp
@@ -3383,6 +3383,8 @@ UTexture* UTextureFactory::ImportTextureUDIM(UClass* Class, UObject* InParent, F
 	// UDIM requires each page to be power-of-2 and be uncompressed
 	const EImageImportFlags ImportFlags = EImageImportFlags::None;
 
+	bool bMismatchedFormats = false;
+
 	for (const auto& It : UDIMIndexToFile)
 	{
 		const FString& TexturePath = It.Value;
@@ -3403,24 +3405,21 @@ UTexture* UTextureFactory::ImportTextureUDIM(UClass* Class, UObject* InParent, F
 					TCSettings = Image.CompressionSettings;
 				}
 
-				// Deal with mismatched formats somehow?  convert?
-				if (ensure(Format == Image.Format && bSRGB == Image.SRGB))
+				if ( Format != Image.Format )
 				{
-					const int32 UDIMIndex = It.Key;
-					FTextureSourceBlock* Block = new(SourceBlocks) FTextureSourceBlock();
-					UE::TextureUtilitiesCommon::ExtractUDIMCoordinates(UDIMIndex, Block->BlockX, Block->BlockY);
-					Block->SizeX = Image.SizeX;
-					Block->SizeY = Image.SizeY;
-					Block->NumSlices = 1;
-					Block->NumMips = Image.NumMips;
-
-					SourceImages.Emplace(MoveTemp(Image));
-					SourceFileNames.Add(TexturePath);
-				}
-				else
-				{
-					Warn->Logf(ELogVerbosity::Warning, TEXT("Mismatched UDIM image formats, skipping file \"%s\""), *TexturePath);
+					bMismatchedFormats = true;
 				}
+
+				const int32 UDIMIndex = It.Key;
+				FTextureSourceBlock* Block = new(SourceBlocks) FTextureSourceBlock();
+				UE::TextureUtilitiesCommon::ExtractUDIMCoordinates(UDIMIndex, Block->BlockX, Block->BlockY);
+				Block->SizeX = Image.SizeX;
+				Block->SizeY = Image.SizeY;
+				Block->NumSlices = 1;
+				Block->NumMips = Image.NumMips;
+
+				SourceImages.Emplace(MoveTemp(Image));
+				SourceFileNames.Add(TexturePath);
 			}
 		}
 	}
@@ -3430,6 +3429,36 @@ UTexture* UTextureFactory::ImportTextureUDIM(UClass* Class, UObject* InParent, F
 		return nullptr;
 	}
 
+	if ( bMismatchedFormats )
+	{
+		Warn->Logf(ELogVerbosity::Warning, TEXT("Mismatched UDIM image formats, converting all to BGRA8 or RGBA16F ..."));
+		
+		if ( FTextureSource::IsHDR(Format) )
+		{
+			Format = TSF_RGBA16F;
+			bSRGB = false;
+		}
+		else
+		{
+			Format = TSF_BGRA8;
+			bSRGB = true;
+		}
+
+		for( FImportImage & Image : SourceImages )
+		{
+			if ( Image.Format != Format )
+			{
+				ERawImageFormat::Type ImageRawFormat = FImageCoreUtils::ConvertToRawImageFormat( Image.Format );
+				FImageView SourceImage( Image.RawData.GetData(), Image.SizeX, Image.SizeY, ImageRawFormat );
+				FImage DestImage( Image.SizeX, Image.SizeY, FImageCoreUtils::ConvertToRawImageFormat(Format) );
+				FImageCore::CopyImage(SourceImage,DestImage);
+
+				Image.RawData = MoveTemp(DestImage.RawData);
+				Image.Format = Format;
+			}				
+		}
+	}
+
 	TArray<const uint8*> SourceImageData;
 	SourceImageData.Reserve(SourceImages.Num());
 	for (const FImportImage& Image : SourceImages)
@@ -4942,7 +4971,7 @@ int32 UTextureExporterGeneric::GetFileCount(UObject* Object) const
 	return FileCount;
 }
 
-FString UTextureExporterGeneric::GetUniqueFilename(const TCHAR* Filename, int32 FileIndex, int32 FileCount)
+FString UTextureExporterGeneric::GetUniqueFilename( UObject* Object, const TCHAR* Filename, int32 FileIndex, int32 FileCount ) const
 {
 	check( FileCount > 0 );
 	if ( FileCount == 1 )
@@ -4953,15 +4982,40 @@ FString UTextureExporterGeneric::GetUniqueFilename(const TCHAR* Filename, int32
 	else
 	{
 		//	for VT with layers/blocks
-		// 
-		// @@!! would like to use LayerIndex+BlockIndex (rather than FileIndex)
-		//  but I'm not given the export Object
-		//	so not possible at the moment
-		// @todo: change GetUniqueFilename API to pass the Object
-		//
-		// it's possible you could even be able to map back to the original UDIM file names (xxx.1001 etc.)
-		//	using the TextureSource asset import info (AssetImportData)
-		return FString::Printf(TEXT("%s%d%s"), *FPaths::GetBaseFilename(Filename, false), FileIndex, *FPaths::GetExtension(Filename, true));
+
+		UTexture* Texture = GetExportTexture(Object);
+		check( Texture != nullptr );
+		
+		int NumBlocks = Texture->Source.GetNumBlocks();
+		int NumLayers = Texture->Source.GetNumLayers();
+		check( FileCount == NumBlocks * NumLayers );
+		
+		// this must match ExportTexture2DGeneric :
+		int BlockIndex = FileIndex / NumLayers;
+		int LayerIndex = FileIndex % NumLayers; 
+		
+		// make a UDIM name like MyTexture.1001.png
+
+		FString OutName = FPaths::GetBaseFilename(Filename, false);
+		// if Filename already has a .1001 in it, remove it
+		OutName.RemoveFromEnd( FString(TEXT(".1001")) );
+
+		if ( NumLayers > 1 )
+		{
+			OutName += FString::Printf(TEXT(".L%d"), LayerIndex);
+		}
+
+		if ( NumBlocks > 1 )
+		{
+			FTextureSourceBlock Block;
+			Texture->Source.GetBlock(BlockIndex,Block);
+			int UDIMCoordinate = 1001 + Block.BlockX + Block.BlockY*10;
+			OutName += FString::Printf(TEXT(".%d"), UDIMCoordinate);
+		}
+
+		OutName += FPaths::GetExtension(Filename, true);
+
+		return OutName;
 	}
 }
 
diff --git a/Engine/Source/Runtime/Engine/Classes/Exporters/Exporter.h b/Engine/Source/Runtime/Engine/Classes/Exporters/Exporter.h
--- a/Engine/Source/Runtime/Engine/Classes/Exporters/Exporter.h
+++ b/Engine/Source/Runtime/Engine/Classes/Exporters/Exporter.h
@@ -141,7 +141,11 @@ class UExporter : public UObject
 	/** 
 	 * Differentiates the filename for objects with multiple files to export. Only needs to be overridden if GetFileCount() returns > 1.
 	 */
-	virtual FString GetUniqueFilename( const TCHAR* Filename, int32 FileIndex, int32 FileCount ) { check( FileIndex == 0 && FileCount == 1 ); return( FString( Filename ) ); }
+	virtual FString GetUniqueFilename( UObject* Object, const TCHAR* Filename, int32 FileIndex, int32 FileCount ) const
+	{
+		check( FileIndex == 0 && FileCount == 1 );
+		return FString( Filename );
+	}
 
 	/**
 	* Return true if the exporter is in batch mode.
diff --git a/Engine/Source/Runtime/Engine/Private/UnrealExporter.cpp b/Engine/Source/Runtime/Engine/Private/UnrealExporter.cpp
--- a/Engine/Source/Runtime/Engine/Private/UnrealExporter.cpp
+++ b/Engine/Source/Runtime/Engine/Private/UnrealExporter.cpp
@@ -377,7 +377,7 @@ bool UExporter::RunAssetExportTask(class UAssetExportTask* Task)
 			FBufferArchive Buffer;
 			if(ExportToArchive(Task->Object, Exporter, Buffer, *Extension, i))
 			{
-				FString UniqueFilename = Exporter->GetUniqueFilename(*Task->Filename, i, FileCount);
+				FString UniqueFilename = Exporter->GetUniqueFilename(Task->Object, *Task->Filename, i, FileCount);
 
 				if(!Task->bReplaceIdentical)
 				{
