AsyncIODeleteTest: fix bug in test that was spuriously deleting the entire binaries directory when async is disabled.
#rnx

diff --git a/Engine/Source/Editor/UnrealEd/Private/Cooker/AsyncIODeleteTest.cpp b/Engine/Source/Editor/UnrealEd/Private/Cooker/AsyncIODeleteTest.cpp
--- a/Engine/Source/Editor/UnrealEd/Private/Cooker/AsyncIODeleteTest.cpp
+++ b/Engine/Source/Editor/UnrealEd/Private/Cooker/AsyncIODeleteTest.cpp
@@ -88,7 +88,14 @@ bool FAsyncIODeleteTest::RunTest(const FString& Parameters)
 		FString DeletionRoot = FString(AsyncIODelete.GetDeletionRoot());
 		StartSection(TEXT("Deleting the temporary files/directories should be finished before WaitForAllTasks returns"));
 		WaitForAllTasksAndVerify(AsyncIODelete);
-		TestTempRootCountsEqual(DeletionRoot, 0, 0, TEXT("AsyncIODelete should have deleted the moved paths before WaitForAllTasks returned."));
+		if (bAsyncEnabled)
+		{
+			TestTrue(TEXT("DeletionRoot is valid after deleting some files"), !DeletionRoot.IsEmpty());
+			if (!DeletionRoot.IsEmpty())
+			{
+				TestTempRootCountsEqual(DeletionRoot, 0, 0, TEXT("AsyncIODelete should have deleted the moved paths before WaitForAllTasks returned."));
+			}
+		}
 
 		StartSection(TEXT("Two FAsyncIODelete constructed at once should be legal, as long as they have different TempRoots"));
 		FAsyncIODelete AsyncIODelete2(SharedTempRoot2);
@@ -103,13 +110,20 @@ bool FAsyncIODeleteTest::RunTest(const FString& Parameters)
 		FString DeletionRoot2 = FString(AsyncIODelete2.GetDeletionRoot());
 		if (bAsyncEnabled)
 		{
-			TestTempRootCountsEqual(DeletionRoot2, 1, 1, TEXT("AsyncIODelete should not have deleted the moved paths because it is paused."));
+			TestTrue(TEXT("DeletionRoot is valid after deleting some files"), !DeletionRoot2.IsEmpty());
+			if (!DeletionRoot2.IsEmpty())
+			{
+				TestTempRootCountsEqual(DeletionRoot2, 1, 1, TEXT("AsyncIODelete should not have deleted the moved paths because it is paused."));
+			}
 		}
 		AsyncIODelete2.SetDeletesPaused(false);
 		WaitForAllTasksAndVerify(AsyncIODelete2);
 		if (bAsyncEnabled)
 		{
-			TestTempRootCountsEqual(DeletionRoot2, 0, 0, TEXT("AsyncIODelete should have deleted the moved paths after unpausing."));
+			if (!DeletionRoot2.IsEmpty())
+			{
+				TestTempRootCountsEqual(DeletionRoot2, 0, 0, TEXT("AsyncIODelete should have deleted the moved paths after unpausing."));
+			}
 		}
 
 		StartSection(TEXT("Verify Teardown() deletes the TempRoot and Setup() creates it"));
@@ -167,12 +181,15 @@ bool FAsyncIODeleteTest::RunTest(const FString& Parameters)
 		}
 
 		StartSection(TEXT("Attempting to delete a parent directory of the temproot, the temproot itself, or a child inside of it fails"));
-		FString SubDirInTempRoot4 = FPaths::Combine(DeletionRoot4, TEXT("SubDir"));
-		FileManager.MakeDirectory(*SubDirInTempRoot4, bApplyToTreeTrue); // Note it's illegal to add files into TempRoot, but we're not currently checking for it and we're not colliding with the DeleteN paths AsyncIODelete uses, so this breaking of the rule will not cause problems
-		TestFalse(TEXT("AsyncIODelete should refuse to delete a parent of its TempRoot."), AsyncIODelete3.DeleteDirectory(TestRoot));
-		TestFalse(TEXT("AsyncIODelete should refuse to delete its TempRoot."), AsyncIODelete3.DeleteDirectory(DeletionRoot4));
-		TestFalse(TEXT("AsyncIODelete should refuse to delete its SharedTempRoot."), AsyncIODelete3.DeleteDirectory(SharedTempRoot4));
-		TestFalse(TEXT("AsyncIODelete should refuse to delete a child of its TempRoot."), AsyncIODelete3.DeleteDirectory(SubDirInTempRoot4));
+		if (bAsyncEnabled && !DeletionRoot4.IsEmpty())
+		{
+			FString SubDirInTempRoot4 = FPaths::Combine(DeletionRoot4, TEXT("SubDir"));
+			FileManager.MakeDirectory(*SubDirInTempRoot4, bApplyToTreeTrue); // Note it's illegal to add files into TempRoot, but we're not currently checking for it and we're not colliding with the DeleteN paths AsyncIODelete uses, so this breaking of the rule will not cause problems
+			TestFalse(TEXT("AsyncIODelete should refuse to delete a parent of its TempRoot."), AsyncIODelete3.DeleteDirectory(TestRoot));
+			TestFalse(TEXT("AsyncIODelete should refuse to delete its TempRoot."), AsyncIODelete3.DeleteDirectory(DeletionRoot4));
+			TestFalse(TEXT("AsyncIODelete should refuse to delete its SharedTempRoot."), AsyncIODelete3.DeleteDirectory(SharedTempRoot4));
+			TestFalse(TEXT("AsyncIODelete should refuse to delete a child of its TempRoot."), AsyncIODelete3.DeleteDirectory(SubDirInTempRoot4));
+		}
 	}
 	if (bAsyncEnabled)
 	{
