Check blueprint class for editor-only flag in addition to Blueprint itself when validating editor-only k2 nodes

diff --git a/Engine/Source/Editor/UnrealEd/Private/Kismet2/BlueprintEditorUtils.cpp b/Engine/Source/Editor/UnrealEd/Private/Kismet2/BlueprintEditorUtils.cpp
--- a/Engine/Source/Editor/UnrealEd/Private/Kismet2/BlueprintEditorUtils.cpp
+++ b/Engine/Source/Editor/UnrealEd/Private/Kismet2/BlueprintEditorUtils.cpp
@@ -4666,7 +4666,11 @@ void FBlueprintEditorUtils::ValidateEditorOnlyNodes(const UK2Node* Node, FCompil
 		const bool bIsEditorOnlyPackage = NodeCDOPackage->HasAllPackagesFlags(PKG_EditorOnly);
 		const bool bIsUncookedOrDev = NodeCDOPackage->HasAnyPackageFlags(PKG_UncookedOnly | PKG_Developer);		
 
-		if (!bIsUncookedOrDev && bIsEditorOnlyPackage && !BP->IsEditorOnly())
+		const UClass* BlueprintClass = BP ? BP->ParentClass : nullptr;
+		const bool bIsEditorOnlyBlueprintBaseClass = BlueprintClass ? IsEditorOnlyObject(BlueprintClass) : false;
+
+		// Check whether the blueprint itself, or its class is marked as editor-only
+		if (!bIsUncookedOrDev && bIsEditorOnlyPackage && !(IsEditorOnlyObject(BP) || bIsEditorOnlyBlueprintBaseClass))
 		{
 			MessageLog.Warning(*LOCTEXT("EditorOnlyConflict_ErrorFmt", "The node '@@' is from an Editor Only module, but is placed in a runtime blueprint! K2 Nodes should only be defined in a Developer or UncookedOnly module.").ToString(), Node);
 		}
