Fixed typo in GetScratchRenderTarget when generating 512x512 RT

diff --git a/Engine/Source/Editor/UnrealEd/Private/EditorEngine.cpp b/Engine/Source/Editor/UnrealEd/Private/EditorEngine.cpp
--- a/Engine/Source/Editor/UnrealEd/Private/EditorEngine.cpp
+++ b/Engine/Source/Editor/UnrealEd/Private/EditorEngine.cpp
@@ -2628,7 +2628,7 @@ UTextureRenderTarget2D* UEditorEngine::GetScratchRenderTarget( uint32 MinSize )
 			ScratchRenderTarget512->TargetGamma = GEngine->DisplayGamma;
 			ScratchRenderTarget512->RenderTargetFormat = RTF_RGBA8;
 
-			ScratchRenderTarget512->InitAutoFormat(515, 512);
+			ScratchRenderTarget512->InitAutoFormat(512, 512);
 		}
 		ScratchRenderTarget = ScratchRenderTarget512;
 	}
