UE-141180 - Loss of precision with String to Float Real conversion

We neglected to redirect Conv_FloatToText nodes to Conv_DoubleToText, which was causing loss of precision when formatting text. However, this revealed another issue with the format text node.

Since the format text node uses wildcards, it needs a sorting priority of Low_UsesDependentWildcard. Otherwise, linked inputs won't be updated before we sync the node's pins. For example, we might be linked to a Conv_FloatToText node that hasn't yet been converted to a Conv_DoubleToText node. As a result, the format text node will use a single precision float input for the connected conversion node instead of a double.

Another issue is that the format text node is marked as causing a structural Blueprint change. The current sorting function prioritizes that over its refresh priority. Fortunately, we don't need to concern ourselves with structural modification if we're compiling on load, which is when we really need the refresh priority to work as expected.

Finally, K2Node_EnumEquality node needs to have its knot dependencies updated prior to reconstruction. Otherwise, it might fail to find a valid UEnum, and disconnect its inputs. Since knot nodes already have this logic built in, we can just force a type propagation on any linked knots.

diff --git a/Engine/Config/BaseEngine.ini b/Engine/Config/BaseEngine.ini
--- a/Engine/Config/BaseEngine.ini
+++ b/Engine/Config/BaseEngine.ini
@@ -1031,9 +1031,7 @@ DefaultWorldPartitionHLODLayer=/Engine/Maps/Templates/HLODs/HLODLayer_Instanced.
 ; Blueprint real number support
 +FunctionRedirects=(OldName="KismetStringLibrary.Conv_FloatToString", NewName="Conv_DoubleToString")
 +PropertyRedirects=(OldName="KismetStringLibrary.Conv_DoubleToString.InFloat", NewName="Conv_DoubleToString.InDouble")
-
 +FunctionRedirects=(OldName="KismetStringLibrary.Conv_StringToFloat", NewName="Conv_StringToDouble")
-
 +FunctionRedirects=(OldName="KismetStringLibrary.BuildString_Float", NewName="BuildString_Double")
 +PropertyRedirects=(OldName="KismetStringLibrary.BuildString_Double.InFloat", NewName="BuildString_Double.InDouble")
 
@@ -1045,7 +1043,6 @@ DefaultWorldPartitionHLODLayer=/Engine/Maps/Templates/HLODs/HLODLayer_Instanced.
 +FunctionRedirects=(OldName="KismetMathLibrary.BreakVector_NetQuantize10", NewName="BreakVector")
 +FunctionRedirects=(OldName="KismetMathLibrary.BreakVector_NetQuantize100", NewName="BreakVector")
 +FunctionRedirects=(OldName="KismetMathLibrary.BreakVector_NetQuantizeNormal", NewName="BreakVector")
-
 +FunctionRedirects=(OldName="KismetMathLibrary.Multiply_FloatFloat", NewName="Multiply_DoubleDouble")
 +FunctionRedirects=(OldName="KismetMathLibrary.Divide_FloatFloat", NewName="Divide_DoubleDouble")
 +FunctionRedirects=(OldName="KismetMathLibrary.Add_FloatFloat", NewName="Add_DoubleDouble")
@@ -1057,6 +1054,8 @@ DefaultWorldPartitionHLODLayer=/Engine/Maps/Templates/HLODs/HLODLayer_Instanced.
 +FunctionRedirects=(OldName="KismetMathLibrary.EqualEqual_FloatFloat", NewName="EqualEqual_DoubleDouble")
 +FunctionRedirects=(OldName="KismetMathLibrary.NotEqual_FloatFloat", NewName="NotEqual_DoubleDouble")
 
++FunctionRedirects=(OldName="KismetTextLibrary.Conv_FloatToText", NewName="Conv_DoubleToText")
+
 [CoreUObject.Metadata]
 ; Note: UnrealHeaderTool should be rerun after making changes to MetadataRedirects to catch any keys specified in class headers
 +MetadataRedirects=(OldKey="K2Protected", NewKey="BlueprintProtected")
diff --git a/Engine/Source/Editor/BlueprintGraph/Classes/K2Node_FormatText.h b/Engine/Source/Editor/BlueprintGraph/Classes/K2Node_FormatText.h
--- a/Engine/Source/Editor/BlueprintGraph/Classes/K2Node_FormatText.h
+++ b/Engine/Source/Editor/BlueprintGraph/Classes/K2Node_FormatText.h
@@ -40,6 +40,7 @@ class UK2Node_FormatText : public UK2Node
 	virtual bool IsConnectionDisallowed(const UEdGraphPin* MyPin, const UEdGraphPin* OtherPin, FString& OutReason) const override;
 	virtual void GetMenuActions(FBlueprintActionDatabaseRegistrar& ActionRegistrar) const override;
 	virtual FText GetMenuCategory() const override;
+	virtual int32 GetNodeRefreshPriority() const override { return EBaseNodeRefreshPriority::Low_UsesDependentWildcard; }
 	//~ End UK2Node Interface.
 
 public:
diff --git a/Engine/Source/Editor/BlueprintGraph/Private/K2Node_EnumEquality.cpp b/Engine/Source/Editor/BlueprintGraph/Private/K2Node_EnumEquality.cpp
--- a/Engine/Source/Editor/BlueprintGraph/Private/K2Node_EnumEquality.cpp
+++ b/Engine/Source/Editor/BlueprintGraph/Private/K2Node_EnumEquality.cpp
@@ -4,6 +4,7 @@
 #include "K2Node_EnumEquality.h"
 #include "EdGraphSchema_K2.h"
 #include "K2Node_CallFunction.h"
+#include "K2Node_Knot.h"
 #include "KismetCompilerMisc.h"
 #include "KismetCompiler.h"
 #include "Kismet/KismetMathLibrary.h"
@@ -77,12 +78,23 @@ void UK2Node_EnumEquality::NotifyPinConnectionListChanged(UEdGraphPin* Pin)
 		}
 		else if (Pin->LinkedTo.Num() > 0)
 		{
+			UEdGraphPin* LinkedPin = Pin->LinkedTo[0];
+			check(LinkedPin);
+			
+			// If we're linked to a knot node, ensure that its enum data is up-to-date.
+			// Otherwise, PinSubCategoryObject might be null if we're reconstructed prior to our knot dependencies.
+			UK2Node_Knot* LinkedKnotNode = Cast<UK2Node_Knot>(LinkedPin->GetOwningNode());
+			if (LinkedKnotNode)
+			{
+				LinkedKnotNode->PropagatePinType();
+			}
+
 			// Make sure the pin is a valid enum
-			if (Pin->LinkedTo[0]->PinType.PinCategory == UEdGraphSchema_K2::PC_Byte &&
-				Pin->LinkedTo[0]->PinType.PinSubCategoryObject.IsValid() &&
-				Pin->LinkedTo[0]->PinType.PinSubCategoryObject.Get()->IsA(UEnum::StaticClass()))
+			if (LinkedPin->PinType.PinCategory == UEdGraphSchema_K2::PC_Byte &&
+				LinkedPin->PinType.PinSubCategoryObject.IsValid() &&
+				LinkedPin->PinType.PinSubCategoryObject.Get()->IsA(UEnum::StaticClass()))
 			{
-				Pin->PinType = Pin->LinkedTo[0]->PinType;
+				Pin->PinType = LinkedPin->PinType;
 
 				UEdGraphPin* OtherPin = (Input1Pin == Pin) ? Input2Pin : Input1Pin;
 
diff --git a/Engine/Source/Editor/UnrealEd/Private/Kismet2/BlueprintEditorUtils.cpp b/Engine/Source/Editor/UnrealEd/Private/Kismet2/BlueprintEditorUtils.cpp
--- a/Engine/Source/Editor/UnrealEd/Private/Kismet2/BlueprintEditorUtils.cpp
+++ b/Engine/Source/Editor/UnrealEd/Private/Kismet2/BlueprintEditorUtils.cpp
@@ -139,9 +139,14 @@ DEFINE_STAT(EKismetCompilerStats_NotifyBlueprintChanged);
 DECLARE_CYCLE_STAT(TEXT("Mark Blueprint as Structurally Modified"), EKismetCompilerStats_MarkBlueprintasStructurallyModified, STATGROUP_KismetCompiler);
 DECLARE_CYCLE_STAT(TEXT("Refresh External DependencyNodes"), EKismetCompilerStats_RefreshExternalDependencyNodes, STATGROUP_KismetCompiler);
 
-struct FCompareNodePriority
+static void SortNodes(TArray<UK2Node*>& AllNodes, bool bSortByPriorityOnly = false)
 {
-	FORCEINLINE bool operator()( const UK2Node& A, const UK2Node& B ) const
+	auto SortNodesInternalByPriorityOnly = [](const UK2Node& A, const UK2Node& B)
+	{
+		return A.GetNodeRefreshPriority() > B.GetNodeRefreshPriority();
+	};
+
+	auto SortNodesInternal = [SortNodesInternalByPriorityOnly](const UK2Node& A, const UK2Node& B)
 	{
 		const bool NodeAChangesStructure = A.NodeCausesStructuralBlueprintChange();
 		const bool NodeBChangesStructure = B.NodeCausesStructuralBlueprintChange();
@@ -150,10 +155,22 @@ struct FCompareNodePriority
 		{
 			return NodeAChangesStructure;
 		}
-		
-		return A.GetNodeRefreshPriority() > B.GetNodeRefreshPriority();
+
+		return SortNodesInternalByPriorityOnly(A, B);
+	};
+
+	if (AllNodes.Num() > 1)
+	{
+		if (bSortByPriorityOnly)
+		{
+			AllNodes.Sort(SortNodesInternalByPriorityOnly);
+		}
+		else
+		{
+			AllNodes.Sort(SortNodesInternal);
+		}
 	}
-};
+}
 
 /**
  * This helper does a depth first search, looking for the highest parent class that
@@ -525,10 +542,7 @@ void FBlueprintEditorUtils::RefreshAllNodes(UBlueprint* Blueprint)
 	FBlueprintEditorUtils::GetAllNodesOfClass(Blueprint, AllNodes);
 
 	const bool bIsMacro = (Blueprint->BlueprintType == BPTYPE_MacroLibrary);
-	if( AllNodes.Num() > 1 )
-	{
-		AllNodes.Sort(FCompareNodePriority());
-	}
+	SortNodes(AllNodes);
 
 	bool bLastChangesStructure = (AllNodes.Num() > 0) ? AllNodes[0]->NodeCausesStructuralBlueprintChange() : true;
 	for( TArray<UK2Node*>::TIterator NodeIt(AllNodes); NodeIt; ++NodeIt )
@@ -580,10 +594,7 @@ void FBlueprintEditorUtils::ReconstructAllNodes(UBlueprint* Blueprint)
 	TArray<UK2Node*> AllNodes;
 	FBlueprintEditorUtils::GetAllNodesOfClass(Blueprint, AllNodes);
 
-	if (AllNodes.Num() > 1)
-	{
-		AllNodes.Sort(FCompareNodePriority());
-	}
+	SortNodes(AllNodes, true);
 
 	for (TArray<UK2Node*>::TIterator NodeIt(AllNodes); NodeIt; ++NodeIt)
 	{
diff --git a/Engine/Source/Runtime/Engine/Classes/Kismet/KismetTextLibrary.h b/Engine/Source/Runtime/Engine/Classes/Kismet/KismetTextLibrary.h
--- a/Engine/Source/Runtime/Engine/Classes/Kismet/KismetTextLibrary.h
+++ b/Engine/Source/Runtime/Engine/Classes/Kismet/KismetTextLibrary.h
@@ -205,12 +205,12 @@ class ENGINE_API UKismetTextLibrary : public UBlueprintFunctionLibrary
 	static FText Conv_Int64ToText(int64 Value, bool bAlwaysSign = false, bool bUseGrouping = true, int32 MinimumIntegralDigits = 1, int32 MaximumIntegralDigits = 324);
 
 	/** Converts a passed in double to text based on formatting options */
-	UFUNCTION(BlueprintPure, meta = (DisplayName = "To Text (Double)", AdvancedDisplay = "1", BlueprintAutocast), Category = "Utilities|Text")
-	static FText Conv_DoubleToText(double InDouble, TEnumAsByte<ERoundingMode> RoundingMode, bool bAlwaysSign = false, bool bUseGrouping = true, int32 MinimumIntegralDigits = 1, int32 MaximumIntegralDigits = 324, int32 MinimumFractionalDigits = 0, int32 MaximumFractionalDigits = 3);
+	UFUNCTION(BlueprintPure, meta = (DisplayName = "To Text (Float)", AdvancedDisplay = "1", BlueprintAutocast), Category = "Utilities|Text")
+	static FText Conv_DoubleToText(double Value, TEnumAsByte<ERoundingMode> RoundingMode, bool bAlwaysSign = false, bool bUseGrouping = true, int32 MinimumIntegralDigits = 1, int32 MaximumIntegralDigits = 324, int32 MinimumFractionalDigits = 0, int32 MaximumFractionalDigits = 3);
 	
 	// Default values are duplicated from FNumberFormattingOptions and should be replicated in all functions and in the struct when changed!
 	/** Converts a passed in float to text based on formatting options */
-	UFUNCTION(BlueprintPure, meta=(DisplayName = "To Text (Float)", AdvancedDisplay = "1", BlueprintAutocast), Category="Utilities|Text")
+	UE_DEPRECATED(5.0, "This method has been deprecated and will be removed.")
 	static FText Conv_FloatToText(float Value, TEnumAsByte<ERoundingMode> RoundingMode, bool bAlwaysSign = false, bool bUseGrouping = true, int32 MinimumIntegralDigits = 1, int32 MaximumIntegralDigits = 324, int32 MinimumFractionalDigits = 0, int32 MaximumFractionalDigits = 3);
 	
 	/**
diff --git a/Engine/Source/Runtime/Engine/Private/KismetTextLibrary.cpp b/Engine/Source/Runtime/Engine/Private/KismetTextLibrary.cpp
--- a/Engine/Source/Runtime/Engine/Private/KismetTextLibrary.cpp
+++ b/Engine/Source/Runtime/Engine/Private/KismetTextLibrary.cpp
@@ -195,7 +195,7 @@ FText UKismetTextLibrary::Conv_Int64ToText(int64 Value, bool bAlwaysSign /*= fal
 	return FText::AsNumber(Value, &NumberFormatOptions);
 }
 
-FText UKismetTextLibrary::Conv_DoubleToText(double InDouble, TEnumAsByte<ERoundingMode> RoundingMode, bool bAlwaysSign, bool bUseGrouping, int32 MinimumIntegralDigits, int32 MaximumIntegralDigits, int32 MinimumFractionalDigits, int32 MaximumFractionalDigits)
+FText UKismetTextLibrary::Conv_DoubleToText(double Value, TEnumAsByte<ERoundingMode> RoundingMode, bool bAlwaysSign, bool bUseGrouping, int32 MinimumIntegralDigits, int32 MaximumIntegralDigits, int32 MinimumFractionalDigits, int32 MaximumFractionalDigits)
 {
 	FNumberFormattingOptions NumberFormatOptions;
 	NumberFormatOptions.AlwaysSign = bAlwaysSign;
@@ -206,21 +206,12 @@ FText UKismetTextLibrary::Conv_DoubleToText(double InDouble, TEnumAsByte<ERoundi
 	NumberFormatOptions.MinimumFractionalDigits = MinimumFractionalDigits;
 	NumberFormatOptions.MaximumFractionalDigits = MaximumFractionalDigits;
 	
-	return FText::AsNumber(InDouble, &NumberFormatOptions);
+	return FText::AsNumber(Value, &NumberFormatOptions);
 }
 
 FText UKismetTextLibrary::Conv_FloatToText(float Value, TEnumAsByte<ERoundingMode> RoundingMode, bool bAlwaysSign/* = false*/, bool bUseGrouping/* = true*/, int32 MinimumIntegralDigits/* = 1*/, int32 MaximumIntegralDigits/* = 324*/, int32 MinimumFractionalDigits/* = 0*/, int32 MaximumFractionalDigits/* = 3*/)
 {
-	FNumberFormattingOptions NumberFormatOptions;
-	NumberFormatOptions.AlwaysSign = bAlwaysSign;
-	NumberFormatOptions.UseGrouping = bUseGrouping;
-	NumberFormatOptions.RoundingMode = RoundingMode;
-	NumberFormatOptions.MinimumIntegralDigits = MinimumIntegralDigits;
-	NumberFormatOptions.MaximumIntegralDigits = MaximumIntegralDigits;
-	NumberFormatOptions.MinimumFractionalDigits = MinimumFractionalDigits;
-	NumberFormatOptions.MaximumFractionalDigits = MaximumFractionalDigits;
-
-	return FText::AsNumber(Value, &NumberFormatOptions);
+	return Conv_DoubleToText(Value, RoundingMode, bAlwaysSign, bUseGrouping, MinimumIntegralDigits, MaximumIntegralDigits, MinimumFractionalDigits, MaximumFractionalDigits);
 }
 
 FText UKismetTextLibrary::AsCurrencyBase(int32 BaseValue, const FString& CurrencyCode)
