Fix crash in FEditorBuildUtils::TriggerNavigationBuilder()

diff --git a/Engine/Source/Editor/UnrealEd/Private/EditorBuildUtils.cpp b/Engine/Source/Editor/UnrealEd/Private/EditorBuildUtils.cpp
--- a/Engine/Source/Editor/UnrealEd/Private/EditorBuildUtils.cpp
+++ b/Engine/Source/Editor/UnrealEd/Private/EditorBuildUtils.cpp
@@ -896,7 +896,7 @@ void FEditorBuildUtils::TriggerNavigationBuilder(UWorld*& InOutWorld, FName Id)
 			bBuildingNavigationFromUserRequest = false;
 		}
 
-		const FString& LongPackageName = InOutWorld->GetPackage()->GetName();
+		const FString& LongPackageName = GetNameSafe(InOutWorld->GetPackage());
 		if (UE::EditorBuildUtils::bNavmeshAllowPartitionedBuildingFromEditor && ULevel::GetIsLevelPartitionedFromPackage(*LongPackageName))
 		{
 			WorldPartitionBuildNavigation(LongPackageName);
