Adding elements support for the marquee selection.

Added a new registry to let the objects provide how their elements should be selected when doing a selection by a volume such as a box or a frustum.
Added a some editor only functions for the marquee selection into the world interface.

diff --git a/Engine/Plugins/Animation/ControlRig/Source/ControlRigEditor/Private/EditMode/ControlRigEditMode.cpp b/Engine/Plugins/Animation/ControlRig/Source/ControlRigEditor/Private/EditMode/ControlRigEditMode.cpp
--- a/Engine/Plugins/Animation/ControlRig/Source/ControlRigEditor/Private/EditMode/ControlRigEditMode.cpp
+++ b/Engine/Plugins/Animation/ControlRig/Source/ControlRigEditor/Private/EditMode/ControlRigEditMode.cpp
@@ -1490,7 +1490,7 @@ bool IntersectsBox( AActor& InActor, const FBox& InBox, FLevelEditorViewportClie
 		UPrimitiveComponent* PrimitiveComponent = Cast<UPrimitiveComponent>(Component);
 		if (PrimitiveComponent && PrimitiveComponent->IsRegistered() && PrimitiveComponent->IsVisibleInEditor())
 		{
-			if (PrimitiveComponent->ComponentIsTouchingSelectionBox(InBox, LevelViewportClient->EngineShowFlags, false, bUseStrictSelection))
+			if (PrimitiveComponent->IsShown(LevelViewportClient->EngineShowFlags) && PrimitiveComponent->ComponentIsTouchingSelectionBox(InBox, false, bUseStrictSelection))
 			{
 				return true;
 			}
@@ -1582,7 +1582,7 @@ bool FControlRigEditMode::FrustumSelect(const FConvexVolume& InFrustum, FEditorV
 				UPrimitiveComponent* PrimitiveComponent = Cast<UPrimitiveComponent>(Component);
 				if (PrimitiveComponent && PrimitiveComponent->IsRegistered() && PrimitiveComponent->IsVisibleInEditor())
 				{
-					if (PrimitiveComponent->ComponentIsTouchingSelectionFrustum(InFrustum, InViewportClient->EngineShowFlags, false /*only bsp*/, false/*encompass entire*/))
+					if (PrimitiveComponent->IsShown(InViewportClient->EngineShowFlags) && PrimitiveComponent->ComponentIsTouchingSelectionFrustum(InFrustum, false /*only bsp*/, false/*encompass entire*/))
 					{
 						if (ShapeActor->IsSelectable() && ShapeActor->ControlRig.IsValid() &&  ShapeActor->ControlRig->GetControlsVisible() )
 						{
diff --git a/Engine/Source/Editor/UnrealEd/Private/DragTool_BoxSelect.cpp b/Engine/Source/Editor/UnrealEd/Private/DragTool_BoxSelect.cpp
--- a/Engine/Source/Editor/UnrealEd/Private/DragTool_BoxSelect.cpp
+++ b/Engine/Source/Editor/UnrealEd/Private/DragTool_BoxSelect.cpp
@@ -10,12 +10,17 @@
 #include "UnrealEdGlobals.h"
 #include "EngineUtils.h"
 #include "EditorModeManager.h"
+#include "Elements/Framework/EngineElementsLibrary.h"
+#include "Elements/Framework/TypedElementRegistry.h"
+#include "Elements/Framework/TypedElementSelectionSet.h"
+#include "Elements/Interfaces/TypedElementWorldInterface.h"
 #include "EditorModes.h"
 #include "ActorEditorUtils.h"
 #include "ScopedTransaction.h"
 #include "Engine/LevelStreaming.h"
 #include "CanvasTypes.h"
 #include "Subsystems/BrushEditingSubsystem.h"
+#include "LevelEditorSubsystem.h"
 
 ///////////////////////////////////////////////////////////////////////////////////////////////////////////////////////
 //
@@ -23,6 +28,31 @@
 //
 ///////////////////////////////////////////////////////////////////////////////////////////////////////////////////////
 
+namespace UE::LevelEditor::Private
+{
+	TArray<FTypedElementHandle> GetElementsIntersectingBox(const AActor* Actor,
+		const FBox& InBox,
+		FLevelEditorViewportClient* Viewport,
+		const FWorldSelectionElementArgs& SelectionArgs)
+	{
+		if (Actor && (!Viewport || !Actor->IsA(AVolume::StaticClass()) || !Viewport->IsVolumeVisibleInViewport(*Actor)))
+		{
+			if (Actor->GetActorLabel() == FString(TEXT("ChaosAsset_Floor3")))
+			{
+				if (FTypedElementHandle ActorHandle = UEngineElementsLibrary::AcquireEditorActorElementHandle(Actor))
+				{
+					if (TTypedElement<ITypedElementWorldInterface> WorldElement = UTypedElementRegistry::GetInstance()->GetElement<ITypedElementWorldInterface>(ActorHandle))
+					{
+						return WorldElement.GetSelectionElementsInBox(InBox, SelectionArgs);
+					}
+				}
+			}
+		}
+
+		return {};
+	}
+}
+
 /**
  * Starts a mouse drag behavior.  The start location is snapped to the editor constraints if bUseSnapping is true.
  *
@@ -81,6 +111,22 @@ void FDragTool_ActorBoxSelect::AddDelta( const FVector& InDelta )
 	{
 		const bool bStrictDragSelection = GetDefault<ULevelEditorViewportSettings>()->bStrictBoxSelection;
 
+		UTypedElementSelectionSet* SelectionSet = GEditor->GetEditorSubsystem<ULevelEditorSubsystem>()->GetSelectionSet();
+
+		UBrushEditingSubsystem* BrushSubsystem = GEditor->GetEditorSubsystem<UBrushEditingSubsystem>();
+		const bool bGeometryMode = BrushSubsystem ? BrushSubsystem->IsGeometryEditorModeActive() : false;
+
+		FWorldSelectionElementArgs SeletionArgs
+		{
+			SelectionSet,
+			ETypedElementSelectionMethod::Primary,
+			FTypedElementSelectionOptions(),
+			&(LevelViewportClient->EngineShowFlags),
+			bStrictDragSelection,
+			bGeometryMode
+		};
+
+
 		// If we are using over feedback calculate a new box from the one being dragged
 		FBox SelBBox;
 		CalculateBox( SelBBox );
@@ -92,7 +138,8 @@ void FDragTool_ActorBoxSelect::AddDelta( const FVector& InDelta )
 		for( FActorIterator It(IteratorWorld); It; ++It )
 		{
 			AActor& Actor = **It;
-			const bool bActorHitByBox = IntersectsBox( Actor, SelBBox, bStrictDragSelection );
+			const bool bActorHitByBox = !UE::LevelEditor::Private::GetElementsIntersectingBox( &Actor, SelBBox, LevelViewportClient, SeletionArgs ).IsEmpty();
+
 
 			if( bActorHitByBox )
 			{
@@ -133,7 +180,10 @@ void FDragTool_ActorBoxSelect::EndDrag()
 {
 	UBrushEditingSubsystem* BrushSubsystem = GEditor->GetEditorSubsystem<UBrushEditingSubsystem>();
 	const bool bGeometryMode = BrushSubsystem ? BrushSubsystem->IsGeometryEditorModeActive() : false;
-	
+
+	UTypedElementSelectionSet* SelectionSet = GEditor->GetEditorSubsystem<ULevelEditorSubsystem>()->GetSelectionSet();
+	FTypedElementSelectionOptions ElementSelectionOption;
+
 	FScopedTransaction Transaction( NSLOCTEXT("ActorFrustumSelect", "MarqueeSelectTransation", "Marquee Select" ) );
 
 	bool bShouldSelect = true;
@@ -149,6 +199,7 @@ void FDragTool_ActorBoxSelect::EndDrag()
 	{
 		// If the user is selecting, but isn't hold down SHIFT, remove all current selections.
 		ModeTools->SelectNone();
+		SelectionSet->ClearSelection(ElementSelectionOption);
 	}
 
 	// Let the editor mode try to handle the box selection.
@@ -162,18 +213,17 @@ void FDragTool_ActorBoxSelect::EndDrag()
 	{
 		const bool bStrictDragSelection = GetDefault<ULevelEditorViewportSettings>()->bStrictBoxSelection;
 
-		if( bControlDown )
+		FWorldSelectionElementArgs SeletionArgs
 		{
-			// If control is down remove from selection
-			bShouldSelect = false;
-		}
-		else if( !bShiftDown )
-		{
-			// If the user is selecting, but isn't hold down SHIFT, remove all current selections.
-			GEditor->SelectNone( true, true );
-		}
-
-		// Select all actors that are within the selection box area.  Be aware that certain modes do special processing below.	
+			SelectionSet,
+			ETypedElementSelectionMethod::Primary,
+			FTypedElementSelectionOptions(),
+			&(LevelViewportClient->EngineShowFlags),
+			bStrictDragSelection,
+			bGeometryMode
+		};
+
+		// Select all element that are within the selection box area.  Be aware that certain modes do special processing below.	
 		bool bSelectionChanged = false;
 		UWorld* IteratorWorld = GWorld;
 		const TArray<FName>& HiddenLayers = LevelViewportClient->ViewHiddenLayers;
@@ -192,11 +242,18 @@ void FDragTool_ActorBoxSelect::EndDrag()
 				}
 			}
 
-			// Select the actor if we need to
-			if( bActorIsVisible && IntersectsBox( *Actor, SelBBox, bStrictDragSelection ) )
+			// Select the actor or its child elements
+			if( bActorIsVisible )
 			{
-				GEditor->SelectActor( Actor, bShouldSelect, false );
-				bSelectionChanged = true;
+				TArray<FTypedElementHandle> Handles = UE::LevelEditor::Private::GetElementsIntersectingBox( Actor, SelBBox, LevelViewportClient, SeletionArgs );
+				if (bShouldSelect)
+				{
+					SelectionSet->SelectElements(Handles, ElementSelectionOption);
+				}
+				else
+				{
+					SelectionSet->DeselectElements(Handles, ElementSelectionOption);
+				}
 			}
 		}
 
@@ -281,54 +338,6 @@ void FDragTool_ActorBoxSelect::CalculateBox( FBox& OutBox )
 	}
 }
 
-
-/** 
- * Returns true if the passed in Actor intersects with the provided box 
- *
- * @param InActor				The actor to check
- * @param InBox					The box to check against
- * @param bUseStrictSelection	true if the actor must be entirely within the frustum
- */
-bool FDragTool_ActorBoxSelect::IntersectsBox( AActor& InActor, const FBox& InBox, bool bUseStrictSelection )
-{
-	bool bActorHitByBox = false;
-
-	UBrushEditingSubsystem* BrushSubsystem = GEditor->GetEditorSubsystem<UBrushEditingSubsystem>();
-	const bool bGeometryMode = BrushSubsystem ? BrushSubsystem->IsGeometryEditorModeActive() : false;
-
-	// Check for special cases (like certain show flags that might hide an actor)
-	bool bActorIsHiddenByShowFlags = false;
-
-	// Check to see that volume actors are visible in the viewport
-	if( InActor.IsA(AVolume::StaticClass()) && (!LevelViewportClient->EngineShowFlags.Volumes || !LevelViewportClient->IsVolumeVisibleInViewport(InActor) ) )
-	{
-		bActorIsHiddenByShowFlags = true;
-	}
-
-	// Never drag-select hidden actors or builder brushes. Also, don't consider actors which haven't been recently rendered.
-	//@TODO - replace with proper check for if this object was visible last frame.  This is viewport dependent and viewports can use different concepts of time
-	//depending on if they are in "realtime" mode or not.  See FLevelEditorViewportClient::Draw for the differing concepts of time.
-	const bool bActorRecentlyRendered = true;//Actor->LastRenderTime > ( GWorld->GetTimeSeconds() - 1.0f );
-	if( !bActorIsHiddenByShowFlags && !InActor.IsHiddenEd() && !FActorEditorUtils::IsABuilderBrush(&InActor) && bActorRecentlyRendered )
-	{
-		// Iterate over all actor components, selecting out primitive components
-		for (UActorComponent* Component : InActor.GetComponents())
-		{
-			UPrimitiveComponent* PrimitiveComponent = Cast<UPrimitiveComponent>(Component);
-			if (PrimitiveComponent && PrimitiveComponent->IsRegistered() && PrimitiveComponent->IsVisibleInEditor())
-			{
-				if (PrimitiveComponent->ComponentIsTouchingSelectionBox(InBox, LevelViewportClient->EngineShowFlags, bGeometryMode, bUseStrictSelection))
-				{
-					bActorHitByBox = true;
-					break;
-				}
-			}
-		}
-	}
-
-	return bActorHitByBox;
-}
-
 /** 
  * Returns true if the provided BSP node intersects with the provided frustum 
  *
diff --git a/Engine/Source/Editor/UnrealEd/Private/DragTool_FrustumSelect.cpp b/Engine/Source/Editor/UnrealEd/Private/DragTool_FrustumSelect.cpp
--- a/Engine/Source/Editor/UnrealEd/Private/DragTool_FrustumSelect.cpp
+++ b/Engine/Source/Editor/UnrealEd/Private/DragTool_FrustumSelect.cpp
@@ -2,6 +2,7 @@
 
 
 #include "DragTool_FrustumSelect.h"
+
 #include "Components/PrimitiveComponent.h"
 #include "CanvasItem.h"
 #include "Settings/LevelEditorViewportSettings.h"
@@ -10,6 +11,10 @@
 #include "UnrealEdGlobals.h"
 #include "EngineUtils.h"
 #include "Engine/Selection.h"
+#include "Elements/Framework/EngineElementsLibrary.h"
+#include "Elements/Framework/TypedElementRegistry.h"
+#include "Elements/Framework/TypedElementSelectionSet.h"
+#include "Elements/Interfaces/TypedElementWorldInterface.h"
 #include "EditorModeManager.h"
 #include "EditorModes.h"
 #include "ActorEditorUtils.h"
@@ -17,6 +22,7 @@
 #include "HModel.h"
 #include "CanvasTypes.h"
 #include "Subsystems/BrushEditingSubsystem.h"
+#include "LevelEditorSubsystem.h"
 
 ///////////////////////////////////////////////////////////////////////////////
 //
@@ -24,6 +30,28 @@
 //
 ///////////////////////////////////////////////////////////////////////////////
 
+namespace UE::LevelEditor::Private
+{
+	TArray<FTypedElementHandle> GetElementsIntersectingFrustum(const AActor* Actor,
+		const FConvexVolume& InFrustum,
+		FLevelEditorViewportClient* Viewport,
+		const FWorldSelectionElementArgs& SelectionArgs)
+	{
+		if (Actor && (!Viewport || !Actor->IsA(AVolume::StaticClass()) || !Viewport->IsVolumeVisibleInViewport(*Actor)))
+		{
+			if (FTypedElementHandle ActorHandle = UEngineElementsLibrary::AcquireEditorActorElementHandle(Actor))
+			{
+				if (TTypedElement<ITypedElementWorldInterface> WorldElement = UTypedElementRegistry::GetInstance()->GetElement<ITypedElementWorldInterface>(ActorHandle))
+				{
+					return WorldElement.GetSelectionElementsInConvexVolume(InFrustum, SelectionArgs);
+				}
+			}
+		}
+
+		return {};
+	}
+}
+
 void FDragTool_ActorFrustumSelect::AddDelta( const FVector& InDelta )
 {
 	FIntPoint MousePos;
@@ -66,7 +94,6 @@ void FDragTool_ActorFrustumSelect::EndDrag()
 	FScopedTransaction Transaction( NSLOCTEXT("ActorFrustumSelect", "MarqueeSelectTransation", "Marquee Select" ) );
 
 	bool bShouldSelect = true;
-	GEditor->GetSelectedActors()->Modify();
 
 	if( !bShiftDown )
 	{
@@ -76,7 +103,6 @@ void FDragTool_ActorFrustumSelect::EndDrag()
 
 	// Does an actor have to be fully contained in the box to be selected
 	const bool bStrictDragSelection = GetDefault<ULevelEditorViewportSettings>()->bStrictBoxSelection;
-	bool bSelectionChanged = false;
 
 	// Let the editor mode try to handle the selection.
 	const bool bEditorModeHandledSelection = ModeTools->FrustumSelect(Frustum, LevelViewportClient, bLeftMouseButtonDown);
@@ -86,10 +112,26 @@ void FDragTool_ActorFrustumSelect::EndDrag()
 
 	if( !bEditorModeHandledSelection && !bComponentVisHandledSelection)
 	{
+		UTypedElementSelectionSet* SelectionSet = GEditor->GetEditorSubsystem<ULevelEditorSubsystem>()->GetSelectionSet();
+		SelectionSet->Modify();
+
+		FWorldSelectionElementArgs SeletionArgs
+		{
+			SelectionSet,
+			ETypedElementSelectionMethod::Primary,
+			FTypedElementSelectionOptions(),
+			&(LevelViewportClient->EngineShowFlags),
+			bStrictDragSelection,
+			bGeometryMode
+		};
+
+		 FTypedElementSelectionOptions ElementSelectionOption;
+
+
 		if( !bShiftDown )
 		{
 			// If the user is selecting, but isn't hold down SHIFT, remove all current selections.
-			GEditor->SelectNone( true, true );
+			SelectionSet->ClearSelection(ElementSelectionOption);
 		}
 
 		const int32 ViewportSizeX = LevelViewportClient->Viewport->GetSizeXY().X;
@@ -112,24 +154,51 @@ void FDragTool_ActorFrustumSelect::EndDrag()
 			for(FActorIterator It(LevelViewportClient->GetWorld()); It; ++It)
 			{
 				AActor* Actor = *It;
-				if (IntersectsFrustum( *Actor, Frustum, bStrictDragSelection))
-				{
-					// Select the actor if it intersected
-					const bool bShouldNotifyNow = false;	// We'll notify after everything has been selected
-					GEditor->SelectActor( Actor, bShouldSelect, bShouldNotifyNow );
-					bSelectionChanged = true;
-				}
+				SelectionSet->SelectElements(UE::LevelEditor::Private::GetElementsIntersectingFrustum(Actor, Frustum, LevelViewportClient, SeletionArgs), ElementSelectionOption);
 			}
 		}
 		else
 		{
 			// Extend the endpoint of the rect to get the actual line
 			// LWC_TODO: Large world scale selection needed.
-			FIntRect BoxRect( FIntPoint( FMath::Max<double>( 0.0, Start.X ), FMath::Max<double>( 0.0, Start.Y ) ), FIntPoint( FMath::Min(ViewportSizeX, FMath::TruncToInt(End.X+1)), FMath::Min( ViewportSizeY, FMath::TruncToInt(End.Y+1) ) ) );
+			FIntRect BoxRect(FIntPoint(FMath::Max<double>(0.0, Start.X), FMath::Max<double>(0.0, Start.Y)), FIntPoint(FMath::Min(ViewportSizeX, FMath::TruncToInt(End.X + 1)), FMath::Min(ViewportSizeY, FMath::TruncToInt(End.Y + 1))));
+
+			// Typed Element selection
+			{
+				FTypedElementListRef ElementList = UTypedElementRegistry::GetInstance()->CreateElementList();
+				LevelViewportClient->Viewport->GetElementHandlesInRect(BoxRect, ElementList);
+
+				if (bStrictDragSelection)
+				{
+					ElementList->ForEachElement<ITypedElementWorldInterface>([bStrictDragSelection, &Frustum, &SelectionSet, &ElementSelectionOption]
+						(const TTypedElement<ITypedElementWorldInterface>& InElement)
+						{
+							if (InElement.IsElementInConvexVolume(Frustum, bStrictDragSelection))
+							{
+								SelectionSet->SelectElement(SelectionSet->GetSelectionElement(InElement, ETypedElementSelectionMethod::Primary), ElementSelectionOption);
+							}
+
+							return true;
+						});
+				}
+				else
+				{
+					// Grab only the selectable handles (this remove the components from the selection and select the actor instead)
+					ElementList->ForEachElementHandle([&SelectionSet, &ElementSelectionOption](const FTypedElementHandle& InHandle)
+						{
+							SelectionSet->SelectElement(SelectionSet->GetSelectionElement(InHandle, ETypedElementSelectionMethod::Primary), ElementSelectionOption);
+							return true;
+						});
+				}
+			}
 
-			TSet<AActor*> HitActors;
+			// We need this old code to support the BSP
+			TSet<AActor*> BSPActors;
 			TSet<UModel*> HitModels;
-			LevelViewportClient->Viewport->GetActorsAndModelsInHitProxy( BoxRect, HitActors, HitModels );
+			LevelViewportClient->Viewport->GetActorsAndModelsInHitProxy( BoxRect, BSPActors, HitModels );
+			BSPActors.Empty(HitModels.Num());
+
+			
 
 			if (HitModels.Num() > 0)
 			{
@@ -144,36 +213,28 @@ void FDragTool_ActorFrustumSelect::EndDrag()
 						{
 							uint32 SurfaceIndex = Model.Nodes[NodeIndex].iSurf;
 							FBspSurf& Surf = Model.Surfs[SurfaceIndex];
-							HitActors.Add( Surf.Actor );
+							BSPActors.Add( Surf.Actor );
 						}
 					}
 				}
 			}
 
-			if( HitActors.Num() > 0 )
+			if( BSPActors.Num() > 0 )
 			{
-				for( auto It = HitActors.CreateConstIterator(); It; ++It )
+				for( auto It = BSPActors.CreateConstIterator(); It; ++It )
 				{
 					AActor* Actor = *It;
-					if( bStrictDragSelection && IntersectsFrustum( *Actor, Frustum, bStrictDragSelection ) )
+					if (bStrictDragSelection)
 					{
-						GEditor->SelectActor(Actor, bShouldSelect, false);
-						bSelectionChanged = true;
+						SelectionSet->SelectElements(UE::LevelEditor::Private::GetElementsIntersectingFrustum(Actor, Frustum, LevelViewportClient, SeletionArgs), ElementSelectionOption);
 					}
-					else if( !bStrictDragSelection )
+					else
 					{
-						GEditor->SelectActor(Actor, bShouldSelect, false);
-						bSelectionChanged = true;
+						SelectionSet->SelectElement(UEngineElementsLibrary::AcquireEditorActorElementHandle(Actor), ElementSelectionOption);
 					}
 				}
 			}
 		}
-
-		if(bSelectionChanged)
-		{
-			// If any selections were made.  Notify that now.
-			GEditor->NoteSelectionChange();
-		}
 	}
 
 	// Clear any hovered objects that might have been created while dragging
@@ -189,44 +250,6 @@ void FDragTool_ActorFrustumSelect::Render(const FSceneView* View, FCanvas* Canva
 	Canvas->DrawItem( BoxItem );
 }
 
-bool FDragTool_ActorFrustumSelect::IntersectsFrustum( AActor& InActor, const FConvexVolume& InFrustum, bool bUseStrictSelection ) const
-{	
-	bool bActorHitByBox = false;
-
-	UBrushEditingSubsystem* BrushSubsystem = GEditor->GetEditorSubsystem<UBrushEditingSubsystem>();
-	const bool bGeometryMode = BrushSubsystem ? BrushSubsystem->IsGeometryEditorModeActive() : false;
-
-	// Check for special cases (like certain show flags that might hide an actor)
-	bool bActorIsHiddenByShowFlags = false;
-
-	// Check to see that volume actors are visible in the viewport
-	check(LevelViewportClient != nullptr);
-	if( InActor.IsA(AVolume::StaticClass()) && (!LevelViewportClient->EngineShowFlags.Volumes || !LevelViewportClient->IsVolumeVisibleInViewport(InActor) ) )
-	{
-		bActorIsHiddenByShowFlags = true;
-	}
-
-	// Never drag-select hidden actors or builder brushes. 
-	if( !bActorIsHiddenByShowFlags && !InActor.IsHiddenEd() && !FActorEditorUtils::IsABuilderBrush(&InActor) )
-	{
-		// Iterate over all actor components, selecting out primitive components
-		for (UActorComponent* Component : InActor.GetComponents())
-		{
-			UPrimitiveComponent* PrimitiveComponent = Cast<UPrimitiveComponent>(Component);
-			if (PrimitiveComponent && PrimitiveComponent->IsRegistered() && PrimitiveComponent->IsVisibleInEditor() )
-			{
-				if (PrimitiveComponent->ComponentIsTouchingSelectionFrustum(InFrustum, LevelViewportClient->EngineShowFlags, bGeometryMode, bUseStrictSelection))
-				{
-					bActorHitByBox = true;
-					break;
-				}
-			}
-		}
-	}
-
-	return bActorHitByBox;
-}
-
 bool FDragTool_ActorFrustumSelect::IntersectsFrustum( const UModel& InModel, int32 NodeIndex, const FConvexVolume& InFrustum, bool bUseStrictSelection ) const
 {
 	FBox NodeBB;
diff --git a/Engine/Source/Editor/UnrealEd/Private/Elements/Actor/ActorElementEditorWorldInterface.cpp b/Engine/Source/Editor/UnrealEd/Private/Elements/Actor/ActorElementEditorWorldInterface.cpp
--- a/Engine/Source/Editor/UnrealEd/Private/Elements/Actor/ActorElementEditorWorldInterface.cpp
+++ b/Engine/Source/Editor/UnrealEd/Private/Elements/Actor/ActorElementEditorWorldInterface.cpp
@@ -5,6 +5,8 @@
 #include "GameFramework/Actor.h"
 
 #include "Elements/Framework/EngineElementsLibrary.h"
+#include "Elements/Framework/TypedElementRegistry.h"
+#include "Elements/Interfaces/TypedElementHierarchyInterface.h"
 
 #include "Editor.h"
 #include "UnrealEdGlobals.h"
@@ -111,3 +113,86 @@ void UActorElementEditorWorldInterface::DuplicateElements(TArrayView<const FType
 		}
 	}
 }
+
+bool UActorElementEditorWorldInterface::IsElementInConvexVolume(const FTypedElementHandle& Handle, const FConvexVolume& InVolume, bool bMustEncompassEntireElement)
+{
+	if (TTypedElement<ITypedElementHierarchyInterface> HierarchyElement = GetRegistry().GetElement<ITypedElementHierarchyInterface>(Handle))
+	{
+		TArray<FTypedElementHandle> ChildElementHandles;
+		HierarchyElement.GetChildElements(ChildElementHandles);
+		for (const FTypedElementHandle& ChildHandle : ChildElementHandles)
+		{
+			if (TTypedElement<ITypedElementWorldInterface> ChildWorldElement = GetRegistry().GetElement<ITypedElementWorldInterface>(ChildHandle))
+			{
+				if (ChildWorldElement.IsElementInConvexVolume(InVolume, bMustEncompassEntireElement))
+				{
+					if (!bMustEncompassEntireElement)
+					{
+						return true;
+					}
+				}
+				else if (bMustEncompassEntireElement)
+				{
+					return false;
+				}
+			}
+		}
+	}
+	else
+	{
+		// We expect the actor elements to have a Hierarchy Interface implementation
+		checkNoEntry();
+	}
+
+	return bMustEncompassEntireElement;
+}
+
+bool UActorElementEditorWorldInterface::IsElementInBox(const FTypedElementHandle& Handle, const FBox& InBox, bool bMustEncompassEntireElement)
+{
+	if (TTypedElement<ITypedElementHierarchyInterface> HierarchyElement = GetRegistry().GetElement<ITypedElementHierarchyInterface>(Handle))
+	{
+		TArray<FTypedElementHandle> ChildElementHandles;
+		HierarchyElement.GetChildElements(ChildElementHandles);
+		for (const FTypedElementHandle& ChildHandle : ChildElementHandles)
+		{
+			if (TTypedElement<ITypedElementWorldInterface> ChildWorldElement = GetRegistry().GetElement<ITypedElementWorldInterface>(ChildHandle))
+			{
+				if (ChildWorldElement.IsElementInBox(InBox, bMustEncompassEntireElement))
+				{
+					if (!bMustEncompassEntireElement)
+					{
+						return true;
+					}
+				}
+				else if (bMustEncompassEntireElement)
+				{
+					return false;
+				}
+			}
+		}
+	}
+	else
+	{
+		// We expect the actor elements to have a Hierarchy Interface implementation
+		checkNoEntry();
+	}
+
+	return bMustEncompassEntireElement;
+}
+
+TArray<FTypedElementHandle> UActorElementEditorWorldInterface::GetSelectionElementsFromSelectionFunction(const FTypedElementHandle& InElementHandle, const FWorldSelectionElementArgs& SelectionArgs, const TFunction<bool(const FTypedElementHandle&, const FWorldSelectionElementArgs&)>& SelectionFunction)
+{
+	if (const AActor* Actor = ActorElementDataUtil::GetActorFromHandle(InElementHandle))
+	{ 
+		if (GetOwnerWorld(InElementHandle)->IsEditorWorld())
+		{
+			if (Actor->IsHiddenEd())
+			{
+				return {};
+			}
+		}
+	}
+
+	return  UActorElementWorldInterface::GetSelectionElementsFromSelectionFunction(InElementHandle, SelectionArgs, SelectionFunction);
+}
+
diff --git a/Engine/Source/Editor/UnrealEd/Private/Elements/Component/ComponentElementEditorWorldInterface.cpp b/Engine/Source/Editor/UnrealEd/Private/Elements/Component/ComponentElementEditorWorldInterface.cpp
--- a/Engine/Source/Editor/UnrealEd/Private/Elements/Component/ComponentElementEditorWorldInterface.cpp
+++ b/Engine/Source/Editor/UnrealEd/Private/Elements/Component/ComponentElementEditorWorldInterface.cpp
@@ -5,6 +5,7 @@
 #include "Components/ActorComponent.h"
 
 #include "Elements/Framework/EngineElementsLibrary.h"
+#include "Elements/Framework/TypedElementSelectionSet.h"
 
 #include "Editor.h"
 #include "UnrealEdGlobals.h"
@@ -80,3 +81,51 @@ void UComponentElementEditorWorldInterface::DuplicateElements(TArrayView<const F
 		}
 	}
 }
+
+bool UComponentElementEditorWorldInterface::IsElementInConvexVolume(const FTypedElementHandle& Handle, const FConvexVolume& InVolume, bool bMustEncompassEntireElement)
+{
+	if (const UPrimitiveComponent* Component = Cast<UPrimitiveComponent>(ComponentElementDataUtil::GetComponentFromHandle(Handle)))
+	{
+		return Component->ComponentIsTouchingSelectionFrustum(InVolume, false, bMustEncompassEntireElement);
+	}
+
+	return ITypedElementWorldInterface::IsElementInConvexVolume(Handle, InVolume, bMustEncompassEntireElement);
+}
+
+bool UComponentElementEditorWorldInterface::IsElementInBox(const FTypedElementHandle& Handle, const FBox& InBox, bool bMustEncompassEntireElement)
+{
+	if (const UPrimitiveComponent* Component = Cast<UPrimitiveComponent>(ComponentElementDataUtil::GetComponentFromHandle(Handle)))
+	{
+		return Component->ComponentIsTouchingSelectionBox(InBox, false, bMustEncompassEntireElement);
+	}
+
+	return ITypedElementWorldInterface::IsElementInBox(Handle, InBox, bMustEncompassEntireElement);
+}
+
+TArray<FTypedElementHandle> UComponentElementEditorWorldInterface::GetSelectionElementsFromSelectionFunction(const FTypedElementHandle& InElementHandle, const FWorldSelectionElementArgs& SelectionArgs, const TFunction<bool(const FTypedElementHandle&, const FWorldSelectionElementArgs&)>& SelectionFunction)
+{
+	if (const UPrimitiveComponent* Component = Cast<UPrimitiveComponent>(ComponentElementDataUtil::GetComponentFromHandle(InElementHandle)))
+	{
+		bool bIsVisible = Component->IsRegistered();
+		if (GetOwnerWorld(InElementHandle)->IsEditorWorld())
+		{
+			bIsVisible &= Component->IsVisibleInEditor();
+		}
+		else
+		{
+			bIsVisible &= Component->IsVisible();
+		}
+
+		if (SelectionArgs.ShowFlags)
+		{
+			bIsVisible &= Component->IsShown(*(SelectionArgs.ShowFlags));
+		}
+		
+		if (bIsVisible)
+		{
+			return ITypedElementWorldInterface::GetSelectionElementsFromSelectionFunction(InElementHandle, SelectionArgs, SelectionFunction);
+		}
+	}
+
+	return {};
+}
diff --git a/Engine/Source/Editor/UnrealEd/Private/Elements/SMInstance/SMInstanceElementEditorWorldInterface.cpp b/Engine/Source/Editor/UnrealEd/Private/Elements/SMInstance/SMInstanceElementEditorWorldInterface.cpp
--- a/Engine/Source/Editor/UnrealEd/Private/Elements/SMInstance/SMInstanceElementEditorWorldInterface.cpp
+++ b/Engine/Source/Editor/UnrealEd/Private/Elements/SMInstance/SMInstanceElementEditorWorldInterface.cpp
@@ -216,3 +216,27 @@ FTypedElementHandle USMInstanceElementEditorWorldInterface::PromoteElement(const
 
 	return FTypedElementHandle();
 }
+
+bool USMInstanceElementEditorWorldInterface::IsElementInConvexVolume(const FTypedElementHandle& Handle, const FConvexVolume& InVolume, bool bMustEncompassEntireElement)
+{
+	FSMInstanceManager SMInstance = SMInstanceElementDataUtil::GetSMInstanceFromHandle(Handle);
+
+	if (UInstancedStaticMeshComponent* Component = SMInstance.GetISMComponent())
+	{
+		return Component->IsInstanceTouchingSelectionFrustum(SMInstance.GetISMInstanceIndex(), InVolume, bMustEncompassEntireElement);
+	}
+
+	return false;
+}
+
+bool USMInstanceElementEditorWorldInterface::IsElementInBox(const FTypedElementHandle& Handle, const FBox& InBox, bool bMustEncompassEntireElement)
+{
+	FSMInstanceManager SMInstance = SMInstanceElementDataUtil::GetSMInstanceFromHandle(Handle);
+
+	if (UInstancedStaticMeshComponent* Component = SMInstance.GetISMComponent())
+	{
+		return Component->IsInstanceTouchingSelectionBox(SMInstance.GetISMInstanceIndex(), InBox, bMustEncompassEntireElement);
+	}
+
+	return false;
+}
diff --git a/Engine/Source/Editor/UnrealEd/Public/DragTool_BoxSelect.h b/Engine/Source/Editor/UnrealEd/Public/DragTool_BoxSelect.h
--- a/Engine/Source/Editor/UnrealEd/Public/DragTool_BoxSelect.h
+++ b/Engine/Source/Editor/UnrealEd/Public/DragTool_BoxSelect.h
@@ -51,15 +51,6 @@ private:
 	 */
 	void CalculateBox( FBox& OutBox );
 
-	/** 
-	 * Returns true if the passed in Actor intersects with the provided box 
-	 *
-	 * @param InActor				The actor to check
-	 * @param InBox					The box to check against
-	 * @param bUseStrictSelection	true if the actor must be entirely within the frustum
-	 */
-	bool IntersectsBox( AActor& InActor, const FBox& InBox, bool bUseStrictSelection );
-
 	/** 
 	 * Returns true if the provided BSP node intersects with the provided frustum 
 	 *
diff --git a/Engine/Source/Editor/UnrealEd/Public/DragTool_FrustumSelect.h b/Engine/Source/Editor/UnrealEd/Public/DragTool_FrustumSelect.h
--- a/Engine/Source/Editor/UnrealEd/Public/DragTool_FrustumSelect.h
+++ b/Engine/Source/Editor/UnrealEd/Public/DragTool_FrustumSelect.h
@@ -52,15 +52,6 @@ private:
 	 */
 	void CalculateFrustum( FSceneView* View, FConvexVolume& OutFrustum, bool bUseBoxFrustum );
 
-	/** 
-	 * Returns true if the passed in Actor intersects with the provided frustum 
-	 *
-	 * @param InActor				The actor to check
-	 * @param InFrustum				The frustum to check against
-	 * @param bUseStrictSelection	true if the actor must be entirely within the frustum
-	 */
-	bool IntersectsFrustum( AActor& InActor, const FConvexVolume& InFrustum, bool bUseStrictSelection ) const;
-
 	/** 
 	 * Returns true if the provided BSP node intersects with the provided frustum 
 	 *
diff --git a/Engine/Source/Editor/UnrealEd/Public/Elements/Actor/ActorElementEditorWorldInterface.h b/Engine/Source/Editor/UnrealEd/Public/Elements/Actor/ActorElementEditorWorldInterface.h
--- a/Engine/Source/Editor/UnrealEd/Public/Elements/Actor/ActorElementEditorWorldInterface.h
+++ b/Engine/Source/Editor/UnrealEd/Public/Elements/Actor/ActorElementEditorWorldInterface.h
@@ -20,4 +20,7 @@ public:
 	virtual bool DeleteElements(TArrayView<const FTypedElementHandle> InElementHandles, UWorld* InWorld, UTypedElementSelectionSet* InSelectionSet, const FTypedElementDeletionOptions& InDeletionOptions) override;
 	virtual bool CanDuplicateElement(const FTypedElementHandle& InElementHandle) override;
 	virtual void DuplicateElements(TArrayView<const FTypedElementHandle> InElementHandles, UWorld* InWorld, const FVector& InLocationOffset, TArray<FTypedElementHandle>& OutNewElements) override;
+	virtual bool IsElementInConvexVolume(const FTypedElementHandle& Handle, const FConvexVolume& InVolume, bool bMustEncompassEntireElement = false);
+	virtual bool IsElementInBox(const FTypedElementHandle& Handle, const FBox& InBox, bool bMustEncompassEntireElement = false);
+	virtual TArray<FTypedElementHandle> GetSelectionElementsFromSelectionFunction(const FTypedElementHandle& InElementHandle, const FWorldSelectionElementArgs& SelectionArgs, const TFunction<bool(const FTypedElementHandle&, const FWorldSelectionElementArgs&)>& SelectionFunction) override;
 };
diff --git a/Engine/Source/Editor/UnrealEd/Public/Elements/Component/ComponentElementEditorWorldInterface.h b/Engine/Source/Editor/UnrealEd/Public/Elements/Component/ComponentElementEditorWorldInterface.h
--- a/Engine/Source/Editor/UnrealEd/Public/Elements/Component/ComponentElementEditorWorldInterface.h
+++ b/Engine/Source/Editor/UnrealEd/Public/Elements/Component/ComponentElementEditorWorldInterface.h
@@ -18,4 +18,7 @@ public:
 	virtual bool DeleteElements(TArrayView<const FTypedElementHandle> InElementHandles, UWorld* InWorld, UTypedElementSelectionSet* InSelectionSet, const FTypedElementDeletionOptions& InDeletionOptions) override;
 	virtual bool CanDuplicateElement(const FTypedElementHandle& InElementHandle) override;
 	virtual void DuplicateElements(TArrayView<const FTypedElementHandle> InElementHandles, UWorld* InWorld, const FVector& InLocationOffset, TArray<FTypedElementHandle>& OutNewElements) override;
+	virtual bool IsElementInConvexVolume(const FTypedElementHandle& Handle, const FConvexVolume& InVolume, bool bMustEncompassEntireElement = false);
+	virtual bool IsElementInBox(const FTypedElementHandle& Handle, const FBox& InBox, bool bMustEncompassEntireElement = false);
+	virtual TArray<FTypedElementHandle> GetSelectionElementsFromSelectionFunction(const FTypedElementHandle& InElementHandle, const FWorldSelectionElementArgs& SelectionArgs, const TFunction<bool(const FTypedElementHandle&, const FWorldSelectionElementArgs&)>& SelectionFunction) override;
 };
diff --git a/Engine/Source/Editor/UnrealEd/Public/Elements/SMInstance/SMInstanceElementEditorWorldInterface.h b/Engine/Source/Editor/UnrealEd/Public/Elements/SMInstance/SMInstanceElementEditorWorldInterface.h
--- a/Engine/Source/Editor/UnrealEd/Public/Elements/SMInstance/SMInstanceElementEditorWorldInterface.h
+++ b/Engine/Source/Editor/UnrealEd/Public/Elements/SMInstance/SMInstanceElementEditorWorldInterface.h
@@ -19,4 +19,6 @@ public:
 	virtual void DuplicateElements(TArrayView<const FTypedElementHandle> InElementHandles, UWorld* InWorld, const FVector& InLocationOffset, TArray<FTypedElementHandle>& OutNewElements) override;
 	virtual bool CanPromoteElement(const FTypedElementHandle& InElementHandle) override;
 	virtual FTypedElementHandle PromoteElement(const FTypedElementHandle& InElementHandle, UWorld* OverrideWorld /* = nullptr */) override;
+	virtual bool IsElementInConvexVolume(const FTypedElementHandle& Handle, const FConvexVolume& InVolume, bool bMustEncompassEntireElement = false);
+	virtual bool IsElementInBox(const FTypedElementHandle& Handle, const FBox& InBox, bool bMustEncompassEntireElement = false);
 };
diff --git a/Engine/Source/Runtime/Engine/Classes/Components/ArrowComponent.h b/Engine/Source/Runtime/Engine/Classes/Components/ArrowComponent.h
--- a/Engine/Source/Runtime/Engine/Classes/Components/ArrowComponent.h
+++ b/Engine/Source/Runtime/Engine/Classes/Components/ArrowComponent.h
@@ -68,8 +68,8 @@ class ENGINE_API UArrowComponent : public UPrimitiveComponent
 	//~ Begin UPrimitiveComponent Interface.
 	virtual FPrimitiveSceneProxy* CreateSceneProxy() override;
 #if WITH_EDITOR
-	virtual bool ComponentIsTouchingSelectionBox(const FBox& InSelBBox, const FEngineShowFlags& ShowFlags, const bool bConsiderOnlyBSP, const bool bMustEncompassEntireComponent) const override;
-	virtual bool ComponentIsTouchingSelectionFrustum(const FConvexVolume& InFrustum, const FEngineShowFlags& ShowFlags, const bool bConsiderOnlyBSP, const bool bMustEncompassEntireComponent) const override;
+	virtual bool ComponentIsTouchingSelectionBox(const FBox& InSelBBox, const bool bConsiderOnlyBSP, const bool bMustEncompassEntireComponent) const override;
+	virtual bool ComponentIsTouchingSelectionFrustum(const FConvexVolume& InFrustum, const bool bConsiderOnlyBSP, const bool bMustEncompassEntireComponent) const override;
 #endif
 	//~ End UPrimitiveComponent Interface.
 
diff --git a/Engine/Source/Runtime/Engine/Classes/Components/BillboardComponent.h b/Engine/Source/Runtime/Engine/Classes/Components/BillboardComponent.h
--- a/Engine/Source/Runtime/Engine/Classes/Components/BillboardComponent.h
+++ b/Engine/Source/Runtime/Engine/Classes/Components/BillboardComponent.h
@@ -77,9 +77,11 @@ class ENGINE_API UBillboardComponent : public UPrimitiveComponent
 	//~ Begin UPrimitiveComponent Interface
 	virtual FPrimitiveSceneProxy* CreateSceneProxy() override;
 	virtual FBoxSphereBounds CalcBounds(const FTransform& LocalToWorld) const override;
+
+	virtual bool IsShown(const FEngineShowFlags& ShowFlags) const override;
 #if WITH_EDITOR
-	virtual bool ComponentIsTouchingSelectionBox(const FBox& InSelBBox, const FEngineShowFlags& ShowFlags, const bool bConsiderOnlyBSP, const bool bMustEncompassEntireComponent) const override;
-	virtual bool ComponentIsTouchingSelectionFrustum(const FConvexVolume& InFrustum, const FEngineShowFlags& ShowFlags, const bool bConsiderOnlyBSP, const bool bMustEncompassEntireComponent) const override;
+	virtual bool ComponentIsTouchingSelectionBox(const FBox& InSelBBox, const bool bConsiderOnlyBSP, const bool bMustEncompassEntireComponent) const override;
+	virtual bool ComponentIsTouchingSelectionFrustum(const FConvexVolume& InFrustum, const bool bConsiderOnlyBSP, const bool bMustEncompassEntireComponent) const override;
 #endif
 	//~ End UPrimitiveComponent Interface
 
diff --git a/Engine/Source/Runtime/Engine/Classes/Components/BrushComponent.h b/Engine/Source/Runtime/Engine/Classes/Components/BrushComponent.h
--- a/Engine/Source/Runtime/Engine/Classes/Components/BrushComponent.h
+++ b/Engine/Source/Runtime/Engine/Classes/Components/BrushComponent.h
@@ -52,9 +52,11 @@ public:
 	virtual void GetUsedMaterials( TArray<UMaterialInterface*>& OutMaterials, bool bGetDebugMaterials = false ) const override;
 	virtual ESceneDepthPriorityGroup GetStaticDepthPriorityGroup() const override;
 	virtual bool IsEditorOnly() const override;
+
+	virtual bool IsShown(const FEngineShowFlags& ShowFlags) const override;
 #if WITH_EDITOR
-	virtual bool ComponentIsTouchingSelectionBox(const FBox& InSelBBox, const FEngineShowFlags& ShowFlags, const bool bConsiderOnlyBSP, const bool bMustEncompassEntireComponent) const override;
-	virtual bool ComponentIsTouchingSelectionFrustum(const FConvexVolume& InFrustum, const FEngineShowFlags& ShowFlags, const bool bConsiderOnlyBSP, const bool bMustEncompassEntireComponent) const override;
+	virtual bool ComponentIsTouchingSelectionBox(const FBox& InSelBox, const bool bConsiderOnlyBSP, const bool bMustEncompassEntireComponent) const override;
+	virtual bool ComponentIsTouchingSelectionFrustum(const FConvexVolume& InFrustum, const bool bConsiderOnlyBSP, const bool bMustEncompassEntireComponent) const override;
 	//~ End UPrimitiveComponent Interface.
 
 	/** Return true if the brush appears to have inverted polys */
diff --git a/Engine/Source/Runtime/Engine/Classes/Components/DrawSphereComponent.h b/Engine/Source/Runtime/Engine/Classes/Components/DrawSphereComponent.h
--- a/Engine/Source/Runtime/Engine/Classes/Components/DrawSphereComponent.h
+++ b/Engine/Source/Runtime/Engine/Classes/Components/DrawSphereComponent.h
@@ -21,8 +21,8 @@ class UDrawSphereComponent : public USphereComponent
 
 public:
 #if WITH_EDITOR
-	virtual bool ComponentIsTouchingSelectionBox(const FBox& InSelBBox, const FEngineShowFlags& ShowFlags, const bool bConsiderOnlyBSP, const bool bMustEncompassEntireComponent) const override;
-	virtual bool ComponentIsTouchingSelectionFrustum(const FConvexVolume& InFrustum, const FEngineShowFlags& ShowFlags, const bool bConsiderOnlyBSP, const bool bMustEncompassEntireComponent) const override;
+	virtual bool ComponentIsTouchingSelectionBox(const FBox& InSelBBox, const bool bConsiderOnlyBSP, const bool bMustEncompassEntireComponent) const override;
+	virtual bool ComponentIsTouchingSelectionFrustum(const FConvexVolume& InFrustum, const bool bConsiderOnlyBSP, const bool bMustEncompassEntireComponent) const override;
 #endif
 };
 
diff --git a/Engine/Source/Runtime/Engine/Classes/Components/InstancedStaticMeshComponent.h b/Engine/Source/Runtime/Engine/Classes/Components/InstancedStaticMeshComponent.h
--- a/Engine/Source/Runtime/Engine/Classes/Components/InstancedStaticMeshComponent.h
+++ b/Engine/Source/Runtime/Engine/Classes/Components/InstancedStaticMeshComponent.h
@@ -290,6 +290,14 @@ class ENGINE_API UInstancedStaticMeshComponent : public UStaticMeshComponent, pu
 	virtual void OnRegister() override;
 	virtual bool SupportsRemoveSwap() const { return false; }
 
+#if WITH_EDITOR
+	virtual bool ComponentIsTouchingSelectionBox(const FBox& InSelBBox, const bool bConsiderOnlyBSP, const bool bMustEncompassEntireComponent) const override;
+	virtual bool ComponentIsTouchingSelectionFrustum(const FConvexVolume& InFrustum, const bool bConsiderOnlyBSP, const bool bMustEncompassEntireComponent) const override;
+
+	virtual bool IsInstanceTouchingSelectionBox(int32 InstanceIndex, const FBox& InBox, const bool bMustEncompassEntireInstance) const;
+	virtual bool IsInstanceTouchingSelectionFrustum(int32 InstanceIndex, const FConvexVolume& InFrustum, const bool bMustEncompassEntireInstance) const;
+#endif
+
 public:
 	/** Render data will be initialized on PostLoad or on demand. Released on the rendering thread. */
 	TSharedPtr<FPerInstanceRenderData, ESPMode::ThreadSafe> PerInstanceRenderData;
diff --git a/Engine/Source/Runtime/Engine/Classes/Components/PrimitiveComponent.h b/Engine/Source/Runtime/Engine/Classes/Components/PrimitiveComponent.h
--- a/Engine/Source/Runtime/Engine/Classes/Components/PrimitiveComponent.h
+++ b/Engine/Source/Runtime/Engine/Classes/Components/PrimitiveComponent.h
@@ -2067,6 +2067,13 @@ public:
 	 */
 	virtual FTransform GetComponentTransformFromBodyInstance(FBodyInstance* UseBI);	
 
+	/**
+	 * Would this primitive be shown with these rendering flags.
+	 * 
+	 * @Note: Currently this only implemented properly for the editor selectable primitives.
+	 */
+	virtual bool IsShown(const FEngineShowFlags& ShowFlags) const;
+
 #if WITH_EDITOR
 	/**
 	 * Determines whether the supplied bounding box intersects with the component.
@@ -2079,8 +2086,21 @@ public:
 	 *
 	 * @return	true if the supplied bounding box is determined to intersect the component (partially or wholly)
 	 */
+	UE_DEPRECATED(5.1, "This function is deprecated. Use the function IsShown and the overload that doesn't take an EngineShowFlags instead.")
 	virtual bool ComponentIsTouchingSelectionBox(const FBox& InSelBBox, const FEngineShowFlags& ShowFlags, const bool bConsiderOnlyBSP, const bool bMustEncompassEntireComponent) const;
 
+	/**
+	 * Determines whether the supplied bounding box intersects with the component.
+	 * Used by the editor in orthographic viewports.
+	 *
+	 * @param	InSelBBox						Bounding box to test against
+	 * @param	bConsiderOnlyBSP				If only BSP geometry should be tested
+	 * @param	bMustEncompassEntireComponent	Whether the component bounding box must lay wholly within the supplied bounding box
+	 *
+	 * @return	true if the supplied bounding box is determined to intersect the component (partially or wholly)
+	 */
+	virtual bool ComponentIsTouchingSelectionBox(const FBox& InSelBBox, const bool bConsiderOnlyBSP, const bool bMustEncompassEntireComponent) const;
+
 	/**
 	 * Determines whether the supplied frustum intersects with the component.
 	 * Used by the editor in perspective viewports.
@@ -2092,7 +2112,20 @@ public:
 	 *
 	 * @return	true if the supplied bounding box is determined to intersect the component (partially or wholly)
 	 */
+	UE_DEPRECATED(5.1, "This function is deprecated. Use the function IsShown and the overload that doesn't take an EngineShowFlags instead.")
 	virtual bool ComponentIsTouchingSelectionFrustum(const FConvexVolume& InFrustum, const FEngineShowFlags& ShowFlags, const bool bConsiderOnlyBSP, const bool bMustEncompassEntireComponent) const;
+
+	/**
+	 * Determines whether the supplied frustum intersects with the component.
+	 * Used by the editor in perspective viewports.
+	 *
+	 * @param	InFrustum						Frustum to test against
+	 * @param	bConsiderOnlyBSP				If only BSP geometry should be tested
+	 * @param	bMustEncompassEntireComponent	Whether the component bounding box must lay wholly within the supplied bounding box
+	 *
+	 * @return	true if the supplied bounding box is determined to intersect the component (partially or wholly)
+	 */
+	virtual bool ComponentIsTouchingSelectionFrustum(const FConvexVolume& InFrustum, const bool bConsiderOnlyBSP, const bool bMustEncompassEntireComponent) const;
 #endif
 
 protected:
diff --git a/Engine/Source/Runtime/Engine/Classes/Components/SkeletalMeshComponent.h b/Engine/Source/Runtime/Engine/Classes/Components/SkeletalMeshComponent.h
--- a/Engine/Source/Runtime/Engine/Classes/Components/SkeletalMeshComponent.h
+++ b/Engine/Source/Runtime/Engine/Classes/Components/SkeletalMeshComponent.h
@@ -1820,9 +1820,10 @@ public:
 	UFUNCTION(BlueprintCallable, Category = "Physics")
 	virtual void AddImpulseToAllBodiesBelow(FVector Impulse, FName BoneName = NAME_None, bool bVelChange = false, bool bIncludeSelf = true);
 
+	virtual bool IsShown(const FEngineShowFlags& ShowFlags) const override;
 #if WITH_EDITOR
-	virtual bool ComponentIsTouchingSelectionBox(const FBox& InSelBBox, const FEngineShowFlags& ShowFlags, const bool bConsiderOnlyBSP, const bool bMustEncompassEntireComponent) const override;
-	virtual bool ComponentIsTouchingSelectionFrustum(const FConvexVolume& InFrustum, const FEngineShowFlags& ShowFlags, const bool bConsiderOnlyBSP, const bool bMustEncompassEntireComponent) const override;
+	virtual bool ComponentIsTouchingSelectionBox(const FBox& InSelBBox, const bool bConsiderOnlyBSP, const bool bMustEncompassEntireComponent) const override;
+	virtual bool ComponentIsTouchingSelectionFrustum(const FConvexVolume& InFrustum, const bool bConsiderOnlyBSP, const bool bMustEncompassEntireComponent) const override;
 #endif
 protected:
 	virtual FTransform GetComponentTransformFromBodyInstance(FBodyInstance* UseBI) override;
diff --git a/Engine/Source/Runtime/Engine/Classes/Components/StaticMeshComponent.h b/Engine/Source/Runtime/Engine/Classes/Components/StaticMeshComponent.h
--- a/Engine/Source/Runtime/Engine/Classes/Components/StaticMeshComponent.h
+++ b/Engine/Source/Runtime/Engine/Classes/Components/StaticMeshComponent.h
@@ -567,10 +567,12 @@ public:
 	virtual bool IsMaterialSlotNameValid(FName MaterialSlotName) const override;
 
 	virtual bool DoCustomNavigableGeometryExport(FNavigableGeometryExport& GeomExport) const override;
+
+	virtual bool IsShown(const FEngineShowFlags& ShowFlags) const override;
 #if WITH_EDITOR
 	virtual void PostStaticMeshCompilation();
-	virtual bool ComponentIsTouchingSelectionBox(const FBox& InSelBBox, const FEngineShowFlags& ShowFlags, const bool bConsiderOnlyBSP, const bool bMustEncompassEntireComponent) const override;
-	virtual bool ComponentIsTouchingSelectionFrustum(const FConvexVolume& InSelBBox, const FEngineShowFlags& ShowFlags, const bool bConsiderOnlyBSP, const bool bMustEncompassEntireComponent) const override;
+	virtual bool ComponentIsTouchingSelectionBox(const FBox& InSelBBox, const bool bConsiderOnlyBSP, const bool bMustEncompassEntireComponent) const override;
+	virtual bool ComponentIsTouchingSelectionFrustum(const FConvexVolume& InSelBBox, const bool bConsiderOnlyBSP, const bool bMustEncompassEntireComponent) const override;
 #endif
 	//~ End UPrimitiveComponent Interface.
 
diff --git a/Engine/Source/Runtime/Engine/Private/Components/ArrowComponent.cpp b/Engine/Source/Runtime/Engine/Private/Components/ArrowComponent.cpp
--- a/Engine/Source/Runtime/Engine/Private/Components/ArrowComponent.cpp
+++ b/Engine/Source/Runtime/Engine/Private/Components/ArrowComponent.cpp
@@ -255,13 +255,13 @@ FPrimitiveSceneProxy* UArrowComponent::CreateSceneProxy()
 }
 
 #if WITH_EDITOR
-bool UArrowComponent::ComponentIsTouchingSelectionBox(const FBox& InSelBBox, const FEngineShowFlags& ShowFlags, const bool bConsiderOnlyBSP, const bool bMustEncompassEntireComponent) const
+bool UArrowComponent::ComponentIsTouchingSelectionBox(const FBox& InSelBBox, const bool bConsiderOnlyBSP, const bool bMustEncompassEntireComponent) const
 {
 	// Arrow components not treated as 'selectable' in editor
 	return false;
 }
 
-bool UArrowComponent::ComponentIsTouchingSelectionFrustum(const FConvexVolume& InFrustum, const FEngineShowFlags& ShowFlags, const bool bConsiderOnlyBSP, const bool bMustEncompassEntireComponent) const
+bool UArrowComponent::ComponentIsTouchingSelectionFrustum(const FConvexVolume& InFrustum, const bool bConsiderOnlyBSP, const bool bMustEncompassEntireComponent) const
 {
 	// Arrow components not treated as 'selectable' in editor
 	return false;
diff --git a/Engine/Source/Runtime/Engine/Private/Components/BillboardComponent.cpp b/Engine/Source/Runtime/Engine/Private/Components/BillboardComponent.cpp
--- a/Engine/Source/Runtime/Engine/Private/Components/BillboardComponent.cpp
+++ b/Engine/Source/Runtime/Engine/Private/Components/BillboardComponent.cpp
@@ -333,12 +333,17 @@ FBoxSphereBounds UBillboardComponent::CalcBounds(const FTransform& LocalToWorld)
 	return FBoxSphereBounds(LocalToWorld.GetLocation(),FVector(NewScale,NewScale,NewScale),FMath::Sqrt(3.0f * FMath::Square(NewScale)));
 }
 
+bool UBillboardComponent::IsShown(const FEngineShowFlags& ShowFlags) const
+{
+	return ShowFlags.BillboardSprites;
+}
+
 #if WITH_EDITOR
-bool UBillboardComponent::ComponentIsTouchingSelectionBox(const FBox& InSelBBox, const FEngineShowFlags& ShowFlags, const bool bConsiderOnlyBSP, const bool bMustEncompassEntireComponent) const
+bool UBillboardComponent::ComponentIsTouchingSelectionBox(const FBox& InSelBBox, const bool bConsiderOnlyBSP, const bool bMustEncompassEntireComponent) const
 {
 	AActor* Actor = GetOwner();
 
-	if (!bConsiderOnlyBSP && ShowFlags.BillboardSprites && Sprite != nullptr && Actor != nullptr)
+	if (!bConsiderOnlyBSP && Sprite != nullptr && Actor != nullptr)
 	{
 		const float Scale = GetComponentTransform().GetMaximumAxisScale();
 
@@ -360,11 +365,11 @@ bool UBillboardComponent::ComponentIsTouchingSelectionBox(const FBox& InSelBBox,
 	return false;
 }
 
-bool UBillboardComponent::ComponentIsTouchingSelectionFrustum(const FConvexVolume& InFrustum, const FEngineShowFlags& ShowFlags, const bool bConsiderOnlyBSP, const bool bMustEncompassEntireComponent) const
+bool UBillboardComponent::ComponentIsTouchingSelectionFrustum(const FConvexVolume& InFrustum, const bool bConsiderOnlyBSP, const bool bMustEncompassEntireComponent) const
 {
 	AActor* Actor = GetOwner();
 
-	if (!bConsiderOnlyBSP && ShowFlags.BillboardSprites && Sprite != nullptr && Actor != nullptr)
+	if (!bConsiderOnlyBSP && Sprite != nullptr && Actor != nullptr)
 	{
 		const float Scale = GetComponentTransform().GetMaximumAxisScale();
 		const float MaxExtent = FMath::Max(Sprite->GetSizeX(), Sprite->GetSizeY());
diff --git a/Engine/Source/Runtime/Engine/Private/Components/BrushComponent.cpp b/Engine/Source/Runtime/Engine/Private/Components/BrushComponent.cpp
--- a/Engine/Source/Runtime/Engine/Private/Components/BrushComponent.cpp
+++ b/Engine/Source/Runtime/Engine/Private/Components/BrushComponent.cpp
@@ -548,10 +548,9 @@ ESceneDepthPriorityGroup UBrushComponent::GetStaticDepthPriorityGroup() const
 	}
 }
 
-#if WITH_EDITOR
-static bool IsComponentTypeShown(AActor* Actor, const FEngineShowFlags& ShowFlags)
+bool UBrushComponent::IsShown(const FEngineShowFlags& ShowFlags) const
 {
-	if (Actor != nullptr)
+	if (const AActor* Actor = GetOwner())
 	{
 		return (Actor->IsA(AVolume::StaticClass())) ? ShowFlags.Volumes : ShowFlags.BSP;
 	}
@@ -559,13 +558,9 @@ static bool IsComponentTypeShown(AActor* Actor, const FEngineShowFlags& ShowFlag
 	return false;
 }
 
-bool UBrushComponent::ComponentIsTouchingSelectionBox(const FBox& InSelBBox, const FEngineShowFlags& ShowFlags, const bool bConsiderOnlyBSP, const bool bMustEncompassEntireComponent) const
+#if WITH_EDITOR
+bool UBrushComponent::ComponentIsTouchingSelectionBox(const FBox& InSelBBox, const bool bConsiderOnlyBSP, const bool bMustEncompassEntireComponent) const
 {
-	if (!IsComponentTypeShown(GetOwner(), ShowFlags))
-	{
-		return false;
-	}
-
 	if (Brush != nullptr && Brush->Polys != nullptr)
 	{
 		TArray<FVector> Vertices;
@@ -648,13 +643,8 @@ bool UBrushComponent::ComponentIsTouchingSelectionBox(const FBox& InSelBBox, con
 }
 
 
-bool UBrushComponent::ComponentIsTouchingSelectionFrustum(const FConvexVolume& InFrustum, const FEngineShowFlags& ShowFlags, const bool bConsiderOnlyBSP, const bool bMustEncompassEntireComponent) const
+bool UBrushComponent::ComponentIsTouchingSelectionFrustum(const FConvexVolume& InFrustum, const bool bConsiderOnlyBSP, const bool bMustEncompassEntireComponent) const
 {
-	if (!IsComponentTypeShown(GetOwner(), ShowFlags))
-	{
-		return false;
-	}
-
 	if (Brush != nullptr && Brush->Polys != nullptr)
 	{
 		TArray<FVector> Vertices;
diff --git a/Engine/Source/Runtime/Engine/Private/Components/DrawSphereComponent.cpp b/Engine/Source/Runtime/Engine/Private/Components/DrawSphereComponent.cpp
--- a/Engine/Source/Runtime/Engine/Private/Components/DrawSphereComponent.cpp
+++ b/Engine/Source/Runtime/Engine/Private/Components/DrawSphereComponent.cpp
@@ -14,13 +14,13 @@ UDrawSphereComponent::UDrawSphereComponent(const FObjectInitializer& ObjectIniti
 }
 
 #if WITH_EDITOR
-bool UDrawSphereComponent::ComponentIsTouchingSelectionBox(const FBox& InSelBBox, const FEngineShowFlags& ShowFlags, const bool bConsiderOnlyBSP, const bool bMustEncompassEntireComponent) const
+bool UDrawSphereComponent::ComponentIsTouchingSelectionBox(const FBox& InSelBBox, const bool bConsiderOnlyBSP, const bool bMustEncompassEntireComponent) const
 {
 	// Draw sphere components not treated as 'selectable' in editor
 	return false;
 }
 
-bool UDrawSphereComponent::ComponentIsTouchingSelectionFrustum(const FConvexVolume& InFrustum, const FEngineShowFlags& ShowFlags, const bool bConsiderOnlyBSP, const bool bMustEncompassEntireComponent) const
+bool UDrawSphereComponent::ComponentIsTouchingSelectionFrustum(const FConvexVolume& InFrustum, const bool bConsiderOnlyBSP, const bool bMustEncompassEntireComponent) const
 {
 	// Draw sphere components not treated as 'selectable' in editor
 	return false;
diff --git a/Engine/Source/Runtime/Engine/Private/Components/PrimitiveComponent.cpp b/Engine/Source/Runtime/Engine/Private/Components/PrimitiveComponent.cpp
--- a/Engine/Source/Runtime/Engine/Private/Components/PrimitiveComponent.cpp
+++ b/Engine/Source/Runtime/Engine/Private/Components/PrimitiveComponent.cpp
@@ -2911,8 +2911,18 @@ bool UPrimitiveComponent::GetOverlapsWithActor(const AActor* Actor, TArray<FOver
 	return GetOverlapsWithActor_Template(Actor, OutOverlaps);
 }
 
+bool UPrimitiveComponent::IsShown(const FEngineShowFlags& ShowFlags) const
+{
+	return true;
+}
+
 #if WITH_EDITOR
 bool UPrimitiveComponent::ComponentIsTouchingSelectionBox(const FBox& InSelBBox, const FEngineShowFlags& ShowFlags, const bool bConsiderOnlyBSP, const bool bMustEncompassEntireComponent) const
+{
+	return IsShown(ShowFlags) && ComponentIsTouchingSelectionBox(InSelBBox, bConsiderOnlyBSP, bMustEncompassEntireComponent);
+}
+
+bool UPrimitiveComponent::ComponentIsTouchingSelectionBox(const FBox& InSelBBox, const bool bConsiderOnlyBSP, const bool bMustEncompassEntireComponent) const
 {
 	if (!bConsiderOnlyBSP)
 	{
@@ -2932,6 +2942,11 @@ bool UPrimitiveComponent::ComponentIsTouchingSelectionBox(const FBox& InSelBBox,
 }
 
 bool UPrimitiveComponent::ComponentIsTouchingSelectionFrustum(const FConvexVolume& InFrustum, const FEngineShowFlags& ShowFlags, const bool bConsiderOnlyBSP, const bool bMustEncompassEntireComponent) const
+{
+	return IsShown(ShowFlags) && ComponentIsTouchingSelectionFrustum(InFrustum, bConsiderOnlyBSP, bMustEncompassEntireComponent);
+}
+
+bool UPrimitiveComponent::ComponentIsTouchingSelectionFrustum(const FConvexVolume& InFrustum, const bool bConsiderOnlyBSP, const bool bMustEncompassEntireComponent) const
 {
 	if (!bConsiderOnlyBSP)
 	{
diff --git a/Engine/Source/Runtime/Engine/Private/Components/SkeletalMeshComponent.cpp b/Engine/Source/Runtime/Engine/Private/Components/SkeletalMeshComponent.cpp
--- a/Engine/Source/Runtime/Engine/Private/Components/SkeletalMeshComponent.cpp
+++ b/Engine/Source/Runtime/Engine/Private/Components/SkeletalMeshComponent.cpp
@@ -3643,11 +3643,16 @@ float USkeletalMeshComponent::CalculateMass(FName BoneName)
 	return Mass;
 }
 
+bool USkeletalMeshComponent::IsShown(const FEngineShowFlags& ShowFlags) const
+{
+	return ShowFlags.SkeletalMeshes;
+}
+
 #if WITH_EDITOR
 
-bool USkeletalMeshComponent::ComponentIsTouchingSelectionBox(const FBox& InSelBBox, const FEngineShowFlags& ShowFlags, const bool bConsiderOnlyBSP, const bool bMustEncompassEntireComponent) const
+bool USkeletalMeshComponent::ComponentIsTouchingSelectionBox(const FBox& InSelBBox, const bool bConsiderOnlyBSP, const bool bMustEncompassEntireComponent) const
 {
-	if (!bConsiderOnlyBSP && ShowFlags.SkeletalMeshes && MeshObject != nullptr)
+	if (!bConsiderOnlyBSP && MeshObject != nullptr)
 	{
 		FSkeletalMeshRenderData* SkelMeshRenderData = GetSkeletalMeshRenderData();
 		check(SkelMeshRenderData);
@@ -3687,9 +3692,9 @@ bool USkeletalMeshComponent::ComponentIsTouchingSelectionBox(const FBox& InSelBB
 	return false;
 }
 
-bool USkeletalMeshComponent::ComponentIsTouchingSelectionFrustum(const FConvexVolume& InFrustum, const FEngineShowFlags& ShowFlags, const bool bConsiderOnlyBSP, const bool bMustEncompassEntireComponent) const
+bool USkeletalMeshComponent::ComponentIsTouchingSelectionFrustum(const FConvexVolume& InFrustum, const bool bConsiderOnlyBSP, const bool bMustEncompassEntireComponent) const
 {
-	if (!bConsiderOnlyBSP && ShowFlags.SkeletalMeshes && MeshObject != nullptr)
+	if (!bConsiderOnlyBSP && MeshObject != nullptr)
 	{
 		FSkeletalMeshRenderData* SkelMeshRenderData = GetSkeletalMeshRenderData();
 		check(SkelMeshRenderData);
diff --git a/Engine/Source/Runtime/Engine/Private/Components/StaticMeshComponent.cpp b/Engine/Source/Runtime/Engine/Private/Components/StaticMeshComponent.cpp
--- a/Engine/Source/Runtime/Engine/Private/Components/StaticMeshComponent.cpp
+++ b/Engine/Source/Runtime/Engine/Private/Components/StaticMeshComponent.cpp
@@ -2822,6 +2822,11 @@ void UStaticMeshComponent::GetNavigationData(FNavigationRelevantData& Data) cons
 	}
 }
 
+bool UStaticMeshComponent::IsShown(const FEngineShowFlags& ShowFlags) const
+{
+	return ShowFlags.StaticMeshes;
+}
+
 #if WITH_EDITOR
 
 // UpdateBounds is currently the closest place we can get from the StaticMesh property
@@ -2861,19 +2866,25 @@ void UStaticMeshComponent::PostStaticMeshCompilation()
 	}
 }
 
-bool UStaticMeshComponent::ComponentIsTouchingSelectionBox(const FBox& InSelBBox, const FEngineShowFlags& ShowFlags, const bool bConsiderOnlyBSP, const bool bMustEncompassEntireComponent) const
+bool UStaticMeshComponent::ComponentIsTouchingSelectionBox(const FBox& InSelBBox, const bool bConsiderOnlyBSP, const bool bMustEncompassEntireComponent) const
 {
-	if (!bConsiderOnlyBSP && ShowFlags.StaticMeshes && GetStaticMesh() != nullptr && GetStaticMesh()->HasValidRenderData())
+	if (!bConsiderOnlyBSP && GetStaticMesh() != nullptr && GetStaticMesh()->HasValidRenderData())
 	{
+		// If the bounds are fully contains assume the static mesh component is fully contained
+		if (Super::ComponentIsTouchingSelectionBox(InSelBBox, bConsiderOnlyBSP, true))
+		{
+			return true;
+		}
+
 		// Check if we are even inside it's bounding box, if we are not, there is no way we colliding via the more advanced checks we will do.
-		if (Super::ComponentIsTouchingSelectionBox(InSelBBox, ShowFlags, bConsiderOnlyBSP, false))
+		if (Super::ComponentIsTouchingSelectionBox(InSelBBox, bConsiderOnlyBSP, false))
 		{
 			TArray<FVector> Vertex;
 
 			FStaticMeshLODResources& LODModel = GetStaticMesh()->GetRenderData()->LODResources[0];
 			FIndexArrayView Indices = LODModel.IndexBuffer.GetArrayView();
 
-			for (const auto& Section : LODModel.Sections)
+			for (const FStaticMeshSection& Section : LODModel.Sections)
 			{
 				// Iterate over each triangle.
 				for (int32 TriangleIndex = 0; TriangleIndex < (int32)Section.NumTriangles; TriangleIndex++)
@@ -2913,31 +2924,52 @@ bool UStaticMeshComponent::ComponentIsTouchingSelectionBox(const FBox& InSelBBox
 	return false;
 }
 
-
-bool UStaticMeshComponent::ComponentIsTouchingSelectionFrustum(const FConvexVolume& InFrustum, const FEngineShowFlags& ShowFlags, const bool bConsiderOnlyBSP, const bool bMustEncompassEntireComponent) const
+bool UStaticMeshComponent::ComponentIsTouchingSelectionFrustum(const FConvexVolume& InFrustum, const bool bConsiderOnlyBSP, const bool bMustEncompassEntireComponent) const
 {
-	if (!bConsiderOnlyBSP && ShowFlags.StaticMeshes && GetStaticMesh() != nullptr && GetStaticMesh()->HasValidRenderData())
+	if (!bConsiderOnlyBSP && GetStaticMesh() && GetStaticMesh()->HasValidRenderData())
 	{
 		// Check if we are even inside it's bounding box, if we are not, there is no way we colliding via the more advanced checks we will do.
-		if (Super::ComponentIsTouchingSelectionFrustum(InFrustum, ShowFlags, bConsiderOnlyBSP, false))
+		bool bIsFullyContained = false;
+		if (InFrustum.IntersectBox(Bounds.Origin, Bounds.BoxExtent, bIsFullyContained))
 		{
-			TArray<FVector> Vertex;
+			// If the bounds are fully contains assume the static mesh component is fully contained
+			if (bIsFullyContained)
+			{
+				return true;
+			}
 
 			FStaticMeshLODResources& LODModel = GetStaticMesh()->GetRenderData()->LODResources[0];
+			FIndexArrayView Indices = LODModel.IndexBuffer.GetArrayView();
 
-			uint32 NumVertices = LODModel.VertexBuffers.StaticMeshVertexBuffer.GetNumVertices();
-			for (uint32 VertexIndex = 0; VertexIndex < NumVertices; ++VertexIndex)
+			const FTransform& TransformToWorld = GetComponentTransform();
+			for (const FStaticMeshSection& Section : LODModel.Sections)
 			{
-				const FVector3f& LocalPosition = LODModel.VertexBuffers.PositionVertexBuffer.VertexPosition(VertexIndex);
-				const FVector WorldPosition = GetComponentTransform().TransformPosition(FVector(LocalPosition));
-				bool bLocationIntersected = InFrustum.IntersectSphere(WorldPosition, 0.0f);
-				if (bLocationIntersected && !bMustEncompassEntireComponent)
-				{
-					return true;
-				}
-				else if (!bLocationIntersected && bMustEncompassEntireComponent)
+				// Iterate over each triangle.
+				for (int32 TriangleIndex = 0; TriangleIndex < (int32)Section.NumTriangles; TriangleIndex++)
 				{
-					return false;
+					int32 Index = TriangleIndex * 3 + Section.FirstIndex;
+				
+					FVector PointA(LODModel.VertexBuffers.PositionVertexBuffer.VertexPosition(Indices[Index]));
+					FVector PointB(LODModel.VertexBuffers.PositionVertexBuffer.VertexPosition(Indices[Index + 1]));
+					FVector PointC(LODModel.VertexBuffers.PositionVertexBuffer.VertexPosition(Indices[Index + 2]));
+
+					PointA = TransformToWorld.TransformPosition(PointA);
+					PointB = TransformToWorld.TransformPosition(PointB);
+					PointC = TransformToWorld.TransformPosition(PointC);
+
+					bool bFullyContained = false;
+					bool bIntersect = InFrustum.IntersectTriangle(PointA, PointB, PointC, bFullyContained);
+
+					if (!bMustEncompassEntireComponent && bIntersect)
+					{
+						// Needn't encompass entire component: any intersection, we consider as touching
+						return true;
+					}
+					else if (bMustEncompassEntireComponent && !bFullyContained)
+					{
+						// Must encompass entire component: any non intersection, we consider as not touching
+						return false;
+					}
 				}
 			}
 
diff --git a/Engine/Source/Runtime/Engine/Private/ConvexVolume.cpp b/Engine/Source/Runtime/Engine/Private/ConvexVolume.cpp
--- a/Engine/Source/Runtime/Engine/Private/ConvexVolume.cpp
+++ b/Engine/Source/Runtime/Engine/Private/ConvexVolume.cpp
@@ -405,6 +405,209 @@ bool FConvexVolume::IntersectSphere(const FVector& Origin,const float& Radius, b
 	return true;
 }
 
+//
+//	FConvexVolume::IntersectTriangle
+//
+
+bool FConvexVolume::IntersectTriangle(const FVector& PointA, const FVector& PointB, const FVector& PointC, bool& bOutFullyContained) const
+{
+		checkSlow(PermutedPlanes.Num() % 4 == 0);
+
+		// Assume that it is not fully contained by default
+		bOutFullyContained = false;
+
+		// Load the points
+		VectorRegister A = VectorLoadFloat3(&PointA);
+		VectorRegister B = VectorLoadFloat3(&PointB);
+		VectorRegister C = VectorLoadFloat3(&PointC);
+		VectorRegister Zero = VectorZero();
+
+		// Splat points into 3 vectors
+		VectorRegister AX = VectorReplicate(A, 0);
+		VectorRegister AY = VectorReplicate(A, 1);
+		VectorRegister AZ = VectorReplicate(A, 2);
+
+		VectorRegister BX = VectorReplicate(B, 0);
+		VectorRegister BY = VectorReplicate(B, 1);
+		VectorRegister BZ = VectorReplicate(B, 2);
+
+		VectorRegister CX = VectorReplicate(C, 0);
+		VectorRegister CY = VectorReplicate(C, 1);
+		VectorRegister CZ = VectorReplicate(C, 2);
+
+		// Since we are moving straight through get a pointer to the data
+		const FPlane* RESTRICT PermutedPlanePtr = (FPlane*)PermutedPlanes.GetData();
+
+		// First test if one or all point are inside the volume
+		bool IsAInside = true;
+		bool IsBInside = true;
+		bool IsCInside = true;
+
+		// Process four planes at a time until we have < 4 left
+		for (int32 Count = 0; Count < PermutedPlanes.Num(); Count += 4)
+		{
+			// Load 4 planes that are already all Xs, Ys, ...
+			VectorRegister PlanesX = VectorLoadAligned(PermutedPlanePtr);
+			PermutedPlanePtr++;
+			VectorRegister PlanesY = VectorLoadAligned(PermutedPlanePtr);
+			PermutedPlanePtr++;
+			VectorRegister PlanesZ = VectorLoadAligned(PermutedPlanePtr);
+			PermutedPlanePtr++;
+			VectorRegister PlanesW = VectorLoadAligned(PermutedPlanePtr);
+			PermutedPlanePtr++;
+
+			// Calculate the distance (x * x) + (y * y) + (z * z) - w and use it to determine if the point is ouside
+
+			// Point A
+			VectorRegister DistX = VectorMultiply(AX, PlanesX);
+			VectorRegister DistY = VectorMultiplyAdd(AY, PlanesY, DistX);
+			VectorRegister DistZ = VectorMultiplyAdd(AZ, PlanesZ, DistY);
+			VectorRegister Distance = VectorSubtract(DistZ, PlanesW);
+			int32 MaskA = VectorAnyGreaterThan(Distance, Zero);
+			if (MaskA)
+			{
+				IsAInside = false;
+			}
+
+			DistX = VectorMultiply(BX, PlanesX);
+			DistY = VectorMultiplyAdd(BY, PlanesY, DistX);
+			DistZ = VectorMultiplyAdd(BZ, PlanesZ, DistY);
+			Distance = VectorSubtract(DistZ, PlanesW);
+			int32 MaskB = VectorAnyGreaterThan(Distance, Zero);
+			if (MaskB)
+			{
+				IsBInside = false;
+			}
+
+			DistX = VectorMultiply(CX, PlanesX);
+			DistY = VectorMultiplyAdd(CY, PlanesY, DistX);
+			DistZ = VectorMultiplyAdd(CZ, PlanesZ, DistY);
+			Distance = VectorSubtract(DistZ, PlanesW);
+			int32 MaskC = VectorAnyGreaterThan(Distance, Zero);
+			if (MaskC)
+			{
+				IsCInside = false;
+			}
+
+			// The points were all outside of at lest one plane
+			if (MaskA & MaskB & MaskC)
+			{
+				return false;
+			}
+		}
+
+		// If a point is inside early exit
+		if (IsAInside || IsBInside || IsCInside)
+		{
+			bOutFullyContained = IsAInside && IsBInside && IsCInside;
+			return true;
+		}
+
+		// Clip the triangle against the planes see if it some part are inside the volume.
+
+		// Arbitrary upper bounds
+		const int32 Slack = 16;
+		TArray<FVector> Vertices;
+		Vertices.Reserve(Slack);
+		Vertices.Add(PointA);
+		Vertices.Add(PointB);
+		Vertices.Add(PointC);
+
+		// True for inside the box
+		TArray<bool> Sides;
+		Sides.Reserve(Slack);
+
+		TArray<FVector> NewVertices;
+		NewVertices.Reserve(Slack);
+
+		for (const FPlane& Plane : Planes)
+		{
+			bool bVertexInside = false;
+			bool bVertexOutside = false;
+
+			VectorRegister PlaneRegist = VectorLoadAligned(&Plane);
+
+
+			for (int32 Index = 0; Index < Vertices.Num(); ++Index)
+			{
+				// Calculate the distance of the point and the plane (x * x) + (y * y) + (z * z) - w and use it to determine on which side the point is
+				const float DotProduct = VectorDot3Scalar(VectorLoadFloat3(&Vertices[Index]), PlaneRegist);
+				if (DotProduct - Plane.W <= 0.0)
+				{
+					bVertexInside = true;
+					Sides.Add(true);
+				}
+				else
+				{
+					bVertexOutside = true;
+					Sides.Add(false);
+				}
+			}
+
+			if (!bVertexInside)
+			{
+				return false;
+			}
+			
+			if (!bVertexOutside)
+			{
+				continue;
+			}
+
+
+			const int32 LastValidIndex = Vertices.Num() - 1;
+			const FVector* PreviousVertex = &Vertices[LastValidIndex];
+			bool PreviousIsInside = Sides[LastValidIndex];
+
+			for (int32 Index = 0; Index < Vertices.Num(); ++Index)
+			{
+				bool bIsInside = Sides[Index];
+				const FVector& Vertex = Vertices[Index];
+				if (bIsInside != PreviousIsInside)
+				{
+					// Cross plane
+					VectorRegister Start = VectorLoadFloat3(PreviousVertex);
+					VectorRegister End = VectorLoadFloat3(&Vertex);
+					VectorRegister Line = VectorSubtract(End, Start);
+
+					// Line plane intersection
+					// Start + Line * ((W - Dot(Start, PlaneNormal)) / Dot(Line, PlaneNormal))
+
+					// (W - Dot(Start, PlaneNormal)) / Dot(Line, PlaneNormal)
+					const FVector::FReal Scalar = (Plane.W - VectorDot3Scalar(Start, PlaneRegist)) / VectorDot3Scalar(Line, PlaneRegist);
+
+					// Start + Line * Scalar
+					VectorRegister Intersection = VectorMultiplyAdd(Line, VectorLoadFloat1(&Scalar), Start);
+
+					NewVertices.Emplace(VectorGetComponent(Intersection, 0), VectorGetComponent(Intersection, 1), VectorGetComponent(Intersection, 2));
+
+					if (bIsInside)
+					{
+						NewVertices.Add(Vertex);
+					}
+				}
+				else if (bIsInside)
+				{
+					NewVertices.Add(Vertex);
+				}
+
+				PreviousIsInside = bIsInside;
+				PreviousVertex = &Vertex;
+			}
+
+
+			if (NewVertices.IsEmpty())
+			{
+				return false;
+			}
+
+			Sides.Reset(NewVertices.Num());
+			Swap(Vertices, NewVertices);
+			NewVertices.Reset();
+		}
+
+	return true;
+}
 
 //
 //	FConvexVolume::IntersectLineSegment
diff --git a/Engine/Source/Runtime/Engine/Private/Elements/Actor/ActorElementWorldInterface.cpp b/Engine/Source/Runtime/Engine/Private/Elements/Actor/ActorElementWorldInterface.cpp
--- a/Engine/Source/Runtime/Engine/Private/Elements/Actor/ActorElementWorldInterface.cpp
+++ b/Engine/Source/Runtime/Engine/Private/Elements/Actor/ActorElementWorldInterface.cpp
@@ -6,8 +6,11 @@
 #include "Elements/Component/ComponentElementData.h"
 #include "Elements/Framework/TypedElementRegistry.h"
 
-#include "GameFramework/Actor.h"
+#include "ActorEditorUtils.h"
 #include "Components/PrimitiveComponent.h"
+#include "GameFramework/Actor.h"
+#include "GameFramework/Volume.h"
+#include "ShowFlags.h"
 
 bool UActorElementWorldInterface::IsTemplateElement(const FTypedElementHandle& InElementHandle)
 {
@@ -181,6 +184,30 @@ bool UActorElementWorldInterface::FindSuitableTransformAlongPath(const FTypedEle
 	return false;
 }
 
+TArray<FTypedElementHandle> UActorElementWorldInterface::GetSelectionElementsFromSelectionFunction(const FTypedElementHandle& InElementHandle, const FWorldSelectionElementArgs& SelectionArgs, const TFunction<bool(const FTypedElementHandle&, const FWorldSelectionElementArgs&)>& SelectionFunction)
+{
+	if (const AActor* Actor = ActorElementDataUtil::GetActorFromHandle(InElementHandle))
+	{
+		if ((SelectionArgs.ShowFlags && !SelectionArgs.ShowFlags->Volumes && Actor->IsA(AVolume::StaticClass())) 
+			|| FActorEditorUtils::IsABuilderBrush(Actor))
+		{
+			return {};
+		}
+
+		if (!GetOwnerWorld(InElementHandle)->IsEditorWorld())
+		{
+			if (Actor->IsHidden())
+			{
+				return {};
+			}
+		}
+
+		return ITypedElementWorldInterface::GetSelectionElementsFromSelectionFunction(InElementHandle, SelectionArgs, SelectionFunction);
+	}
+
+	return {};
+}
+
 bool UActorElementWorldInterface::FindSuitableTransformAlongPath_WorldSweep(const UWorld* InWorld, const FVector& InPathStart, const FVector& InPathEnd, const FCollisionShape& InTestShape, TArrayView<const FTypedElementHandle> InElementsToIgnore, FCollisionQueryParams& InOutParams, FTransform& OutSuitableTransform)
 {
 	for (const FTypedElementHandle& ElementToIgnore : InElementsToIgnore)
diff --git a/Engine/Source/Runtime/Engine/Private/Elements/Interfaces/TypedElementWorldInterface.cpp b/Engine/Source/Runtime/Engine/Private/Elements/Interfaces/TypedElementWorldInterface.cpp
--- a/Engine/Source/Runtime/Engine/Private/Elements/Interfaces/TypedElementWorldInterface.cpp
+++ b/Engine/Source/Runtime/Engine/Private/Elements/Interfaces/TypedElementWorldInterface.cpp
@@ -3,9 +3,140 @@
 #include "Elements/Interfaces/TypedElementWorldInterface.h"
 
 #include "Elements/Framework/TypedElementRegistry.h"
+#include "Elements/Framework/TypedElementSelectionSet.h"
+#include "Elements/Interfaces/TypedElementHierarchyInterface.h"
 
+#include "ConvexVolume.h"
+#include "Math/Box.h"
+#include "ShowFlags.h"
 #include "UObject/Stack.h"
 
+bool ITypedElementWorldInterface::IsElementInConvexVolume(const FTypedElementHandle& Handle, const FConvexVolume& InVolume, bool bMustEncompassEntireElement)
+{
+	FBoxSphereBounds Bounds;
+	if (GetBounds(Handle, Bounds))
+	{
+		bool bIsFullyContained = false;
+		if (InVolume.IntersectBox(Bounds.Origin, Bounds.BoxExtent, bIsFullyContained))
+		{
+			return !bMustEncompassEntireElement || bIsFullyContained;
+		}
+	}
+
+	return false;
+}
+
+bool ITypedElementWorldInterface::IsElementInBox(const FTypedElementHandle& Handle, const FBox& InBox, bool bMustEncompassEntireElement)
+{
+	FBoxSphereBounds Bounds;
+	if (GetBounds(Handle, Bounds))
+	{
+		const FBox& ElementBounds = Bounds.GetBox();
+
+		// Check the component bounds versus the selection box
+		// If the selection box must encompass the entire component, then both the min and max vector of the bounds must be inside in the selection
+		// box to be valid. If the selection box only has to touch the component, then it is sufficient to check if it intersects with the bounds.
+		if ((!bMustEncompassEntireElement && InBox.Intersect(ElementBounds))
+			|| (bMustEncompassEntireElement && InBox.IsInside(ElementBounds)))
+		{
+			return true;
+		}
+	}
+
+	return false;
+}
+
+TArray<FTypedElementHandle> ITypedElementWorldInterface::GetSelectionElementsInConvexVolume(const FTypedElementHandle& Handle, const FConvexVolume& InVolume, const FWorldSelectionElementArgs& SelectionArgs)
+{
+	UTypedElementRegistry& Registry = GetRegistry();
+	return GetSelectionElementsFromSelectionFunction(Handle, SelectionArgs, [&InVolume, &Registry](const FTypedElementHandle& ElementHandle, const FWorldSelectionElementArgs& SelectionArgs) -> bool
+	{
+		if (TTypedElement<ITypedElementWorldInterface> WorldElement = Registry.GetElement<ITypedElementWorldInterface>(ElementHandle))
+		{
+			return WorldElement.IsElementInConvexVolume(InVolume, SelectionArgs.bMustEncompassEntireElement);
+		}
+
+		return false;
+	});
+}
+
+TArray<FTypedElementHandle> ITypedElementWorldInterface::GetSelectionElementsInBox(const FTypedElementHandle& Handle, const FBox& InBox, const FWorldSelectionElementArgs& SelectionArgs)
+{
+	UTypedElementRegistry& Registry = GetRegistry();
+	return GetSelectionElementsFromSelectionFunction(Handle, SelectionArgs, [&InBox, &Registry](const FTypedElementHandle& ElementHandle, const FWorldSelectionElementArgs& SelectionArgs) -> bool
+	{
+		if (TTypedElement<ITypedElementWorldInterface> WorldElement = Registry.GetElement<ITypedElementWorldInterface>(ElementHandle))
+		{
+			return WorldElement.IsElementInBox(InBox, SelectionArgs.bMustEncompassEntireElement);
+		}
+
+		return false;
+	});
+}
+
+TArray<FTypedElementHandle> ITypedElementWorldInterface::GetSelectionElementsFromSelectionFunction(const FTypedElementHandle& InElementHandle, const FWorldSelectionElementArgs& SelectionArgs, const TFunction<bool(const FTypedElementHandle& /*ElementHandle*/, const FWorldSelectionElementArgs& /*SelectionArgs*/)>& SelectionFunction)
+{
+	// The default implementation simply iterate the child elements
+
+	TArray<FTypedElementHandle> SelectedSubElements;
+
+	if (TTypedElement<ITypedElementHierarchyInterface> HierarchyElement = GetRegistry().GetElement<ITypedElementHierarchyInterface>(InElementHandle))
+	{
+		TArray<FTypedElementHandle> ChildElementHandles;
+		HierarchyElement.GetChildElements(ChildElementHandles);
+		for (const FTypedElementHandle& ChildHandle : ChildElementHandles)
+		{
+			if (TTypedElement<ITypedElementWorldInterface> ChildWorldElement = GetRegistry().GetElement<ITypedElementWorldInterface>(ChildHandle))
+			{
+				TArray<FTypedElementHandle> ElementsFromChildSelection = ChildWorldElement.GetSelectionElementsFromSelectionFunction(SelectionArgs, SelectionFunction);
+				if (ElementsFromChildSelection.Num() > 1)
+				{
+					SelectedSubElements.Reserve(SelectedSubElements.Num() + ElementsFromChildSelection.Num());
+				}
+
+				for (const FTypedElementHandle& Handle : ElementsFromChildSelection)
+				{
+					if (Handle == InElementHandle)
+					{
+						return { InElementHandle };
+					}
+
+					SelectedSubElements.Add(Handle);
+				}
+			}
+		}
+
+		// If all the elements selected are the child elements, Select the parent instead if possible.
+		if (!ChildElementHandles.IsEmpty() && ChildElementHandles == SelectedSubElements)
+		{
+			if (SelectionArgs.SelectionSet)
+			{
+				SelectedSubElements.Add(SelectionArgs.SelectionSet->GetSelectionElement(InElementHandle, SelectionArgs.SelectionMethod));
+			}
+			else
+			{
+				SelectedSubElements.Add(InElementHandle);
+			}
+		}
+	}
+	else
+	{
+		if (SelectionFunction(InElementHandle, SelectionArgs))
+		{
+			if (SelectionArgs.SelectionSet)
+			{
+				SelectedSubElements.Add(SelectionArgs.SelectionSet->GetSelectionElement(InElementHandle, SelectionArgs.SelectionMethod));
+			}
+			else
+			{
+				SelectedSubElements.Add(InElementHandle);
+			}
+		}
+	}
+
+	return SelectedSubElements;
+}
+
 bool ITypedElementWorldInterface::IsTemplateElement(const FScriptTypedElementHandle& InElementHandle)
 {
 	FTypedElementHandle NativeHandle = InElementHandle.GetTypedElementHandle();
diff --git a/Engine/Source/Runtime/Engine/Private/Elements/SMInstance/SMInstanceElementWorldInterface.cpp b/Engine/Source/Runtime/Engine/Private/Elements/SMInstance/SMInstanceElementWorldInterface.cpp
--- a/Engine/Source/Runtime/Engine/Private/Elements/SMInstance/SMInstanceElementWorldInterface.cpp
+++ b/Engine/Source/Runtime/Engine/Private/Elements/SMInstance/SMInstanceElementWorldInterface.cpp
@@ -9,6 +9,7 @@
 #include "Elements/SMInstance/SMInstanceElementData.h"
 
 #include "Engine/StaticMesh.h"
+#include "ShowFlags.h"
 
 bool USMInstanceElementWorldInterface::CanEditElement(const FTypedElementHandle& InElementHandle)
 {
@@ -117,3 +118,25 @@ void USMInstanceElementWorldInterface::NotifyMovementEnded(const FTypedElementHa
 		SMInstance.NotifySMInstanceMovementEnded();
 	}
 }
+
+TArray<FTypedElementHandle> USMInstanceElementWorldInterface::GetSelectionElementsFromSelectionFunction(const FTypedElementHandle& InElementHandle, const FWorldSelectionElementArgs& SelectionArgs, const TFunction<bool(const FTypedElementHandle&, const FWorldSelectionElementArgs&)>& SelectionFunction)
+{
+	if (SelectionArgs.bBSPSelectionOnly && SelectionArgs.ShowFlags && !SelectionArgs.ShowFlags->StaticMeshes)
+	{
+		return {};
+	}
+
+	if (SelectionFunction(InElementHandle, SelectionArgs))
+	{
+		if (SelectionArgs.SelectionSet)
+		{
+			return {SelectionArgs.SelectionSet->GetSelectionElement(InElementHandle, SelectionArgs.SelectionMethod)};
+		}
+		else
+		{
+			return {InElementHandle};
+		}
+	}
+
+	return {};
+}
diff --git a/Engine/Source/Runtime/Engine/Private/InstancedStaticMesh.cpp b/Engine/Source/Runtime/Engine/Private/InstancedStaticMesh.cpp
--- a/Engine/Source/Runtime/Engine/Private/InstancedStaticMesh.cpp
+++ b/Engine/Source/Runtime/Engine/Private/InstancedStaticMesh.cpp
@@ -26,14 +26,11 @@
 #include "UObject/MobileObjectVersion.h"
 #include "EngineStats.h"
 #include "Interfaces/ITargetPlatform.h"
-#if WITH_EDITOR
-#include "DeviceProfiles/DeviceProfile.h"
-#include "DeviceProfiles/DeviceProfileManager.h"
-#endif // WITH_EDITOR
 #include "MeshMaterialShader.h"
 #include "ProfilingDebugging/LoadTimeTracker.h"
 #include "NaniteSceneProxy.h"
 #include "MaterialCachedData.h"
+#include "Collision.h"
 
 #include "Elements/Framework/EngineElementsLibrary.h"
 #include "Elements/SMInstance/SMInstanceElementData.h"
@@ -4177,6 +4174,193 @@ void UInstancedStaticMeshComponent::OnRegister()
 	}
 }
 
+#if WITH_EDITOR
+
+bool UInstancedStaticMeshComponent::ComponentIsTouchingSelectionBox(const FBox& InSelBBox, const bool bConsiderOnlyBSP, const bool bMustEncompassEntireComponent) const
+{
+	if (bConsiderOnlyBSP)
+	{
+		return false;
+	}
+
+	for (int32 Index = 0; Index < PerInstanceSMData.Num(); ++Index)
+	{
+		if (IsInstanceTouchingSelectionBox(Index, InSelBBox, bMustEncompassEntireComponent) && !bMustEncompassEntireComponent)
+		{
+			return true;
+		}
+		else if (bMustEncompassEntireComponent)
+		{
+			return false;
+		}
+	}
+
+	return bMustEncompassEntireComponent;
+}
+
+bool UInstancedStaticMeshComponent::ComponentIsTouchingSelectionFrustum(const FConvexVolume& InFrustum, const bool bConsiderOnlyBSP, const bool bMustEncompassEntireComponent) const
+{
+	if (bConsiderOnlyBSP)
+	{
+		return false;
+	}
+
+	for (int32 Index = 0; Index < PerInstanceSMData.Num(); ++Index)
+	{
+		if (IsInstanceTouchingSelectionFrustum(Index, InFrustum, bMustEncompassEntireComponent) && !bMustEncompassEntireComponent)
+		{
+			return true;
+		}
+		else if (bMustEncompassEntireComponent)
+		{
+			return false;
+		}
+	}
+
+	return bMustEncompassEntireComponent;
+}
+
+bool UInstancedStaticMeshComponent::IsInstanceTouchingSelectionBox(int32 InstanceIndex, const FBox& InBox, const bool bMustEncompassEntireInstance) const
+{
+	if (PerInstanceSMData.IsValidIndex(InstanceIndex))
+	{
+		const UStaticMesh* Mesh = GetStaticMesh();
+		if (Mesh && Mesh->HasValidRenderData())
+		{
+			const FBoxSphereBounds& MeshBoundingBox = Mesh->GetBounds();
+
+			FTransform InstanceTransfrom;
+			const bool bWorldSpace = true;
+			GetInstanceTransform(InstanceIndex, InstanceTransfrom, bWorldSpace);
+
+			const FBox& InstanceBoundingBox = MeshBoundingBox.TransformBy(InstanceTransfrom).GetBox();
+
+			// If the bounds are fully contains assume the static mesh instance is fully contained
+			if (InstanceBoundingBox.IsInside(InBox))
+			{
+				return true;
+			}
+
+			// Test the bounding box first to avoid heavy computation.
+			if (InBox.Intersect(InstanceBoundingBox))
+			{
+				TArray<FVector> Vertex;
+
+				FStaticMeshLODResources& LODModel = GetStaticMesh()->GetRenderData()->LODResources[0];
+				FIndexArrayView Indices = LODModel.IndexBuffer.GetArrayView();
+
+				for (const FStaticMeshSection& Section : LODModel.Sections)
+				{
+					// Iterate over each triangle.
+					for (int32 TriangleIndex = 0; TriangleIndex < (int32)Section.NumTriangles; TriangleIndex++)
+					{
+						Vertex.Empty(3);
+
+						int32 FirstIndex = TriangleIndex * 3 + Section.FirstIndex;
+						for (int32 i = 0; i < 3; i++)
+						{
+							int32 VertexIndex = Indices[FirstIndex + i];
+							FVector LocalPosition(LODModel.VertexBuffers.PositionVertexBuffer.VertexPosition(VertexIndex));
+							Vertex.Emplace(InstanceTransfrom.TransformPosition(LocalPosition));
+						}
+
+						// Check if the triangle is colliding with the bounding box.
+						FSeparatingAxisPointCheck ThePointCheck(Vertex, InBox.GetCenter(), InBox.GetExtent(), false);
+						if (!bMustEncompassEntireInstance && ThePointCheck.bHit)
+						{
+							// Needn't encompass entire component: any intersection, we consider as touching
+							return true;
+						}
+						else if (bMustEncompassEntireInstance && !ThePointCheck.bHit)
+						{
+							// Must encompass entire component: any non intersection, we consider as not touching
+							return false;
+						}
+					}
+				}
+
+				// Either:
+				// a) It must encompass the entire component and all points were intersected (return true), or;
+				// b) It needn't encompass the entire component but no points were intersected (return false)
+				return bMustEncompassEntireInstance;
+			}
+		}
+	}
+
+	return false;
+}
+
+bool UInstancedStaticMeshComponent::IsInstanceTouchingSelectionFrustum(int32 InstanceIndex, const FConvexVolume& InFrustum, const bool bMustEncompassEntireInstance) const
+{
+	if (PerInstanceSMData.IsValidIndex(InstanceIndex))
+	{
+		const UStaticMesh* Mesh = GetStaticMesh();
+		if (Mesh && Mesh->HasValidRenderData())
+		{
+			const FBoxSphereBounds& MeshBoundingBox = Mesh->GetBounds();
+
+			FTransform InstanceTransfrom;
+			const bool bWorldSpace = true;
+			GetInstanceTransform(InstanceIndex, InstanceTransfrom, bWorldSpace);
+
+			const FBoxSphereBounds& InstanceBoundingBox = MeshBoundingBox.TransformBy(InstanceTransfrom);
+
+			// Test the bounding box first to avoid heavy computation.
+			bool bIsFullyInside = false;
+			if (InFrustum.IntersectBox(InstanceBoundingBox.Origin, InstanceBoundingBox.BoxExtent, bIsFullyInside))
+			{
+				// If the bounds are fully contains assume the static mesh instance is fully contained
+				if (bIsFullyInside)
+				{
+					return true;
+				}
+
+				FStaticMeshLODResources& LODModel = GetStaticMesh()->GetRenderData()->LODResources[0];
+				FIndexArrayView Indices = LODModel.IndexBuffer.GetArrayView();
+
+				for (const FStaticMeshSection& Section : LODModel.Sections)
+				{
+					// Iterate over each triangle.
+					for (int32 TriangleIndex = 0; TriangleIndex < (int32)Section.NumTriangles; TriangleIndex++)
+					{
+						int32 Index = TriangleIndex * 3 + Section.FirstIndex;
+
+						FVector PointA(LODModel.VertexBuffers.PositionVertexBuffer.VertexPosition(Indices[Index]));
+						FVector PointB(LODModel.VertexBuffers.PositionVertexBuffer.VertexPosition(Indices[Index + 1]));
+						FVector PointC(LODModel.VertexBuffers.PositionVertexBuffer.VertexPosition(Indices[Index + 2]));
+
+						PointA = InstanceTransfrom.TransformPosition(PointA);
+						PointB = InstanceTransfrom.TransformPosition(PointB);
+						PointC = InstanceTransfrom.TransformPosition(PointC);
+
+						bool bFullyContained = false;
+						bool bIntersect = InFrustum.IntersectTriangle(PointA, PointB, PointC, bFullyContained);
+
+						if (!bMustEncompassEntireInstance && bIntersect)
+						{
+							// Needn't encompass entire component: any intersection, we consider as touching
+							return true;
+						}
+						else if (bMustEncompassEntireInstance && !bFullyContained)
+						{
+							// Must encompass entire component: any non intersection, we consider as not touching
+							return false;
+						}
+					}
+				}
+
+				// Either:
+				// a) It must encompass the entire component and all points were intersected (return true), or;
+				// b) It needn't encompass the entire component but no points were intersected (return false)
+				return bMustEncompassEntireInstance;
+			}
+		}
+	}
+
+	return false;
+}
+#endif //WITH_EDITOR
+
 void UInstancedStaticMeshComponent::PostLoad()
 {
 	Super::PostLoad();
diff --git a/Engine/Source/Runtime/Engine/Public/ConvexVolume.h b/Engine/Source/Runtime/Engine/Public/ConvexVolume.h
--- a/Engine/Source/Runtime/Engine/Public/ConvexVolume.h
+++ b/Engine/Source/Runtime/Engine/Public/ConvexVolume.h
@@ -122,6 +122,13 @@ public:
      */
 	bool IntersectSphere(const FVector& Origin,const float& Radius, bool& bOutFullyContained) const;
 
+	/**
+     * Intersection test with a triangle
+	 * param bOutFullyContained to know if the triangle was fully contained 
+     * @returns true if this convex volume intersects the given triangle.
+     */
+	bool IntersectTriangle(const FVector& PointA, const FVector& PointB, const FVector& PointC, bool& bOutFullyContained) const;
+
 	/**
      * Intersection test with line segment
      * @param Start of the segment.
diff --git a/Engine/Source/Runtime/Engine/Public/Elements/Actor/ActorElementWorldInterface.h b/Engine/Source/Runtime/Engine/Public/Elements/Actor/ActorElementWorldInterface.h
--- a/Engine/Source/Runtime/Engine/Public/Elements/Actor/ActorElementWorldInterface.h
+++ b/Engine/Source/Runtime/Engine/Public/Elements/Actor/ActorElementWorldInterface.h
@@ -24,7 +24,7 @@ public:
 	virtual bool SetRelativeTransform(const FTypedElementHandle& InElementHandle, const FTransform& InTransform) override;
 	virtual bool FindSuitableTransformAtPoint(const FTypedElementHandle& InElementHandle, const FTransform& InPotentialTransform, FTransform& OutSuitableTransform) override;
 	virtual bool FindSuitableTransformAlongPath(const FTypedElementHandle& InElementHandle, const FVector& InPathStart, const FVector& InPathEnd, const FCollisionShape& InTestShape, TArrayView<const FTypedElementHandle> InElementsToIgnore, FTransform& OutSuitableTransform) override;
-
+	virtual TArray<FTypedElementHandle> GetSelectionElementsFromSelectionFunction(const FTypedElementHandle& InElementHandle, const FWorldSelectionElementArgs& SelectionArgs, const TFunction<bool(const FTypedElementHandle&, const FWorldSelectionElementArgs&)>& SelectionFunction) override;
 	static bool FindSuitableTransformAlongPath_WorldSweep(const UWorld* InWorld, const FVector& InPathStart, const FVector& InPathEnd, const FCollisionShape& InTestShape, TArrayView<const FTypedElementHandle> InElementsToIgnore, FCollisionQueryParams& InOutParams, FTransform& OutSuitableTransform);
 	static void AddIgnoredCollisionQueryElement(const FTypedElementHandle& InElementHandle, FCollisionQueryParams& InOutParams);
 };
diff --git a/Engine/Source/Runtime/Engine/Public/Elements/Interfaces/TypedElementWorldInterface.h b/Engine/Source/Runtime/Engine/Public/Elements/Interfaces/TypedElementWorldInterface.h
--- a/Engine/Source/Runtime/Engine/Public/Elements/Interfaces/TypedElementWorldInterface.h
+++ b/Engine/Source/Runtime/Engine/Public/Elements/Interfaces/TypedElementWorldInterface.h
@@ -3,14 +3,21 @@
 #pragma once
 
 #include "Elements/Framework/TypedElementHandle.h"
+
+#include "Elements/Interfaces/TypedElementSelectionInterface.h"
+
+#include "CoreFwd.h"
 #include "UObject/Interface.h"
 
 #include "TypedElementWorldInterface.generated.h"
 
 class ULevel;
+class UTypedElementSelectionSet; 
 class UWorld;
-class UTypedElementSelectionSet;
 struct FCollisionShape;
+struct FConvexVolume;
+struct FEngineShowFlags;
+struct FWorldSelectionElementArgs;
 
 UENUM()
 enum class ETypedElementWorldType : uint8
@@ -195,7 +202,7 @@ public:
 	/**
 	 * Can the element be promoted
 	 * Generally available when the element is lighter representation of another element.
-	 * Like an instance for example.
+	 * Like promoting a static mesh instance from an instanced static static component to a static actor for example.
 	 */
 	virtual bool CanPromoteElement(const FTypedElementHandle& InElementHandle)
 	{
@@ -205,7 +212,7 @@ public:
 	/**
 	 * Promote a element when possible
 	 * Generally available when the element is lighter representation of another element.
-	 * Like an instance for example.
+	 * Like promoting a static mesh instance from an instanced static static component to a static actor for example.
 	 * 
 	 * @param OverrideWorld Override the world in which the promotion might create new elements. Leave it to null to use the world from the handle.
 	 */
@@ -214,6 +221,35 @@ public:
 		return FTypedElementHandle();
 	}
 
+	/**
+	 * Return true if the element is in the volume
+	 * Note: This should only be use for editor tools since it doesn't use the physics system and the performance would probably be to slow for a runtime application.
+	 */
+	virtual bool IsElementInConvexVolume(const FTypedElementHandle& Handle, const FConvexVolume& InVolume, bool bMustEncompassEntireElement = false);
+
+	/**
+	 * Return true if the element is in the Box
+	 * Note: This should only be use for editor tools since it doesn't use the physics system and the performance would probably be to slow for a runtime application.
+	 */
+	virtual bool IsElementInBox(const FTypedElementHandle& Handle, const FBox& InBox, bool bMustEncompassEntireElement = false);
+
+	/**
+	 * Get the selectable elements from this element that are inside the the convex volume
+	 */
+	TArray<FTypedElementHandle> GetSelectionElementsInConvexVolume(const FTypedElementHandle& Handle, const FConvexVolume& InVolume, const FWorldSelectionElementArgs& SelectionArgs);
+
+	/**
+	 * Get the selectable elements from this element that are inside the the convex volume
+	 */
+	TArray<FTypedElementHandle> GetSelectionElementsInBox(const FTypedElementHandle& Handle, const FBox& InBox, const FWorldSelectionElementArgs& SelectionArgs);
+
+	virtual TArray<FTypedElementHandle> GetSelectionElementsFromSelectionFunction(
+		const FTypedElementHandle& Handle,
+		const FWorldSelectionElementArgs& SelectionArgs,
+		const TFunction<bool(const FTypedElementHandle& /*ElementHandle*/, const FWorldSelectionElementArgs& /*SelectionArgs*/)>& SelectionFunction
+		);
+
+
 	/**
 	 * Script Api
 	 */
@@ -352,7 +388,7 @@ public:
 	UFUNCTION(BlueprintCallable, Category="TypedElementInterfaces|World")
 	virtual FScriptTypedElementHandle PromoteElement(const FScriptTypedElementHandle& InElementHandle, UWorld* OverrideWorld = nullptr);
 
-private:
+protected:
 	/**
 	 * Return the registry associated with this interface implementation
 	 */
@@ -385,4 +421,30 @@ struct TTypedElement<ITypedElementWorldInterface> : public TTypedElementBase<ITy
 	FTypedElementHandle DuplicateElement(UWorld* InWorld, const FVector& InLocationOffset) const { return InterfacePtr->DuplicateElement(*this, InWorld, InLocationOffset); }
 	bool CanPromoteElement() const { return InterfacePtr->CanPromoteElement(*this); }
 	FTypedElementHandle PromoteElement(UWorld* OverrideWorld = nullptr) const { return InterfacePtr->PromoteElement(*this, OverrideWorld); }
+	bool IsElementInConvexVolume(const FConvexVolume& InVolume, bool bMustEncompassEntireElement = false) const { return InterfacePtr->IsElementInConvexVolume(*this, InVolume, bMustEncompassEntireElement); }
+	bool IsElementInBox(const FBox& InBox, bool bMustEncompassEntireElement = false) const { return InterfacePtr->IsElementInBox(*this, InBox, bMustEncompassEntireElement); }
+	TArray<FTypedElementHandle> GetSelectionElementsInConvexVolume(const FConvexVolume& InVolume, const FWorldSelectionElementArgs& SelectionArgs) const { return InterfacePtr->GetSelectionElementsInConvexVolume(*this, InVolume, SelectionArgs); }
+	TArray<FTypedElementHandle> GetSelectionElementsInBox(const FBox& InBox, const FWorldSelectionElementArgs& SelectionArgs) const { return InterfacePtr->GetSelectionElementsInBox(*this, InBox, SelectionArgs); }
+	TArray<FTypedElementHandle> GetSelectionElementsFromSelectionFunction(const FWorldSelectionElementArgs& SelectionArgs, const TFunction<bool(const FTypedElementHandle&, const FWorldSelectionElementArgs&)>& SelectionFunction) const
+	{
+		return InterfacePtr->GetSelectionElementsFromSelectionFunction(*this, SelectionArgs, SelectionFunction);
+	}
+};
+
+struct FWorldSelectionElementArgs
+{
+	// Optional selection set to use to get the selectable elements only
+	const UTypedElementSelectionSet* SelectionSet;
+
+	ETypedElementSelectionMethod SelectionMethod = ETypedElementSelectionMethod::Primary;
+
+	FTypedElementSelectionOptions SelectionOptions;
+
+	// Optional show flag to avoid selecting non rendered elements
+	FEngineShowFlags* ShowFlags;
+
+	bool bMustEncompassEntireElement = false;
+
+	// To deprecate once the BSP tools are removed from the engine
+	bool bBSPSelectionOnly = false;
 };
diff --git a/Engine/Source/Runtime/Engine/Public/Elements/SMInstance/SMInstanceElementWorldInterface.h b/Engine/Source/Runtime/Engine/Public/Elements/SMInstance/SMInstanceElementWorldInterface.h
--- a/Engine/Source/Runtime/Engine/Public/Elements/SMInstance/SMInstanceElementWorldInterface.h
+++ b/Engine/Source/Runtime/Engine/Public/Elements/SMInstance/SMInstanceElementWorldInterface.h
@@ -24,4 +24,5 @@ public:
 	virtual void NotifyMovementStarted(const FTypedElementHandle& InElementHandle) override;
 	virtual void NotifyMovementOngoing(const FTypedElementHandle& InElementHandle) override;
 	virtual void NotifyMovementEnded(const FTypedElementHandle& InElementHandle) override;
+	virtual TArray<FTypedElementHandle> GetSelectionElementsFromSelectionFunction(const FTypedElementHandle& InElementHandle, const FWorldSelectionElementArgs& SelectionArgs, const TFunction<bool(const FTypedElementHandle&, const FWorldSelectionElementArgs&)>& SelectionFunction) override;
 };
diff --git a/Engine/Source/Runtime/Landscape/Classes/LandscapeComponent.h b/Engine/Source/Runtime/Landscape/Classes/LandscapeComponent.h
--- a/Engine/Source/Runtime/Landscape/Classes/LandscapeComponent.h
+++ b/Engine/Source/Runtime/Landscape/Classes/LandscapeComponent.h
@@ -824,12 +824,12 @@ public:
 	bool GetPendingLayerCollisionDataUpdate() const { return bPendingLayerCollisionDataUpdate; }
 #endif 
 
+	virtual bool IsShown(const FEngineShowFlags& ShowFlags) const override;
+
 #if WITH_EDITOR
 	virtual int32 GetNumMaterials() const override;
 	virtual UMaterialInterface* GetMaterial(int32 ElementIndex) const override;
 	virtual void SetMaterial(int32 ElementIndex, UMaterialInterface* Material) override;
-	virtual bool ComponentIsTouchingSelectionBox(const FBox& InSelBBox, const FEngineShowFlags& ShowFlags, const bool bConsiderOnlyBSP, const bool bMustEncompassEntireComponent) const override;
-	virtual bool ComponentIsTouchingSelectionFrustum(const FConvexVolume& InFrustum, const FEngineShowFlags& ShowFlags, const bool bConsiderOnlyBSP, const bool bMustEncompassEntireComponent) const override;
 	virtual void PreFeatureLevelChange(ERHIFeatureLevel::Type PendingFeatureLevel) override;
 #endif
 	//~ End UPrimitiveComponent Interface.
diff --git a/Engine/Source/Runtime/Landscape/Classes/LandscapeHeightfieldCollisionComponent.h b/Engine/Source/Runtime/Landscape/Classes/LandscapeHeightfieldCollisionComponent.h
--- a/Engine/Source/Runtime/Landscape/Classes/LandscapeHeightfieldCollisionComponent.h
+++ b/Engine/Source/Runtime/Landscape/Classes/LandscapeHeightfieldCollisionComponent.h
@@ -206,10 +206,7 @@ public:
 
 	//~ Begin UPrimitiveComponent Interface
 	virtual bool DoCustomNavigableGeometryExport(FNavigableGeometryExport& GeomExport) const override;
-#if WITH_EDITOR
-	virtual bool ComponentIsTouchingSelectionBox(const FBox& InSelBBox, const FEngineShowFlags& ShowFlags, const bool bConsiderOnlyBSP, const bool bMustEncompassEntireComponent) const override;
-	virtual bool ComponentIsTouchingSelectionFrustum(const FConvexVolume& InFrustum, const FEngineShowFlags& ShowFlags, const bool bConsiderOnlyBSP, const bool bMustEncompassEntireComponent) const override;
-#endif
+	virtual bool IsShown(const FEngineShowFlags& ShowFlags) const override;
 	//End UPrimitiveComponent interface
 
 	//~ Begin INavRelevantInterface Interface
diff --git a/Engine/Source/Runtime/Landscape/Private/Landscape.cpp b/Engine/Source/Runtime/Landscape/Private/Landscape.cpp
--- a/Engine/Source/Runtime/Landscape/Private/Landscape.cpp
+++ b/Engine/Source/Runtime/Landscape/Private/Landscape.cpp
@@ -1524,6 +1524,11 @@ FPrimitiveSceneProxy* ULandscapeComponent::CreateSceneProxy()
 	return Proxy;
 }
 
+bool ULandscapeComponent::IsShown(const FEngineShowFlags& ShowFlags) const
+{
+	return ShowFlags.Landscape;
+}
+
 void ULandscapeComponent::DestroyComponent(bool bPromoteChildren/*= false*/)
 {
 	ALandscapeProxy* Proxy = GetLandscapeProxy();
diff --git a/Engine/Source/Runtime/Landscape/Private/LandscapeCollision.cpp b/Engine/Source/Runtime/Landscape/Private/LandscapeCollision.cpp
--- a/Engine/Source/Runtime/Landscape/Private/LandscapeCollision.cpp
+++ b/Engine/Source/Runtime/Landscape/Private/LandscapeCollision.cpp
@@ -2342,27 +2342,12 @@ void ULandscapeHeightfieldCollisionComponent::PostEditUndo()
 		FNavigationSystem::UpdateComponentData(*this);
 	}
 }
+#endif // WITH_EDITOR
 
-bool ULandscapeHeightfieldCollisionComponent::ComponentIsTouchingSelectionBox(const FBox& InSelBBox, const FEngineShowFlags& ShowFlags, const bool bConsiderOnlyBSP, const bool bMustEncompassEntireComponent) const
-{
-	if (ShowFlags.Landscape)
-	{
-		return Super::ComponentIsTouchingSelectionBox(InSelBBox, ShowFlags, bConsiderOnlyBSP, bMustEncompassEntireComponent);
-	}
-
-	return false;
-}
-
-bool ULandscapeHeightfieldCollisionComponent::ComponentIsTouchingSelectionFrustum(const FConvexVolume& InFrustum, const FEngineShowFlags& ShowFlags, const bool bConsiderOnlyBSP, const bool bMustEncompassEntireComponent) const
+bool ULandscapeHeightfieldCollisionComponent::IsShown(const FEngineShowFlags& ShowFlags) const
 {
-	if (ShowFlags.Landscape)
-	{
-		return Super::ComponentIsTouchingSelectionFrustum(InFrustum, ShowFlags, bConsiderOnlyBSP, bMustEncompassEntireComponent);
-	}
-
-	return false;
+	return ShowFlags.Landscape;
 }
-#endif // WITH_EDITOR
 
 bool ULandscapeHeightfieldCollisionComponent::DoCustomNavigableGeometryExport(FNavigableGeometryExport& GeomExport) const
 {
diff --git a/Engine/Source/Runtime/Landscape/Private/LandscapeEdit.cpp b/Engine/Source/Runtime/Landscape/Private/LandscapeEdit.cpp
--- a/Engine/Source/Runtime/Landscape/Private/LandscapeEdit.cpp
+++ b/Engine/Source/Runtime/Landscape/Private/LandscapeEdit.cpp
@@ -623,26 +623,6 @@ void ULandscapeComponent::SetMaterial(int32 ElementIndex, class UMaterialInterfa
 	}
 }
 
-bool ULandscapeComponent::ComponentIsTouchingSelectionBox(const FBox& InSelBBox, const FEngineShowFlags& ShowFlags, const bool bConsiderOnlyBSP, const bool bMustEncompassEntireComponent) const
-{
-	if (ShowFlags.Landscape)
-	{
-		return Super::ComponentIsTouchingSelectionBox(InSelBBox, ShowFlags, bConsiderOnlyBSP, bMustEncompassEntireComponent);
-	}
-
-	return false;
-}
-
-bool ULandscapeComponent::ComponentIsTouchingSelectionFrustum(const FConvexVolume& InFrustum, const FEngineShowFlags& ShowFlags, const bool bConsiderOnlyBSP, const bool bMustEncompassEntireComponent) const
-{
-	if (ShowFlags.Landscape)
-	{
-		return Super::ComponentIsTouchingSelectionFrustum(InFrustum, ShowFlags, bConsiderOnlyBSP, bMustEncompassEntireComponent);
-	}
-
-	return false;
-}
-
 void ULandscapeComponent::PreFeatureLevelChange(ERHIFeatureLevel::Type PendingFeatureLevel)
 {
 	Super::PreFeatureLevelChange(PendingFeatureLevel);
diff --git a/Engine/Source/Runtime/NavigationSystem/Private/NavLinkRenderingComponent.cpp b/Engine/Source/Runtime/NavigationSystem/Private/NavLinkRenderingComponent.cpp
--- a/Engine/Source/Runtime/NavigationSystem/Private/NavLinkRenderingComponent.cpp
+++ b/Engine/Source/Runtime/NavigationSystem/Private/NavLinkRenderingComponent.cpp
@@ -86,13 +86,13 @@ FPrimitiveSceneProxy* UNavLinkRenderingComponent::CreateSceneProxy()
 }
 
 #if WITH_EDITOR
-bool UNavLinkRenderingComponent::ComponentIsTouchingSelectionBox(const FBox& InSelBBox, const FEngineShowFlags& ShowFlags, const bool bConsiderOnlyBSP, const bool bMustEncompassEntireComponent) const
+bool UNavLinkRenderingComponent::ComponentIsTouchingSelectionBox(const FBox& InSelBBox, const bool bConsiderOnlyBSP, const bool bMustEncompassEntireComponent) const
 {
 	// NavLink rendering components not treated as 'selectable' in editor
 	return false;
 }
 
-bool UNavLinkRenderingComponent::ComponentIsTouchingSelectionFrustum(const FConvexVolume& InFrustum, const FEngineShowFlags& ShowFlags, const bool bConsiderOnlyBSP, const bool bMustEncompassEntireComponent) const
+bool UNavLinkRenderingComponent::ComponentIsTouchingSelectionFrustum(const FConvexVolume& InFrustum, const bool bConsiderOnlyBSP, const bool bMustEncompassEntireComponent) const
 {
 	// NavLink rendering components not treated as 'selectable' in editor
 	return false;
diff --git a/Engine/Source/Runtime/NavigationSystem/Public/NavLinkRenderingComponent.h b/Engine/Source/Runtime/NavigationSystem/Public/NavLinkRenderingComponent.h
--- a/Engine/Source/Runtime/NavigationSystem/Public/NavLinkRenderingComponent.h
+++ b/Engine/Source/Runtime/NavigationSystem/Public/NavLinkRenderingComponent.h
@@ -22,8 +22,8 @@ class NAVIGATIONSYSTEM_API UNavLinkRenderingComponent : public UPrimitiveCompone
 	/** Should recreate proxy one very update */
 	virtual bool ShouldRecreateProxyOnUpdateTransform() const override { return true; }
 #if WITH_EDITOR
-	virtual bool ComponentIsTouchingSelectionBox(const FBox& InSelBBox, const FEngineShowFlags& ShowFlags, const bool bConsiderOnlyBSP, const bool bMustEncompassEntireComponent) const override;
-	virtual bool ComponentIsTouchingSelectionFrustum(const FConvexVolume& InFrustum, const FEngineShowFlags& ShowFlags, const bool bConsiderOnlyBSP, const bool bMustEncompassEntireComponent) const override;
+	virtual bool ComponentIsTouchingSelectionBox(const FBox& InSelBBox, const bool bConsiderOnlyBSP, const bool bMustEncompassEntireComponent) const override;
+	virtual bool ComponentIsTouchingSelectionFrustum(const FConvexVolume& InFrustum, const bool bConsiderOnlyBSP, const bool bMustEncompassEntireComponent) const override;
 	virtual bool IgnoreBoundsForEditorFocus() const override { return true; }
 #endif
 	//~ End UPrimitiveComponent Interface
