Optimize staticmesh reimport.
- Remove static meshdescription compact call from the importer.
- Port Get/Set color vertex data to use directly meshdescription and avoid rawmesh conversion round trip.
- Make sure the existing vertex color buffer use FVector3f instead of FVector, it allow the code to use reference instead of copy. For both static and skeletal mesh.

#rnx

diff --git a/Engine/Plugins/Experimental/Interchange/Runtime/Source/Import/Private/Mesh/InterchangeStaticMeshFactory.cpp b/Engine/Plugins/Experimental/Interchange/Runtime/Source/Import/Private/Mesh/InterchangeStaticMeshFactory.cpp
--- a/Engine/Plugins/Experimental/Interchange/Runtime/Source/Import/Private/Mesh/InterchangeStaticMeshFactory.cpp
+++ b/Engine/Plugins/Experimental/Interchange/Runtime/Source/Import/Private/Mesh/InterchangeStaticMeshFactory.cpp
@@ -154,7 +154,7 @@ UObject* UInterchangeStaticMeshFactory::CreateAsset(const FCreateAssetParams& Ar
 	int32 FinalLodCount = LodCount;
 
 	// If we are reimporting, cache the existing vertex colors so they can be optionally reapplied after reimport
-	TMap<FVector, FColor> ExisitingVertexColorData;
+	TMap<FVector3f, FColor> ExisitingVertexColorData;
 	if (Arguments.ReimportObject)
 	{
 		if(ExistingAsset)
@@ -265,7 +265,7 @@ UObject* UInterchangeStaticMeshFactory::CreateAsset(const FCreateAssetParams& Ar
 						//If we have old vertex color (reimport), we want to keep it if the option is ignore
 						if (ExisitingVertexColorData.Num() > 0)
 						{
-							FVector VertexPosition = FVector(LodMeshDescription->GetVertexPosition(LodMeshDescription->GetVertexInstanceVertex(VertexInstanceID)));
+							const FVector3f& VertexPosition = LodMeshDescription->GetVertexPosition(LodMeshDescription->GetVertexInstanceVertex(VertexInstanceID));
 							const FColor* PaintedColor = ExisitingVertexColorData.Find(VertexPosition);
 							if (PaintedColor)
 							{
diff --git a/Engine/Source/Editor/UnrealEd/Private/Fbx/FbxSkeletalMeshImport.cpp b/Engine/Source/Editor/UnrealEd/Private/Fbx/FbxSkeletalMeshImport.cpp
--- a/Engine/Source/Editor/UnrealEd/Private/Fbx/FbxSkeletalMeshImport.cpp
+++ b/Engine/Source/Editor/UnrealEd/Private/Fbx/FbxSkeletalMeshImport.cpp
@@ -1317,7 +1317,7 @@ bool UnFbx::FFbxImporter::ImportBones(TArray<FbxNode*>& NodeArray, FSkeletalMesh
 	return true;
 }
 
-bool UnFbx::FFbxImporter::FillSkeletalMeshImportData(TArray<FbxNode*>& NodeArray, UFbxSkeletalMeshImportData* TemplateImportData, TArray<FbxShape*> *FbxShapeArray, FSkeletalMeshImportData* OutData, TArray<FbxNode*>& OutImportedSkeletonLinkNodes, TArray<FName> &LastImportedMaterialNames, const bool bIsReimport, const TMap<FVector, FColor>& ExistingVertexColorData)
+bool UnFbx::FFbxImporter::FillSkeletalMeshImportData(TArray<FbxNode*>& NodeArray, UFbxSkeletalMeshImportData* TemplateImportData, TArray<FbxShape*> *FbxShapeArray, FSkeletalMeshImportData* OutData, TArray<FbxNode*>& OutImportedSkeletonLinkNodes, TArray<FName> &LastImportedMaterialNames, const bool bIsReimport, const TMap<FVector3f, FColor>& ExistingVertexColorData)
 {
 	if (NodeArray.Num() == 0)
 	{
@@ -1682,7 +1682,7 @@ USkeletalMesh* UnFbx::FFbxImporter::ImportSkeletalMesh(FImportSkeletalMeshArgs &
 		}
 	}
 
-	TMap<FVector, FColor> ExistingVertexColorData;
+	TMap<FVector3f, FColor> ExistingVertexColorData;
 	if (!ExistingSkelMesh)
 	{
 		// When we are not re-importing we want to create the mesh here to be sure there is no material
@@ -3100,7 +3100,7 @@ void UnFbx::FFbxImporter::CleanUpUnusedMaterials(FSkeletalMeshImportData& Import
 	}
 }
 
-bool UnFbx::FFbxImporter::FillSkelMeshImporterFromFbx( FSkeletalMeshImportData& ImportData, FbxMesh*& Mesh, FbxSkin* Skin, FbxShape* FbxShape, TArray<FbxNode*> &SortedLinks, const TArray<FbxSurfaceMaterial*>& FbxMaterials, FbxNode *RootNode, const TMap<FVector, FColor>& ExistingVertexColorData)
+bool UnFbx::FFbxImporter::FillSkelMeshImporterFromFbx( FSkeletalMeshImportData& ImportData, FbxMesh*& Mesh, FbxSkin* Skin, FbxShape* FbxShape, TArray<FbxNode*> &SortedLinks, const TArray<FbxSurfaceMaterial*>& FbxMaterials, FbxNode *RootNode, const TMap<FVector3f, FColor>& ExistingVertexColorData)
 {
 	FbxNode* Node = Mesh->GetNode();
 
@@ -3701,7 +3701,7 @@ bool UnFbx::FFbxImporter::FillSkelMeshImporterFromFbx( FSkeletalMeshImportData&
 		{
 			for (int32 VertexIndex = 0; VertexIndex < 3; VertexIndex++)
 			{
-				const FVector VertexPosition = (FVector)ImportData.Points[TmpWedges[VertexIndex].VertexIndex];
+				const FVector3f& VertexPosition = ImportData.Points[TmpWedges[VertexIndex].VertexIndex];
 				const FColor* PaintedColor = ExistingVertexColorData.Find(VertexPosition);
 				
 				// try to match this wedge current vertex with one that existed in the previous mesh.
diff --git a/Engine/Source/Editor/UnrealEd/Private/Fbx/FbxStaticMeshImport.cpp b/Engine/Source/Editor/UnrealEd/Private/Fbx/FbxStaticMeshImport.cpp
--- a/Engine/Source/Editor/UnrealEd/Private/Fbx/FbxStaticMeshImport.cpp
+++ b/Engine/Source/Editor/UnrealEd/Private/Fbx/FbxStaticMeshImport.cpp
@@ -63,6 +63,7 @@ struct FRestoreReimportData
 
 	FRestoreReimportData(UStaticMesh* Mesh)
 	{
+		TRACE_CPUPROFILER_EVENT_SCOPE(FRestoreReimportData);
 		if (!ensure(Mesh))
 		{
 			return;
@@ -81,6 +82,7 @@ struct FRestoreReimportData
 
 	void RestoreMesh(UnFbx::FFbxImporter* FbxImporter)
 	{
+		TRACE_CPUPROFILER_EVENT_SCOPE(RestoreMesh);
 		if (!ensure(EditorObject && DupObject))
 		{
 			return;
@@ -125,6 +127,7 @@ struct FRestoreReimportData
 
 	void CleanupDuplicateMesh()
 	{
+		TRACE_CPUPROFILER_EVENT_SCOPE(CleanupDuplicateMesh);
 		if(DupObject)
 		{
 			DupObject->RemoveFromRoot();
@@ -369,7 +372,7 @@ float UnFbx::FFbxImporter::GetTriangleAreaThreshold() const
 
 
 bool UnFbx::FFbxImporter::BuildStaticMeshFromGeometry(FbxNode* Node, UStaticMesh* StaticMesh, TArray<FFbxMaterial>& MeshMaterials, int32 LODIndex,
-	EVertexColorImportOption::Type VertexColorImportOption, const TMap<FVector, FColor>& ExistingVertexColorData, const FColor& VertexOverrideColor)
+	EVertexColorImportOption::Type VertexColorImportOption, const TMap<FVector3f, FColor>& ExistingVertexColorData, const FColor& VertexOverrideColor)
 {
 	TRACE_CPUPROFILER_EVENT_SCOPE(FFbxImporter::BuildStaticMeshFromGeometry);
 	FFbxScopedOperation ScopedImportOperation(this);
@@ -384,11 +387,6 @@ bool UnFbx::FFbxImporter::BuildStaticMeshFromGeometry(FbxNode* Node, UStaticMesh
 	check(MeshDescription);
 	FStaticMeshAttributes Attributes(*MeshDescription);
 
-	//This slow task doesn't have any text description and is used to unify all the other sub-tasks of this function.
-	FScopedSlowTask BuildStaticMeshSlowTask(Mesh->IsTriangleMesh() ? 1 : 2);
-	BuildStaticMeshSlowTask.MakeDialog();
-	BuildStaticMeshSlowTask.EnterProgressFrame(1);
-
 	//Get the base layer of the mesh
 	FbxLayer* BaseLayer = Mesh->GetLayer(0);
 	if (BaseLayer == NULL)
@@ -487,13 +485,6 @@ bool UnFbx::FFbxImporter::BuildStaticMeshFromGeometry(FbxNode* Node, UStaticMesh
 			AddTokenizedErrorMessage(FTokenizedMessage::Create(EMessageSeverity::Warning, FText::Format(LOCTEXT("Error_FailedToTriangulateWithOmission", "Unable to triangulate mesh '{0}': it will be omitted."), FText::FromString(FbxNodeName))), FFbxErrors::Generic_Mesh_TriangulationFailed);
 			return false;
 		}
-		
-		BuildStaticMeshSlowTask.EnterProgressFrame(1);
-		if (ImportOptions->bIsImportCancelable && BuildStaticMeshSlowTask.ShouldCancel())
-		{ 
-			bImportOperationCanceled = true; 
-			return false; 
-		}
 	}
 	
 	// renew the base layer
@@ -778,7 +769,7 @@ bool UnFbx::FFbxImporter::BuildStaticMeshFromGeometry(FbxNode* Node, UStaticMesh
 
 				MeshDescription->ReserveNewPolygons(PolygonCount);
 				MeshDescription->ReserveNewVertexInstances(TotalVertexCount);
-				MeshDescription->ReserveNewEdges(TotalVertexCount);
+				MeshDescription->ReserveNewEdges(Mesh->GetMeshEdgeCount());
 			}
 
 			bool  bBeginGetMeshEdgeIndexForPolygonCalled   = false;
@@ -853,7 +844,7 @@ bool UnFbx::FFbxImporter::BuildStaticMeshFromGeometry(FbxNode* Node, UStaticMesh
 					CornerInstanceIDs[CornerIndex] = VertexInstanceID;
 					const int32 ControlPointIndex = Mesh->GetPolygonVertex(PolygonIndex, CornerIndex);
 					const FVertexID VertexID(VertexOffset + ControlPointIndex);
-					const FVector VertexPosition = (FVector)VertexPositions[VertexID];
+					const FVector3f& VertexPosition = VertexPositions[VertexID];
 					CornerVerticesIDs[CornerIndex] = VertexID;
 
 					FVertexInstanceID AddedVertexInstanceId = MeshDescription->CreateVertexInstance(VertexID);
@@ -1176,11 +1167,7 @@ bool UnFbx::FFbxImporter::BuildStaticMeshFromGeometry(FbxNode* Node, UStaticMesh
 
 			if (SkippedVertexInstance > 0)
 			{
-				//We must compact the sparse array before reserving new space
-				//When we reserve it will make a hole in the sparse array if the last reserve was not fully use
-				//The importer assume there will be no hole when importing a mesh
-				FElementIDRemappings OutRemappings;
-				MeshDescription->Compact(OutRemappings);
+				check(MeshDescription->Triangles().Num() == MeshDescription->Triangles().GetArraySize());
 			}
 		}
 	}
@@ -1681,7 +1668,7 @@ UStaticMesh* UnFbx::FFbxImporter::ImportStaticMeshAsSingle(UObject* InParent, TA
 	UObject* ExistingObject = NULL;
 
 	// A mapping of vertex positions to their color in the existing static mesh
-	TMap<FVector, FColor>		ExistingVertexColorData;
+	TMap<FVector3f, FColor>		ExistingVertexColorData;
 
 	EVertexColorImportOption::Type VertexColorImportOption = ImportOptions->VertexColorImportOption;
 	FString NewPackageName;
@@ -1798,10 +1785,12 @@ UStaticMesh* UnFbx::FFbxImporter::ImportStaticMeshAsSingle(UObject* InParent, TA
 	bool bAllDegenerated = true;
 	TArray<FFbxMaterial> MeshMaterials;
 	int32 NodeFailCount = 0;
+	
+	FScopedSlowTask BuildMeshSlowTask(MeshNodeArray.Num(), FText::Format(LOCTEXT("FbxStaticMeshBuildGeometryTask", "Importing Static Mesh Geometries {0} of {1}."), FText::AsNumber(0), FText::AsNumber(MeshNodeArray.Num())));
+	BuildMeshSlowTask.MakeDialog();
 	for (int32 MeshIndex = 0; MeshIndex < MeshNodeArray.Num(); MeshIndex++)
 	{
-		FScopedSlowTask BuildMeshSlowTask(MeshNodeArray.Num(), FText::Format(LOCTEXT("FbxStaticMeshBuildGeometryTask", "Importing Static Mesh Geometry {0} of {1}."), FText::AsNumber(1), FText::AsNumber(MeshNodeArray.Num())));
-		
+		BuildMeshSlowTask.EnterProgressFrame(1, FText::Format(LOCTEXT("FbxStaticMeshBuildGeometryTask", "Importing Static Mesh Geometry {0} of {1}."), FText::AsNumber(MeshIndex+1), FText::AsNumber(MeshNodeArray.Num())));
 		FbxNode* Node = MeshNodeArray[MeshIndex];
 		if (Node->GetMesh())
 		{
@@ -1817,8 +1806,6 @@ UStaticMesh* UnFbx::FFbxImporter::ImportStaticMeshAsSingle(UObject* InParent, TA
 			{
 				bAllDegenerated = false;
 
-				//Only cancel the operation if we are creating a new asset, as we don't support existing asset restoration yet.
-				BuildMeshSlowTask.EnterProgressFrame(1, FText::Format(LOCTEXT("FbxStaticMeshBuildGeometryTask", "Importing Static Mesh Geometry {0} of {1}."), FText::AsNumber(MeshIndex + 1), FText::AsNumber(MeshNodeArray.Num())));
 				if(ImportOptions->bIsImportCancelable && BuildMeshSlowTask.ShouldCancel())
 				{
 					bImportOperationCanceled = true; 
diff --git a/Engine/Source/Editor/UnrealEd/Public/FbxImporter.h b/Engine/Source/Editor/UnrealEd/Public/FbxImporter.h
--- a/Engine/Source/Editor/UnrealEd/Public/FbxImporter.h
+++ b/Engine/Source/Editor/UnrealEd/Public/FbxImporter.h
@@ -1416,7 +1416,7 @@ protected:
 	 * @return bool true if set up successfully
 	 */
 	bool BuildStaticMeshFromGeometry(FbxNode* Node, UStaticMesh* StaticMesh, TArray<FFbxMaterial>& MeshMaterials, int LODIndex,
-									 EVertexColorImportOption::Type VertexColorImportOption, const TMap<FVector, FColor>& ExistingVertexColorData, const FColor& VertexOverrideColor);
+									 EVertexColorImportOption::Type VertexColorImportOption, const TMap<FVector3f, FColor>& ExistingVertexColorData, const FColor& VertexOverrideColor);
 	
 	/**
 	 * Clean up for destroy the Importer.
@@ -1491,7 +1491,7 @@ protected:
 	* @returns bool*	true if import successfully.
 	*/
     bool FillSkelMeshImporterFromFbx(FSkeletalMeshImportData& ImportData, FbxMesh*& Mesh, FbxSkin* Skin, 
-										FbxShape* Shape, TArray<FbxNode*> &SortedLinks, const TArray<FbxSurfaceMaterial*>& FbxMaterials, FbxNode *RootNode, const TMap<FVector, FColor>& ExistingVertexColorData);
+										FbxShape* Shape, TArray<FbxNode*> &SortedLinks, const TArray<FbxSurfaceMaterial*>& FbxMaterials, FbxNode *RootNode, const TMap<FVector3f, FColor>& ExistingVertexColorData);
 public:
 
 	/**
@@ -1507,7 +1507,7 @@ public:
 	*/
 	bool FillSkeletalMeshImportData(TArray<FbxNode*>& NodeArray, UFbxSkeletalMeshImportData* TemplateImportData, TArray<FbxShape*> *FbxShapeArray,
 									FSkeletalMeshImportData* OutData, TArray<FbxNode*>& OutImportedSkeletonLinkNodes, TArray<FName> &LastImportedMaterialNames, 
-									const bool bIsReimport, const TMap<FVector, FColor>& ExistingVertexColorData);
+									const bool bIsReimport, const TMap<FVector3f, FColor>& ExistingVertexColorData);
 
 protected:
 
diff --git a/Engine/Source/Runtime/Engine/Classes/Engine/SkeletalMesh.h b/Engine/Source/Runtime/Engine/Classes/Engine/SkeletalMesh.h
--- a/Engine/Source/Runtime/Engine/Classes/Engine/SkeletalMesh.h
+++ b/Engine/Source/Runtime/Engine/Classes/Engine/SkeletalMesh.h
@@ -2836,7 +2836,7 @@ public:
 	 *
 	 *	@return	VertexColorData		Returns a map of vertex position and their associated color.
 	 */
-	TMap<FVector, FColor> GetVertexColorData(const uint32 PaintingMeshLODIndex = 0) const;
+	TMap<FVector3f, FColor> GetVertexColorData(const uint32 PaintingMeshLODIndex = 0) const;
 
 	/** Called to rebuild an out-of-date or invalid socket map */
 	void RebuildSocketMap();
diff --git a/Engine/Source/Runtime/Engine/Classes/Engine/StaticMesh.h b/Engine/Source/Runtime/Engine/Classes/Engine/StaticMesh.h
--- a/Engine/Source/Runtime/Engine/Classes/Engine/StaticMesh.h
+++ b/Engine/Source/Runtime/Engine/Classes/Engine/StaticMesh.h
@@ -1811,7 +1811,7 @@ public:
 	 *
 	 *	@param	VertexColorData		(out)A map of vertex position data and its color. The method fills this map.
 	 */
-	ENGINE_API void GetVertexColorData(TMap<FVector, FColor>& VertexColorData);
+	ENGINE_API void GetVertexColorData(TMap<FVector3f, FColor>& VertexColorData);
 
 	/**
 	 * Sets vertex color data by position.
@@ -1820,7 +1820,7 @@ public:
 	 *
 	 *	@param	VertexColorData		A map of vertex position data and color.
 	 */
-	ENGINE_API void SetVertexColorData(const TMap<FVector, FColor>& VertexColorData);
+	ENGINE_API void SetVertexColorData(const TMap<FVector3f, FColor>& VertexColorData);
 
 	/** Removes all vertex colors from this mesh and rebuilds it (Editor only */
 	ENGINE_API void RemoveVertexColors();
diff --git a/Engine/Source/Runtime/Engine/Private/SkeletalMesh.cpp b/Engine/Source/Runtime/Engine/Private/SkeletalMesh.cpp
--- a/Engine/Source/Runtime/Engine/Private/SkeletalMesh.cpp
+++ b/Engine/Source/Runtime/Engine/Private/SkeletalMesh.cpp
@@ -4134,9 +4134,9 @@ USkeletalMeshSocket* USkeletalMesh::GetSocketByIndex(int32 Index) const
 	return nullptr;
 }
 
-TMap<FVector, FColor> USkeletalMesh::GetVertexColorData(const uint32 PaintingMeshLODIndex) const
+TMap<FVector3f, FColor> USkeletalMesh::GetVertexColorData(const uint32 PaintingMeshLODIndex) const
 {
-	TMap<FVector, FColor> VertexColorData;
+	TMap<FVector3f, FColor> VertexColorData;
 #if WITH_EDITOR
 	const FSkeletalMeshModel* SkeletalMeshModel = GetImportedModel();
 	if (GetHasVertexColors() && SkeletalMeshModel && SkeletalMeshModel->LODModels.IsValidIndex(PaintingMeshLODIndex))
@@ -4149,7 +4149,7 @@ TMap<FVector, FColor> USkeletalMesh::GetVertexColorData(const uint32 PaintingMes
 			
 			for (int32 VertexIndex = 0; VertexIndex < SoftVertices.Num(); ++VertexIndex)
 			{
-				FVector Position(SoftVertices[VertexIndex].Position);
+				const FVector3f& Position(SoftVertices[VertexIndex].Position);
 				FColor& Color = VertexColorData.FindOrAdd(Position);
 				Color = SoftVertices[VertexIndex].Color;
 			}
diff --git a/Engine/Source/Runtime/Engine/Private/StaticMesh.cpp b/Engine/Source/Runtime/Engine/Private/StaticMesh.cpp
--- a/Engine/Source/Runtime/Engine/Private/StaticMesh.cpp
+++ b/Engine/Source/Runtime/Engine/Private/StaticMesh.cpp
@@ -6843,8 +6843,9 @@ void UStaticMesh::MarkAsNotHavingNavigationData()
  *
  *	@param	VertexColorData		(out)A map of vertex position data and its color. The method fills this map.
  */
-void UStaticMesh::GetVertexColorData(TMap<FVector, FColor>& VertexColorData)
+void UStaticMesh::GetVertexColorData(TMap<FVector3f, FColor>& VertexColorData)
 {
+	TRACE_CPUPROFILER_EVENT_SCOPE(UStaticMesh::GetVertexColorData);
 	VertexColorData.Empty();
 #if WITH_EDITOR
 	// What LOD to get vertex colors from.  
@@ -6854,21 +6855,18 @@ void UStaticMesh::GetVertexColorData(TMap<FVector, FColor>& VertexColorData)
 	{
 		if (!GetSourceModel(PaintingMeshLODIndex).IsRawMeshEmpty())
 		{
-			// Extract the raw mesh.
-			FRawMesh Mesh;
-			GetSourceModel(PaintingMeshLODIndex).LoadRawMesh(Mesh);
-			// Nothing to copy if there are no colors stored.
-			if (Mesh.WedgeColors.Num() != 0 && Mesh.WedgeColors.Num() == Mesh.WedgeIndices.Num())
+			FMeshDescription* MeshDescription = GetMeshDescription(PaintingMeshLODIndex);
+			VertexColorData.Reserve(MeshDescription->Vertices().Num());
+			TVertexAttributesConstRef<FVector3f> Positions = FStaticMeshConstAttributes(*MeshDescription).GetVertexPositions();
+			TVertexInstanceAttributesConstRef<FVector4f> Colors = FStaticMeshConstAttributes(*MeshDescription).GetVertexInstanceColors();
+			for(FVertexInstanceID VertexInstanceID : MeshDescription->VertexInstances().GetElementIDs())
 			{
-				// Build a mapping of vertex positions to vertex colors.
-				for (int32 WedgeIndex = 0; WedgeIndex < Mesh.WedgeIndices.Num(); ++WedgeIndex)
+				FVertexID VertexID = MeshDescription->GetVertexInstanceVertex(VertexInstanceID);
+				FColor Color = FLinearColor(Colors[VertexInstanceID]).ToFColorSRGB();
+				const FVector3f& Position = Positions[VertexID];
+				if (!VertexColorData.Contains(Position))
 				{
-					FVector Position(Mesh.VertexPositions[Mesh.WedgeIndices[WedgeIndex]]);
-					FColor Color = Mesh.WedgeColors[WedgeIndex];
-					if (!VertexColorData.Contains(Position))
-					{
-						VertexColorData.Add(Position, Color);
-					}
+					VertexColorData.Add(Position, Color);
 				}
 			}
 		}
@@ -6883,8 +6881,9 @@ void UStaticMesh::GetVertexColorData(TMap<FVector, FColor>& VertexColorData)
  *
  *	@param	VertexColorData		A map of vertex position data and color.
  */
-void UStaticMesh::SetVertexColorData(const TMap<FVector, FColor>& VertexColorData)
+void UStaticMesh::SetVertexColorData(const TMap<FVector3f, FColor>& VertexColorData)
 {
+	TRACE_CPUPROFILER_EVENT_SCOPE(UStaticMesh::SetVertexColorData);
 #if WITH_EDITOR
 	// What LOD to get vertex colors from.  
 	// Currently mesh painting only allows for painting on the first lod.
@@ -6893,34 +6892,23 @@ void UStaticMesh::SetVertexColorData(const TMap<FVector, FColor>& VertexColorDat
 	{
 		if (GetSourceModel(PaintingMeshLODIndex).IsRawMeshEmpty() == false)
 		{
-			// Extract the raw mesh.
-			FRawMesh Mesh;
-			GetSourceModel(PaintingMeshLODIndex).LoadRawMesh(Mesh);
-
-			// Reserve space for the new vertex colors.
-			if (Mesh.WedgeColors.Num() == 0 || Mesh.WedgeColors.Num() != Mesh.WedgeIndices.Num())
-			{
-				Mesh.WedgeColors.Empty(Mesh.WedgeIndices.Num());
-				Mesh.WedgeColors.AddUninitialized(Mesh.WedgeIndices.Num());
-			}
-
-			// Build a mapping of vertex positions to vertex colors.
-			for (int32 WedgeIndex = 0; WedgeIndex < Mesh.WedgeIndices.Num(); ++WedgeIndex)
+			FMeshDescription* MeshDescription = GetMeshDescription(PaintingMeshLODIndex);
+			TVertexAttributesRef<FVector3f> Positions = FStaticMeshAttributes(*MeshDescription).GetVertexPositions();
+			TVertexInstanceAttributesRef<FVector4f> Colors = FStaticMeshAttributes(*MeshDescription).GetVertexInstanceColors();
+			for (FVertexInstanceID VertexInstanceID : MeshDescription->VertexInstances().GetElementIDs())
 			{
-				FVector Position(Mesh.VertexPositions[Mesh.WedgeIndices[WedgeIndex]]);
-				const FColor* Color = VertexColorData.Find(Position);
-				if (Color)
+				FVertexID VertexID = MeshDescription->GetVertexInstanceVertex(VertexInstanceID);
+				const FVector3f& Position = Positions[VertexID];
+				if (const FColor* Color = VertexColorData.Find(Position))
 				{
-					Mesh.WedgeColors[WedgeIndex] = *Color;
+					Colors[VertexInstanceID] = FVector4f(FLinearColor::FromSRGBColor(*Color));
 				}
 				else
 				{
-					Mesh.WedgeColors[WedgeIndex] = FColor(255, 255, 255, 255);
+					Colors[VertexInstanceID] = FVector4f(FLinearColor::White);
 				}
 			}
-
-			// Save the new raw mesh.
-			GetSourceModel(PaintingMeshLODIndex).SaveRawMesh(Mesh);
+			CommitMeshDescription(PaintingMeshLODIndex);
 		}
 	}
 	// TODO_STATICMESH: Build?
@@ -6937,17 +6925,14 @@ ENGINE_API void UStaticMesh::RemoveVertexColors()
 		FStaticMeshSourceModel& SourceModel = GetSourceModel(LodIndex);
 		if (!SourceModel.IsRawMeshEmpty())
 		{
-			FRawMesh RawMesh;
-			SourceModel.LoadRawMesh(RawMesh);
-
-			if (RawMesh.WedgeColors.Num() > 0)
+			FMeshDescription* MeshDescription = GetMeshDescription(LodIndex);
+			TVertexInstanceAttributesRef<FVector4f> Colors = FStaticMeshAttributes(*MeshDescription).GetVertexInstanceColors();
+			for (FVertexInstanceID VertexInstanceID : MeshDescription->VertexInstances().GetElementIDs())
 			{
-				RawMesh.WedgeColors.Empty();
-
-				SourceModel.SaveRawMesh(RawMesh);
-
-				bRemovedVertexColors = true;
+				Colors[VertexInstanceID] = FVector4f(FLinearColor::White);
 			}
+			CommitMeshDescription(LodIndex);
+			bRemovedVertexColors = true;
 		}
 	}
 
