Asset Bulk Data Diff Commandlet

Adds "Diff Tags" to aid in determining why a given bulk data change occurred. Other large bulk data assets should add their own diff tags (i.e. StaticMesh). Currently only Texture is filled out.

diff --git a/Engine/Source/Editor/UnrealEd/Classes/Commandlets/DiffAssetBulkDataCommandlet.h b/Engine/Source/Editor/UnrealEd/Classes/Commandlets/DiffAssetBulkDataCommandlet.h
new file mode 100644
--- /dev/null
+++ b/Engine/Source/Editor/UnrealEd/Classes/Commandlets/DiffAssetBulkDataCommandlet.h
@@ -0,0 +1,23 @@
+// Copyright Epic Games, Inc. All Rights Reserved.
+#pragma once
+
+#include "CoreTypes.h"
+#include "UObject/ObjectMacros.h"
+#include "Commandlets/Commandlet.h"
+#include "DiffAssetBulkDataCommandlet.generated.h"
+
+// Detailed comments in DiffAssetBulkDataCommandlet.cpp
+UCLASS()
+class UDiffAssetBulkDataCommandlet : public UCommandlet
+{
+	GENERATED_UCLASS_BODY()
+
+public:
+
+	// Begin UCommandlet Interface
+	virtual int32 Main(const FString& Params) override;
+	// End UCommandlet Interface
+
+private:
+
+};
diff --git a/Engine/Source/Editor/UnrealEd/Private/Commandlets/DiffAssetBulkDataCommandlet.cpp b/Engine/Source/Editor/UnrealEd/Private/Commandlets/DiffAssetBulkDataCommandlet.cpp
new file mode 100644
--- /dev/null
+++ b/Engine/Source/Editor/UnrealEd/Private/Commandlets/DiffAssetBulkDataCommandlet.cpp
@@ -0,0 +1,475 @@
+// Copyright Epic Games, Inc. All Rights Reserved.
+
+#include "Commandlets/DiffAssetBulkDataCommandlet.h"
+
+#include "AssetRegistry/AssetData.h"
+#include "AssetRegistry/IAssetRegistry.h"
+#include "AssetRegistryState.h"
+#include "HAL/FileManager.h"
+#include "Serialization/LargeMemoryReader.h"
+
+
+DEFINE_LOG_CATEGORY_STATIC(LogDiffAssetBulk, Display, All);
+
+/**
+ * Diff Asset Bulk Data
+ * 
+ * This loads two asset registries newer than FAssetRegistryVersion::AddedChunkHashes,
+ * and attempts to find the reason for bulk data differences.
+ * 
+ * First, it finds what bulk datas changed by using the hash of the bulk data,
+ * then it uses "Diff Tags" to try and determine at what point during the derived data
+ * build the change occurred.
+ * 
+ * 
+ * Diff Tags
+ * 
+ * Diff Tags are cook tags added during the cook process using Ar.CookContext()->CookTagList() (see CookTagList.h)
+ * and are of the form "Cook_Diff_##_Key":
+ *
+ * 		- "Cook_": 	Added automatically by the the cook tag system.
+ * 		- "Diff_": 	Identifies the tag as a diff tag.
+ * 		- "##":		Specifies where in the build process the tag represents (Ordering).
+ * 		- "_Key":	Descriptive text for the tag.
+ * 
+ * If a bulk data difference is found, the diff tags are checked for differences in order, and the first
+ * diff tag that changed is assigned the "blame" for the change under the assumption that later
+ * tags will necessarily change as a result of the earlier change.
+ * 
+ * If diff tags are present for the asset and none of the diff tags changed, then it is assumed that a build determinism 
+ * issue has caused the change.
+ *
+ */
+
+/**
+ *  The list of known cook diff tags - this is just used to provide explanations in the output for the reader.
+ */
+static struct FBuiltinDiffTagHelp {const TCHAR* TagName; const TCHAR* TagHelp;} GBuiltinDiffTagHelp[] = 
+{
+	{TEXT("Cook_Diff_20_Tex2D_CacheKey"), TEXT("Texture settings or referenced data changed (DDC2)")},
+	{TEXT("Cook_Diff_20_Tex2D_DDK"), TEXT("Texture settings or referenced data changed (DDC1)")},
+	{TEXT("Cook_Diff_10_Tex2D_Source"), TEXT("Texture source data changed")}
+};
+
+
+UDiffAssetBulkDataCommandlet::UDiffAssetBulkDataCommandlet(const FObjectInitializer& ObjectInitializer)
+: Super(ObjectInitializer)
+{
+
+}
+
+
+int32 UDiffAssetBulkDataCommandlet::Main(const FString& FullCommandLine)
+{
+	FString BaseFileName, CurrentFileName;
+	const TCHAR* CmdLine = *FullCommandLine;
+	if (FParse::Value(CmdLine, TEXT("Base="), BaseFileName) == false ||
+		FParse::Value(CmdLine, TEXT("Current="), CurrentFileName) == false)
+	{
+		UE_LOG(LogDiffAssetBulk, Display, TEXT(""));
+		UE_LOG(LogDiffAssetBulk, Display, TEXT("Diff Asset Bulk Data"));
+		UE_LOG(LogDiffAssetBulk, Display, TEXT(""));
+		UE_LOG(LogDiffAssetBulk, Display, TEXT("Loads two development asset registries and finds all bulk data changes, and tries to find why"));
+		UE_LOG(LogDiffAssetBulk, Display, TEXT("the bulk data changed. Development asset registries are in the cooked /Metadata directory."));
+		UE_LOG(LogDiffAssetBulk, Display, TEXT(""));
+		UE_LOG(LogDiffAssetBulk, Display, TEXT("Parameters:"));
+		UE_LOG(LogDiffAssetBulk, Display, TEXT(""));
+		UE_LOG(LogDiffAssetBulk, Display, TEXT("    -Base=<path/to/file>              Base Development Asset Registry (Required)"));
+		UE_LOG(LogDiffAssetBulk, Display, TEXT("    -Current=<path/to/file>           New Development Asset Registry (Required)"));
+		return 1;
+	}
+
+	// Convert the static init help text to a map
+	TMap<FName, const TCHAR*> BuiltinDiffTagHelpMap;
+	for (FBuiltinDiffTagHelp& DiffTagHelp : GBuiltinDiffTagHelp)
+	{
+		BuiltinDiffTagHelpMap.Add(DiffTagHelp.TagName, DiffTagHelp.TagHelp);
+	}
+
+	FAssetRegistryState BaseState, CurrentState;
+	FAssetRegistryVersion::Type BaseVersion, CurrentVersion;
+	if (FAssetRegistryState::LoadFromDisk(*BaseFileName, FAssetRegistryLoadOptions(), BaseState, &BaseVersion) == false)
+	{
+		UE_LOG(LogDiffAssetBulk, Error, TEXT("Failed load base (%s)"), *BaseFileName);
+		return 1;
+	}
+	if (FAssetRegistryState::LoadFromDisk(*CurrentFileName, FAssetRegistryLoadOptions(), CurrentState, &CurrentVersion) == false)
+	{
+		UE_LOG(LogDiffAssetBulk, Error, TEXT("Failed load current (%s)"), *CurrentFileName);
+		return 1;
+	}
+
+	//
+	// The cook process adds the hash for almost all iochunks to the asset registry - 
+	// so as long as both asset registries have that data, we get what we want.
+	//
+	if (BaseVersion < FAssetRegistryVersion::AddedChunkHashes)
+	{
+		UE_LOG(LogDiffAssetBulk, Error, TEXT("Base asset registry version is too old (%d, need %d)"), BaseVersion, FAssetRegistryVersion::AddedChunkHashes);
+		return 1;
+	}
+	if (CurrentVersion < FAssetRegistryVersion::AddedChunkHashes)
+	{
+		UE_LOG(LogDiffAssetBulk, Error, TEXT("Current asset registry version is too old (%d, need %d)"), CurrentVersion, FAssetRegistryVersion::AddedChunkHashes);
+		return 1;
+	}
+
+	const TMap<FName, const FAssetPackageData*>& BasePackages = BaseState.GetAssetPackageDataMap();
+	const TMap<FName, const FAssetPackageData*>& CurrentPackages = CurrentState.GetAssetPackageDataMap();
+
+	TArray<FName> NewPackages, DeletedPackages, SharedPackages;
+	{
+		for (const TPair<FName, const FAssetPackageData*>& NamePackageDataPair : BasePackages)
+		{
+			if (CurrentState.GetAssetPackageData(NamePackageDataPair.Key) == nullptr)
+			{
+				DeletedPackages.Add(NamePackageDataPair.Key);
+				continue;
+			}
+
+			SharedPackages.Add(NamePackageDataPair.Key);
+		}
+
+		for (const TPair<FName, const FAssetPackageData*>& NamePackageDataPair : CurrentPackages)
+		{
+			if (BaseState.GetAssetPackageData(NamePackageDataPair.Key) == nullptr)
+			{
+				NewPackages.Add(NamePackageDataPair.Key);
+			}
+		}
+	}
+
+	// Now we need to see what changed. We current care about BulkData (or DerivedData when it gets turned on).
+	// OptionalBulkData and MemoryMappedBulkData not yet supported...
+	//
+	// This whole thing assumes that the index parameter of CreateIoChunkId is always 0. This is likely not going
+	// to be true with FDerivedData, once that gets turned on, but should be easy to update when the time comes.
+	//
+	TArray<FName> ChangedBulkData, NewBulkData, DeletedBulkData;
+	for (const FName& SharedPackage : SharedPackages)
+	{
+		const FAssetPackageData* BasePackage = BasePackages[SharedPackage];
+		const FAssetPackageData* CurrentPackage = CurrentPackages[SharedPackage];
+
+		for (const TPair<FIoChunkId, FIoHash>& ChunkHashPair : BasePackage->ChunkHashes)
+		{
+			if (ChunkHashPair.Key.GetChunkType() != EIoChunkType::BulkData)
+			{
+				continue;
+			}		
+
+			const FIoHash* CurrentHash = CurrentPackage->ChunkHashes.Find(ChunkHashPair.Key);
+			if (CurrentHash == nullptr)
+			{
+				check(DeletedBulkData.Contains(SharedPackage) == false); // Because only 0 chunk index
+				DeletedBulkData.Add(SharedPackage);
+				continue;
+			}
+
+			if (*CurrentHash != ChunkHashPair.Value)
+			{
+				check(ChangedBulkData.Contains(SharedPackage) == false); // Because only 0 chunk index
+				ChangedBulkData.Add(SharedPackage);
+			}
+		}
+
+		for (const TPair<FIoChunkId, FIoHash>& ChunkHashPair : CurrentPackage->ChunkHashes)
+		{
+			if (ChunkHashPair.Key.GetChunkType() != EIoChunkType::BulkData)
+			{
+				continue;
+			}
+
+			if (BasePackage->ChunkHashes.Contains(ChunkHashPair.Key) == false)
+			{
+				check(NewBulkData.Contains(SharedPackage) == false); // Because only 0 chunk index
+				NewBulkData.Add(SharedPackage);
+			}
+		}
+	}
+
+	//
+	// We know what bulk datas *packages* changed. Try and see if any of the assets in the package have
+	// diff blame tags for us to determine cause. _usually_ there's one asset per package, but it's definitely
+	// possible to have more. Additionally _usually_ there's a good single candidate for assigning the data
+	// cost, however it is possible to have e.g. an importer create a lot of assets in a single package that
+	// all add bulk data to the package.
+	// 
+	// Once we have FDerivedData we might be able to keep what data belongs to which asset.
+	//
+	struct FDiffResult
+	{
+		FString ChangedAssetObjectPath;
+		FString TagBaseValue;
+		FString TagCurrentValue;
+	};
+
+	TMap<FName /* TagName */, TMap<FName /* AssetClass */, TArray<FDiffResult>>> Results;
+	TMap<FName /* AssetClass */, TArray<FName /* PackageName */>> NoTagPackagesByAssumedClass;
+	TArray<FName /* PackageName */> PackagesWithUnassignableDiffsAndUntaggedAssets;
+	TArray<FName /* PackageName */> PackagesWithUnassignableDiffs;
+	
+	for (FName& ChangedPackageName : ChangedBulkData)
+	{
+		TConstArrayView<FAssetData const*> BaseAssetDatas = BaseState.GetAssetsByPackageName(ChangedPackageName);		
+		TConstArrayView<FAssetData const*> CurrentAssetDatas = CurrentState.GetAssetsByPackageName(ChangedPackageName);
+
+		struct FDiffTag
+		{
+			// Order is used to sort the diff blame keys so that the correct thing is blamed. This is
+			// so that e.g. changing the texture source (which would change the ddc key) gets properly blamed
+			// as it is lower order.
+			int Order;
+			FName TagName;
+			FString BaseValue;
+			FString CurrentValue;
+
+			const FAssetData* BaseAssetData;
+			const FAssetData* CurrentAssetData;
+		};
+		
+		// We want to find all the tags that are in both base/current.
+		TMap<FName /* AssetName */, TArray<FDiffTag>> PackageDiffTags;
+		bool bPackageHasUntaggedAsset = false;
+		for (const FAssetData* BaseAssetData : BaseAssetDatas)
+		{
+			BaseAssetData->EnumerateTags([&PackageDiffTags, BaseAssetData, CurrentAssetDatas](TPair<FName, FAssetTagValueRef> TagAndValue)
+			{
+				TCHAR Name[NAME_SIZE];
+				TagAndValue.Key.GetPlainNameString(Name);
+				if (FCString::Strncmp(Name, TEXT("Cook_Diff_"), 10))
+				{
+					return;
+				}
+
+				// This is O(N) but like 99.9% of the time there's only 1 asset.
+				const FAssetData* const* CurrentAssetData = CurrentAssetDatas.FindByPredicate([SearchAssetName = &BaseAssetData->AssetName](const FAssetData* AssetData) { return (AssetData->AssetName == *SearchAssetName); });
+				if (CurrentAssetData == nullptr)
+				{
+					return;
+				}
+
+				FString CurrentValue;
+				if (CurrentAssetData[0]->GetTagValue(TagAndValue.Key, CurrentValue) == false)
+				{
+					// Both version don't have the tag so we can't compare.
+					return;
+				}
+
+				TArray<FDiffTag>& AssetDiffTags = PackageDiffTags.FindOrAdd(BaseAssetData->AssetName);
+				FDiffTag& Tag = AssetDiffTags.AddDefaulted_GetRef();
+				Tag.Order = FCString::Atoi(Name + FCString::Strlen(TEXT("Cook_Diff_"))); // this gets optimized to +10
+				Tag.TagName = TagAndValue.Key;
+				Tag.BaseValue =  TagAndValue.Value.AsString();
+				Tag.CurrentValue = MoveTemp(CurrentValue);
+				Tag.BaseAssetData = BaseAssetData;
+				Tag.CurrentAssetData = *CurrentAssetData;
+			});
+
+			if (PackageDiffTags.Contains(BaseAssetData->AssetName) == false)
+			{
+				bPackageHasUntaggedAsset = true;
+				// An asset exists in the package that doesn't have any tags - make a note so that
+				// we can suggest this caused the bulk data diff if we don't find a blame.
+			}
+		}
+
+		bool bPackageHasUntaggedAndTaggedAssets = false;
+		if (PackageDiffTags.Num())
+		{
+			if (bPackageHasUntaggedAsset)
+			{
+				bPackageHasUntaggedAndTaggedAssets = true;
+			}
+		}
+		else
+		{
+			// Nothing has anything to use for diff blaming for this package.
+			// Try to find a representative asset class from the assets in the package.
+
+			// Find a candidate asset.
+			// If there's a "UAsset", then we use that as the asset to get the sizes.
+			// If not, then we look for a "TopLevelAsset", i.e. one that shows up in the content browser.
+			// 
+			// If we have multiple TLS or no TLAs, we dump those in a "??" list because
+			// we can't really make any decisions there.
+			// 
+			FAssetData const* UseAsset = nullptr;
+			{
+				FAssetData const* Found_UAsset = nullptr;
+				FAssetData const* Found_TLA = nullptr;
+				bool bMultipleTLAs = false;
+
+				for (FAssetData const* const Asset : CurrentAssetDatas)
+				{
+					if (Asset->IsUAsset())
+					{
+						Found_UAsset = Asset;
+					}
+					if (Asset->IsTopLevelAsset())
+					{
+						if (Found_TLA)
+						{
+							bMultipleTLAs = true;
+						}
+						Found_TLA = Asset;
+					}
+				}
+
+				UseAsset = Found_UAsset;
+				if (UseAsset == nullptr)
+				{
+					if (bMultipleTLAs == false)
+					{
+						UseAsset = Found_TLA; // only accept the TLA if there's one, otherwise leave this as unassigned.
+					}
+				}
+			}
+
+			if (UseAsset == nullptr)
+			{
+				NoTagPackagesByAssumedClass.FindOrAdd(NAME_None).Add(ChangedPackageName);
+			}
+			else
+			{
+				NoTagPackagesByAssumedClass.FindOrAdd(UseAsset->AssetClass).Add(ChangedPackageName);
+			}
+			
+			continue;
+		}
+
+		// Now we check and see if any of the diff tags can tell us why the package changed.
+		// We could find multiple assets that caused the change.
+		bool bFoundDiffTag = false;
+		for (TPair<FName, TArray<FDiffTag>>& AssetDiffTagPair : PackageDiffTags)
+		{
+			TArray<FDiffTag>& AssetDiffTags = AssetDiffTagPair.Value;
+			Algo::SortBy(AssetDiffTags, &FDiffTag::Order);
+			
+			for (FDiffTag& Tag : AssetDiffTags)
+			{
+				if (Tag.BaseValue != Tag.CurrentValue)
+				{
+					TMap<FName, TArray<FDiffResult>>& TagResults = Results.FindOrAdd(Tag.TagName);
+
+					TArray<FDiffResult>& ClassResults = TagResults.FindOrAdd(Tag.BaseAssetData->AssetClass);
+
+					FDiffResult& Result = ClassResults.AddDefaulted_GetRef();
+					Result.ChangedAssetObjectPath = Tag.BaseAssetData->ObjectPath.ToString();
+					Result.TagBaseValue = MoveTemp(Tag.BaseValue);
+					Result.TagCurrentValue = MoveTemp(Tag.CurrentValue);
+					bFoundDiffTag = true;
+					break;
+				}
+			}
+		}
+
+
+		if (bFoundDiffTag == false)
+		{
+			// This means that all the tags they added didn't change, but the asset did.
+			// Assuming that a DDC key tag has been added, this means either:
+			// 
+			// A) The asset changed independent of DDC key, which is a build consistency / determinism alert.
+			// B) The package had an asset with tags and an asset without tags, and the asset without tags caused
+			//	  the bulk data change.
+			//
+			// Unfortunately A) is a Big Deal and needs a warning, but B might end up being common due to blueprint classes,
+			// so we segregate the lists.
+			if (bPackageHasUntaggedAndTaggedAssets)
+			{
+				PackagesWithUnassignableDiffsAndUntaggedAssets.Add(ChangedPackageName);
+			}
+			else
+			{
+				PackagesWithUnassignableDiffs.Add(ChangedPackageName);
+			}
+		}
+	}
+	
+	UE_LOG(LogDiffAssetBulk, Display, TEXT("Packages with Bulk Data: %d new, %d deleted, %d changed"), NewBulkData.Num(), DeletedBulkData.Num(), ChangedBulkData.Num());
+	if (ChangedBulkData.Num() == 0)
+	{
+		return 0;
+	}
+
+	UE_LOG(LogDiffAssetBulk, Display, TEXT("Changed package breakdown:"));
+	UE_LOG(LogDiffAssetBulk, Display, TEXT("    No blame information available:"));
+	UE_LOG(LogDiffAssetBulk, Display, TEXT("        Can't determine asset class   : %d"), NoTagPackagesByAssumedClass.FindOrAdd(NAME_None).Num());
+	for (TPair<FName, TArray<FName>> ClassPackages : NoTagPackagesByAssumedClass)
+	{
+		if (ClassPackages.Key == NAME_None)
+		{
+			continue;
+		}
+
+		UE_LOG(LogDiffAssetBulk, Display, TEXT("        %-30s: %d"), *ClassPackages.Key.ToString(), ClassPackages.Value.Num());
+	}
+
+	if (PackagesWithUnassignableDiffs.Num())
+	{
+		Algo::Sort(PackagesWithUnassignableDiffs, FNameLexicalLess());
+
+		UE_LOG(LogDiffAssetBulk, Warning, TEXT("    Can't determine blame:        : %-7d // Assets had blame tags but all matched - check determinism!"), PackagesWithUnassignableDiffs.Num());
+		for (FName& PackageName : PackagesWithUnassignableDiffs)
+		{
+			UE_LOG(LogDiffAssetBulk, Warning, TEXT("        %s"), *PackageName.ToString());
+		}
+	}
+
+	if (PackagesWithUnassignableDiffsAndUntaggedAssets.Num())
+	{
+		Algo::Sort(PackagesWithUnassignableDiffsAndUntaggedAssets, FNameLexicalLess());
+		
+		UE_LOG(LogDiffAssetBulk, Display, TEXT("    Potential untagged assets:    : %-7d // Package had assets with blame tags that matched, but also untagged assets. Might be determinism!"), PackagesWithUnassignableDiffsAndUntaggedAssets.Num());
+		for (FName& PackageName : PackagesWithUnassignableDiffsAndUntaggedAssets)
+		{
+			UE_LOG(LogDiffAssetBulk, Display, TEXT("        %s"), *PackageName.ToString());
+		}
+
+	}
+
+	if (Results.Num())
+	{
+		UE_LOG(LogDiffAssetBulk, Display, TEXT("    Summary changes by blame tag:"));
+
+		for (TPair<FName, TMap<FName, TArray<FDiffResult>>>& TagResults : Results)
+		{
+			uint32 TagCount = 0;
+			for (TPair<FName, TArray<FDiffResult>>& ClassResults : TagResults.Value)
+			{
+				TagCount += ClassResults.Value.Num();
+			}
+
+			const TCHAR** TagHelp = BuiltinDiffTagHelpMap.Find(TagResults.Key);
+			if (TagHelp != nullptr)
+			{
+				UE_LOG(LogDiffAssetBulk, Display, TEXT("        %-30s: %-7d // %s"), *TagResults.Key.ToString(), TagCount, *TagHelp);
+			}
+			else
+			{
+				UE_LOG(LogDiffAssetBulk, Display, TEXT("        %-30s: %-7d"), *TagResults.Key.ToString(), TagCount);
+			}
+		}
+
+		UE_LOG(LogDiffAssetBulk, Display, TEXT("    Asset changes by blame tag:"));
+
+		for (TPair<FName, TMap<FName, TArray<FDiffResult>>>& TagResults : Results)
+		{
+			UE_LOG(LogDiffAssetBulk, Display, TEXT("        %s"), *TagResults.Key.ToString());
+
+			for (TPair<FName, TArray<FDiffResult>>& ClassResults : TagResults.Value)
+			{
+				Algo::SortBy(ClassResults.Value, &FDiffResult::ChangedAssetObjectPath);
+				UE_LOG(LogDiffAssetBulk, Display, TEXT("            %s [%d]"), *ClassResults.Key.ToString(), ClassResults.Value.Num());
+
+				for (FDiffResult& Result : ClassResults.Value)
+				{
+					UE_LOG(LogDiffAssetBulk, Display, TEXT("                %s [%s -> %s]"), *Result.ChangedAssetObjectPath, *Result.TagBaseValue, *Result.TagCurrentValue);
+				}
+			}
+		}
+	}
+
+	return 0;
+}
diff --git a/Engine/Source/Runtime/CoreUObject/Public/AssetRegistry/CookTagList.h b/Engine/Source/Runtime/CoreUObject/Public/AssetRegistry/CookTagList.h
--- a/Engine/Source/Runtime/CoreUObject/Public/AssetRegistry/CookTagList.h
+++ b/Engine/Source/Runtime/CoreUObject/Public/AssetRegistry/CookTagList.h
@@ -41,6 +41,9 @@ struct FCookTagList
 	*	Adds a tag to the list to be added to the development asset registry
 	*	for this object. Note that we only carry over the name/value, not
 	*	any of the other metadata available when adding tags via GetAssetRegistryTags.
+	* 
+	*	Tags that start with "Diff_" are used for bulk data diff blaming - see DiffAssetBulkDataCommandlet.cpp. When
+	* 	adding diff tags, be sure to add explanatory prose in that commandlet to GBuiltinDiffTagHelp.
 	*/
 	void Add(UObject* InObject, FName InTagName, FString&& InTagValue)
 	{
diff --git a/Engine/Source/Runtime/Engine/Private/TextureDerivedData.cpp b/Engine/Source/Runtime/Engine/Private/TextureDerivedData.cpp
--- a/Engine/Source/Runtime/Engine/Private/TextureDerivedData.cpp
+++ b/Engine/Source/Runtime/Engine/Private/TextureDerivedData.cpp
@@ -2467,6 +2467,22 @@ static void SerializePlatformData(
 			}
 
 			CookTags->Add(Texture, "Format", FString(GPixelFormats[PlatformData->PixelFormat].Name));
+
+			// Add in diff keys for change detection/blame.
+			{
+				// Did the source change?
+				CookTags->Add(Texture, "Diff_10_Tex2D_Source", Texture->Source.GetIdString());
+
+				// Did the settings change?
+				if (const UE::DerivedData::FCacheKeyProxy* CacheKey = PlatformData->DerivedDataKey.TryGet<UE::DerivedData::FCacheKeyProxy>())
+				{
+					CookTags->Add(Texture, "Diff_20_Tex2D_CacheKey", *WriteToString<64>(*CacheKey->AsCacheKey()));
+				}
+				else if (const FString* DDK = PlatformData->DerivedDataKey.TryGet<FString>())
+				{
+					CookTags->Add(Texture, "Diff_20_Tex2D_DDK", FString(*DDK));
+				}
+			}
 		}
 #endif
 	}
