Fix SoundWave asset thumbnails disappearing when dragging and dropping

diff --git a/Engine/Source/Editor/UnrealEd/Private/ThumbnailRendering/SoundWaveThumbnailRenderer.cpp b/Engine/Source/Editor/UnrealEd/Private/ThumbnailRendering/SoundWaveThumbnailRenderer.cpp
--- a/Engine/Source/Editor/UnrealEd/Private/ThumbnailRendering/SoundWaveThumbnailRenderer.cpp
+++ b/Engine/Source/Editor/UnrealEd/Private/ThumbnailRendering/SoundWaveThumbnailRenderer.cpp
@@ -24,8 +24,6 @@ bool USoundWaveThumbnailRenderer::CanVisualizeAsset(UObject* Object)
 
 void USoundWaveThumbnailRenderer::Draw(UObject* Object, int32 X, int32 Y, uint32 Width, uint32 Height, FRenderTarget* RenderTarget, FCanvas* Canvas, bool bAdditionalViewFamily)
 {
-	bool bCanDrawWaveform = true;
-
 	USoundWave* SoundWave = Cast<USoundWave>(Object);
 	if (!SoundWave)
 	{
@@ -177,7 +175,7 @@ bool USoundWaveThumbnailRenderer::AllowsRealtimeThumbnails(UObject* Object) cons
 {
 	if (USoundWave* SoundWave = Cast<USoundWave>(Object))
 	{
-		return SoundWave->GetRedrawThumbnail();
+		return SoundWave->CanVisualizeAsset();
 	}
 
 	return false;
