Fix for 'Create Animation Blueprint' window being created off-screen when using Windows scale window

diff --git a/Engine/Source/Editor/UnrealEd/Private/Factories/AnimBlueprintFactory.cpp b/Engine/Source/Editor/UnrealEd/Private/Factories/AnimBlueprintFactory.cpp
--- a/Engine/Source/Editor/UnrealEd/Private/Factories/AnimBlueprintFactory.cpp
+++ b/Engine/Source/Editor/UnrealEd/Private/Factories/AnimBlueprintFactory.cpp
@@ -209,7 +209,7 @@ public:
 		TSharedRef<SWindow> Window = SNew(SWindow)
 		.Title( LOCTEXT("CreateAnimBlueprintOptions", "Create Animation Blueprint") )
 		.ClientSize(FVector2D(500, 500))
-		.SizingRule(ESizingRule::Autosized)
+		.SizingRule(ESizingRule::FixedSize)
 		.SupportsMinimize(false)
 		.SupportsMaximize(false)
 		[
