//UE5/Main - Run Enterprise UnrealEditor Python Tests - Crash on commit of dataprep graph

diff --git a/Engine/Source/Editor/UnrealEd/Private/EditorLevelUtils.cpp b/Engine/Source/Editor/UnrealEd/Private/EditorLevelUtils.cpp
--- a/Engine/Source/Editor/UnrealEd/Private/EditorLevelUtils.cpp
+++ b/Engine/Source/Editor/UnrealEd/Private/EditorLevelUtils.cpp
@@ -726,8 +726,11 @@ ULevelStreaming* UEditorLevelUtils::CreateNewStreamingLevelForWorld(UWorld& InWo
 	}
 	check(bIsPartitioned == PersistentLevel->bIsPartitioned);
 
-	// Call lambda before saving level
-	InPreSaveLevelOperation(PersistentLevel);
+	if (InPreSaveLevelOperation)
+	{
+		// Call lambda before saving level
+		InPreSaveLevelOperation(PersistentLevel);
+	}
 
 	bool bNewWorldSaved = false;
 	FString NewPackageName = DefaultFilename;
