Fix for directional light warning showing in thumbnails.

diff --git a/Engine/Source/Editor/UnrealEd/Private/ThumbnailHelpers.cpp b/Engine/Source/Editor/UnrealEd/Private/ThumbnailHelpers.cpp
--- a/Engine/Source/Editor/UnrealEd/Private/ThumbnailHelpers.cpp
+++ b/Engine/Source/Editor/UnrealEd/Private/ThumbnailHelpers.cpp
@@ -50,10 +50,12 @@ FThumbnailPreviewScene::FThumbnailPreviewScene()
 	// Add additional lights
 	UDirectionalLightComponent* DirectionalLight2 = NewObject<UDirectionalLightComponent>();
 	DirectionalLight2->Intensity = 5.0f;
+	DirectionalLight2->ForwardShadingPriority = 1;
 	AddComponent(DirectionalLight2, FTransform( FRotator(-40,-144.678, 0) ));
 
 	UDirectionalLightComponent* DirectionalLight3 = NewObject<UDirectionalLightComponent>();
 	DirectionalLight3->Intensity = 1.0f;
+	DirectionalLight2->ForwardShadingPriority = 2;
 	AddComponent(DirectionalLight3, FTransform( FRotator(299.235,144.993, 0) ));
 
 	SetSkyCubemap(GUnrealEd->GetThumbnailManager()->AmbientCubemap);
