EditorFactories nop change trying to make static analysis happy

EditorFactories.cpp(2846): warning V512: Instantiation of PNGDataFill < uint8, uint32, 2, 1, 0, 3 >: A call of the 'memcpy' function will lead to overflow of the buffer '& WhiteWithZeroAlpha'.

diff --git a/Engine/Source/Editor/UnrealEd/Private/Factories/EditorFactories.cpp b/Engine/Source/Editor/UnrealEd/Private/Factories/EditorFactories.cpp
--- a/Engine/Source/Editor/UnrealEd/Private/Factories/EditorFactories.cpp
+++ b/Engine/Source/Editor/UnrealEd/Private/Factories/EditorFactories.cpp
@@ -2843,7 +2843,7 @@ public:
 		{
 			check( sizeof(ColorDataType) == 8 );
 			uint16 RGBA[4] = { 0xFFFF,0xFFFF,0xFFFF, 0 };
-			memcpy(&WhiteWithZeroAlpha,RGBA,8);
+			memcpy(&WhiteWithZeroAlpha,RGBA,sizeof(ColorDataType));
 		}
 
 		// Left -> Right
