Fix BP Debugger Step Over not behaving correctly with Sequence nodes with a branch as the end of the exec line

diff --git a/Engine/Source/Editor/UnrealEd/Private/Kismet2/KismetDebugUtilities.cpp b/Engine/Source/Editor/UnrealEd/Private/Kismet2/KismetDebugUtilities.cpp
--- a/Engine/Source/Editor/UnrealEd/Private/Kismet2/KismetDebugUtilities.cpp
+++ b/Engine/Source/Editor/UnrealEd/Private/Kismet2/KismetDebugUtilities.cpp
@@ -32,6 +32,7 @@
 #include "K2Node_FunctionEntry.h"
 #include "K2Node_Knot.h"
 #include "K2Node_MacroInstance.h"
+#include "K2Node_Composite.h"
 #include "K2Node_Message.h"
 #include "Kismet2/KismetEditorUtilities.h"
 #include "Kismet2/BlueprintEditorUtils.h"
@@ -144,44 +145,43 @@ void FKismetDebugUtilities::RequestStepOver()
 	FKismetDebugUtilitiesData& Data = FKismetDebugUtilitiesData::Get();
 	const TArray<const FFrame*>& ScriptStack = FBlueprintContextTracker::Get().GetScriptStack();
 
-	if(ScriptStack.Num() > 0)
+	if (ScriptStack.Num() > 0)
 	{
 		Data.TargetGraphStackDepth = ScriptStack.Num();
-		
-		// get the current graph that we're stopped at:
-		const FFrame* CurrentFrame = ScriptStack.Last();
-		if(CurrentFrame->Object)
+
+		if (const UEdGraphNode* StoppedNode = Data.MostRecentStoppedNode.Get())
 		{
-			if(UBlueprintGeneratedClass* BPGC = Cast<UBlueprintGeneratedClass>(CurrentFrame->Object->GetClass()))
+			if (StoppedNode->IsA<UK2Node_MacroInstance>() || StoppedNode->IsA<UK2Node_Composite>())
 			{
-				const int32 BreakpointOffset = CurrentFrame->Code - CurrentFrame->Node->Script.GetData() - 1;
-				UEdGraphNode* BlueprintNode = BPGC->DebugData.FindSourceNodeFromCodeLocation(CurrentFrame->Node, BreakpointOffset, true);
-				if(BlueprintNode)
+				for (const UEdGraphPin* Pin : StoppedNode->Pins)
 				{
 					// add any nodes connected via execs as TargetGraphNodes:
-					for(UEdGraphPin* Pin : BlueprintNode->Pins)
+					if (Pin->Direction == EGPD_Output && Pin->PinType.PinCategory == UEdGraphSchema_K2::PC_Exec && Pin->LinkedTo.Num() > 0)
 					{
-						if(Pin->Direction == EGPD_Output && Pin->PinType.PinCategory == UEdGraphSchema_K2::PC_Exec && Pin->LinkedTo.Num() > 0)
+						for (UEdGraphPin* LinkedTo : Pin->LinkedTo)
 						{
-							for(UEdGraphPin* LinkedTo : Pin->LinkedTo)
+							UEdGraphNode* GraphNode = LinkedTo->GetOwningNode();
+							if (UK2Node_Knot* Knot = Cast<UK2Node_Knot>(GraphNode))
 							{
-								UEdGraphNode* GraphNode = LinkedTo->GetOwningNode();
-								if(UK2Node_Knot* Knot = Cast<UK2Node_Knot>(GraphNode))
-								{
-									// search the knot chain to find the actual node:
-									GraphNode = Knot->GetExecTerminal();
-								}
+								// search the knot chain to find the actual node:
+								GraphNode = Knot->GetExecTerminal();
+							}
 
-								if(GraphNode)
-								{
-									Data.TargetGraphNodes.AddUnique(GraphNode);
-								}
+							if (GraphNode)
+							{
+								Data.TargetGraphNodes.AddUnique(GraphNode);
 							}
 						}
 					}
+
 				}
+
+				return;
 			}
 		}
+
+		Data.bIsSingleStepping = false;
+		Data.bIsSteppingOut = true;
 	}
 }
 
