Interchange Framework - Adding a reimport interface to Interchange so that UInterchangeReimportHandler no longer relies on hardcoded asset types.

We now use the registered Translators and Factories to fetch the asset source file and determine if a reimport is supported.

diff --git a/Engine/Plugins/Experimental/Interchange/Runtime/Source/Import/Private/Animation/InterchangeAnimSequenceFactory.cpp b/Engine/Plugins/Experimental/Interchange/Runtime/Source/Import/Private/Animation/InterchangeAnimSequenceFactory.cpp
--- a/Engine/Plugins/Experimental/Interchange/Runtime/Source/Import/Private/Animation/InterchangeAnimSequenceFactory.cpp
+++ b/Engine/Plugins/Experimental/Interchange/Runtime/Source/Import/Private/Animation/InterchangeAnimSequenceFactory.cpp
@@ -453,6 +453,30 @@ void UInterchangeAnimSequenceFactory::PreImportPreCompletedCallback(const FImpor
 #endif
 }
 
+bool UInterchangeAnimSequenceFactory::GetSourceFilenames(const UObject* Object, TArray<FString>& OutSourceFilenames) const
+{
+#if WITH_EDITORONLY_DATA
+	if (const UAnimSequence* AnimSequence = Cast<UAnimSequence>(Object))
+	{
+		return UE::Interchange::FFactoryCommon::GetSourceFilenames(AnimSequence->AssetImportData.Get(), OutSourceFilenames);
+	}
+#endif
+
+	return false;
+}
+
+bool UInterchangeAnimSequenceFactory::SetSourceFilename(const UObject* Object, const FString& SourceFilename, int32 SourceIndex) const
+{
+#if WITH_EDITORONLY_DATA
+	if (const UAnimSequence* AnimSequence = Cast<UAnimSequence>(Object))
+	{
+		return UE::Interchange::FFactoryCommon::SetSourceFilename(AnimSequence->AssetImportData.Get(), SourceFilename, SourceIndex);
+	}
+#endif
+
+	return false;
+}
+
 /************************************************************************/
 /* Automation tests                                                     */
 /************************************************************************/
diff --git a/Engine/Plugins/Experimental/Interchange/Runtime/Source/Import/Private/Material/InterchangeMaterialFactory.cpp b/Engine/Plugins/Experimental/Interchange/Runtime/Source/Import/Private/Material/InterchangeMaterialFactory.cpp
--- a/Engine/Plugins/Experimental/Interchange/Runtime/Source/Import/Private/Material/InterchangeMaterialFactory.cpp
+++ b/Engine/Plugins/Experimental/Interchange/Runtime/Source/Import/Private/Material/InterchangeMaterialFactory.cpp
@@ -354,6 +354,30 @@ void UInterchangeMaterialFactory::PreImportPreCompletedCallback(const FImportPre
 #endif
 }
 
+bool UInterchangeMaterialFactory::GetSourceFilenames(const UObject* Object, TArray<FString>& OutSourceFilenames) const
+{
+#if WITH_EDITORONLY_DATA
+	if (const UMaterialInterface* MaterialInterface = Cast<UMaterialInterface>(Object))
+	{
+		return UE::Interchange::FFactoryCommon::GetSourceFilenames(MaterialInterface->AssetImportData.Get(), OutSourceFilenames);
+	}
+#endif
+
+	return false;
+}
+
+bool UInterchangeMaterialFactory::SetSourceFilename(const UObject* Object, const FString& SourceFilename, int32 SourceIndex) const
+{
+#if WITH_EDITORONLY_DATA
+	if (const UMaterialInterface* MaterialInterface = Cast<UMaterialInterface>(Object))
+	{
+		return UE::Interchange::FFactoryCommon::SetSourceFilename(MaterialInterface->AssetImportData.Get(), SourceFilename, SourceIndex);
+	}
+#endif
+
+	return false;
+}
+
 #if WITH_EDITOR
 void UInterchangeMaterialFactory::SetupMaterial(UMaterial* Material, const FCreateAssetParams& Arguments, const UInterchangeBaseMaterialFactoryNode* BaseMaterialFactoryNode)
 {
diff --git a/Engine/Plugins/Experimental/Interchange/Runtime/Source/Import/Private/Mesh/InterchangeSkeletalMeshFactory.cpp b/Engine/Plugins/Experimental/Interchange/Runtime/Source/Import/Private/Mesh/InterchangeSkeletalMeshFactory.cpp
--- a/Engine/Plugins/Experimental/Interchange/Runtime/Source/Import/Private/Mesh/InterchangeSkeletalMeshFactory.cpp
+++ b/Engine/Plugins/Experimental/Interchange/Runtime/Source/Import/Private/Mesh/InterchangeSkeletalMeshFactory.cpp
@@ -1503,3 +1503,39 @@ void UInterchangeSkeletalMeshFactory::PreImportPreCompletedCallback(const FImpor
 	}
 #endif
 }
+
+bool UInterchangeSkeletalMeshFactory::GetSourceFilenames(const UObject* Object, TArray<FString>& OutSourceFilenames) const
+{
+#if WITH_EDITORONLY_DATA
+	if (const USkeletalMesh* SkeletalMesh = Cast<USkeletalMesh>(Object))
+	{
+		return UE::Interchange::FFactoryCommon::GetSourceFilenames(SkeletalMesh->GetAssetImportData(), OutSourceFilenames);
+	}
+#endif
+
+	return false;
+}
+
+bool UInterchangeSkeletalMeshFactory::SetSourceFilename(const UObject* Object, const FString& SourceFilename, int32 SourceIndex) const
+{
+#if WITH_EDITORONLY_DATA
+	if (const USkeletalMesh* SkeletalMesh = Cast<USkeletalMesh>(Object))
+	{
+		return UE::Interchange::FFactoryCommon::SetSourceFilename(SkeletalMesh->GetAssetImportData(), SourceFilename, SourceIndex);
+	}
+#endif
+
+	return false;
+}
+
+bool UInterchangeSkeletalMeshFactory::SetReimportSourceIndex(const UObject* Object, int32 SourceIndex) const
+{
+#if WITH_EDITORONLY_DATA
+	if (const USkeletalMesh* SkeletalMesh = Cast<USkeletalMesh>(Object))
+	{
+		return UE::Interchange::FFactoryCommon::SetReimportSourceIndex(SkeletalMesh, SkeletalMesh->GetAssetImportData(), SourceIndex);
+	}
+#endif
+
+	return false;
+}
diff --git a/Engine/Plugins/Experimental/Interchange/Runtime/Source/Import/Private/Mesh/InterchangeStaticMeshFactory.cpp b/Engine/Plugins/Experimental/Interchange/Runtime/Source/Import/Private/Mesh/InterchangeStaticMeshFactory.cpp
--- a/Engine/Plugins/Experimental/Interchange/Runtime/Source/Import/Private/Mesh/InterchangeStaticMeshFactory.cpp
+++ b/Engine/Plugins/Experimental/Interchange/Runtime/Source/Import/Private/Mesh/InterchangeStaticMeshFactory.cpp
@@ -1338,3 +1338,27 @@ bool UInterchangeStaticMeshFactory::ImportSockets(const FCreateAssetParams& Argu
 
 	return true;
 }
+
+bool UInterchangeStaticMeshFactory::GetSourceFilenames(const UObject* Object, TArray<FString>& OutSourceFilenames) const
+{
+#if WITH_EDITORONLY_DATA
+	if (const UStaticMesh* StaticMesh = Cast<UStaticMesh>(Object))
+	{
+		return UE::Interchange::FFactoryCommon::GetSourceFilenames(StaticMesh->GetAssetImportData(), OutSourceFilenames);
+	}
+#endif
+
+	return false;
+}
+
+bool UInterchangeStaticMeshFactory::SetSourceFilename(const UObject* Object, const FString& SourceFilename, int32 SourceIndex) const
+{
+#if WITH_EDITORONLY_DATA
+	if (const UStaticMesh* StaticMesh = Cast<UStaticMesh>(Object))
+	{
+		return UE::Interchange::FFactoryCommon::SetSourceFilename(StaticMesh->GetAssetImportData(), SourceFilename, SourceIndex);
+	}
+#endif
+
+	return false;
+}
diff --git a/Engine/Plugins/Experimental/Interchange/Runtime/Source/Import/Private/Texture/InterchangeTextureFactory.cpp b/Engine/Plugins/Experimental/Interchange/Runtime/Source/Import/Private/Texture/InterchangeTextureFactory.cpp
--- a/Engine/Plugins/Experimental/Interchange/Runtime/Source/Import/Private/Texture/InterchangeTextureFactory.cpp
+++ b/Engine/Plugins/Experimental/Interchange/Runtime/Source/Import/Private/Texture/InterchangeTextureFactory.cpp
@@ -1293,3 +1293,27 @@ void UInterchangeTextureFactory::PreImportPreCompletedCallback(const FImportPreC
 	}
 #endif
 }
+
+bool UInterchangeTextureFactory::GetSourceFilenames(const UObject* Object, TArray<FString>& OutSourceFilenames) const
+{
+#if WITH_EDITORONLY_DATA
+	if (const UTexture* Texture = Cast<UTexture>(Object))
+	{
+		return UE::Interchange::FFactoryCommon::GetSourceFilenames(Texture->AssetImportData.Get(), OutSourceFilenames);
+	}
+#endif
+
+	return false;
+}
+
+bool UInterchangeTextureFactory::SetSourceFilename(const UObject* Object, const FString& SourceFilename, int32 SourceIndex) const
+{
+#if WITH_EDITORONLY_DATA
+	if (const UTexture* Texture = Cast<UTexture>(Object))
+	{
+		return UE::Interchange::FFactoryCommon::SetSourceFilename(Texture->AssetImportData.Get(), SourceFilename, SourceIndex);
+	}
+#endif
+
+	return false;
+}
\ No newline at end of file
diff --git a/Engine/Plugins/Experimental/Interchange/Runtime/Source/Import/Public/Animation/InterchangeAnimSequenceFactory.h b/Engine/Plugins/Experimental/Interchange/Runtime/Source/Import/Public/Animation/InterchangeAnimSequenceFactory.h
--- a/Engine/Plugins/Experimental/Interchange/Runtime/Source/Import/Public/Animation/InterchangeAnimSequenceFactory.h
+++ b/Engine/Plugins/Experimental/Interchange/Runtime/Source/Import/Public/Animation/InterchangeAnimSequenceFactory.h
@@ -30,6 +30,8 @@ public:
 		//Currently we cannot use the anim sequence controller outside of the game thread
 		return false;
 	}
+	virtual bool GetSourceFilenames(const UObject* Object, TArray<FString>& OutSourceFilenames) const override;
+	virtual bool SetSourceFilename(const UObject* Object, const FString& SourceFilename, int32 SourceIndex) const override;
 
 	// Interchange factory base interface end
 	//////////////////////////////////////////////////////////////////////////
diff --git a/Engine/Plugins/Experimental/Interchange/Runtime/Source/Import/Public/Material/InterchangeMaterialFactory.h b/Engine/Plugins/Experimental/Interchange/Runtime/Source/Import/Public/Material/InterchangeMaterialFactory.h
--- a/Engine/Plugins/Experimental/Interchange/Runtime/Source/Import/Public/Material/InterchangeMaterialFactory.h
+++ b/Engine/Plugins/Experimental/Interchange/Runtime/Source/Import/Public/Material/InterchangeMaterialFactory.h
@@ -28,6 +28,8 @@ public:
 	virtual UObject* CreateEmptyAsset(const FCreateAssetParams& Arguments) override;
 	virtual UObject* CreateAsset(const FCreateAssetParams& Arguments) override;
 	virtual void PreImportPreCompletedCallback(const FImportPreCompletedCallbackParams& Arguments) override;
+	virtual bool GetSourceFilenames(const UObject* Object, TArray<FString>& OutSourceFilenames) const override;
+	virtual bool SetSourceFilename(const UObject* Object, const FString& SourceFilename, int32 SourceIndex) const override;
 
 	// Interchange factory base interface end
 	//////////////////////////////////////////////////////////////////////////
diff --git a/Engine/Plugins/Experimental/Interchange/Runtime/Source/Import/Public/Mesh/InterchangeSkeletalMeshFactory.h b/Engine/Plugins/Experimental/Interchange/Runtime/Source/Import/Public/Mesh/InterchangeSkeletalMeshFactory.h
--- a/Engine/Plugins/Experimental/Interchange/Runtime/Source/Import/Public/Mesh/InterchangeSkeletalMeshFactory.h
+++ b/Engine/Plugins/Experimental/Interchange/Runtime/Source/Import/Public/Mesh/InterchangeSkeletalMeshFactory.h
@@ -25,6 +25,9 @@ public:
 	virtual UObject* CreateAsset(const FCreateAssetParams& Arguments) override;
 	virtual void PreImportPreCompletedCallback(const FImportPreCompletedCallbackParams& Arguments) override;
 	//virtual void PostImportPreCompletedCallback(const FImportPreCompletedCallbackParams& Arguments) const override;
+	virtual bool GetSourceFilenames(const UObject* Object, TArray<FString>& OutSourceFilenames) const override;
+	virtual bool SetSourceFilename(const UObject* Object, const FString& SourceFilename, int32 SourceIndex) const override;
+	virtual bool SetReimportSourceIndex(const UObject* Object, int32 SourceIndex) const override;
 
 	// Interchange factory base interface end
 	//////////////////////////////////////////////////////////////////////////
diff --git a/Engine/Plugins/Experimental/Interchange/Runtime/Source/Import/Public/Mesh/InterchangeStaticMeshFactory.h b/Engine/Plugins/Experimental/Interchange/Runtime/Source/Import/Public/Mesh/InterchangeStaticMeshFactory.h
--- a/Engine/Plugins/Experimental/Interchange/Runtime/Source/Import/Public/Mesh/InterchangeStaticMeshFactory.h
+++ b/Engine/Plugins/Experimental/Interchange/Runtime/Source/Import/Public/Mesh/InterchangeStaticMeshFactory.h
@@ -32,6 +32,8 @@ public:
 	virtual UObject* CreateEmptyAsset(const FCreateAssetParams& Arguments) override;
 	virtual UObject* CreateAsset(const FCreateAssetParams& Arguments) override;
 	virtual void PreImportPreCompletedCallback(const FImportPreCompletedCallbackParams& Arguments) override;
+	virtual bool GetSourceFilenames(const UObject* Object, TArray<FString>& OutSourceFilenames) const override;
+	virtual bool SetSourceFilename(const UObject* Object, const FString& SourceFilename, int32 SourceIndex) const override;
 
 	// Interchange factory base interface end
 	//////////////////////////////////////////////////////////////////////////
diff --git a/Engine/Plugins/Experimental/Interchange/Runtime/Source/Import/Public/Texture/InterchangeTextureFactory.h b/Engine/Plugins/Experimental/Interchange/Runtime/Source/Import/Public/Texture/InterchangeTextureFactory.h
--- a/Engine/Plugins/Experimental/Interchange/Runtime/Source/Import/Public/Texture/InterchangeTextureFactory.h
+++ b/Engine/Plugins/Experimental/Interchange/Runtime/Source/Import/Public/Texture/InterchangeTextureFactory.h
@@ -56,6 +56,9 @@ public:
 	virtual UObject* CreateEmptyAsset(const FCreateAssetParams& Arguments) override;
 	virtual UObject* CreateAsset(const FCreateAssetParams& Arguments) override;
 	virtual void PreImportPreCompletedCallback(const FImportPreCompletedCallbackParams& Arguments) override;
+	virtual bool GetSourceFilenames(const UObject* Object, TArray<FString>& OutSourceFilenames) const override;
+	virtual bool SetSourceFilename(const UObject* Object, const FString& SourceFilename, int32 SourceIndex) const override;
+
 	// Interchange factory base interface end
 	//////////////////////////////////////////////////////////////////////////
 
diff --git a/Engine/Source/Editor/UnrealEd/Private/Factories/InterchangeReimportHandler.cpp b/Engine/Source/Editor/UnrealEd/Private/Factories/InterchangeReimportHandler.cpp
--- a/Engine/Source/Editor/UnrealEd/Private/Factories/InterchangeReimportHandler.cpp
+++ b/Engine/Source/Editor/UnrealEd/Private/Factories/InterchangeReimportHandler.cpp
@@ -10,6 +10,7 @@
 #include "Engine/Texture.h"
 #include "Factories/Factory.h"
 #include "InterchangeAssetImportData.h"
+#include "InterchangeFactoryBase.h"
 #include "InterchangeManager.h"
 #include "Settings/EditorExperimentalSettings.h"
 #include "UObject/UObjectGlobals.h"
@@ -25,113 +26,38 @@ UInterchangeReimportHandler::UInterchangeReimportHandler(const FObjectInitialize
 //~ Begin FReimportHandler Interface
 bool UInterchangeReimportHandler::CanReimport(UObject* Obj, TArray<FString>& OutFilenames)
 {
-	auto CanReimportAsset = [&OutFilenames](UAssetImportData* AssetImportData)
-	{
-		if (!AssetImportData)
-		{
-			return false;
-		}
-		UInterchangeAssetImportData* InterchangeAssetImportData = Cast<UInterchangeAssetImportData>(AssetImportData);
-		if (InterchangeAssetImportData == nullptr)
-		{
-			return false;
-		}
-		AssetImportData->ExtractFilenames(OutFilenames);
-		return true;
-	};
-
-	
 	const UEditorExperimentalSettings* EditorExperimentalSettings = GetDefault<UEditorExperimentalSettings>();
-	const bool bUseInterchangeFramework = EditorExperimentalSettings->bEnableInterchangeFramework;
-	const bool bUseInterchangeFrameworkForTextureOnly = (!bUseInterchangeFramework) && EditorExperimentalSettings->bEnableInterchangeFrameworkForTextureOnly;
+	const bool bUseInterchange = EditorExperimentalSettings->bEnableInterchangeFramework 
+		|| (EditorExperimentalSettings->bEnableInterchangeFrameworkForTextureOnly && Cast<UTexture>(Obj) != nullptr);
 
-	USkeletalMesh* SkeletalMesh = Cast<USkeletalMesh>(Obj);
-	UStaticMesh* StaticMesh = Cast<UStaticMesh>(Obj);
-	UAnimSequence* AnimSequence = Cast<UAnimSequence>(Obj);
-	UTexture* Texture = Cast<UTexture>(Obj);
-	if (bUseInterchangeFramework && SkeletalMesh)
-	{
-		if (UAssetImportData* AssetImportData = SkeletalMesh->GetAssetImportData())
-		{
-			return CanReimportAsset(AssetImportData);
-		}
-	}
-	else if (bUseInterchangeFramework && StaticMesh)
-	{
-		if (UAssetImportData* AssetImportData = StaticMesh->GetAssetImportData())
-		{
-			return CanReimportAsset(AssetImportData);
-		}
-	}
-	else if (bUseInterchangeFramework && AnimSequence)
-	{
-		if (UAssetImportData* AssetImportData = AnimSequence->AssetImportData)
-		{
-			return CanReimportAsset(AssetImportData);
-		}
-	}
-	else if ((bUseInterchangeFramework || bUseInterchangeFrameworkForTextureOnly) && Texture)
+	if (bUseInterchange)
 	{
-		if (UAssetImportData* AssetImportData = Texture->AssetImportData)
-		{
-			return CanReimportAsset(AssetImportData);
-		}
+		const UInterchangeManager& InterchangeManager = UInterchangeManager::GetInterchangeManager();
+		return InterchangeManager.CanReimport(Obj, OutFilenames);
 	}
+
 	return false;
 }
 
 void UInterchangeReimportHandler::SetReimportPaths(UObject* Obj, const FString& NewReimportPath, const int32 SourceFileIndex)
 {
-	USkeletalMesh* SkeletalMesh = Cast<USkeletalMesh>(Obj);
-	UStaticMesh* StaticMesh = Cast<UStaticMesh>(Obj);
-	UAnimSequence* AnimSequence = Cast<UAnimSequence>(Obj);
-	UTexture* Texture = Cast<UTexture>(Obj);
+	const UInterchangeManager& InterchangeManager = UInterchangeManager::GetInterchangeManager();
 
-	UAssetImportData* AssetImportData = nullptr;
-	if (SkeletalMesh)
+	if (const UClass* FactoryClass = InterchangeManager.GetRegisteredFactoryClass(Obj->GetClass()))
 	{
-		AssetImportData = SkeletalMesh->GetAssetImportData();
-	}
-	else if (StaticMesh)
-	{
-		AssetImportData = StaticMesh->GetAssetImportData();
-	}
-	else if (AnimSequence)
-	{
-		AssetImportData = AnimSequence->AssetImportData;
-	}
-	else if (Texture)
-	{
-		AssetImportData = Texture->AssetImportData;
-	}
-
-	if (AssetImportData)
-	{
-		int32 RealSourceFileIndex = SourceFileIndex == INDEX_NONE ? 0 : SourceFileIndex;
-		if (RealSourceFileIndex < AssetImportData->GetSourceFileCount())
-		{
-			AssetImportData->UpdateFilenameOnly(NewReimportPath, SourceFileIndex);
-		}
-		else
-		{
-			//Create a source file entry, this case happen when user import a specific content for the first time
-			FString SourceIndexLabel = USkeletalMesh::GetSourceFileLabelFromIndex(RealSourceFileIndex).ToString();
-			AssetImportData->AddFileName(NewReimportPath, RealSourceFileIndex, SourceIndexLabel);
-		}
+		UInterchangeFactoryBase* FactoryBase = FactoryClass->GetDefaultObject<UInterchangeFactoryBase>();
+		FactoryBase->SetSourceFilename(Obj, NewReimportPath, SourceFileIndex);
 	}
 }
 
 void UInterchangeReimportHandler::SetReimportSourceIndex(UObject* Obj, const int32 SourceIndex)
 {
-	if (USkeletalMesh* SkeletalMesh = Cast<USkeletalMesh>(Obj))
+	const UInterchangeManager& InterchangeManager = UInterchangeManager::GetInterchangeManager();
+
+	if (const UClass* FactoryClass = InterchangeManager.GetRegisteredFactoryClass(Obj->GetClass()))
 	{
-		if (UInterchangeAssetImportData* InterchangeAssetImportData = Cast<UInterchangeAssetImportData>(SkeletalMesh->GetAssetImportData()))
-		{
-			for (UInterchangePipelineBase* PipelineBase : InterchangeAssetImportData->Pipelines)
-			{
-				PipelineBase->ScriptedSetReimportSourceIndex(Obj->GetClass(), SourceIndex);
-			}
-		}
+		UInterchangeFactoryBase* FactoryBase = FactoryClass->GetDefaultObject<UInterchangeFactoryBase>();
+		FactoryBase->SetReimportSourceIndex(Obj, SourceIndex);
 	}
 }
 
diff --git a/Engine/Source/Runtime/Experimental/Interchange/Core/Public/InterchangeFactoryBase.h b/Engine/Source/Runtime/Experimental/Interchange/Core/Public/InterchangeFactoryBase.h
--- a/Engine/Source/Runtime/Experimental/Interchange/Core/Public/InterchangeFactoryBase.h
+++ b/Engine/Source/Runtime/Experimental/Interchange/Core/Public/InterchangeFactoryBase.h
@@ -165,6 +165,31 @@ public:
 		return;
 	}
 
+	/**
+	 * Fills the OutSourceFilenames array with the list of source files contained in the asset source data.
+	 * Returns true if the operation was successful.
+	 */
+	virtual bool GetSourceFilenames(const UObject* Object, TArray<FString>& OutSourceFilenames) const
+	{
+		return false;
+	}
+
+	/**
+	 * Sets the object's source at the specified index to the given SourceFileName
+	 */
+	virtual bool SetSourceFilename(const UObject* Object, const FString& SourceFilename, int32 SourceIndex) const
+	{
+		return false;
+	}
+
+	/**
+	 * Set the object's reimport source at the specified index value.
+	 */
+	virtual bool SetReimportSourceIndex(const UObject* Object, int32 SourceIndex) const
+	{
+		return false;
+	}
+
 	/**
 	 * This function is used to add the given message object directly into the results for this operation.
 	 */
diff --git a/Engine/Source/Runtime/Experimental/Interchange/Engine/Private/InterchangeImportCommon.cpp b/Engine/Source/Runtime/Experimental/Interchange/Engine/Private/InterchangeImportCommon.cpp
--- a/Engine/Source/Runtime/Experimental/Interchange/Engine/Private/InterchangeImportCommon.cpp
+++ b/Engine/Source/Runtime/Experimental/Interchange/Engine/Private/InterchangeImportCommon.cpp
@@ -181,6 +181,53 @@ namespace UE
 			// Return the asset import data so it can be set on the imported asset.
 			return AssetImportData;
 		}
+
+		bool FFactoryCommon::GetSourceFilenames(const UAssetImportData* AssetImportData, TArray<FString>& OutSourceFilenames)
+		{
+			if (Cast<UInterchangeAssetImportData>(AssetImportData) != nullptr)
+			{
+				AssetImportData->ExtractFilenames(OutSourceFilenames);
+				return true;
+			}
+			return false;
+		}
+
+		bool FFactoryCommon::SetSourceFilename(UAssetImportData* AssetImportData, const FString& SourceFilename, int32 SourceIndex)
+		{
+			if (AssetImportData)
+			{
+				int32 RealSourceFileIndex = SourceIndex == INDEX_NONE ? 0 : SourceIndex;
+				if (RealSourceFileIndex < AssetImportData->GetSourceFileCount())
+				{
+					AssetImportData->UpdateFilenameOnly(SourceFilename, SourceIndex);
+				}
+				else
+				{
+					//Create a source file entry, this case happen when user import a specific content for the first time
+					FString SourceIndexLabel = USkeletalMesh::GetSourceFileLabelFromIndex(RealSourceFileIndex).ToString();
+					AssetImportData->AddFileName(SourceFilename, RealSourceFileIndex, SourceIndexLabel);
+				}
+				return true;
+			}
+
+			return false;
+		}
+		
+		bool FFactoryCommon::SetReimportSourceIndex(const UObject* Object, UAssetImportData* AssetImportData, int32 SourceIndex)
+		{
+			UInterchangeAssetImportData* InterchangeAssetImportData = Cast<UInterchangeAssetImportData>(AssetImportData);
+			if (!InterchangeAssetImportData)
+			{
+				return false;
+			}
+
+			for (UInterchangePipelineBase* PipelineBase : InterchangeAssetImportData->Pipelines)
+			{
+				PipelineBase->ScriptedSetReimportSourceIndex(Object->GetClass(), SourceIndex);
+			}
+			return true;
+		}
+
 #endif //#if WITH_EDITORONLY_DATA
 
 
diff --git a/Engine/Source/Runtime/Experimental/Interchange/Engine/Private/InterchangeManager.cpp b/Engine/Source/Runtime/Experimental/Interchange/Engine/Private/InterchangeManager.cpp
--- a/Engine/Source/Runtime/Experimental/Interchange/Engine/Private/InterchangeManager.cpp
+++ b/Engine/Source/Runtime/Experimental/Interchange/Engine/Private/InterchangeManager.cpp
@@ -496,6 +496,34 @@ bool UInterchangeManager::CanTranslateSourceData(const UInterchangeSourceData* S
 	return GetTranslatorForSourceData(SourceData) != nullptr;
 }
 
+bool UInterchangeManager::CanReimport(const UObject* Object, TArray<FString>& OutFilenames) const
+{
+	const UClass* RegisteredFactoryClass = GetRegisteredFactoryClass(Object->GetClass());
+	if (!RegisteredFactoryClass)
+	{
+		return false;
+	}
+
+	UInterchangeFactoryBase* Factory = RegisteredFactoryClass->GetDefaultObject<UInterchangeFactoryBase>();
+	if(!Factory->GetSourceFilenames(Object, OutFilenames))
+	{
+		return false;
+	}
+
+	for (const FString& Filename : OutFilenames)
+	{
+		UE::Interchange::FScopedSourceData ScopedSourceData(Filename);
+
+		if (CanTranslateSourceData(ScopedSourceData.GetSourceData()))
+		{
+			return true;
+		}
+	}
+
+	OutFilenames.Empty();
+	return false;
+}
+
 void UInterchangeManager::StartQueuedTasks(bool bCancelAllTasks /*= false*/)
 {
 	if (!ensure(IsInGameThread()))
diff --git a/Engine/Source/Runtime/Experimental/Interchange/Engine/Public/InterchangeImportCommon.h b/Engine/Source/Runtime/Experimental/Interchange/Engine/Public/InterchangeImportCommon.h
--- a/Engine/Source/Runtime/Experimental/Interchange/Engine/Public/InterchangeImportCommon.h
+++ b/Engine/Source/Runtime/Experimental/Interchange/Engine/Public/InterchangeImportCommon.h
@@ -77,6 +77,23 @@ namespace UE
 			 * @return The source data that should be stored on the asset or nullptr if a parameter is invalid
 			 */
 			static UAssetImportData* SetImportAssetData(FSetImportAssetDataParameters& Parameters);
+
+			/**
+			 * Fills the OutSourceFilenames array with the list of source files contained in the asset source data.
+			 * Returns true if the operation was successful.
+			 */
+			static bool GetSourceFilenames(const UAssetImportData* AssetImportData, TArray<FString>& OutSourceFilenames);
+
+			/**
+			 * Sets the SourceFileName value at the specified index.
+			 */
+			static bool SetSourceFilename(UAssetImportData* AssetImportData, const FString& SourceFilename, int32 SourceIndex);
+
+			/**
+			 * Set the object's reimport source at the specified index value.
+			 */
+			static bool SetReimportSourceIndex(const UObject* Object, UAssetImportData* AssetImportData, int32 SourceIndex);
+
 #endif // WITH_EDITORONLY_DATA
 
 			/**
diff --git a/Engine/Source/Runtime/Experimental/Interchange/Engine/Public/InterchangeManager.h b/Engine/Source/Runtime/Experimental/Interchange/Engine/Public/InterchangeManager.h
--- a/Engine/Source/Runtime/Experimental/Interchange/Engine/Public/InterchangeManager.h
+++ b/Engine/Source/Runtime/Experimental/Interchange/Engine/Public/InterchangeManager.h
@@ -364,6 +364,13 @@ public:
 	 */
 	bool CanTranslateSourceData(const UInterchangeSourceData* SourceData) const;
 
+	/**
+	 * Returns true if Interchange can create that type of assets and is able to translate its source file.
+	 * @Param Object - The object we want to reimport.
+	 * @Param OutFilenames - An array that is filled with the object's source filenames if the operation is successful.
+	 */
+	bool CanReimport(const UObject* Object, TArray<FString>& OutFilenames) const;
+
 	/**
 	 * Call this to start an import asset process, the caller must specify a source data.
 	 * This import process can import many different asset, but all in the game content.
