Change how lazy load is configured to allow it to be driven by callback instead of by metadata.  Added a callback from CookOnTheFlyServer to allow it to lazy load references to packages it doesn't intend to save (ie: iteratively skipped packages).  This is only effective when cooking in a commandlet (not in editor).

diff --git a/Engine/Source/Editor/UnrealEd/Classes/CookOnTheSide/CookOnTheFlyServer.h b/Engine/Source/Editor/UnrealEd/Classes/CookOnTheSide/CookOnTheFlyServer.h
--- a/Engine/Source/Editor/UnrealEd/Classes/CookOnTheSide/CookOnTheFlyServer.h
+++ b/Engine/Source/Editor/UnrealEd/Classes/CookOnTheSide/CookOnTheFlyServer.h
@@ -33,6 +33,8 @@ class IPlugin;
 class ITargetPlatform;
 struct FPropertyChangedEvent;
 
+namespace UE::LinkerLoad { enum class EImportBehavior : uint8; }
+
 enum class ECookInitializationFlags
 {
 	None =										0x00000000, // No flags
@@ -294,6 +296,8 @@ private:
 
 	void ProcessAccessedIniSettings(const FConfigFile* Config, FIniSettingContainer& AccessedIniStrings) const;
 
+	void OnRequestClusterCompleted(const UE::Cook::FRequestCluster& RequestCluster);
+
 	/**
 	* OnTargetPlatformChangedSupportedFormats
 	* called when target platform changes the return value of supports shader formats 
@@ -1116,6 +1120,11 @@ private:
 	/** Allocate a new FCookSavePackageContext and ICookedPackageWriter for the given platform. */
 	UE::Cook::FCookSavePackageContext* CreateSaveContext(const ITargetPlatform* TargetPlatform);
 
+	void ConditionalInstallImportBehaviorCallback();
+	void ConditionalUninstallImportBehaviorCallback();
+	static void PropertyImportBehaviorCallback(const FObjectImport& Import, const FLinkerLoad& LinkerLoad, UE::LinkerLoad::EImportBehavior& OutBehavior);
+	
+	static UCookOnTheFlyServer* ActiveCOTFS;
 	uint32		StatLoadedPackageCount = 0;
 	uint32		StatSavedPackageCount = 0;
 
@@ -1135,6 +1144,7 @@ private:
 	bool bHybridIterativeDebug = false;
 	bool bHasBlockedOnAssetRegistry = false;
 	bool bHasDeferredInitializeCookOnTheFly = false;
+	bool bImportBehaviorCallbackInstalled = false;
 
 
 	/** Timers for tracking how long we have been busy, to manage retries and warnings of deadlock */
diff --git a/Engine/Source/Editor/UnrealEd/Private/CookOnTheFlyServer.cpp b/Engine/Source/Editor/UnrealEd/Private/CookOnTheFlyServer.cpp
--- a/Engine/Source/Editor/UnrealEd/Private/CookOnTheFlyServer.cpp
+++ b/Engine/Source/Editor/UnrealEd/Private/CookOnTheFlyServer.cpp
@@ -37,6 +37,7 @@
 #include "Editor.h"
 #include "Editor/UnrealEdEngine.h"
 #include "EditorDomain/EditorDomain.h"
+#include "EditorDomain/EditorDomainUtils.h"
 #include "Engine/AssetManager.h"
 #include "Engine/Level.h"
 #include "Engine/LevelStreaming.h"
@@ -114,6 +115,7 @@
 #include "UObject/Class.h"
 #include "UObject/ConstructorHelpers.h"
 #include "UObject/GarbageCollection.h"
+#include "UObject/LinkerLoadImportBehavior.h"
 #include "UObject/MetaData.h"
 #include "UObject/ObjectSaveContext.h"
 #include "UObject/Package.h"
@@ -182,6 +184,7 @@ static FAutoConsoleVariableRef CVarCookDisplayWarnBusyTime(
 ////////////////////////////////////////////////////////////////
 /// Cook on the fly server
 ///////////////////////////////////////////////////////////////
+UCookOnTheFlyServer* UCookOnTheFlyServer::ActiveCOTFS = nullptr;
 
 constexpr float TickCookableObjectsFrameTime = .100f;
 
@@ -2010,18 +2013,7 @@ void UCookOnTheFlyServer::PumpRequests(UE::Cook::FTickStackData& StackData, int3
 				PackageData->SendToState(EPackageState::Idle, ESendFlags::QueueAdd);
 			}
 			RequestClusters.PopFront();
-			if (IsCookByTheBookMode() && bHybridIterativeDebug)
-			{
-				for (FPackageData* PackageData : *PackageDatas)
-				{
-					if (!PackageData->AreAllRequestedPlatformsExplored())
-					{
-						UE_LOG(LogCook, Warning, TEXT("Missing dependency: existing requested Package %s was not explored in the first cluster."),
-							*WriteToString<256>(PackageData->GetPackageName()));
-					}
-				}
-				PackageDatas->SetLogDiscoveredPackages(true); // Turn this on for the rest of the cook after the initial cluster
-			}
+			OnRequestClusterCompleted(RequestCluster);
 		}
 		if (CookerTimer.IsTimeUp())
 		{
@@ -3845,6 +3837,7 @@ void UCookOnTheFlyServer::ShutdownCookOnTheFly()
 		CookOnTheFlyRequestManager->Shutdown();
 		CookOnTheFlyRequestManager.Reset();
 		ClearHierarchyTimers();
+		ConditionalUninstallImportBehaviorCallback();
 	}
 }
 
@@ -4143,6 +4136,7 @@ void UCookOnTheFlyServer::PostGarbageCollect()
 void UCookOnTheFlyServer::BeginDestroy()
 {
 	ShutdownCookOnTheFly();
+	ConditionalUninstallImportBehaviorCallback();
 
 	Super::BeginDestroy();
 }
@@ -5709,6 +5703,28 @@ bool UCookOnTheFlyServer::SaveCurrentIniSettings(const ITargetPlatform* TargetPl
 
 }
 
+void UCookOnTheFlyServer::OnRequestClusterCompleted(const UE::Cook::FRequestCluster& RequestCluster)
+{
+	using namespace UE::Cook;
+	if (IsCookByTheBookMode())
+	{
+		if (bHybridIterativeDebug)
+		{
+			for (FPackageData* PackageData : *PackageDatas)
+			{
+				if (!PackageData->AreAllRequestedPlatformsExplored())
+				{
+					UE_LOG(LogCook, Warning, TEXT("Missing dependency: existing requested Package %s was not explored in the first cluster."),
+						*WriteToString<256>(PackageData->GetPackageName()));
+				}
+			}
+			PackageDatas->SetLogDiscoveredPackages(true); // Turn this on for the rest of the cook after the initial cluster
+		}
+
+		ConditionalInstallImportBehaviorCallback();
+	}
+}
+
 FAsyncIODelete& UCookOnTheFlyServer::GetAsyncIODelete()
 {
 	if (AsyncIODelete)
@@ -6961,6 +6977,8 @@ void UCookOnTheFlyServer::CookByTheBookFinished()
 	BuildDefinitions->Wait();
 	
 	GetDerivedDataCacheRef().WaitForQuiescence(true);
+
+	ConditionalUninstallImportBehaviorCallback();
 	
 	UCookerSettings const* CookerSettings = GetDefault<UCookerSettings>();
 
@@ -9180,5 +9198,86 @@ UE::Cook::FCookSavePackageContext& UCookOnTheFlyServer::FindOrCreateSaveContext(
 	return *SavePackageContexts.Add_GetRef(CreateSaveContext(TargetPlatform));
 }
 
+void UCookOnTheFlyServer::ConditionalInstallImportBehaviorCallback()
+{
+	if (!bImportBehaviorCallbackInstalled && !IsCookingInEditor() && FEditorDomain::Get())
+	{
+		check(ActiveCOTFS == nullptr);
+		ActiveCOTFS = this;
+		UE::LinkerLoad::SetPropertyImportBehaviorCallback(&PropertyImportBehaviorCallback);
+		bImportBehaviorCallbackInstalled = true;
+	}
+}
+
+void UCookOnTheFlyServer::ConditionalUninstallImportBehaviorCallback()
+{
+	if (bImportBehaviorCallbackInstalled)
+	{
+		check(ActiveCOTFS == this);
+		UE::LinkerLoad::SetPropertyImportBehaviorCallback((UE::LinkerLoad::PropertyImportBehaviorFunction*)nullptr);
+		ActiveCOTFS = nullptr;
+		bImportBehaviorCallbackInstalled = false;
+	}
+}
+
+void UCookOnTheFlyServer::PropertyImportBehaviorCallback(const FObjectImport& Import, const FLinkerLoad& LinkerLoad, UE::LinkerLoad::EImportBehavior& OutBehavior)
+{
+	auto GetImportPackageName = [&LinkerLoad](const FObjectImport& Import)
+	{
+		const FObjectImport* CurrentResource = &Import;
+		for (int32 NumCycles = 0; NumCycles < LinkerLoad.ImportMap.Num(); ++NumCycles)
+		{
+			// If the import has a package name set, then that's the import package name,
+			if (CurrentResource->HasPackageName())
+			{
+				return CurrentResource->GetPackageName();
+			}
+			// If our outer is null, then we have a package
+			else if (CurrentResource->OuterIndex.IsNull())
+			{
+				return CurrentResource->ObjectName;
+			}
+			if (!CurrentResource->OuterIndex.IsImport())
+			{
+				return FName();
+			}
+			CurrentResource = &LinkerLoad.Imp(CurrentResource->OuterIndex);
+		}
+		return FName();
+	};
+
+	if (ActiveCOTFS != nullptr)
+	{
+		FName ImportPackageName = GetImportPackageName(Import);
+		if (!ImportPackageName.IsNone())
+		{
+			UE::Cook::FPackageData* PackageData = ActiveCOTFS->PackageDatas->FindPackageDataByPackageName(ImportPackageName);
+
+			if (PackageData && PackageData->IsInProgress())
+			{
+				OutBehavior = UE::LinkerLoad::EImportBehavior::Eager;
+			}
+			else
+			{
+				FNameBuilder ClassPathBuilder(Import.ClassPackage);
+				ClassPathBuilder.AppendChar(TEXT('.'));
+				ClassPathBuilder << Import.ClassName;
+				FName ClassPath(ClassPathBuilder.ToView());
+				UE::EditorDomain::FClassDigestMap& ClassDigests = UE::EditorDomain::GetClassDigests();
+				FReadScopeLock ClassDigestsScopeLock(ClassDigests.Lock);
+				UE::EditorDomain::FClassDigestData* ExistingData = ClassDigests.Map.Find(ClassPath);
+				if (!ExistingData || !ExistingData->bTargetIterativeEnabled)
+				{
+					OutBehavior = UE::LinkerLoad::EImportBehavior::Eager;
+				}
+				else
+				{
+					OutBehavior = UE::LinkerLoad::EImportBehavior::LazyOnDemand;
+				}
+			}
+		}
+	}
+}
+
 
 #undef LOCTEXT_NAMESPACE
diff --git a/Engine/Source/Runtime/CoreUObject/Internal/UObject/LinkerLoadImportBehavior.h b/Engine/Source/Runtime/CoreUObject/Internal/UObject/LinkerLoadImportBehavior.h
new file mode 100644
--- /dev/null
+++ b/Engine/Source/Runtime/CoreUObject/Internal/UObject/LinkerLoadImportBehavior.h
@@ -0,0 +1,30 @@
+// Copyright Epic Games, Inc. All Rights Reserved.
+
+#pragma once
+#include "UObject/ObjectHandle.h"
+
+#if UE_WITH_OBJECT_HANDLE_LATE_RESOLVE
+#define UE_API COREUOBJECT_API
+
+class FLinkerLoad;
+struct FObjectImport;
+
+namespace UE::LinkerLoad
+{
+enum class EImportBehavior : uint8
+{
+	Eager = 0,
+	// @TODO: OBJPTR: we want to permit lazy background loading in the future
+	//LazyBackground,
+	LazyOnDemand,
+};
+
+using PropertyImportBehaviorFunction = void(const FObjectImport& Import, const FLinkerLoad& LinkerLoad, EImportBehavior& OutBehavior);
+
+UE_API void SetPropertyImportBehaviorCallback(PropertyImportBehaviorFunction* Function);
+
+void GetPropertyImportLoadBehavior(const FObjectImport& Import, const FLinkerLoad& LinkerLoad, EImportBehavior& OutBehavior);
+}
+
+#undef UE_API
+#endif // UE_WITH_OBJECT_HANDLE_LATE_RESOLVE
diff --git a/Engine/Source/Runtime/CoreUObject/Private/UObject/LinkerLoad.cpp b/Engine/Source/Runtime/CoreUObject/Private/UObject/LinkerLoad.cpp
--- a/Engine/Source/Runtime/CoreUObject/Private/UObject/LinkerLoad.cpp
+++ b/Engine/Source/Runtime/CoreUObject/Private/UObject/LinkerLoad.cpp
@@ -11,6 +11,7 @@
 #include "Misc/ObjectThumbnail.h"
 #include "Misc/App.h"
 #include "UObject/MetaData.h"
+#include "UObject/LinkerLoadImportBehavior.h"
 #include "UObject/ObjectRedirector.h"
 #include "UObject/Package.h"
 #include "UObject/PackageResourceManager.h"
@@ -1963,98 +1964,6 @@ FStructuredArchiveSlot FLinkerLoad::GetExportSlot(FPackageIndex InExportIndex)
 	return ExportReaders[Index]->GetRoot();
 }
 
-#if UE_WITH_OBJECT_HANDLE_LATE_RESOLVE
-FLinkerLoad::EImportLoadBehavior FLinkerLoad::ParseImportLoadBehavior(const FString* LoadBehaviorMeta)
-{
-	EImportLoadBehavior LoadBehavior = EImportLoadBehavior::Eager;
-	if (LoadBehaviorMeta)
-	{
-		if (*LoadBehaviorMeta == "Eager")
-		{
-			LoadBehavior = EImportLoadBehavior::Eager;
-		}
-		// @TODO: OBJPTR: we want to permit lazy background loading in the future
-		// else if (*LoadBehaviorMeta == "LazyBackground")
-		// {
-		// 	LoadBehavior = EImportLoadBehavior::LazyBackground;
-		// }
-		else if (*LoadBehaviorMeta == "LazyOnDemand")
-		{
-			LoadBehavior = EImportLoadBehavior::LazyOnDemand;
-		}
-	}
-
-	return LoadBehavior;
-}
-
-FLinkerLoad::EImportLoadBehavior FLinkerLoad::GetCurrentPropertyImportLoadBehavior(FPackageIndex ImportIndex)
-{
-	static const FName Name_LoadBehavior("LoadBehavior");
-	const FObjectImport& Import = Imp(ImportIndex);
-	if (Import.bImportSearchedFor)
-	{
-		// If it was something that's been searched for, we've already attempted a resolve, might as well use it
-		return EImportLoadBehavior::Eager;
-	}
-
-	if (!IsImportLazyLoadEnabled() || !IsAllowingLazyLoading() || (LinkerRoot && LinkerRoot->HasAnyPackageFlags(PKG_PlayInEditor)))
-	{
-		return EImportLoadBehavior::Eager;
-	}
-
-	UObject* ClassPackage = FindObject<UPackage>(nullptr, *Import.ClassPackage.ToString());
-	UClass* FindClass = ClassPackage ? FindObject<UClass>(ClassPackage, *Import.ClassName.ToString()) : nullptr;
-
-	// Attempt to get the meta from the referenced class.  This only looks in already loaded classes.  May need to resolve the class in the future.
-	EImportLoadBehavior LoadBehavior = EImportLoadBehavior::Eager;
-	static const bool bDefaultLoadBehaviorTest = FParse::Param(FCommandLine::Get(), TEXT("DefaultLoadBehaviorTest"));
-	if (bDefaultLoadBehaviorTest)
-	{
-		LoadBehavior = EImportLoadBehavior::LazyOnDemand;
-	}
-
-	if (FindClass)
-	{
-		if (const FString* LoadBehaviorMeta = FindClass->FindMetaData(Name_LoadBehavior))
-		{
-			LoadBehavior = ParseImportLoadBehavior(LoadBehaviorMeta);
-		}
-	}
-
-	const FArchiveSerializedPropertyChain* PropChain = GetSerializedPropertyChain();
-	if (!PropChain)
-	{
-		return LoadBehavior;
-	}
-
-	const int32 NumPropsInChain = PropChain->GetNumProperties();
-	const FProperty* CurrentProperty = NumPropsInChain > 0 ? PropChain->GetPropertyFromStack(0) : nullptr;
-	if (!CurrentProperty)
-	{
-		return LoadBehavior;
-	}
-
-	const FProperty* ParentProperty = NumPropsInChain > 1 ? PropChain->GetPropertyFromStack(1) : nullptr;
-	// Exclude properties that are in hash sensitive containers as they will require the object to be loaded to compute a hash
-	// @TODO: OBJPTR: Forcing resolve for ObjectPtr hashing.  See GetTypeHash in ObjectHandle.h.
-	if (ParentProperty && 
-		(
-			ParentProperty->IsA<FSetProperty>() ||
-			(ParentProperty->IsA<FMapProperty>() && (CurrentProperty == static_cast<const FMapProperty*>(ParentProperty)->GetKeyProperty()))
-		))
-	{
-		return EImportLoadBehavior::Eager;
-	}
-
-	if (const FString* LoadBehaviorMeta = CurrentProperty->FindMetaData(Name_LoadBehavior))
-	{
-		LoadBehavior = ParseImportLoadBehavior(LoadBehaviorMeta);
-	}
-
-	return LoadBehavior;
-}
-#endif // UE_WITH_OBJECT_HANDLE_LATE_RESOLVE
-
 FString ExtractObjectName(const FString& InFullPath)
 {
 	FString ObjectName = InFullPath;
@@ -2580,15 +2489,15 @@ void FLinkerLoad::Verify()
 {
 	if (!bHaveImportsBeenVerified)
 	{
-		bool bShouldVerifiedAllImports = IsRunningCommandlet();
+		bool bShouldVerifyAllImports = IsRunningCommandlet();
 
 #if WITH_EDITOR
 		// In editor builds using OFPA, we need to resolve imports for BP classes referenced by the level script in order to be able
 		// to properly reinstanced them. We could filter out imports to resolve here, but we resolve all of them instead.
-		bShouldVerifiedAllImports = true;
+		bShouldVerifyAllImports = true;
 #endif
 
-		if (!IsImportLazyLoadEnabled() && bShouldVerifiedAllImports)
+		if (!IsImportLazyLoadEnabled() && bShouldVerifyAllImports)
 		{
 #if WITH_EDITOR
 			TOptional<FScopedSlowTask> SlowTask;
@@ -5512,8 +5421,10 @@ FArchive& FLinkerLoad::operator<<( FObjectPtr& ObjectPtr )
 #if UE_WITH_OBJECT_HANDLE_LATE_RESOLVE
 	if (Index.IsImport())
 	{
-		EImportLoadBehavior LoadBehavior = GetCurrentPropertyImportLoadBehavior(Index);
-		if (LoadBehavior == EImportLoadBehavior::LazyOnDemand)
+		using namespace UE::LinkerLoad;
+		EImportBehavior LoadBehavior = EImportBehavior::Eager;
+		GetPropertyImportLoadBehavior(Imp(Index), *this, LoadBehavior);
+		if (LoadBehavior == EImportBehavior::LazyOnDemand)
 		{
 			const FObjectImport& Import = Imp(Index);
 			FObjectPathId ObjectPath;
@@ -5765,8 +5676,22 @@ bool FLinkerLoad::RemoveKnownMissingPackage(FName PackageName)
 #if UE_WITH_OBJECT_HANDLE_LATE_RESOLVE
 bool FLinkerLoad::IsImportLazyLoadEnabled()
 {
-	static const bool bLazyLoadImports = FParse::Param(FCommandLine::Get(), TEXT("LazyLoadImports"));
-	return bLazyLoadImports;
+	auto ImportLazyLoadEnabled = []()
+	{
+		if (FParse::Param(FCommandLine::Get(), TEXT("LazyLoadImports")))
+		{
+			return true;
+		}
+		else if (GConfig)
+		{
+			bool bLazyLoadImportsConfig = false;
+			GConfig->GetBool(TEXT("Core.System.Experimental"), TEXT("LazyLoadImports"), bLazyLoadImportsConfig, GEngineIni);
+			return bLazyLoadImportsConfig;
+		}
+		return false;
+	};
+	static const bool bImportLazyLoadEnabled = ImportLazyLoadEnabled();
+	return bImportLazyLoadEnabled;
 }
 #endif
 
diff --git a/Engine/Source/Runtime/CoreUObject/Private/UObject/LinkerLoadImportBehavior.cpp b/Engine/Source/Runtime/CoreUObject/Private/UObject/LinkerLoadImportBehavior.cpp
new file mode 100644
--- /dev/null
+++ b/Engine/Source/Runtime/CoreUObject/Private/UObject/LinkerLoadImportBehavior.cpp
@@ -0,0 +1,41 @@
+// Copyright Epic Games, Inc. All Rights Reserved.
+
+#include "UObject/LinkerLoadImportBehavior.h"
+#include "UObject/Package.h"
+
+#if UE_WITH_OBJECT_HANDLE_LATE_RESOLVE
+
+namespace UE::LinkerLoad
+{
+PropertyImportBehaviorFunction* PropertyImportBehaviorCallback = nullptr;
+
+void SetPropertyImportBehaviorCallback(PropertyImportBehaviorFunction* Function)
+{
+	void* PreviousValue = (void*)FPlatformAtomics::InterlockedExchangePtr((void**)&PropertyImportBehaviorCallback, (void*)Function);
+	checkf((Function == nullptr) || (PreviousValue == nullptr), TEXT("LinkerLoad property import behavior callback set more than once.  This is not supported."))
+}
+
+void GetPropertyImportLoadBehavior(const FObjectImport& Import, const FLinkerLoad& LinkerLoad, EImportBehavior& OutBehavior)
+{
+	OutBehavior = EImportBehavior::Eager;
+
+	if (Import.bImportSearchedFor)
+	{
+		// If it was something that's been searched for, we've already attempted a resolve, might as well use it
+		return;
+	}
+
+	if (!LinkerLoad.IsImportLazyLoadEnabled() || !LinkerLoad.IsAllowingLazyLoading() || (LinkerLoad.LinkerRoot && LinkerLoad.LinkerRoot->HasAnyPackageFlags(PKG_PlayInEditor)))
+	{
+		return;
+	}
+
+	if (PropertyImportBehaviorCallback)
+	{
+		PropertyImportBehaviorCallback(Import, LinkerLoad, OutBehavior);
+	}
+}
+}
+
+#endif // UE_WITH_OBJECT_HANDLE_LATE_RESOLVE
+
diff --git a/Engine/Source/Runtime/CoreUObject/Public/UObject/LinkerLoad.h b/Engine/Source/Runtime/CoreUObject/Public/UObject/LinkerLoad.h
--- a/Engine/Source/Runtime/CoreUObject/Public/UObject/LinkerLoad.h
+++ b/Engine/Source/Runtime/CoreUObject/Public/UObject/LinkerLoad.h
@@ -1083,18 +1083,6 @@ private:
 
 	/** Sets the basic linker archive info */
 	void ResetStatusInfo();
-
-#if UE_WITH_OBJECT_HANDLE_LATE_RESOLVE
-	enum class EImportLoadBehavior
-	{
-		Eager = 0,
-		// @TODO: OBJPTR: we want to permit lazy background loading in the future
-		//LazyBackground,
-		LazyOnDemand,
-	};
-	static EImportLoadBehavior ParseImportLoadBehavior(const FString* LoadBehaviorMeta);
-	EImportLoadBehavior GetCurrentPropertyImportLoadBehavior(FPackageIndex ImportIndex);
-#endif // UE_WITH_OBJECT_HANDLE_LATE_RESOLVE
 public:
 	/**
 	 * Serializes the gatherable text data container.
diff --git a/Engine/Source/Runtime/Engine/Classes/Sound/SoundCue.h b/Engine/Source/Runtime/Engine/Classes/Sound/SoundCue.h
--- a/Engine/Source/Runtime/Engine/Classes/Sound/SoundCue.h
+++ b/Engine/Source/Runtime/Engine/Classes/Sound/SoundCue.h
@@ -84,7 +84,7 @@ public:
 /**
  * The behavior of audio playback is defined within Sound Cues.
  */
-UCLASS(hidecategories=object, BlueprintType, meta= (LoadBehavior = "LazyOnDemand"))
+UCLASS(hidecategories=object, BlueprintType)
 class ENGINE_API USoundCue : public USoundBase
 {
 	GENERATED_UCLASS_BODY()
diff --git a/Engine/Source/Runtime/Engine/Classes/Sound/SoundWave.h b/Engine/Source/Runtime/Engine/Classes/Sound/SoundWave.h
--- a/Engine/Source/Runtime/Engine/Classes/Sound/SoundWave.h
+++ b/Engine/Source/Runtime/Engine/Classes/Sound/SoundWave.h
@@ -370,7 +370,7 @@ struct ISoundWaveClient
 	virtual bool OnIsReadyForFinishDestroy(class USoundWave* Wave) const = 0;
 	virtual void OnFinishDestroy(class USoundWave* Wave) = 0;
 };
-UCLASS(hidecategories=Object, editinlinenew, BlueprintType, meta= (LoadBehavior = "LazyOnDemand"))
+UCLASS(hidecategories=Object, editinlinenew, BlueprintType)
 class ENGINE_API USoundWave : public USoundBase, public IAudioProxyDataFactory, public IInterface_AsyncCompilation
 {
 	GENERATED_UCLASS_BODY()
