Adding AssetSizeQuery commandlet for rapidly getting asset size breakdown by class. Also can save to CSV, either the class summary or all assets. Assets can be filtered by wildcard match.

diff --git a/Engine/Source/Editor/UnrealEd/Classes/Commandlets/AssetSizeQueryCommandlet.h b/Engine/Source/Editor/UnrealEd/Classes/Commandlets/AssetSizeQueryCommandlet.h
new file mode 100644
--- /dev/null
+++ b/Engine/Source/Editor/UnrealEd/Classes/Commandlets/AssetSizeQueryCommandlet.h
@@ -0,0 +1,22 @@
+// Copyright Epic Games, Inc. All Rights Reserved.
+#pragma once
+
+#include "CoreTypes.h"
+#include "UObject/ObjectMacros.h"
+#include "Commandlets/Commandlet.h"
+#include "AssetSizeQueryCommandlet.generated.h"
+
+UCLASS()
+class UAssetSizeQueryCommandlet : public UCommandlet
+{
+	GENERATED_UCLASS_BODY()
+
+public:
+
+	// Begin UCommandlet Interface
+	virtual int32 Main(const FString& Params) override;
+	// End UCommandlet Interface
+
+private:
+
+};
diff --git a/Engine/Source/Editor/UnrealEd/Private/Commandlets/AssetSizeQueryCommandlet.cpp b/Engine/Source/Editor/UnrealEd/Private/Commandlets/AssetSizeQueryCommandlet.cpp
new file mode 100644
--- /dev/null
+++ b/Engine/Source/Editor/UnrealEd/Private/Commandlets/AssetSizeQueryCommandlet.cpp
@@ -0,0 +1,245 @@
+// Copyright Epic Games, Inc. All Rights Reserved.
+
+#include "Commandlets/AssetSizeQueryCommandlet.h"
+
+#include "AssetRegistry/AssetData.h"
+#include "AssetRegistry/IAssetRegistry.h"
+#include "AssetRegistryState.h"
+#include "Misc/FileHelper.h"
+#include "Misc/Paths.h"
+#include "Misc/WildcardString.h"
+
+DEFINE_LOG_CATEGORY_STATIC(LogAssetSize, Display, All);
+
+UAssetSizeQueryCommandlet::UAssetSizeQueryCommandlet(const FObjectInitializer& ObjectInitializer)
+: Super(ObjectInitializer)
+{
+
+}
+
+enum class EOutputCSVType : uint8
+{
+	None,
+	Classes,
+	Assets
+};
+
+int32 UAssetSizeQueryCommandlet::Main(const FString& FullCommandLine)
+{
+	FString FileName;
+	if (FParse::Value(*FullCommandLine, TEXT("AssetRegistry="), FileName) == false)
+	{
+		UE_LOG(LogAssetSize, Error, TEXT("No AssetRegistry specified."));
+		UE_LOG(LogAssetSize, Display, TEXT(""));
+		UE_LOG(LogAssetSize, Display, TEXT("AssetSizeQueryCommandlet"));
+		UE_LOG(LogAssetSize, Display, TEXT(""));
+		UE_LOG(LogAssetSize, Display, TEXT("Used to find asset type size breakdown in a staged/compressed project."));
+		UE_LOG(LogAssetSize, Display, TEXT(""));
+		UE_LOG(LogAssetSize, Display, TEXT("Params:"));
+		UE_LOG(LogAssetSize, Display, TEXT("    -AssetRegistry=path                     Provides the path to the Development asset registry. This"));
+		UE_LOG(LogAssetSize, Display, TEXT("                                            asset registry must have staging size metadata via ProjectSettings->"));
+		UE_LOG(LogAssetSize, Display, TEXT("                                            Packaging->WriteBackMetadataToAssetRegistry or iostore -AssetRegistryWriteback."));
+		UE_LOG(LogAssetSize, Display, TEXT("    -Filter=wildcard                        (optional) Filteres the list of assets using a wildcard match."));
+		UE_LOG(LogAssetSize, Display, TEXT("    -Show=#                                 (optional) Shows only the top # classes, sorted on size (0 is all, default 10)."));
+		UE_LOG(LogAssetSize, Display, TEXT("    -CSV=path                               (optional) Output the filtered per class infomation to the given CSV file."));
+		UE_LOG(LogAssetSize, Display, TEXT("    -CSVType=(Assets,Classes)               (optional) Specifies whether to write the class summary or all matching assets to the csv file."));
+		UE_LOG(LogAssetSize, Display, TEXT(""));
+		UE_LOG(LogAssetSize, Display, TEXT("Note that a project must be staged in order to determine compressed sizes, and that additionally"));
+		UE_LOG(LogAssetSize, Display, TEXT("some platforms stage in a platform specific manner that precludes finding compressed sizes."));
+		return 1;
+	}
+
+	FAssetRegistryState AssetRegistry;
+	if (FAssetRegistryState::LoadFromDisk(*FileName, FAssetRegistryLoadOptions(), AssetRegistry) == false)
+	{
+		UE_LOG(LogAssetSize, Error, TEXT("Failed load asset registry (%s)"), *FileName);		
+		return 1;
+	}
+
+	FString AssetFilter;
+	FParse::Value(*FullCommandLine, TEXT("Filter="), AssetFilter);
+
+	int ShowCount = 10;
+	FString ShowCountString;
+	if (FParse::Value(*FullCommandLine, TEXT("Show="), ShowCountString))
+	{
+		ShowCount = FCString::Atoi(*ShowCountString);
+		if (ShowCount < 0)
+		{
+			UE_LOG(LogAssetSize, Warning, TEXT("Invalid 'Show' count specified (%s), using \"show all\""), *ShowCountString);
+			ShowCount = 0;
+		}
+	}
+
+	FString OutputCSVPath;
+	EOutputCSVType OutputCSVType = EOutputCSVType::None;
+	if (FParse::Value(*FullCommandLine, TEXT("-CSV="), OutputCSVPath))
+	{
+		OutputCSVType = EOutputCSVType::Classes;
+		FString RawCSVType;
+		if (FParse::Value(*FullCommandLine, TEXT("-CSVType="), RawCSVType))
+		{
+			if (RawCSVType.Compare(TEXT("classes"), ESearchCase::IgnoreCase) == 0)
+			{
+				OutputCSVType = EOutputCSVType::Classes;
+				UE_LOG(LogAssetSize, Display, TEXT("CSV Type: Classes"));
+			}
+			else if (RawCSVType.Compare(TEXT("assets"), ESearchCase::IgnoreCase) == 0)
+			{
+				OutputCSVType = EOutputCSVType::Assets;
+				UE_LOG(LogAssetSize, Display, TEXT("CSV Type: Assets"));
+			}
+			else
+			{
+				UE_LOG(LogAssetSize, Error, TEXT("Invalid -CSVType: %s"), *RawCSVType);
+				return 1;
+			}
+		}
+		else
+		{
+			UE_LOG(LogAssetSize, Display, TEXT("CSV Type: Default (Classes)"));
+		}
+	}
+
+	uint64 TotalDiskSize = 0;
+	const TMap<FName, const FAssetPackageData*>& PackageDatas = AssetRegistry.GetAssetPackageDataMap();
+	for (const TPair<FName, const FAssetPackageData*>& Pair : PackageDatas)
+	{
+		if (Pair.Value->DiskSize != -1)
+		{
+			TotalDiskSize += Pair.Value->DiskSize;
+		}
+	}
+
+	int64 ImportantAssetCount = 0;
+	int64 MatchedAssetCount = 0;
+	int64 FilteredCompressedSize = 0;
+	int64 TotalCompressedSize = 0;
+
+	struct FMatchedAssetInfo
+	{
+		FName ObjectPath;
+		int64 CompressedSize;
+	};
+	TMap<FName /* AssetClass */, int64> FilteredClassCompressedSizes;
+	TMap<FName /* AssetClass */, TArray<FMatchedAssetInfo>> FilteredClassMatchedAssets;
+	AssetRegistry.EnumerateAllAssets(TSet<FName>(), 
+		[&ImportantAssetCount,
+		 &MatchedAssetCount,
+		 &TotalCompressedSize,
+		 &FilteredClassMatchedAssets, 
+		 &FilteredCompressedSize, 
+		 &FilteredClassCompressedSizes, 
+		 &AssetFilter](const FAssetData& AssetData)
+	{
+		FString CompressedSize;
+		if (AssetData.GetTagValue("Stage_ChunkCompressedSize", CompressedSize) == false ||
+			CompressedSize.Len() == 0)
+		{
+			return true;
+		}
+
+		int64 AssetCompressedSize = FCString::Atoi64(*CompressedSize);
+		ImportantAssetCount++;
+		TotalCompressedSize += AssetCompressedSize;
+
+		bool bMatched = true;
+		if (AssetFilter.Len())
+		{
+			FString ObjectPath = AssetData.ObjectPath.ToString();
+			bMatched = FWildcardString::IsMatchSubstring(*AssetFilter, *ObjectPath, *ObjectPath + ObjectPath.Len(), ESearchCase::IgnoreCase);
+		} // end if filter exists
+
+		if (bMatched)
+		{
+			MatchedAssetCount++;
+
+			FMatchedAssetInfo& Info = FilteredClassMatchedAssets.FindOrAdd(AssetData.AssetClass).AddDefaulted_GetRef();
+			Info.ObjectPath = AssetData.ObjectPath;
+			Info.CompressedSize = AssetCompressedSize;
+			
+			FilteredCompressedSize += AssetCompressedSize;
+			int64& FilteredClassCompressedSize = FilteredClassCompressedSizes.FindOrAdd(AssetData.AssetClass);
+			FilteredClassCompressedSize += AssetCompressedSize;
+		}
+
+		return true;
+	}); // end EnumerateAssets
+
+	if (ImportantAssetCount == 0)
+	{
+		UE_LOG(LogAssetSize, Display, TEXT("No assets with size information found - staging metadata needs to be added to the asset registry"));
+		UE_LOG(LogAssetSize, Display, TEXT("via ProjectSettings->Packaging->WriteBackMetadataToAssetRegistry or using iostore "));
+		UE_LOG(LogAssetSize, Display, TEXT("-AssetRegistryWriteback."));
+		return 1;
+	}
+
+	FilteredClassCompressedSizes.ValueSort(TGreater<int64>());
+
+	
+	UE_LOG(LogAssetSize, Display, TEXT("Filter:                                      %s"), AssetFilter.Len() ? *AssetFilter : TEXT("<all>"));
+	UE_LOG(LogAssetSize, Display, TEXT("Assets with size information:                %s"), *FText::AsNumber(ImportantAssetCount).ToString());
+	UE_LOG(LogAssetSize, Display, TEXT("Filtered to:                                 %s (%.1f%%)"), *FText::AsNumber(MatchedAssetCount).ToString(), 100.0 * MatchedAssetCount / ImportantAssetCount);
+	UE_LOG(LogAssetSize, Display, TEXT("Total compressed size (bytes):               %s"), *FText::AsNumber(TotalCompressedSize).ToString());
+	UE_LOG(LogAssetSize, Display, TEXT("Filtered compressed size (bytes):            %s (%.1f%%)"), *FText::AsNumber(FilteredCompressedSize).ToString(), 100.0 * FilteredCompressedSize / TotalCompressedSize);
+	if (ShowCount)
+	{
+		UE_LOG(LogAssetSize, Display, TEXT("Top %2s filtered class sizes:                 bytes          (pct of filtered total)"), *FText::AsNumber(ShowCount).ToString());
+	}
+	else
+	{
+		UE_LOG(LogAssetSize, Display, TEXT("Filtered class sizes:                         bytes          (pct of filtered total)"));
+	}
+	for (TPair<FName, int64>& ClassSizePair : FilteredClassCompressedSizes)
+	{
+		UE_LOG(LogAssetSize, Display, TEXT("    %-40s %-14s (%.1f%%)"), *ClassSizePair.Key.ToString(), *FText::AsNumber(ClassSizePair.Value).ToString(), 100.0 * ClassSizePair.Value / FilteredCompressedSize);
+		if (ShowCount)
+		{
+			ShowCount--;
+			if (ShowCount == 0)
+			{
+				break;
+			}
+		}
+	}
+
+	if (OutputCSVType != EOutputCSVType::None)
+	{
+		TArray<FString> Lines;
+
+		if (OutputCSVType == EOutputCSVType::Classes)
+		{
+			Lines.Add(TEXT("AssetClass,AssetCount,TotalCompressedSize"));
+			for (TPair<FName, int64>& ClassSizePair : FilteredClassCompressedSizes)
+			{
+				// we add to both maps at the same time to we know the lookup succeeds.
+				Lines.Add(FString::Printf(TEXT("%s,%lld,%lld"), *ClassSizePair.Key.ToString(), FilteredClassMatchedAssets[ClassSizePair.Key].Num(), ClassSizePair.Value));
+			}
+		}
+		else if (OutputCSVType == EOutputCSVType::Assets)
+		{
+			Lines.Add(TEXT("AssetName,AssetType,CompressedSize"));
+			for (TPair<FName, TArray<FMatchedAssetInfo>>& ClassAssetsPair : FilteredClassMatchedAssets)
+			{
+				// we add to both maps at the same time to we know the lookup succeeds.
+				for (const FMatchedAssetInfo& AssetInfo : ClassAssetsPair.Value)
+				{
+					Lines.Add(FString::Printf(TEXT("%s,%s,%lld"), *AssetInfo.ObjectPath.ToString(), *ClassAssetsPair.Key.ToString(), AssetInfo.CompressedSize));
+				}
+			}
+		}
+		else
+		{
+			check(0); // added a type ??
+		}
+
+		if (FFileHelper::SaveStringArrayToFile(Lines, *OutputCSVPath) == false)
+		{
+			UE_LOG(LogAssetSize, Error, TEXT("Unable to write CSV file %s"), *OutputCSVPath);
+			return 1;
+		}
+		UE_LOG(LogAssetSize, Display, TEXT("Saved CSV file: %s"), *OutputCSVPath);
+	} // end if outputting csv file
+
+	return 0;
+}
