Fix crash on "Build -> Build All"
Prevent issue where the UWorld pointer would not be valid after forcing a world unload/reload in some of the build commands.

diff --git a/Engine/Source/Editor/UnrealEd/Private/EditorBuildUtils.cpp b/Engine/Source/Editor/UnrealEd/Private/EditorBuildUtils.cpp
--- a/Engine/Source/Editor/UnrealEd/Private/EditorBuildUtils.cpp
+++ b/Engine/Source/Editor/UnrealEd/Private/EditorBuildUtils.cpp
@@ -337,6 +337,8 @@ bool FEditorBuildUtils::EditorBuild( UWorld* InWorld, FName Id, const bool bAllo
 		BuildProgressWidget.Pin()->SetBuildType(BuildType);
 	}
 
+	TSoftObjectPtr<UWorld> World = InWorld;
+
 	bool bShouldMapCheck = !FParse::Param(FCommandLine::Get(), TEXT("SkipMapCheck"));
 	if (Id == FBuildOptions::BuildGeometry)
 	{
@@ -495,6 +497,11 @@ bool FEditorBuildUtils::EditorBuild( UWorld* InWorld, FName Id, const bool bAllo
 		bDoBuild = false;
 	}
 
+	// It's possible the world was unloaded & reloaded if external commands were run.
+	// To work around this, reassign the initial world from it's soft object path.
+	check(World.IsValid());
+	InWorld = World.Get();
+
 	// Check map for errors (only if build operation happened)
 	if ( bShouldMapCheck && bDoBuild && !GEditor->GetMapBuildCancelled() )
 	{
@@ -1155,9 +1162,16 @@ void FBuildAllHandler::ProcessBuild(const TWeakPtr<SBuildProgressWidget>& BuildP
 {
 	const FScopedBusyCursor BusyCursor;
 
+	TSoftObjectPtr<UWorld> World = CurrentWorld;
+
 	// Loop until we finish, or we start an async step.
 	while (true)
 	{
+		// It's possible the world was unloaded & reloaded if external commands were run.
+		// To work around this, reassign the initial world from it's soft object path.
+		check(World.IsValid());
+		CurrentWorld = World.Get();
+
 		if (GEditor->GetMapBuildCancelled())
 		{
 			// Build cancelled, so bail.
