-Add option to minimize editor when running non-VR PIE as well

diff --git a/Engine/Source/Editor/UnrealEd/Classes/Settings/LevelEditorPlaySettings.h b/Engine/Source/Editor/UnrealEd/Classes/Settings/LevelEditorPlaySettings.h
--- a/Engine/Source/Editor/UnrealEd/Classes/Settings/LevelEditorPlaySettings.h
+++ b/Engine/Source/Editor/UnrealEd/Classes/Settings/LevelEditorPlaySettings.h
@@ -247,6 +247,10 @@ public:
 	UPROPERTY(config, EditAnywhere, Category = PlayInEditor, meta = (ToolTip = "Whether or not the editor is minimized on VR PIE"))
 	bool ShouldMinimizeEditorOnVRPIE;
 
+	/** Should we minimize the editor when non-VR PIE is clicked (default = false) */
+	UPROPERTY(config, EditAnywhere, Category = PlayInEditor, meta = (ToolTip = "Whether or not the editor is minimized on non-VR PIE"))
+	bool bShouldMinimizeEditorOnNonVRPIE;
+
 	/** Whether to automatically recompile blueprints on PIE */
 	UPROPERTY(config, EditAnywhere, Category=PlayInEditor, meta=(ToolTip="Automatically recompile blueprints used by the current level when initiating a Play In Editor session"))
 	bool AutoRecompileBlueprints;
diff --git a/Engine/Source/Editor/UnrealEd/Private/PlayLevel.cpp b/Engine/Source/Editor/UnrealEd/Private/PlayLevel.cpp
--- a/Engine/Source/Editor/UnrealEd/Private/PlayLevel.cpp
+++ b/Engine/Source/Editor/UnrealEd/Private/PlayLevel.cpp
@@ -3311,7 +3311,8 @@ TSharedRef<SPIEViewport> UEditorEngine::GeneratePIEViewportWindow(const FRequest
 	InViewportClient->SetViewportOverlayWidget(PieWindow, ViewportOverlayWidgetRef);
 	InViewportClient->SetGameLayerManager(GameLayerManagerRef);
 
-	bool bShouldMinimizeRootWindow = bVRPreview && GEngine->XRSystem.IsValid() && InSessionParams.EditorPlaySettings->ShouldMinimizeEditorOnVRPIE;
+	const bool bShouldMinimizeRootWindowForVRPreview = bVRPreview && GEngine->XRSystem.IsValid() && InSessionParams.EditorPlaySettings->ShouldMinimizeEditorOnVRPIE;
+	const bool bShouldMinimizeRootWindowForNonVRPreview = !bVRPreview && InSessionParams.EditorPlaySettings->bShouldMinimizeEditorOnNonVRPIE;
 	// Set up a notification when the window is closed so we can clean up PIE
 	{
 		struct FLocal
@@ -3374,7 +3375,7 @@ TSharedRef<SPIEViewport> UEditorEngine::GeneratePIEViewportWindow(const FRequest
 
 		PieWindow->SetRequestDestroyWindowOverride(FRequestDestroyWindowOverride::CreateStatic(&FLocal::RequestDestroyPIEWindowOverride, TWeakObjectPtr<UEditorEngine>(this)));
 		PieWindow->SetOnWindowClosed(FOnWindowClosed::CreateStatic(&FLocal::OnPIEWindowClosed, TWeakPtr<SViewport>(PieViewportWidget), TWeakObjectPtr<UEditorEngine>(this),
-			InViewportIndex, bShouldMinimizeRootWindow, PreviewFeatureLevelChangedHandle));
+			InViewportIndex, bShouldMinimizeRootWindowForVRPreview || bShouldMinimizeRootWindowForNonVRPreview, PreviewFeatureLevelChangedHandle));
 	}
 
 	// Create a new viewport that the viewport widget will use to render the game
@@ -3398,18 +3399,18 @@ TSharedRef<SPIEViewport> UEditorEngine::GeneratePIEViewportWindow(const FRequest
 
 	// Change the system resolution to match our window, to make sure game and slate window are kept synchronized
 	FSystemResolution::RequestResolutionChange(WindowSize.X, WindowSize.Y, EWindowMode::Windowed);
-
+	
 	const bool bHMDIsReady = (GEngine && GEngine->XRSystem.IsValid() && GEngine->XRSystem->GetHMDDevice() && GEngine->XRSystem->GetHMDDevice()->IsHMDConnected());
 	if (bVRPreview && bHMDIsReady)
 	{
 		GEngine->StereoRenderingDevice->EnableStereo(true);
+	}
 
-		// minimize the root window to provide max performance for the preview.
-		TSharedPtr<SWindow> RootWindow = FGlobalTabmanager::Get()->GetRootWindow();
-		if (RootWindow.IsValid() && bShouldMinimizeRootWindow)
-		{
-			RootWindow->Minimize();
-		}
+	// minimize the root window to provide max performance for the preview.
+	TSharedPtr<SWindow> RootWindow = FGlobalTabmanager::Get()->GetRootWindow();
+	if (RootWindow.IsValid() && (bShouldMinimizeRootWindowForVRPreview || bShouldMinimizeRootWindowForNonVRPreview))
+	{
+		RootWindow->Minimize();
 	}
 
 	return PieViewportWidget;
