Updating diff asset bulk data to use new cleaner functions.

diff --git a/Engine/Source/Editor/UnrealEd/Private/Commandlets/DiffAssetBulkDataCommandlet.cpp b/Engine/Source/Editor/UnrealEd/Private/Commandlets/DiffAssetBulkDataCommandlet.cpp
--- a/Engine/Source/Editor/UnrealEd/Private/Commandlets/DiffAssetBulkDataCommandlet.cpp
+++ b/Engine/Source/Editor/UnrealEd/Private/Commandlets/DiffAssetBulkDataCommandlet.cpp
@@ -209,7 +209,7 @@ int32 UDiffAssetBulkDataCommandlet::Main(const FString& FullCommandLine)
 	TArray<FName /* PackageName */> PackagesWithUnassignableDiffsAndUntaggedAssets;
 	TArray<FName /* PackageName */> PackagesWithUnassignableDiffs;
 	
-	for (FName& ChangedPackageName : ChangedBulkData)
+	for (const FName& ChangedPackageName : ChangedBulkData)
 	{
 		TConstArrayView<FAssetData const*> BaseAssetDatas = BaseState.GetAssetsByPackageName(ChangedPackageName);		
 		TConstArrayView<FAssetData const*> CurrentAssetDatas = CurrentState.GetAssetsByPackageName(ChangedPackageName);
@@ -286,53 +286,14 @@ int32 UDiffAssetBulkDataCommandlet::Main(const FString& FullCommandLine)
 		{
 			// Nothing has anything to use for diff blaming for this package.
 			// Try to find a representative asset class from the assets in the package.
-
-			// Find a candidate asset.
-			// If there's a "UAsset", then we use that as the asset to get the sizes.
-			// If not, then we look for a "TopLevelAsset", i.e. one that shows up in the content browser.
-			// 
-			// If we have multiple TLS or no TLAs, we dump those in a "??" list because
-			// we can't really make any decisions there.
-			// 
-			FAssetData const* UseAsset = nullptr;
-			{
-				FAssetData const* Found_UAsset = nullptr;
-				FAssetData const* Found_TLA = nullptr;
-				bool bMultipleTLAs = false;
-
-				for (FAssetData const* const Asset : CurrentAssetDatas)
-				{
-					if (Asset->IsUAsset())
-					{
-						Found_UAsset = Asset;
-					}
-					if (Asset->IsTopLevelAsset())
-					{
-						if (Found_TLA)
-						{
-							bMultipleTLAs = true;
-						}
-						Found_TLA = Asset;
-					}
-				}
-
-				UseAsset = Found_UAsset;
-				if (UseAsset == nullptr)
-				{
-					if (bMultipleTLAs == false)
-					{
-						UseAsset = Found_TLA; // only accept the TLA if there's one, otherwise leave this as unassigned.
-					}
-				}
-			}
-
-			if (UseAsset == nullptr)
+			FAssetData const* RepresentativeAsset = UE::AssetRegistry::GetMostImportantAsset(CurrentAssetDatas, true);
+			if (RepresentativeAsset == nullptr)
 			{
 				NoTagPackagesByAssumedClass.FindOrAdd(NAME_None).Add(ChangedPackageName);
 			}
 			else
 			{
-				NoTagPackagesByAssumedClass.FindOrAdd(UseAsset->AssetClass).Add(ChangedPackageName);
+				NoTagPackagesByAssumedClass.FindOrAdd(RepresentativeAsset->AssetClass).Add(ChangedPackageName);
 			}
 			
 			continue;
@@ -393,9 +354,11 @@ int32 UDiffAssetBulkDataCommandlet::Main(const FString& FullCommandLine)
 		return 0;
 	}
 
+	TArray<FName>& CantDetermineAssetClassPackages = NoTagPackagesByAssumedClass.FindOrAdd(NAME_None);
+
 	UE_LOG(LogDiffAssetBulk, Display, TEXT("Changed package breakdown:"));
 	UE_LOG(LogDiffAssetBulk, Display, TEXT("    No blame information available:"));
-	UE_LOG(LogDiffAssetBulk, Display, TEXT("        Can't determine asset class   : %d"), NoTagPackagesByAssumedClass.FindOrAdd(NAME_None).Num());
+	UE_LOG(LogDiffAssetBulk, Display, TEXT("        Can't determine asset class   : %d"), CantDetermineAssetClassPackages.Num());
 	for (TPair<FName, TArray<FName>> ClassPackages : NoTagPackagesByAssumedClass)
 	{
 		if (ClassPackages.Key == NAME_None)
@@ -406,12 +369,22 @@ int32 UDiffAssetBulkDataCommandlet::Main(const FString& FullCommandLine)
 		UE_LOG(LogDiffAssetBulk, Display, TEXT("        %-30s: %d"), *ClassPackages.Key.ToString(), ClassPackages.Value.Num());
 	}
 
+	if (CantDetermineAssetClassPackages.Num())
+	{
+		Algo::Sort(CantDetermineAssetClassPackages, FNameLexicalLess());
+		UE_LOG(LogDiffAssetBulk, Display, TEXT("    Can't determine asset class:  : %-7d // Couldn't pick a representative asset in the package."), CantDetermineAssetClassPackages.Num());
+		for (const FName& PackageName : CantDetermineAssetClassPackages)
+		{
+			UE_LOG(LogDiffAssetBulk, Display, TEXT("        %s"), *PackageName.ToString());
+		}
+	}
+
 	if (PackagesWithUnassignableDiffs.Num())
 	{
 		Algo::Sort(PackagesWithUnassignableDiffs, FNameLexicalLess());
 
 		UE_LOG(LogDiffAssetBulk, Warning, TEXT("    Can't determine blame:        : %-7d // Assets had blame tags but all matched - check determinism!"), PackagesWithUnassignableDiffs.Num());
-		for (FName& PackageName : PackagesWithUnassignableDiffs)
+		for (const FName& PackageName : PackagesWithUnassignableDiffs)
 		{
 			UE_LOG(LogDiffAssetBulk, Warning, TEXT("        %s"), *PackageName.ToString());
 		}
@@ -422,7 +395,7 @@ int32 UDiffAssetBulkDataCommandlet::Main(const FString& FullCommandLine)
 		Algo::Sort(PackagesWithUnassignableDiffsAndUntaggedAssets, FNameLexicalLess());
 		
 		UE_LOG(LogDiffAssetBulk, Display, TEXT("    Potential untagged assets:    : %-7d // Package had assets with blame tags that matched, but also untagged assets. Might be determinism!"), PackagesWithUnassignableDiffsAndUntaggedAssets.Num());
-		for (FName& PackageName : PackagesWithUnassignableDiffsAndUntaggedAssets)
+		for (const FName& PackageName : PackagesWithUnassignableDiffsAndUntaggedAssets)
 		{
 			UE_LOG(LogDiffAssetBulk, Display, TEXT("        %s"), *PackageName.ToString());
 		}
