UEditorLevelUtils::PrivateRemoveLevelFromWorld: Clean up LevelCollection ref to ULevel.

This was causing a forced Garbage Collection when deleting/unloading a LevelInstance actor in projects with 'gc.PendingKillEnabled = false'

diff --git a/Engine/Source/Editor/UnrealEd/Private/EditorLevelUtils.cpp b/Engine/Source/Editor/UnrealEd/Private/EditorLevelUtils.cpp
--- a/Engine/Source/Editor/UnrealEd/Private/EditorLevelUtils.cpp
+++ b/Engine/Source/Editor/UnrealEd/Private/EditorLevelUtils.cpp
@@ -926,6 +926,11 @@ void UEditorLevelUtils::PrivateRemoveLevelFromWorld(ULevel* InLevel)
 		check(InLevel->bIsVisible == false);
 	}
 
+	if (FLevelCollection* LC = InLevel->GetCachedLevelCollection())
+	{
+		LC->RemoveLevel(InLevel);
+	}
+
 	InLevel->ReleaseRenderingResources();
 
 	IStreamingManager::Get().RemoveLevel(InLevel);
