Fix crash when re-importing and there is some bad influence data.

diff --git a/Engine/Source/Editor/UnrealEd/Private/Fbx/FbxSkeletalMeshImport.cpp b/Engine/Source/Editor/UnrealEd/Private/Fbx/FbxSkeletalMeshImport.cpp
--- a/Engine/Source/Editor/UnrealEd/Private/Fbx/FbxSkeletalMeshImport.cpp
+++ b/Engine/Source/Editor/UnrealEd/Private/Fbx/FbxSkeletalMeshImport.cpp
@@ -4180,7 +4180,7 @@ bool UnFbx::FFbxImporter::ImportSkeletalMeshLOD(USkeletalMesh* InSkeletalMesh, U
 			for (int32 InfluenceIndex = InfluenceNumber-1; InfluenceIndex >= 0; --InfluenceIndex)
 			{
 				SkeletalMeshImportData::FRawBoneInfluence& Influence = LODImportData.Influences[InfluenceIndex];
-				Influence.BoneIndex = ImportDataBoneRemap[Influence.BoneIndex];
+				Influence.BoneIndex = ImportDataBoneRemap.IsValidIndex(Influence.BoneIndex) ? ImportDataBoneRemap[Influence.BoneIndex] : INDEX_NONE;
 				if (Influence.BoneIndex == INDEX_NONE)
 				{
 					const int32 DeleteCount = 1;
