Fbx: Fix crash on exporting transform track when no scene component is found. This fixes an issue exporting a live link component (actor component) which doesn't have a transform associated with it.

diff --git a/Engine/Source/Editor/UnrealEd/Private/Fbx/FbxMainExport.cpp b/Engine/Source/Editor/UnrealEd/Private/Fbx/FbxMainExport.cpp
--- a/Engine/Source/Editor/UnrealEd/Private/Fbx/FbxMainExport.cpp
+++ b/Engine/Source/Editor/UnrealEd/Private/Fbx/FbxMainExport.cpp
@@ -3184,6 +3184,35 @@ void FFbxExporter::ExportLevelSequenceBaked3DTransformTrack(IAnimTrackAdapter& A
 	const bool bIsLightActor = BoundActor ? BoundActor->IsA(ALight::StaticClass()) : BoundComponent ? BoundComponent->IsA(ULightComponent::StaticClass()) : false;
 	const bool bBakeRotations = bIsCameraActor || bIsLightActor;
 
+	USceneComponent* InterrogatedComponent = nullptr;
+	if (BoundComponent)
+	{
+		InterrogatedComponent = BoundComponent;
+	}
+	else if (BoundActor)
+	{
+		InterrogatedComponent = BoundActor->GetRootComponent();
+	}
+
+	if (bIsCameraActor)
+	{
+		if (InterrogatedComponent && InterrogatedComponent->IsA<UCameraComponent>())
+		{
+			// all set
+		}
+		else if (BoundActor && BoundActor->IsA(ACameraActor::StaticClass()))
+		{
+			ACameraActor* CameraActor = Cast<ACameraActor>(BoundActor);
+			InterrogatedComponent = CameraActor->GetCameraComponent();
+		}
+	}
+
+	if (!InterrogatedComponent)
+	{
+		UE_LOG(LogFbx, Warning, TEXT("Export transform track for %s failed because could not find suitable scene component"), *BoundObject->GetName());
+		return;
+	}
+
 	if (!FbxNode)
 	{
 		FbxNode = CreateNode(TransformTrack->GetDisplayName().ToString());
@@ -3237,29 +3266,6 @@ void FFbxExporter::ExportLevelSequenceBaked3DTransformTrack(IAnimTrackAdapter& A
 
 	FMovieSceneTimeTransform LocalToRootTransform = RootToLocalTransform.InverseLinearOnly();
 
-	USceneComponent* InterrogatedComponent = nullptr;
-	if (BoundComponent)
-	{
-		InterrogatedComponent = BoundComponent;
-	}
-	else if (BoundActor)
-	{
-		InterrogatedComponent = BoundActor->GetRootComponent();
-	}
-
-	if (bIsCameraActor)
-	{
-		if (InterrogatedComponent && InterrogatedComponent->IsA<UCameraComponent>())
-		{
-			// all set
-		}
-		else if (BoundActor && BoundActor->IsA(ACameraActor::StaticClass()))
-		{
-			ACameraActor* CameraActor = Cast<ACameraActor>(BoundActor);
-			InterrogatedComponent = CameraActor->GetCameraComponent();
-		}
-	}
-
 	TArray<FTransform> RelativeTransforms;
 	int32 LocalStartFrame = FFrameRate::TransformTime(FFrameTime(DiscreteInclusiveLower(InPlaybackRange)), TickResolution, DisplayRate).RoundToFrame().Value;
 	int32 AnimationLength = FFrameRate::TransformTime(FFrameTime(FFrameNumber(DiscreteSize(InPlaybackRange))), TickResolution, DisplayRate).RoundToFrame().Value + 1; // Add one so that we export a key for the end frame
