Reduce staticmesh build time by half during FBX import caused by too many SectionInfoMap being created by error

diff --git a/Engine/Source/Editor/UnrealEd/Private/Fbx/FbxStaticMeshImport.cpp b/Engine/Source/Editor/UnrealEd/Private/Fbx/FbxStaticMeshImport.cpp
--- a/Engine/Source/Editor/UnrealEd/Private/Fbx/FbxStaticMeshImport.cpp
+++ b/Engine/Source/Editor/UnrealEd/Private/Fbx/FbxStaticMeshImport.cpp
@@ -595,25 +595,29 @@ bool UnFbx::FFbxImporter::BuildStaticMeshFromGeometry(FbxNode* Node, UStaticMesh
 		bool bEnableCollision = bImportedCollision || (GBuildStaticMeshCollision && LODIndex == 0 && ImportOptions->bRemoveDegenerates);
 		for(int32 SectionIndex=MaterialIndexOffset; SectionIndex<MaterialIndexOffset+MaterialCount; SectionIndex++)
 		{
-			FMeshSectionInfo Info = StaticMesh->GetSectionInfoMap().Get(LODIndex, SectionIndex);
-		
-			Info.bEnableCollision = bEnableCollision;
-			//Make sure LOD greater then 0 copy the LOD 0 sections collision flags
-			if (LODIndex != 0)
+			if (StaticMesh->GetSectionInfoMap().IsValidSection(LODIndex, SectionIndex))
 			{
-				//Match the material slot index
-				for (int32 LodZeroSectionIndex = 0; LodZeroSectionIndex < StaticMesh->GetSectionInfoMap().GetSectionNumber(0); ++LodZeroSectionIndex)
+				FMeshSectionInfo Info = StaticMesh->GetSectionInfoMap().Get(LODIndex, SectionIndex);
+		
+				Info.bEnableCollision = bEnableCollision;
+				//Make sure LOD greater then 0 copy the LOD 0 sections collision flags
+				if (LODIndex != 0)
 				{
-					FMeshSectionInfo InfoLodZero = StaticMesh->GetSectionInfoMap().Get(0, LodZeroSectionIndex);
-					if (InfoLodZero.MaterialIndex == Info.MaterialIndex)
+					//Match the material slot index
+					for (int32 LodZeroSectionIndex = 0; LodZeroSectionIndex < StaticMesh->GetSectionInfoMap().GetSectionNumber(0); ++LodZeroSectionIndex)
 					{
-						Info.bEnableCollision = InfoLodZero.bEnableCollision;
-						Info.bCastShadow = InfoLodZero.bCastShadow;
-						break;
+						FMeshSectionInfo InfoLodZero = StaticMesh->GetSectionInfoMap().Get(0, LodZeroSectionIndex);
+						if (InfoLodZero.MaterialIndex == Info.MaterialIndex)
+						{
+							Info.bEnableCollision = InfoLodZero.bEnableCollision;
+							Info.bCastShadow = InfoLodZero.bCastShadow;
+							break;
+						}
 					}
 				}
+
+				StaticMesh->GetSectionInfoMap().Set(LODIndex, SectionIndex, Info);
 			}
-			StaticMesh->GetSectionInfoMap().Set(LODIndex, SectionIndex, Info);
 		}
 	}
 
@@ -1863,34 +1867,30 @@ UStaticMesh* UnFbx::FFbxImporter::ImportStaticMeshAsSingle(UObject* InParent, TA
 			}
 			UMaterialInterface* Material = MeshMaterials.IsValidIndex(MaterialIndex) ? MeshMaterials[MaterialIndex].Material : UMaterial::GetDefaultMaterial(MD_Surface);
 			FStaticMaterial StaticMaterial(Material, MaterialSlotName, ImportedMaterialSlotName);
-			(LODIndex > 0) ? MaterialToAdd.Add(StaticMaterial) : StaticMesh->GetStaticMaterials().Add(StaticMaterial);
+			MaterialToAdd.Add(StaticMaterial);
 		}
-		if (LODIndex > 0)
+
 		{
 			//Insert the new materials in the static mesh
 			//The build function will search for imported slot name to find the appropriate slot
-			int32 StaticMeshMaterialCount = StaticMesh->GetStaticMaterials().Num();
-			if (StaticMeshMaterialCount > 0)
+			for (int32 MaterialToAddIndex = 0; MaterialToAddIndex < MaterialToAdd.Num(); ++MaterialToAddIndex)
 			{
-				for (int32 MaterialToAddIndex = 0; MaterialToAddIndex < MaterialToAdd.Num(); ++MaterialToAddIndex)
+				const FStaticMaterial& CandidateMaterial = MaterialToAdd[MaterialToAddIndex];
+				bool FoundExistingMaterial = false;
+				//Found matching existing material
+				for (int32 StaticMeshMaterialIndex = 0; StaticMeshMaterialIndex < StaticMesh->GetStaticMaterials().Num(); ++StaticMeshMaterialIndex)
 				{
-					const FStaticMaterial& CandidateMaterial = MaterialToAdd[MaterialToAddIndex];
-					bool FoundExistingMaterial = false;
-					//Found matching existing material
-					for (int32 StaticMeshMaterialIndex = 0; StaticMeshMaterialIndex < StaticMeshMaterialCount; ++StaticMeshMaterialIndex)
+					const FStaticMaterial& StaticMeshMaterial = StaticMesh->GetStaticMaterials()[StaticMeshMaterialIndex];
+					if (StaticMeshMaterial.ImportedMaterialSlotName == CandidateMaterial.ImportedMaterialSlotName)
 					{
-						const FStaticMaterial& StaticMeshMaterial = StaticMesh->GetStaticMaterials()[StaticMeshMaterialIndex];
-						if (StaticMeshMaterial.ImportedMaterialSlotName == CandidateMaterial.ImportedMaterialSlotName)
-						{
-							FoundExistingMaterial = true;
-							break;
-						}
-					}
-					if (!FoundExistingMaterial)
-					{
-						StaticMesh->GetStaticMaterials().Add(CandidateMaterial);
+						FoundExistingMaterial = true;
+						break;
 					}
 				}
+				if (!FoundExistingMaterial)
+				{
+					StaticMesh->GetStaticMaterials().Add(CandidateMaterial);
+				}
 			}
 			
 			//Set the Section Info Map to fit the real StaticMaterials array
