Fixed a crash in the selection code when exiting a pie session.
Some editor mode selection logic has dependancy on the state of the editor and the place where selection of the actors happened is right in the middle of the state transition from pie to level edit resulting in crash.

diff --git a/Engine/Source/Editor/UnrealEd/Private/PlayLevel.cpp b/Engine/Source/Editor/UnrealEd/Private/PlayLevel.cpp
--- a/Engine/Source/Editor/UnrealEd/Private/PlayLevel.cpp
+++ b/Engine/Source/Editor/UnrealEd/Private/PlayLevel.cpp
@@ -464,19 +464,6 @@ void UEditorEngine::EndPlayMap()
 	FWorldContext& EditorWorldContext = GEditor->GetEditorWorldContext();
 	UPackage* Package = EditorWorldContext.World()->GetOutermost();
 
-	// For every actor that was selected previously, make sure it's editor equivalent is selected. We do that after the cleanup in case some actor where removed.
-	GEditor->GetSelectedActors()->BeginBatchSelectOperation();
-	for (int32 ActorIndex = 0; ActorIndex < SelectedActors.Num(); ++ActorIndex)
-	{
-		AActor* Actor = SelectedActors[ActorIndex].Get();
-		if (Actor)
-		{
-			// We need to notify or else the manipulation transform widget won't appear, but only notify once at the end because OnEditorSelectionChanged is expensive for large groups. 
-			SelectActor(Actor, true, false);
-		}
-	}
-	GEditor->GetSelectedActors()->EndBatchSelectOperation(true);
-
 	// Spawn note actors dropped in PIE.
 	if(GEngine->PendingDroppedNotes.Num() > 0)
 	{
@@ -583,6 +570,21 @@ void UEditorEngine::EndPlayMap()
 		PRAGMA_ENABLE_DEPRECATION_WARNINGS
 	}
 
+	/**
+	 * For every actor that was selected previously, make sure it's editor equivalent is selected. We do that after the cleanup in case some actor where removed.
+	 * This must be done at after bIsSimulatingInEditor is set to false because some of the selection logic can have a dependency on the state of this variable.
+	 */
+	GEditor->GetSelectedActors()->BeginBatchSelectOperation();
+	for (int32 ActorIndex = 0; ActorIndex < SelectedActors.Num(); ++ActorIndex)
+	{
+		AActor* Actor = SelectedActors[ActorIndex].Get();
+		if (Actor)
+		{
+			// We need to notify or else the manipulation transform widget won't appear, but only notify once at the end because OnEditorSelectionChanged is expensive for large groups. 
+			SelectActor(Actor, true, false);
+		}
+	}
+	GEditor->GetSelectedActors()->EndBatchSelectOperation(true);
 }
 
 void UEditorEngine::CleanupPIEOnlineSessions(TArray<FName> OnlineIdentifiers)
