Fix missing input/output pins onto 'make attribute' and 'break attribute' nodes when strata is enabled.

diff --git a/Engine/Source/Editor/UnrealEd/Classes/MaterialGraph/MaterialGraph.h b/Engine/Source/Editor/UnrealEd/Classes/MaterialGraph/MaterialGraph.h
--- a/Engine/Source/Editor/UnrealEd/Classes/MaterialGraph/MaterialGraph.h
+++ b/Engine/Source/Editor/UnrealEd/Classes/MaterialGraph/MaterialGraph.h
@@ -53,7 +53,7 @@ struct FMaterialInputInfo
 
 	bool IsVisiblePin(const UMaterial* Material, bool bIgnoreMaterialAttributes = false) const
 	{
-		if (!Material->IsPropertySupported(Property))
+		if (!Material->IsPropertySupported(Property) && !bIgnoreMaterialAttributes)
 		{
 			return false;
 		}
