Fix crash duplicating an actor with a skeletalmesh using an ABP that uses an template anim BP as its parent

diff --git a/Engine/Source/Editor/UnrealEd/Private/EditorObject.cpp b/Engine/Source/Editor/UnrealEd/Private/EditorObject.cpp
--- a/Engine/Source/Editor/UnrealEd/Private/EditorObject.cpp
+++ b/Engine/Source/Editor/UnrealEd/Private/EditorObject.cpp
@@ -596,6 +596,9 @@ static const TCHAR* ImportProperties(
 						NewFlags |= RF_DefaultSubObject;
 					}
 
+					// fix crash when comment below is true as it is now forbidden to clear RF_PendingKill using ClearFlags (see UObjectBaseUtility.h)
+					ComponentTemplate->ClearGarbage();
+					
 					// Make sure desired flags are set - existing object could be pending kill
 					ComponentTemplate->ClearGarbage();
 					ComponentTemplate->ClearFlags(RF_AllFlags);
