Fix a crash if no group actor is selected when attempting to add actors to a selected group

PR #9113: Update GroupActor.cpp (Contributed by [REDACTED_PERSON])

diff --git a/Engine/Source/Editor/UnrealEd/Private/GroupActor.cpp b/Engine/Source/Editor/UnrealEd/Private/GroupActor.cpp
--- a/Engine/Source/Editor/UnrealEd/Private/GroupActor.cpp
+++ b/Engine/Source/Editor/UnrealEd/Private/GroupActor.cpp
@@ -549,8 +549,8 @@ void AGroupActor::AddSelectedActorsToSelectedGroup()
 			}
 		}
 
-		AGroupActor* SelectedGroup = Cast<AGroupActor>(EditorWorld->ActiveGroupActors[SelectedGroupIndex]);
-		if( SelectedGroupIndex != -1 && SelectedGroup != NULL )
+		AGroupActor* SelectedGroup = SelectedGroupIndex != -1 ? Cast<AGroupActor>(EditorWorld->ActiveGroupActors[SelectedGroupIndex]) : nullptr;
+		if( SelectedGroup != nullptr )
 		{
 			ULevel* GroupLevel = SelectedGroup->GetLevel();
 
