Fix a crash in the mode manager when some mode's exit sequence accidentally calls enter on itself.

diff --git a/Engine/Source/Editor/UnrealEd/Private/EditorModeManager.cpp b/Engine/Source/Editor/UnrealEd/Private/EditorModeManager.cpp
--- a/Engine/Source/Editor/UnrealEd/Private/EditorModeManager.cpp
+++ b/Engine/Source/Editor/UnrealEd/Private/EditorModeManager.cpp
@@ -51,8 +51,6 @@
 	The master class that handles tracking of the current mode.
 ------------------------------------------------------------------------------*/
 
-
-
 FEditorModeTools::FEditorModeTools()
 	: PivotShown(false)
 	, Snapping(false)
@@ -631,13 +629,17 @@ bool FEditorModeTools::TestAllModes(TFunctionRef<bool(UEdMode*)> InCalllback, bo
 
 void FEditorModeTools::ExitAllModesPendingDeactivate()
 {
-	// Make a copy so we can modify the pending deactivate modes map
+	bIsExitingModesDuringTick = true;
+
+	// Make a copy so we can modify the pending deactivate modes map during ExitMode
 	TMap<FEditorModeID, UEdMode*> PendingDeactivateModesCopy(PendingDeactivateModes);
 	for (auto& Pair : PendingDeactivateModesCopy)
 	{
 		ExitMode(Pair.Value);
 	}
 
+	bIsExitingModesDuringTick = false;
+
 	check(PendingDeactivateModes.Num() == 0);
 }
 
@@ -851,8 +853,15 @@ void FEditorModeTools::ActivateMode(FEditorModeID InID, bool bToggle)
 	if (!ScriptableMode)
 	{
 		ScriptableMode = PendingDeactivateModes.FindRef(InID);
+
 		if (ScriptableMode)
 		{
+			// If we are actively exiting modes, don't re-activate the mode
+			if (bIsExitingModesDuringTick)
+			{
+				return;
+			}
+
 			bNeedsEnter = false;
 		}
 	}
diff --git a/Engine/Source/Editor/UnrealEd/Public/EditorModeManager.h b/Engine/Source/Editor/UnrealEd/Public/EditorModeManager.h
--- a/Engine/Source/Editor/UnrealEd/Public/EditorModeManager.h
+++ b/Engine/Source/Editor/UnrealEd/Public/EditorModeManager.h
@@ -613,6 +613,9 @@ private:
 	/** Flag set between calls to StartTracking() and EndTracking() */
 	bool bIsTracking;
 
+	/** Guard to prevent modes from entering as part of their exit routine */
+	bool bIsExitingModesDuringTick = false;
+
 	FEditorViewportClient* HoveredViewportClient = nullptr;
 	FEditorViewportClient* FocusedViewportClient = nullptr;
 
