Wait for static meshes to be compiled before generating navigation data when using UResavePackagesCommandlet

diff --git a/Engine/Source/Editor/UnrealEd/Private/Commandlets/ContentCommandlets.cpp b/Engine/Source/Editor/UnrealEd/Private/Commandlets/ContentCommandlets.cpp
--- a/Engine/Source/Editor/UnrealEd/Private/Commandlets/ContentCommandlets.cpp
+++ b/Engine/Source/Editor/UnrealEd/Private/Commandlets/ContentCommandlets.cpp
@@ -54,6 +54,7 @@
 #include "AssetCompilingManager.h"
 #include "PackageHelperFunctions.h"
 #include "PackageTools.h"
+#include "StaticMeshCompiler.h"
 #include "String/ParseTokens.h"
 #include "UObject/PackageTrailer.h"
 #include "UObject/SavePackage.h"
@@ -2043,6 +2044,9 @@ void UResavePackagesCommandlet::PerformAdditionalOperations(class UWorld* World,
 
 			if (bShouldBuildNavigationData)
 			{
+				// Make sure static meshes have compiled before generating navigation data
+				FStaticMeshCompilingManager::Get().FinishAllCompilation();
+				
 				// Make sure navigation is added and initialized in EditorMode
 				FNavigationSystem::AddNavigationSystemToWorld(*World, FNavigationSystemRunMode::EditorMode);
 
