Expose cooking switch for saving Optional Objects packages

diff --git a/Engine/Source/Editor/UnrealEd/Classes/Commandlets/CookCommandlet.h b/Engine/Source/Editor/UnrealEd/Classes/Commandlets/CookCommandlet.h
--- a/Engine/Source/Editor/UnrealEd/Classes/Commandlets/CookCommandlet.h
+++ b/Engine/Source/Editor/UnrealEd/Classes/Commandlets/CookCommandlet.h
@@ -39,6 +39,8 @@ class UCookCommandlet
 	bool bSkipEditorContent;
 	/** Save all cooked packages without versions. These are then assumed to be current version on load. This is dangerous but results in smaller patch sizes. */
 	bool bUnversioned;
+	/** Produce editor optional package output when cooking. */
+	bool bCookEditorOptional;
 	/** Generate manifests for building streaming install packages */
 	bool bGenerateStreamingInstallManifests;
 	/** Error if we access engine content (useful for dlc) */
diff --git a/Engine/Source/Editor/UnrealEd/Classes/CookOnTheSide/CookOnTheFlyServer.h b/Engine/Source/Editor/UnrealEd/Classes/CookOnTheSide/CookOnTheFlyServer.h
--- a/Engine/Source/Editor/UnrealEd/Classes/CookOnTheSide/CookOnTheFlyServer.h
+++ b/Engine/Source/Editor/UnrealEd/Classes/CookOnTheSide/CookOnTheFlyServer.h
@@ -57,6 +57,8 @@ enum class ECookInitializationFlags
 	IterateSharedBuild =						0x00010000, // iterate from a build in the SharedIterativeBuild directory 
 	IgnoreIniSettingsOutOfDate =				0x00020000, // if the inisettings say the cook is out of date keep using the previously cooked build
 	IgnoreScriptPackagesOutOfDate =				0x00040000, // for incremental cooking, ignore script package changes
+	//unused =									0x00080000,
+	CookEditorOptional =						0x00100000,	// enable producing the cooked output of optional editor packages that can be use in the editor when loading cooked data
 };
 ENUM_CLASS_FLAGS(ECookInitializationFlags);
 
diff --git a/Engine/Source/Editor/UnrealEd/Private/Commandlets/CookCommandlet.cpp b/Engine/Source/Editor/UnrealEd/Private/Commandlets/CookCommandlet.cpp
--- a/Engine/Source/Editor/UnrealEd/Private/Commandlets/CookCommandlet.cpp
+++ b/Engine/Source/Editor/UnrealEd/Private/Commandlets/CookCommandlet.cpp
@@ -421,6 +421,7 @@ bool UCookCommandlet::CookOnTheFly( FGuid InstanceId, int32 Timeout, bool bForce
 	CookFlags |= bIterativeCooking ? IterateFlags : ECookInitializationFlags::None;
 	CookFlags |= bSkipEditorContent ? ECookInitializationFlags::SkipEditorContent : ECookInitializationFlags::None;
 	CookFlags |= bUnversioned ? ECookInitializationFlags::Unversioned : ECookInitializationFlags::None;
+	CookFlags |= bCookEditorOptional ? ECookInitializationFlags::CookEditorOptional : ECookInitializationFlags::None;
 	CookOnTheFlyServer->Initialize( ECookMode::CookOnTheFly, CookFlags );
 
 	UCookOnTheFlyServer::FCookOnTheFlyOptions CookOnTheFlyStartupOptions;
@@ -627,6 +628,7 @@ int32 UCookCommandlet::Main(const FString& CmdLineParams)
 	bCookOnTheFly = Switches.Contains(TEXT("COOKONTHEFLY"));   // Prototype cook-on-the-fly server
 	bCookAll = Switches.Contains(TEXT("COOKALL"));   // Cook everything
 	bUnversioned = Switches.Contains(TEXT("UNVERSIONED"));   // Save all cooked packages without versions. These are then assumed to be current version on load. This is dangerous but results in smaller patch sizes.
+	bCookEditorOptional = Switches.Contains(TEXT("EDITOROPTIONAL")); // Produce the optional editor package data alongside the cooked data.
 	bGenerateStreamingInstallManifests = Switches.Contains(TEXT("MANIFESTS"));   // Generate manifests for building streaming install packages
 	bIterativeCooking = Switches.Contains(TEXT("ITERATE"));
 	bSkipEditorContent = Switches.Contains(TEXT("SKIPEDITORCONTENT")); // This won't save out any packages in Engine/Content/Editor*
@@ -746,6 +748,7 @@ bool UCookCommandlet::CookByTheBook( const TArray<ITargetPlatform*>& Platforms)
 	CookFlags |= bSkipEditorContent ? ECookInitializationFlags::SkipEditorContent : ECookInitializationFlags::None;	
 	CookFlags |= bUseSerializationForGeneratingPackageDependencies ? ECookInitializationFlags::UseSerializationForPackageDependencies : ECookInitializationFlags::None;
 	CookFlags |= bUnversioned ? ECookInitializationFlags::Unversioned : ECookInitializationFlags::None;
+	CookFlags |= bCookEditorOptional ? ECookInitializationFlags::CookEditorOptional : ECookInitializationFlags::None;
 	CookFlags |= bVerboseCookerWarnings ? ECookInitializationFlags::OutputVerboseCookerWarnings : ECookInitializationFlags::None;
 	CookFlags |= bPartialGC ? ECookInitializationFlags::EnablePartialGC : ECookInitializationFlags::None;
 	bool bTestCook = Switches.Contains(TEXT("TestCook"));
diff --git a/Engine/Source/Editor/UnrealEd/Private/CookOnTheFlyServer.cpp b/Engine/Source/Editor/UnrealEd/Private/CookOnTheFlyServer.cpp
--- a/Engine/Source/Editor/UnrealEd/Private/CookOnTheFlyServer.cpp
+++ b/Engine/Source/Editor/UnrealEd/Private/CookOnTheFlyServer.cpp
@@ -4333,6 +4333,7 @@ void FSaveCookedPackageContext::SetupPackage()
 
 	SaveFlags = SAVE_KeepGUID | SAVE_Async
 		| (COTFS.IsCookFlagSet(ECookInitializationFlags::Unversioned) ? SAVE_Unversioned : 0);
+	SaveFlags |= COTFS.IsCookFlagSet(ECookInitializationFlags::CookEditorOptional) ? SAVE_Optional : SAVE_None;
 
 	// removing editor only packages only works when cooking in commandlet and non iterative cooking
 	// also doesn't work in multiprocess cooking
