SDlgPickAssetPath: Validate DefaultAssetPath against AllowReadOnlyFolders.
FFoliageEditUtility::SaveFoliageTypeObject initialize SDlgPickAssetPath with AllowReadonlyFolders(false)

diff --git a/Engine/Source/Editor/FoliageEdit/Private/FoliageEditUtility.cpp b/Engine/Source/Editor/FoliageEdit/Private/FoliageEditUtility.cpp
--- a/Engine/Source/Editor/FoliageEdit/Private/FoliageEditUtility.cpp
+++ b/Engine/Source/Editor/FoliageEdit/Private/FoliageEditUtility.cpp
@@ -35,7 +35,8 @@ UFoliageType* FFoliageEditUtility::SaveFoliageTypeObject(UFoliageType* InFoliage
 		TSharedRef<SDlgPickAssetPath> SaveFoliageTypeDialog =
 			SNew(SDlgPickAssetPath)
 			.Title(LOCTEXT("SaveFoliageTypeDialogTitle", "Choose Location for Foliage Type Asset"))
-			.DefaultAssetPath(FText::FromString(PackageName));
+			.DefaultAssetPath(FText::FromString(PackageName))
+			.AllowReadOnlyFolders(false);
 
 		if (SaveFoliageTypeDialog->ShowModal() != EAppReturnType::Cancel)
 		{
diff --git a/Engine/Source/Editor/UnrealEd/Private/Dialogs/DlgPickAssetPath.cpp b/Engine/Source/Editor/UnrealEd/Private/Dialogs/DlgPickAssetPath.cpp
--- a/Engine/Source/Editor/UnrealEd/Private/Dialogs/DlgPickAssetPath.cpp
+++ b/Engine/Source/Editor/UnrealEd/Private/Dialogs/DlgPickAssetPath.cpp
@@ -13,16 +13,25 @@
 #include "Editor.h"
 #include "IContentBrowserSingleton.h"
 #include "ContentBrowserModule.h"
+#include "AssetToolsModule.h"
 #include "SPrimaryButton.h"
 
 #define LOCTEXT_NAMESPACE "DlgPickAssetPath"
 
 void SDlgPickAssetPath::Construct(const FArguments& InArgs)
 {
-	AssetPath = FText::FromString(FPackageName::GetLongPackagePath(InArgs._DefaultAssetPath.ToString()));
-	AssetName = FText::FromString(FPackageName::GetLongPackageAssetName(InArgs._DefaultAssetPath.ToString()));
+	const FString DefaultAssetPath = InArgs._DefaultAssetPath.ToString();
+
 	bAllowReadOnlyFolders = InArgs._AllowReadOnlyFolders;
 
+	// Make sure DefaultAssetPath is allowed before setting AssetPath/Name
+	FAssetToolsModule& AssetToolsModule = FModuleManager::LoadModuleChecked<FAssetToolsModule>("AssetTools");
+	if (bAllowReadOnlyFolders || AssetToolsModule.Get().GetWritableFolderPermissionList()->PassesStartsWithFilter(DefaultAssetPath))
+	{
+		AssetPath = FText::FromString(FPackageName::GetLongPackagePath(DefaultAssetPath));
+		AssetName = FText::FromString(FPackageName::GetLongPackageAssetName(DefaultAssetPath));
+	}
+	
 	FPathPickerConfig PathPickerConfig;
 	PathPickerConfig.DefaultPath = AssetPath.ToString();
 	PathPickerConfig.OnPathSelected = FOnPathSelected::CreateSP(this, &SDlgPickAssetPath::OnPathChange);
