Turn on hash based guid for Fbx importer to avoid DDC miss on reimport
  - 45s to 21s to reimport a test mesh

#rnx

diff --git a/Engine/Source/Editor/UnrealEd/Private/Fbx/FbxStaticMeshImport.cpp b/Engine/Source/Editor/UnrealEd/Private/Fbx/FbxStaticMeshImport.cpp
--- a/Engine/Source/Editor/UnrealEd/Private/Fbx/FbxStaticMeshImport.cpp
+++ b/Engine/Source/Editor/UnrealEd/Private/Fbx/FbxStaticMeshImport.cpp
@@ -1760,7 +1760,10 @@ UStaticMesh* UnFbx::FFbxImporter::ImportStaticMeshAsSingle(UObject* InParent, TA
 	{
 		MeshDescription = StaticMesh->CreateMeshDescription(LODIndex);
 		check(MeshDescription != nullptr);
-		StaticMesh->CommitMeshDescription(LODIndex);
+
+		UStaticMesh::FCommitMeshDescriptionParams Params;
+		Params.bUseHashAsGuid = true;
+		StaticMesh->CommitMeshDescription(LODIndex, Params);
 
 		//Make sure an imported mesh do not get reduce if there was no mesh data before reimport.
 		//In this case we have a generated LOD convert to a custom LOD
@@ -1927,8 +1930,11 @@ UStaticMesh* UnFbx::FFbxImporter::ImportStaticMeshAsSingle(UObject* InParent, TA
 				SectionIndex++;
 			}
 		}
+
 		//Set the original mesh description to be able to do non destructive reduce
-		StaticMesh->CommitMeshDescription(LODIndex);
+		UStaticMesh::FCommitMeshDescriptionParams Params;
+		Params.bUseHashAsGuid = true;
+		StaticMesh->CommitMeshDescription(LODIndex, Params);
 
 		// Setup default LOD settings based on the selected LOD group.
 		if (LODIndex == 0)
