SAssetThumbnail: Create the viewport for the thumbnail with the desired dimensions to fix problems with using SAssetThumbnail in a ListViewRow where the thumbnail desired height remained at 240 (=the default), stretching the row height in the process.

diff --git a/Engine/Source/Editor/UnrealEd/Private/AssetThumbnail.cpp b/Engine/Source/Editor/UnrealEd/Private/AssetThumbnail.cpp
--- a/Engine/Source/Editor/UnrealEd/Private/AssetThumbnail.cpp
+++ b/Engine/Source/Editor/UnrealEd/Private/AssetThumbnail.cpp
@@ -175,7 +175,8 @@ public:
 				// In VR editor every widget is in the world and gamma corrected by the scene renderer.  Thumbnails will have already been gamma
 				// corrected and so they need to be reversed
 				.ReverseGammaCorrection(IVREditorModule::Get().IsVREditorModeActive())
-				.EnableBlending(true);
+				.EnableBlending(true)
+				.ViewportSize(AssetThumbnail->GetSize());
 
 			Viewport->SetViewportInterface( AssetThumbnail.ToSharedRef() );
 			AssetThumbnail->GetViewportRenderTargetTexture(); // Access the render texture to push it on the stack if it isn't already rendered
