CookPackageSplitter: Add a Teardown function called when all generated packages have been saved. Suppress "package still in memory" warnings when the cooker needs to keep the packages referenced.
#rnx

diff --git a/Engine/Source/Editor/UnrealEd/Classes/CookOnTheSide/CookOnTheFlyServer.h b/Engine/Source/Editor/UnrealEd/Classes/CookOnTheSide/CookOnTheFlyServer.h
--- a/Engine/Source/Editor/UnrealEd/Classes/CookOnTheSide/CookOnTheFlyServer.h
+++ b/Engine/Source/Editor/UnrealEd/Classes/CookOnTheSide/CookOnTheFlyServer.h
@@ -777,6 +777,9 @@ private:
 	*/
 	void CookByTheBookFinished();
 
+	/** Clears session-lifetime data from COTFS and CookPackageDatas. */
+	void ShutdownCookSession();
+
 	/** Print some stats when finishing or cancelling CookByTheBook */
 	void PrintFinishStats();
 
diff --git a/Engine/Source/Editor/UnrealEd/Private/CookOnTheFlyServer.cpp b/Engine/Source/Editor/UnrealEd/Private/CookOnTheFlyServer.cpp
--- a/Engine/Source/Editor/UnrealEd/Private/CookOnTheFlyServer.cpp
+++ b/Engine/Source/Editor/UnrealEd/Private/CookOnTheFlyServer.cpp
@@ -2963,34 +2963,27 @@ void UCookOnTheFlyServer::ReleaseCookedPlatformData(UE::Cook::FPackageData& Pack
 	if (bCompletedSave)
 	{
 		FGeneratorPackage* GeneratorPackage = PackageData.GetGeneratorPackage();
-		if (GeneratorPackage && PackageData.HasCompletedGeneration())
+		if (GeneratorPackage)
 		{
-			UObject* SplitObject = GeneratorPackage->FindSplitDataObject();
-			if (!SplitObject || !PackageData.GetPackage())
+			if (PackageData.HasCompletedGeneration()) // bCompletedSave can be true if !HasCompletedGeneration in the case of a GC request from BeginPrepareSave
 			{
-				UE_LOG(LogCook, Error, TEXT("PackageSplitter: %s was GarbageCollected before we finished saving it. This will cause a failure later when we attempt to save it again and generate a second time. Splitter=%s."),
-					(!PackageData.GetPackage() ? TEXT("UPackage") : TEXT("SplitDataObject")), *GeneratorPackage->GetSplitDataObjectName().ToString());
+				GeneratorPackage->SetGeneratorSaved(PackageData.GetPackage());
 			}
-			else
-			{
-				GeneratorPackage->GetCookPackageSplitterInstance()->PostSaveGeneratorPackage(PackageData.GetPackage(), SplitObject);
-			}
-
-			PackageData.ResetGenerationProgress();
-			if (IsCookByTheBookMode() && GeneratorPackage->IsComplete())
+			if (GeneratorPackage->IsComplete())
 			{
 				PackageData.DestroyGeneratorPackage();
 			}
 		}
-		else
-		{
-			PackageData.ResetGenerationProgress();
-		}
+		PackageData.ResetGenerationProgress();
 
 		FGeneratorPackage* OwnerGeneratorPackage = PackageData.GetGeneratedOwner();
 		if (OwnerGeneratorPackage)
 		{
 			OwnerGeneratorPackage->SetGeneratedSaved(PackageData);
+			if (OwnerGeneratorPackage->IsComplete())
+			{
+				PackageData.DestroyGeneratorPackage();
+			}
 		}
 	}
 	PackageData.ClearCookedPlatformData();
@@ -3837,7 +3830,8 @@ void UCookOnTheFlyServer::ShutdownCookOnTheFly()
 		CookOnTheFlyRequestManager->Shutdown();
 		CookOnTheFlyRequestManager.Reset();
 		ClearHierarchyTimers();
-		ConditionalUninstallImportBehaviorCallback();
+
+		ShutdownCookSession();
 	}
 }
 
@@ -3972,6 +3966,7 @@ void UCookOnTheFlyServer::PreGarbageCollect()
 	}
 
 	TArray<UPackage*> GCKeepPackages;
+	TArray<FPackageData*> GCKeepPackageDatas;
 	for (FPackageData* PackageData : PackageDatas->GetSaveQueue())
 	{
 		// Generator packages that have started generation and have not yet finished saving should not be garbage
@@ -3980,6 +3975,7 @@ void UCookOnTheFlyServer::PreGarbageCollect()
 		if (PackageData->GetGeneratorPackage())
 		{
 			GCKeepPackages.Add(PackageData->GetPackage());
+			GCKeepPackageDatas.Add(PackageData);
 		}
 	}
 
@@ -3987,23 +3983,24 @@ void UCookOnTheFlyServer::PreGarbageCollect()
 	// and make sure that they are not garbage collected until the jobs have
 	// completed.
 	{
-		TSet<UPackage*> UniquePendingPackages;
+		TMap<FPackageData*, UPackage*> UniquePendingPackages;
 		for (FPendingCookedPlatformData& PendingData : PackageDatas->GetPendingCookedPlatformDatas())
 		{
 			if (UObject* Object = PendingData.Object.Get())
 			{	
 				if (UPackage* Package = Object->GetPackage())
 				{
-					UniquePendingPackages.Add(Package);
+					UniquePendingPackages.Add(&PendingData.PackageData, Package);
 				}	
 			}
 		}
 
 		GCKeepPackages.Reserve(GCKeepPackages.Num() + UniquePendingPackages.Num());
-		for (UPackage* Package : UniquePendingPackages)
+		for (const TPair<FPackageData*,UPackage*>& Pair : UniquePendingPackages)
 		{
-			GCKeepPackages.Add(Package);
-		}	
+			GCKeepPackages.Add(Pair.Value);
+			GCKeepPackageDatas.Add(Pair.Key);
+		}
 	}
 
 	const bool bPartialGC = IsCookFlagSet(ECookInitializationFlags::EnablePartialGC);
@@ -4017,16 +4014,16 @@ void UCookOnTheFlyServer::PreGarbageCollect()
 		TMap<const FPackageData*, int32> DependenciesCount;
 
 		TSet<FName> KeepPackages;
-		for (const FPackageData* PackageData : *PackageDatas)
+		for (FPackageData* PackageData : *PackageDatas)
 		{
-			if (PackageData->GetState() == UE::Cook::EPackageState::Save)
+			if (!PackageData->IsInProgress())
 			{
-				// already handled above
 				continue;
 			}
 			const TArray<FName>& NeededPackages = GetFullPackageDependencies(PackageData->GetPackageName());
 			DependenciesCount.Add(PackageData, NeededPackages.Num());
 			KeepPackages.Append(NeededPackages);
+			GCKeepPackageDatas.Add(PackageData);
 		}
 
 		TSet<FName> LoadedPackages;
@@ -4085,6 +4082,10 @@ void UCookOnTheFlyServer::PreGarbageCollect()
 			}
 		}
 	}
+	for (FPackageData* PackageData : GCKeepPackageDatas)
+	{
+		PackageData->SetKeepReferencedDuringGC(true);
+	}
 }
 
 void UCookOnTheFlyServer::CookerAddReferencedObjects(FReferenceCollector& Collector)
@@ -4129,6 +4130,10 @@ void UCookOnTheFlyServer::PostGarbageCollect()
 			GeneratorPackage->PostGarbageCollect();
 		}
 	}
+	for (FPackageData* PackageData : *PackageDatas.Get())
+	{
+		PackageData->SetKeepReferencedDuringGC(false);
+	}
 
 	CookedPackageCountSinceLastGC = 0;
 }
@@ -6978,8 +6983,6 @@ void UCookOnTheFlyServer::CookByTheBookFinished()
 	BuildDefinitions->Wait();
 	
 	GetDerivedDataCacheRef().WaitForQuiescence(true);
-
-	ConditionalUninstallImportBehaviorCallback();
 	
 	UCookerSettings const* CookerSettings = GetDefault<UCookerSettings>();
 
@@ -7200,7 +7203,7 @@ void UCookOnTheFlyServer::CookByTheBookFinished()
 	{
 		FCoreUObjectDelegates::PackageCreatedForLoad.RemoveAll(this);
 	}
-	PlatformManager->ClearSessionPlatforms();
+	ShutdownCookSession();
 
 	PrintFinishStats();
 
@@ -7210,6 +7213,17 @@ void UCookOnTheFlyServer::CookByTheBookFinished()
 	UE_LOG(LogCook, Display, TEXT("Done!"));
 }
 
+void UCookOnTheFlyServer::ShutdownCookSession()
+{
+	PlatformManager->ClearSessionPlatforms();
+	for (UE::Cook::FPackageData* PackageData : *PackageDatas)
+	{
+		PackageData->DestroyGeneratorPackage();
+	}
+
+	ConditionalUninstallImportBehaviorCallback();
+}
+
 void UCookOnTheFlyServer::PrintFinishStats()
 {
 	const float TotalCookTime = (float)(FPlatformTime::Seconds() - CookByTheBookOptions->CookStartTime);
diff --git a/Engine/Source/Editor/UnrealEd/Private/Cooker/CookPackageData.cpp b/Engine/Source/Editor/UnrealEd/Private/Cooker/CookPackageData.cpp
--- a/Engine/Source/Editor/UnrealEd/Private/Cooker/CookPackageData.cpp
+++ b/Engine/Source/Editor/UnrealEd/Private/Cooker/CookPackageData.cpp
@@ -56,7 +56,7 @@ FPackageData::FPackageData(FPackageDatas& PackageDatas, const FName& InPackageNa
 	, bIsVisited(0), bIsPreloadAttempted(0)
 	, bIsPreloaded(0), bHasSaveCache(0), bHasBeginPrepareSaveFailed(0), bCookedPlatformDataStarted(0)
 	, bCookedPlatformDataCalled(0), bCookedPlatformDataComplete(0), bMonitorIsCooked(0)
-	, bInitializedGeneratorSave(0), bCompletedGeneration(0), bGenerated(0)
+	, bInitializedGeneratorSave(0), bCompletedGeneration(0), bGenerated(0), bKeepReferencedDuringGC(0)
 {
 	SetState(EPackageState::Idle);
 	SendToState(EPackageState::Idle, ESendFlags::QueueAdd);
@@ -1127,9 +1127,19 @@ FGeneratorPackage::FGeneratorPackage(UE::Cook::FPackageData& InOwner, const UObj
 
 FGeneratorPackage::~FGeneratorPackage()
 {
+	ConditionalNotifyCompletion(ICookPackageSplitter::ETeardown::Canceled);
 	ClearGeneratedPackages();
 }
 
+void FGeneratorPackage::ConditionalNotifyCompletion(ICookPackageSplitter::ETeardown Status)
+{
+	if (!bNotifiedCompletion)
+	{
+		bNotifiedCompletion = true;
+		CookPackageSplitterInstance->Teardown(Status);
+	}
+}
+
 void FGeneratorPackage::ClearGeneratedPackages()
 {
 	for (FGeneratedStruct& GeneratedStruct : PackagesToGenerate)
@@ -1237,7 +1247,7 @@ void FGeneratorPackage::PostGarbageCollect()
 	}
 	else
 	{
-		// After the Generator Package is saved, we drop its referenced and it can be garbage collected
+		// After the Generator Package is saved, we drop our references to it and it can be garbage collected
 		// If we have any packages left to populate, our splitter contract requires that it be garbage collected;
 		// we promise that the package is not partially GC'd during calls to TryPopulateGeneratedPackage
 		// The splitter can opt-out of this contract and keep it referenced itself if it desires.
@@ -1245,6 +1255,7 @@ void FGeneratorPackage::PostGarbageCollect()
 		if (OwnerPackage)
 		{
 			if (RemainingToPopulate > 0 &&
+				!Owner.IsKeepReferencedDuringGC() &&
 				!CookPackageSplitterInstance->UseInternalReferenceToAvoidGarbageCollect())
 			{
 				UE_LOG(LogCook, Error, TEXT("PackageSplitter found the Generator package still in memory after it should have been deleted by GC.")
@@ -1265,16 +1276,21 @@ void FGeneratorPackage::PostGarbageCollect()
 	bool bHasIssuedWarning = false;
 	for (FGeneratedStruct& GeneratedStruct : PackagesToGenerate)
 	{
-		GeneratedStruct.bHasCreatedPackage = false;
-		if (!GeneratedStruct.bHasSaved && !bHasIssuedWarning)
+		if (FindObject<UPackage>(nullptr, *GeneratedStruct.PackageData->GetPackageName().ToString()))
 		{
-			if (FindObject<UPackage>(nullptr, *GeneratedStruct.PackageData->GetPackageName().ToString()))
+			if (!GeneratedStruct.PackageData->IsKeepReferencedDuringGC() &&
+				!GeneratedStruct.bHasSaved &&
+				!bHasIssuedWarning)
 			{
 				UE_LOG(LogCook, Warning, TEXT("PackageSplitter found a package it generated that was not removed from memory during garbage collection. This will cause errors later during population.")
 					TEXT("\n\tSplitter=%s, Generated=%s."), *GetSplitDataObjectName().ToString(), *GeneratedStruct.PackageData->GetPackageName().ToString());
 				bHasIssuedWarning = true; // Only issue the warning once per GC
 			}
 		}
+		else
+		{
+			GeneratedStruct.bHasCreatedPackage = false;
+		}
 	}
 }
 
@@ -1290,6 +1306,25 @@ UPackage* FGeneratorPackage::CreateGeneratedUPackage(FGeneratorPackage::FGenerat
 	return GeneratedPackage;
 }
 
+void FGeneratorPackage::SetGeneratorSaved(UPackage* GeneratorUPackage)
+{
+	UObject* SplitObject = FindSplitDataObject();
+	if (!SplitObject || !GeneratorUPackage)
+	{
+		UE_LOG(LogCook, Error, TEXT("PackageSplitter: %s was GarbageCollected before we finished saving it. This will cause a failure later when we attempt to save it again and generate a second time. Splitter=%s."),
+			(!GeneratorUPackage ? TEXT("UPackage") : TEXT("SplitDataObject")), *GetSplitDataObjectName().ToString());
+	}
+	else
+	{
+		GetCookPackageSplitterInstance()->PostSaveGeneratorPackage(GeneratorUPackage, SplitObject);
+	}
+	bOwnerHasSaved = true;
+	if (IsComplete())
+	{
+		ConditionalNotifyCompletion(ICookPackageSplitter::ETeardown::Complete);
+	}
+}
+
 void FGeneratorPackage::SetGeneratedSaved(FPackageData& PackageData)
 {
 	FGeneratedStruct* GeneratedStruct = FindGeneratedStruct(&PackageData);
@@ -1307,6 +1342,10 @@ void FGeneratorPackage::SetGeneratedSaved(FPackageData& PackageData)
 	GeneratedStruct->bHasSaved = true;
 	--RemainingToPopulate;
 	check(RemainingToPopulate >= 0);
+	if (IsComplete())
+	{
+		ConditionalNotifyCompletion(ICookPackageSplitter::ETeardown::Complete);
+	}
 }
 
 bool FGeneratorPackage::IsComplete() const
diff --git a/Engine/Source/Editor/UnrealEd/Private/Cooker/CookPackageData.h b/Engine/Source/Editor/UnrealEd/Private/Cooker/CookPackageData.h
--- a/Engine/Source/Editor/UnrealEd/Private/Cooker/CookPackageData.h
+++ b/Engine/Source/Editor/UnrealEd/Private/Cooker/CookPackageData.h
@@ -453,6 +453,12 @@ public:
 	 */
 	const FInstigator& GetInstigator() const { return Instigator; }
 
+	/** Get whether COTFS is keeping this package referenced referenced during GC. */
+	bool IsKeepReferencedDuringGC() const { return static_cast<bool>(bKeepReferencedDuringGC); }
+	/** Set whether COTFS is keeping this package referenced referenced during GC. */
+	void SetKeepReferencedDuringGC(bool Value) { bKeepReferencedDuringGC = Value != 0; }
+
+
 private:
 	friend struct UE::Cook::FPackageDatas;
 
@@ -561,6 +567,7 @@ private:
 	uint32 bInitializedGeneratorSave : 1;
 	uint32 bCompletedGeneration : 1;
 	uint32 bGenerated : 1;
+	uint32 bKeepReferencedDuringGC : 1;
 };
 
 /**
@@ -638,12 +645,16 @@ public:
 	/** Call CreatePackage and set PackageData for deterministic generated packages, and update status. */
 	UPackage* CreateGeneratedUPackage(FGeneratedStruct& GeneratedStruct,
 		const UPackage* OwnerPackage, const TCHAR* GeneratedPackageName);
-	/** Mark that the generated package has been populated, to keep track of when *this is no longer needed. */
+	/** Mark that the generator package has saved, to keep track of when *this is no longer needed. */
+	void SetGeneratorSaved(UPackage* GeneratorUPackage);
+	/** Mark that the generated package has saved, to keep track of when *this is no longer needed. */
 	void SetGeneratedSaved(FPackageData& PackageData);
 	/** Return whether list has been generated and all generated packages have been populated */
 	bool IsComplete() const;
 
 private:
+	void ConditionalNotifyCompletion(ICookPackageSplitter::ETeardown Status);
+
 	/** PackageData for the package that is being split */
 	const UE::Cook::FPackageData& Owner;
 	/** Name of the object that prompted the splitter creation */
@@ -660,6 +671,8 @@ private:
 	bool bClearedOldPackagesWithGC = false;
 	bool bQueuedGeneratedPackages = false;
 	bool bWasOwnerReloaded = false;
+	bool bOwnerHasSaved = false;
+	bool bNotifiedCompletion = false;
 };
 
 
diff --git a/Engine/Source/Editor/UnrealEd/Public/CookPackageSplitter.h b/Engine/Source/Editor/UnrealEd/Public/CookPackageSplitter.h
--- a/Engine/Source/Editor/UnrealEd/Public/CookPackageSplitter.h
+++ b/Engine/Source/Editor/UnrealEd/Public/CookPackageSplitter.h
@@ -20,6 +20,14 @@ class ICookPackageSplitter
 public:
 	virtual ~ICookPackageSplitter() {}
 
+	enum class ETeardown
+	{
+		Complete,
+		Canceled,
+	};
+	/** Do teardown actions after all packages have saved, or when the cook is cancelled. Always called before destruction. */
+	virtual void Teardown(ETeardown Status) {}
+
 	/** Data sent to the cooker to describe each desired generated package */
 	struct FGeneratedPackage
 	{
