CookOnTheFlyServer: Reduce GCookProgressRetryBusyTime from 1 second to 10ms. It is only used when there is no other work to do, and calling PumpSaves/PumpLoads is cheap enough to do every 10ms. The most expensive part of PumpSaves when there is no work to do is PollPendingCookedPlatformDatas, and that function has a separate internal limiter on the frequency at which it runs.
#rnx

diff --git a/Engine/Source/Editor/UnrealEd/Private/CookOnTheFlyServer.cpp b/Engine/Source/Editor/UnrealEd/Private/CookOnTheFlyServer.cpp
--- a/Engine/Source/Editor/UnrealEd/Private/CookOnTheFlyServer.cpp
+++ b/Engine/Source/Editor/UnrealEd/Private/CookOnTheFlyServer.cpp
@@ -167,11 +167,11 @@ static FAutoConsoleVariableRef CVarCookDisplayRepeatTime(
 	TEXT("Controls the time before the cooker will repeat the same progress message.\n"),
 	ECVF_Default);
 
-float GCookProgressRetryBusyTime = 1.0f;
+float GCookProgressRetryBusyTime = 0.01f;
 static FAutoConsoleVariableRef CVarCookRetryBusyTime(
 	TEXT("Cook.retrybusytime"),
 	GCookProgressRetryBusyTime,
-	TEXT("Controls the time between retry attempts at save and load when the save and load queues are busy.\n"),
+	TEXT("Controls the time between retry attempts at save and load when the save and load queues are busy and there is no other work to do.\n"),
 	ECVF_Default);
 
 float GCookProgressWarnBusyTime = 120.0f;
