Optionally ignore inability to load .spc (PSO cache files).

diff --git a/Engine/Source/Editor/UnrealEd/Private/Commandlets/ShaderPipelineCacheToolsCommandlet.cpp b/Engine/Source/Editor/UnrealEd/Private/Commandlets/ShaderPipelineCacheToolsCommandlet.cpp
--- a/Engine/Source/Editor/UnrealEd/Private/Commandlets/ShaderPipelineCacheToolsCommandlet.cpp
+++ b/Engine/Source/Editor/UnrealEd/Private/Commandlets/ShaderPipelineCacheToolsCommandlet.cpp
@@ -39,6 +39,15 @@ static FAutoConsoleVariableRef CVarShaderPipelineCacheDoNotPrecompileComputePSO(
 	ECVF_Default
 );
 
+int32 GShaderPipelineCacheTools_IgnoreObsoleteStableCacheFiles = 0;
+static FAutoConsoleVariableRef CVarShaderPipelineCacheIgnoreObsoleteStableCacheFiles(
+	TEXT("r.ShaderPipelineCacheTools.IgnoreObsoleteStableCacheFiles"),
+	GShaderPipelineCacheTools_IgnoreObsoleteStableCacheFiles,
+	TEXT("When set to the default value of 0, building the cache (and usually the whole cook) will fail if any .spc file can't be loaded, to prevent further testing.\n")
+	TEXT("By setting to 1, a project may choose to ignore this instead (warning will still be issued)."),
+	ECVF_Default
+);
+
 
 struct FSCDataChunk
 {
@@ -2272,10 +2281,6 @@ int32 BuildPSOSC(const TArray<FString>& Tokens, const TMap<FString, FString>& Pa
 						{
 							UE_LOG(LogShaderPipelineCacheTools, Display, TEXT("Loaded %d stable PSOs from %s. %d PSOs rejected, %d PSOs merged"), PSOs.Num(), *FileName, PSOsRejected, PSOsMerged);
 						}
-						else
-						{
-
-						}
 					}, TStatId(), &PreReqs);
 			}
 		}
@@ -2301,7 +2306,14 @@ int32 BuildPSOSC(const TArray<FString>& Tokens, const TMap<FString, FString>& Pa
 
 		if (!PSOsByFile[FileIndex].Num())
 		{
-			return 1;
+			if (GShaderPipelineCacheTools_IgnoreObsoleteStableCacheFiles)
+			{
+				continue;
+			}
+			else
+			{
+				return 1;
+			}
 		}
 
 		check(TargetShaderFormat == NAME_None || TargetShaderFormat == TargetShaderFormatByFile[FileIndex]);
@@ -2309,11 +2321,6 @@ int32 BuildPSOSC(const TArray<FString>& Tokens, const TMap<FString, FString>& Pa
 
 		TSet<FPipelineCacheFileFormatPSO>& CurrentFilePSOs = PSOsByFile[FileIndex];
 
-		if (!CurrentFilePSOs.Num())
-		{
-			continue;
-		}
-
 		// Now merge this file PSO set with main PSO set (this is going to be slow as we need to incrementally reprocess each existing PSO per file to get reasonable bindcount averages).
 		// Can't sum all and avg: A) Overflow and B) Later ones want to remain high so only start to get averaged from the point they are added onwards:
 		// 1) New PSO goes in with it's bindcount intact for this iteration - if it's the last file then it keeps it bindcount
