Speculative fix crash on actor delete due to non-existent viewport.

diff --git a/Engine/Source/Editor/UnrealEd/Private/UnrealEdSrv.cpp b/Engine/Source/Editor/UnrealEd/Private/UnrealEdSrv.cpp
--- a/Engine/Source/Editor/UnrealEd/Private/UnrealEdSrv.cpp
+++ b/Engine/Source/Editor/UnrealEd/Private/UnrealEdSrv.cpp
@@ -2446,29 +2446,32 @@ bool UUnrealEdEngine::Exec_Actor( UWorld* InWorld, const TCHAR* Str, FOutputDevi
 	}
 	else if( FParse::Command(&Str,TEXT("DELETE")) )		// ACTOR SELECT DELETE
 	{
-		if (TSharedPtr<ILevelEditor> LevelEditor = GCurrentLevelEditingViewportClient->ParentLevelEditor.Pin())
+		if (GCurrentLevelEditingViewportClient)
 		{
-			if (UTypedElementCommonActions* CommonActions = LevelEditor->GetCommonActions())
+			if (TSharedPtr<ILevelEditor> LevelEditor = GCurrentLevelEditingViewportClient->ParentLevelEditor.Pin())
 			{
-				UTypedElementSelectionSet* SelectionSet = LevelEditor->GetMutableElementSelectionSet();
-				
-				const FScopedTransaction Transaction(NSLOCTEXT("UnrealEd", "DeleteElements", "Delete Elements"));
-				if (SelectionSet->GetNumSelectedElements() == 0)
+				if (UTypedElementCommonActions* CommonActions = LevelEditor->GetCommonActions())
 				{
-					// HACK: Not all modes will select elements, so allow them a shot at deletion if we don't think anything else is selected
-					// TODO: Move this logic into FLevelEditorActionCallbacks and have it call into the mode directly
-					if (!GLevelEditorModeTools().ProcessEditDelete())
+					UTypedElementSelectionSet* SelectionSet = LevelEditor->GetMutableElementSelectionSet();
+				
+					const FScopedTransaction Transaction(NSLOCTEXT("UnrealEd", "DeleteElements", "Delete Elements"));
+					if (SelectionSet->GetNumSelectedElements() == 0)
+					{
+						// HACK: Not all modes will select elements, so allow them a shot at deletion if we don't think anything else is selected
+						// TODO: Move this logic into FLevelEditorActionCallbacks and have it call into the mode directly
+						if (!GLevelEditorModeTools().ProcessEditDelete())
+						{
+							// HACK: Call these directly for an empty selection so that folder deletion in the outliner still works
+							// TODO: Move this logic into FLevelEditorActionCallbacks and have it call into the outliner directly
+							FEditorDelegates::OnDeleteActorsBegin.Broadcast();
+							FEditorDelegates::OnDeleteActorsEnd.Broadcast();
+						}
+					}
+					else
 					{
-						// HACK: Call these directly for an empty selection so that folder deletion in the outliner still works
-						// TODO: Move this logic into FLevelEditorActionCallbacks and have it call into the outliner directly
-						FEditorDelegates::OnDeleteActorsBegin.Broadcast();
-						FEditorDelegates::OnDeleteActorsEnd.Broadcast();
+						CommonActions->DeleteSelectedElements(SelectionSet, InWorld, FTypedElementDeletionOptions());
 					}
 				}
-				else
-				{
-					CommonActions->DeleteSelectedElements(SelectionSet, InWorld, FTypedElementDeletionOptions());
-				}
 			}
 		}
 		return true;
@@ -2639,32 +2642,35 @@ bool UUnrealEdEngine::Exec_Actor( UWorld* InWorld, const TCHAR* Str, FOutputDevi
 	}
 	else if( FParse::Command(&Str,TEXT("DUPLICATE")) )
 	{
-		if (TSharedPtr<ILevelEditor> LevelEditor = GCurrentLevelEditingViewportClient->ParentLevelEditor.Pin())
+		if (GCurrentLevelEditingViewportClient)
 		{
-			if (UTypedElementCommonActions* CommonActions = LevelEditor->GetCommonActions())
+			if (TSharedPtr<ILevelEditor> LevelEditor = GCurrentLevelEditingViewportClient->ParentLevelEditor.Pin())
 			{
-				UTypedElementSelectionSet* SelectionSet = LevelEditor->GetMutableElementSelectionSet();
-
-				const FScopedTransaction Transaction(NSLOCTEXT("UnrealEd", "DuplicateElements", "Duplicate Elements"));
-				if (SelectionSet->GetNumSelectedElements() == 0)
+				if (UTypedElementCommonActions* CommonActions = LevelEditor->GetCommonActions())
 				{
-					// HACK: Not all modes will select elements, so allow them a shot at duplication if we don't think anything else is selected
-					// TODO: Move this logic into FLevelEditorActionCallbacks and have it call into the mode directly
-					if (!GLevelEditorModeTools().ProcessEditDuplicate())
+					UTypedElementSelectionSet* SelectionSet = LevelEditor->GetMutableElementSelectionSet();
+
+					const FScopedTransaction Transaction(NSLOCTEXT("UnrealEd", "DuplicateElements", "Duplicate Elements"));
+					if (SelectionSet->GetNumSelectedElements() == 0)
 					{
-						// HACK: Call these directly for an empty selection so that folder duplication in the outliner still works
-						// TODO: Move this logic into FLevelEditorActionCallbacks and have it call into the outliner directly
-						FEditorDelegates::OnDuplicateActorsBegin.Broadcast();
-						FEditorDelegates::OnDuplicateActorsEnd.Broadcast();
+						// HACK: Not all modes will select elements, so allow them a shot at duplication if we don't think anything else is selected
+						// TODO: Move this logic into FLevelEditorActionCallbacks and have it call into the mode directly
+						if (!GLevelEditorModeTools().ProcessEditDuplicate())
+						{
+							// HACK: Call these directly for an empty selection so that folder duplication in the outliner still works
+							// TODO: Move this logic into FLevelEditorActionCallbacks and have it call into the outliner directly
+							FEditorDelegates::OnDuplicateActorsBegin.Broadcast();
+							FEditorDelegates::OnDuplicateActorsEnd.Broadcast();
+						}
 					}
-				}
-				else
-				{
-					const TArray<FTypedElementHandle> DuplicatedElements = CommonActions->DuplicateSelectedElements(SelectionSet, InWorld, GEditor->GetGridLocationOffset(/*bUniformOffset*/false));
-					if (DuplicatedElements.Num() > 0)
+					else
 					{
-						SelectionSet->SetSelection(DuplicatedElements, FTypedElementSelectionOptions());
-						SelectionSet->NotifyPendingChanges();
+						const TArray<FTypedElementHandle> DuplicatedElements = CommonActions->DuplicateSelectedElements(SelectionSet, InWorld, GEditor->GetGridLocationOffset(/*bUniformOffset*/false));
+						if (DuplicatedElements.Num() > 0)
+						{
+							SelectionSet->SetSelection(DuplicatedElements, FTypedElementSelectionOptions());
+							SelectionSet->NotifyPendingChanges();
+						}
 					}
 				}
 			}
