Added support for root motion mirroring in animation preview and when using the mirror node

diff --git a/Engine/Source/Editor/UnrealEd/Private/Animation/DebugSkelMeshComponent.cpp b/Engine/Source/Editor/UnrealEd/Private/Animation/DebugSkelMeshComponent.cpp
--- a/Engine/Source/Editor/UnrealEd/Private/Animation/DebugSkelMeshComponent.cpp
+++ b/Engine/Source/Editor/UnrealEd/Private/Animation/DebugSkelMeshComponent.cpp
@@ -3,6 +3,7 @@
 
 #include "Animation/DebugSkelMeshComponent.h"
 #include "Animation/AnimSequence.h"
+#include "Animation/MirrorDataTable.h"
 #include "BonePose.h"
 #include "Materials/Material.h"
 #include "Animation/AnimMontage.h"
@@ -179,7 +180,7 @@ void UDebugSkelMeshComponent::ConsumeRootMotion(const FVector& FloorMin, const F
 	}
 
 	// Helper to extract root motion manually from the current asset
-	auto ExtractRootMotionFromAnimationAsset = [](const UAnimationAsset* Animation, float StartPosition, float EndPosition) -> FTransform
+	auto ExtractRootMotionFromAnimationAsset = [](const UAnimationAsset* Animation, const UMirrorDataTable* MirrorDataTable, float StartPosition, float EndPosition) -> FTransform
 	{
 		if (const UAnimMontage* Anim = Cast<UAnimMontage>(Animation))
 		{
@@ -195,7 +196,19 @@ void UDebugSkelMeshComponent::ConsumeRootMotion(const FVector& FloorMin, const F
 		
 		if (const UAnimSequence* Anim = Cast<UAnimSequence>(Animation))
 		{
-			return Anim->ExtractRootMotionFromRange(StartPosition, EndPosition);
+			FTransform RootMotion = Anim->ExtractRootMotionFromRange(StartPosition, EndPosition);
+			if (MirrorDataTable)
+			{
+				FVector T = RootMotion.GetTranslation();
+				T = FAnimationRuntime::MirrorVector(T, MirrorDataTable->MirrorAxis);
+
+				FQuat Q = RootMotion.GetRotation();
+				Q = FAnimationRuntime::MirrorQuat(Q, MirrorDataTable->MirrorAxis);
+
+				FVector S = RootMotion.GetScale3D();
+				return FTransform(Q, T, S);
+			}
+			return RootMotion;
 		}
 
 		return FTransform::Identity;
@@ -209,7 +222,19 @@ void UDebugSkelMeshComponent::ConsumeRootMotion(const FVector& FloorMin, const F
 
 	//Extract root motion regardless of where we use it so that we don't hit problems with it building up in the instance
 	FRootMotionMovementParams ExtractedRootMotion = ConsumeRootMotion_Internal(1.0f);
-	
+	if (PreviewInstance->GetMirrorDataTable())
+	{
+		FTransform RootMotion = ExtractedRootMotion.GetRootMotionTransform();
+		FVector T = RootMotion.GetTranslation();
+		T = FAnimationRuntime::MirrorVector(T, PreviewInstance->GetMirrorDataTable()->MirrorAxis);
+
+		FQuat Q = RootMotion.GetRotation();
+		Q = FAnimationRuntime::MirrorQuat(Q, PreviewInstance->GetMirrorDataTable()->MirrorAxis);
+
+		FVector S = RootMotion.GetScale3D();
+		ExtractedRootMotion.Set(FTransform(Q, T, S));
+	}
+
 	const float CurrentTime = PreviewInstance->GetCurrentTime();
 	const float PreviousTime = ConsumeRootMotionPreviousPlaybackTime;
 	ConsumeRootMotionPreviousPlaybackTime = CurrentTime;
@@ -266,7 +291,7 @@ void UDebugSkelMeshComponent::ConsumeRootMotion(const FVector& FloorMin, const F
 						{
 							float StartTime, EndTime;
 							Montage->GetSectionStartAndEndTime(LastSectionIdx, StartTime, EndTime);
-							const FTransform RootMotionDelta = ExtractRootMotionFromAnimationAsset(PreviewInstance->CurrentAsset, EndTime, CurrentTime);
+							const FTransform RootMotionDelta = ExtractRootMotionFromAnimationAsset(PreviewInstance->CurrentAsset, PreviewInstance->GetMirrorDataTable(), EndTime, CurrentTime);
 							SetRelativeTransform(RootMotionDelta);
 						}
 						// Otherwise, reset the position of the mesh back to the origin and apply root motion from the beginning of the FirstSection to CurrentTime
@@ -274,7 +299,7 @@ void UDebugSkelMeshComponent::ConsumeRootMotion(const FVector& FloorMin, const F
 						{
 							float StartTime, EndTime;
 							Montage->GetSectionStartAndEndTime(FirstSectionIdx, StartTime, EndTime);
-							const FTransform RootMotionDelta = ExtractRootMotionFromAnimationAsset(PreviewInstance->CurrentAsset, StartTime, CurrentTime);
+							const FTransform RootMotionDelta = ExtractRootMotionFromAnimationAsset(PreviewInstance->CurrentAsset, PreviewInstance->GetMirrorDataTable(), StartTime, CurrentTime);
 							SetRelativeTransform(RootMotionDelta);
 						}
 
@@ -292,12 +317,12 @@ void UDebugSkelMeshComponent::ConsumeRootMotion(const FVector& FloorMin, const F
 					{
 						if (PreviewInstance->IsReverse())
 						{
-							const FTransform RootMotionDelta = ExtractRootMotionFromAnimationAsset(PreviewInstance->CurrentAsset, PreviewInstance->CurrentAsset->GetPlayLength(), CurrentTime);
+							const FTransform RootMotionDelta = ExtractRootMotionFromAnimationAsset(PreviewInstance->CurrentAsset, PreviewInstance->GetMirrorDataTable(),  PreviewInstance->CurrentAsset->GetPlayLength(), CurrentTime);
 							SetRelativeTransform(RootMotionDelta);
 						}
 						else
 						{
-							const FTransform RootMotionDelta = ExtractRootMotionFromAnimationAsset(PreviewInstance->CurrentAsset, 0.f, CurrentTime);
+							const FTransform RootMotionDelta = ExtractRootMotionFromAnimationAsset(PreviewInstance->CurrentAsset, PreviewInstance->GetMirrorDataTable(), 0.f, CurrentTime);
 							SetRelativeTransform(RootMotionDelta);
 						}
 					}
@@ -310,7 +335,7 @@ void UDebugSkelMeshComponent::ConsumeRootMotion(const FVector& FloorMin, const F
 		}
 		else // Not Playing. When not playing user can still scrub the time line but animation is not ticking so we have to extract and apply root motion manually
 		{
-			const FTransform RootMotion = ExtractRootMotionFromAnimationAsset(PreviewInstance->CurrentAsset, PreviousTime, CurrentTime);
+			const FTransform RootMotion = ExtractRootMotionFromAnimationAsset(PreviewInstance->CurrentAsset, PreviewInstance->GetMirrorDataTable(), PreviousTime, CurrentTime);
 			AddLocalTransform(RootMotion);
 		}
 	}
diff --git a/Engine/Source/Runtime/Engine/Private/Animation/AnimSync.cpp b/Engine/Source/Runtime/Engine/Private/Animation/AnimSync.cpp
--- a/Engine/Source/Runtime/Engine/Private/Animation/AnimSync.cpp
+++ b/Engine/Source/Runtime/Engine/Private/Animation/AnimSync.cpp
@@ -4,6 +4,7 @@
 #include "Animation/AnimNodeBase.h"
 #include "Animation/AnimInstanceProxy.h"
 #include "AnimationUtils.h"
+#include "Animation/MirrorDataTable.h"
 
 namespace UE { namespace Anim {
 
@@ -49,6 +50,27 @@ void FAnimSync::TickAssetPlayerInstances(FAnimInstanceProxy& InProxy, float InDe
 
 	SCOPE_CYCLE_COUNTER(STAT_TickAssetPlayerInstances);
 
+	auto AccumulateRootMotion = [&InProxy](FAnimTickRecord& TickRecord, FAnimAssetTickContext& TickContext)
+	{
+		if (TickRecord.MirrorDataTable)
+		{
+			FTransform InTransform = TickContext.RootMotionMovementParams.GetRootMotionTransform();
+			FVector T = InTransform.GetTranslation();
+			T = FAnimationRuntime::MirrorVector(T, TickRecord.MirrorDataTable->MirrorAxis);
+
+			FQuat Q = InTransform.GetRotation();
+			Q = FAnimationRuntime::MirrorQuat(Q, TickRecord.MirrorDataTable->MirrorAxis);
+
+			FVector S = InTransform.GetScale3D();
+			FTransform MirroredTransform = FTransform(Q, T, S);
+			InProxy.ExtractedRootMotion.AccumulateWithBlend(MirroredTransform, TickRecord.GetRootMotionWeight());
+		}
+		else
+		{
+			InProxy.ExtractedRootMotion.AccumulateWithBlend(TickContext.RootMotionMovementParams, TickRecord.GetRootMotionWeight());
+		}
+	};
+
 	const ERootMotionMode::Type RootMotionMode = InProxy.GetRootMotionMode();
 
 	// Handle all players inside sync groups
@@ -105,7 +127,7 @@ void FAnimSync::TickAssetPlayerInstances(FAnimInstanceProxy& InProxy, float InDe
 
 				if (RootMotionMode == ERootMotionMode::RootMotionFromEverything && TickContext.RootMotionMovementParams.bHasRootMotion)
 				{
-					InProxy.ExtractedRootMotion.AccumulateWithBlend(TickContext.RootMotionMovementParams, GroupLeader.GetRootMotionWeight());
+					AccumulateRootMotion(GroupLeader, TickContext);
 				}
 
 				// if we're not using marker based sync, we don't care, get out
@@ -198,7 +220,7 @@ void FAnimSync::TickAssetPlayerInstances(FAnimInstanceProxy& InProxy, float InDe
 					}
 					if (RootMotionMode == ERootMotionMode::RootMotionFromEverything && TickContext.RootMotionMovementParams.bHasRootMotion)
 					{
-						InProxy.ExtractedRootMotion.AccumulateWithBlend(TickContext.RootMotionMovementParams, AssetPlayer.GetRootMotionWeight());
+						AccumulateRootMotion(AssetPlayer, TickContext);
 					}
 				}
 			}
@@ -228,7 +250,7 @@ void FAnimSync::TickAssetPlayerInstances(FAnimInstanceProxy& InProxy, float InDe
 		}
 		if (RootMotionMode == ERootMotionMode::RootMotionFromEverything && TickContext.RootMotionMovementParams.bHasRootMotion)
 		{
-			InProxy.ExtractedRootMotion.AccumulateWithBlend(TickContext.RootMotionMovementParams, AssetPlayerToTick.GetRootMotionWeight());
+			AccumulateRootMotion(AssetPlayerToTick, TickContext);
 		}
 
 #if ANIM_TRACE_ENABLED
