BulkDataRegistry: Fix EditorBuildInputResolver to set Status to Error in the cases where a registered BulkData is found but it cannot find its payload.
Remove bValid from FMetaData and FData return structs; it is redundant with the RawHash/Buffer data already present in those structs.
#rnx

diff --git a/Engine/Source/Editor/UnrealEd/Private/EditorDomain/BulkDataRegistryEditorDomain.cpp b/Engine/Source/Editor/UnrealEd/Private/EditorDomain/BulkDataRegistryEditorDomain.cpp
--- a/Engine/Source/Editor/UnrealEd/Private/EditorDomain/BulkDataRegistryEditorDomain.cpp
+++ b/Engine/Source/Editor/UnrealEd/Private/EditorDomain/BulkDataRegistryEditorDomain.cpp
@@ -314,7 +314,7 @@ TFuture<UE::BulkDataRegistry::FMetaData> FBulkDataRegistryEditorDomain::GetMeta(
 		if (!Existing || !Existing->bPayloadAvailable)
 		{
 			TPromise<UE::BulkDataRegistry::FMetaData> Promise;
-			Promise.SetValue(UE::BulkDataRegistry::FMetaData{ false, FIoHash(), 0 });
+			Promise.SetValue(UE::BulkDataRegistry::FMetaData{ FIoHash(), 0 });
 			return Promise.GetFuture();
 		}
 
@@ -322,7 +322,7 @@ TFuture<UE::BulkDataRegistry::FMetaData> FBulkDataRegistryEditorDomain::GetMeta(
 		if (!BulkData.HasPlaceholderPayloadId())
 		{
 			TPromise<UE::BulkDataRegistry::FMetaData> Promise;
-			Promise.SetValue(UE::BulkDataRegistry::FMetaData{ true, BulkData.GetPayloadId(), static_cast<uint64>(BulkData.GetPayloadSize()) });
+			Promise.SetValue(UE::BulkDataRegistry::FMetaData{ BulkData.GetPayloadId(), static_cast<uint64>(BulkData.GetPayloadSize()) });
 			return Promise.GetFuture();
 		}
 
@@ -344,7 +344,7 @@ TFuture<UE::BulkDataRegistry::FMetaData> FBulkDataRegistryEditorDomain::GetMeta(
 		TFuture<UE::BulkDataRegistry::FMetaData> Future = Promise.GetFuture();
 		UpdatingPayload.Requesters.Add([Promise = MoveTemp(Promise)](bool bValid, const FCompressedBuffer& Buffer) mutable
 			{
-				Promise.SetValue(UE::BulkDataRegistry::FMetaData{ bValid, Buffer.GetRawHash(), Buffer.GetRawSize() });
+				Promise.SetValue(UE::BulkDataRegistry::FMetaData{ Buffer.GetRawHash(), Buffer.GetRawSize() });
 			});
 		return Future;
 	}
@@ -366,7 +366,7 @@ TFuture<UE::BulkDataRegistry::FData> FBulkDataRegistryEditorDomain::GetData(cons
 			if (!Existing || !Existing->bPayloadAvailable)
 			{
 				TPromise<UE::BulkDataRegistry::FData> Result;
-				Result.SetValue(UE::BulkDataRegistry::FData{ false, FCompressedBuffer() });
+				Result.SetValue(UE::BulkDataRegistry::FData{ FCompressedBuffer() });
 				return Result.GetFuture();
 			}
 
@@ -410,7 +410,7 @@ TFuture<UE::BulkDataRegistry::FData> FBulkDataRegistryEditorDomain::GetData(cons
 			TFuture<UE::BulkDataRegistry::FData> Future = Promise.GetFuture();
 			UpdatingPayload.Requesters.Add([Promise=MoveTemp(Promise)](bool bValid, const FCompressedBuffer& Buffer) mutable
 				{
-					Promise.SetValue(UE::BulkDataRegistry::FData{ bValid, Buffer });
+					Promise.SetValue(UE::BulkDataRegistry::FData{ Buffer });
 				});
 			return Future;
 		}
@@ -421,7 +421,7 @@ TFuture<UE::BulkDataRegistry::FData> FBulkDataRegistryEditorDomain::GetData(cons
 	// does not read from the BulkData after returning from GetCompressedPayload, so a read-after-free is not possible.
 	return CopyBulk->GetCompressedPayload().Next([](FCompressedBuffer Payload)
 		{
-			return UE::BulkDataRegistry::FData{ true, Payload };
+			return UE::BulkDataRegistry::FData{ Payload };
 		});
 }
 
diff --git a/Engine/Source/Editor/UnrealEd/Private/EditorDomain/EditorBuildInputResolver.cpp b/Engine/Source/Editor/UnrealEd/Private/EditorDomain/EditorBuildInputResolver.cpp
--- a/Engine/Source/Editor/UnrealEd/Private/EditorDomain/EditorBuildInputResolver.cpp
+++ b/Engine/Source/Editor/UnrealEd/Private/EditorDomain/EditorBuildInputResolver.cpp
@@ -14,6 +14,9 @@
 #include "Serialization/BulkDataRegistry.h"
 #include "Templates/Tuple.h"
 
+// MRPTODO
+PRAGMA_DISABLE_OPTIMIZATION;
+
 namespace UE::DerivedData
 {
 
@@ -47,7 +50,7 @@ private:
 	 * Move BulkData's fields from the BulkDataRegistry's format into the BuildInputResolver format.
 	 * Implementation varies depending on whether the template type is for MetaData or Data.
 	*/
-	void AssignBulkPayload(CallbackPayloadType& OutPayload, BulkPayloadType& InPayload);
+	void AssignBulkPayload(CallbackPayloadType& OutPayload, bool& bOutPayloadIsValid, BulkPayloadType& InPayload);
 	/**
 	 * Log that a BulkData has failed to resolve.
 	 * The text of the message varies depending on whether the template type is for MetaData or Data.
@@ -272,11 +275,9 @@ void TResolveInputRequest<BulkPayloadType, CallbackPayloadType, CallbackType>::S
 		PayloadData.Future = InPair.template Get<TFuture<BulkPayloadType>>().Next(
 			[pThis = TRefCountPtr<ThisType>(this), &Payload](BulkPayloadType&& Result)
 			{
-				if (Result.bValid)
-				{
-					pThis->AssignBulkPayload(Payload, Result);
-				}
-				else
+				bool bValid;
+				pThis->AssignBulkPayload(Payload, bValid, Result);
+				if (!bValid)
 				{
 					pThis->LogBulkDataError(Payload.Key);
 					pThis->bAllSucceeded.store(false, std::memory_order_relaxed);
@@ -332,16 +333,18 @@ void TResolveInputRequest<BulkPayloadType, CallbackPayloadType, CallbackType>::C
 }
 
 template<>
-void FResolveInputMetaRequest::AssignBulkPayload(FBuildInputMetaByKey& OutPayload, UE::BulkDataRegistry::FMetaData& InPayload)
+void FResolveInputMetaRequest::AssignBulkPayload(FBuildInputMetaByKey& OutPayload, bool& bOutPayloadIsValid, UE::BulkDataRegistry::FMetaData& InPayload)
 {
 	OutPayload.RawHash = InPayload.RawHash;
 	OutPayload.RawSize = InPayload.RawSize;
+	bOutPayloadIsValid = !OutPayload.RawHash.IsZero();
 }
 
 template<>
-void FResolveInputDataRequest::AssignBulkPayload(FBuildInputDataByKey& OutPayload, UE::BulkDataRegistry::FData& InPayload)
+void FResolveInputDataRequest::AssignBulkPayload(FBuildInputDataByKey& OutPayload, bool& bOutPayloadIsValid, UE::BulkDataRegistry::FData& InPayload)
 {
 	OutPayload.Data = MoveTemp(InPayload.Buffer);
+	bOutPayloadIsValid = !OutPayload.Data.IsNull();
 }
 
 template<>
diff --git a/Engine/Source/Runtime/CoreUObject/Private/Serialization/BulkDataRegistry.cpp b/Engine/Source/Runtime/CoreUObject/Private/Serialization/BulkDataRegistry.cpp
--- a/Engine/Source/Runtime/CoreUObject/Private/Serialization/BulkDataRegistry.cpp
+++ b/Engine/Source/Runtime/CoreUObject/Private/Serialization/BulkDataRegistry.cpp
@@ -56,13 +56,13 @@ public:
 	virtual TFuture<UE::BulkDataRegistry::FMetaData> GetMeta(const FGuid& BulkDataId) override
 	{
 		TPromise<UE::BulkDataRegistry::FMetaData> Promise;
-		Promise.SetValue(UE::BulkDataRegistry::FMetaData{ false, FIoHash(), 0 });
+		Promise.SetValue(UE::BulkDataRegistry::FMetaData{ FIoHash(), 0 });
 		return Promise.GetFuture();
 	}
 	virtual TFuture<UE::BulkDataRegistry::FData> GetData(const FGuid& BulkDataId) override
 	{
 		TPromise<UE::BulkDataRegistry::FData> Promise;
-		Promise.SetValue(UE::BulkDataRegistry::FData{ false, FCompressedBuffer() });
+		Promise.SetValue(UE::BulkDataRegistry::FData{ FCompressedBuffer() });
 		return Promise.GetFuture();
 	}
 	virtual bool TryGetBulkData(const FGuid & BulkDataId, UE::Serialization::FEditorBulkData* OutBulk = nullptr,
diff --git a/Engine/Source/Runtime/CoreUObject/Public/Serialization/BulkDataRegistry.h b/Engine/Source/Runtime/CoreUObject/Public/Serialization/BulkDataRegistry.h
--- a/Engine/Source/Runtime/CoreUObject/Public/Serialization/BulkDataRegistry.h
+++ b/Engine/Source/Runtime/CoreUObject/Public/Serialization/BulkDataRegistry.h
@@ -47,20 +47,22 @@ namespace UE::BulkDataRegistry
 	/** Results of GetMeta call. */
 	struct FMetaData
 	{
-		/** True if data was found, else false. */
-		bool bValid;
-		/** IoHash of the uncompressed bytes of the data that will be returned from GetData. */
+		/**
+		 * IoHash of the uncompressed bytes of the data that will be returned from GetData.
+		 * FIoHash::Zero if and only if data was not found.
+		 */
 		FIoHash RawHash;
-		/** Size of the uncompressed bytes of the data that will be returned from GetData. */
+		/**
+		 * Size of the uncompressed bytes of the data that will be returned from GetData.
+		 * 0 if data was not found, but can be 0 for valid data as well.
+		 */
 		uint64 RawSize;
 	};
 
 	/** Results of GetData call. */
 	struct FData
 	{
-		/** True if data was found, else false. */
-		bool bValid;
-		/** The discovered data. Empty if data was not found. */
+		/** The discovered data. Null buffer and only if data was not found. */
 		FCompressedBuffer Buffer;
 	};
 }
