Remove accidentally-submitted debug code.
#rnx

diff --git a/Engine/Source/Editor/UnrealEd/Private/EditorDomain/EditorBuildInputResolver.cpp b/Engine/Source/Editor/UnrealEd/Private/EditorDomain/EditorBuildInputResolver.cpp
--- a/Engine/Source/Editor/UnrealEd/Private/EditorDomain/EditorBuildInputResolver.cpp
+++ b/Engine/Source/Editor/UnrealEd/Private/EditorDomain/EditorBuildInputResolver.cpp
@@ -14,9 +14,6 @@
 #include "Serialization/BulkDataRegistry.h"
 #include "Templates/Tuple.h"
 
-// MRPTODO
-PRAGMA_DISABLE_OPTIMIZATION;
-
 namespace UE::DerivedData
 {
 
