Fixed issue where CreateNewStreamingLevel scripting function would ignore NewLevelPath argument

diff --git a/Engine/Source/Editor/UnrealEd/Private/EditorLevelUtils.cpp b/Engine/Source/Editor/UnrealEd/Private/EditorLevelUtils.cpp
--- a/Engine/Source/Editor/UnrealEd/Private/EditorLevelUtils.cpp
+++ b/Engine/Source/Editor/UnrealEd/Private/EditorLevelUtils.cpp
@@ -647,7 +647,8 @@ ULevelStreaming* UEditorLevelUtils::CreateNewStreamingLevel(TSubclassOf<ULevelSt
 	{
 		if (ensureAsRuntimeWarning(LevelStreamingClass.Get() != nullptr))
 		{
-			return CreateNewStreamingLevelForWorld(*GEditor->GetEditorWorldContext().World(), LevelStreamingClass, Filename, bMoveSelectedActorsIntoNewLevel);
+			bool bUseSaveAs = PackagePath.IsEmpty();
+			return CreateNewStreamingLevelForWorld(*GEditor->GetEditorWorldContext().World(), LevelStreamingClass, Filename, bMoveSelectedActorsIntoNewLevel, nullptr, bUseSaveAs);
 		}
 	}
 
