CookOnTheFlyServer - minor changes to reduce code paths in StartCookByTheBook.
#rnx

diff --git a/Engine/Source/Editor/UnrealEd/Classes/CookOnTheSide/CookOnTheFlyServer.h b/Engine/Source/Editor/UnrealEd/Classes/CookOnTheSide/CookOnTheFlyServer.h
--- a/Engine/Source/Editor/UnrealEd/Classes/CookOnTheSide/CookOnTheFlyServer.h
+++ b/Engine/Source/Editor/UnrealEd/Classes/CookOnTheSide/CookOnTheFlyServer.h
@@ -737,6 +737,9 @@ private:
 	void AddFileToCook( TArray<FName>& InOutFilesToCook, TMap<FName, UE::Cook::FInstigator>& InOutInstigators,
 		const FString &InFilename, const UE::Cook::FInstigator& Instigator) const;
 
+	/** Return the name to use for the project's global shader library */
+	FString GetProjectShaderLibraryName() const;
+
 	/**
 	* Invokes the necessary FShaderCodeLibrary functions to start cooking the shader code library.
 	*/
@@ -772,6 +775,8 @@ private:
 	*/
 	void CleanShaderCodeLibraries();
 
+	void RegisterShaderChunkDataGenerator();
+
 	/**
 	* Call back from the TickCookOnTheSide when a cook by the book finishes (when started form StartCookByTheBook)
 	*/
@@ -811,6 +816,8 @@ private:
 	*/
 	void WriteMapDependencyGraph(const ITargetPlatform* TargetPlatform);
 
+	void RegisterLocalizationChunkDataGenerator();
+
 	//////////////////////////////////////////////////////////////////////////
 	// cook on the fly specific functions
 
@@ -1147,7 +1154,8 @@ private:
 	bool bPreexploreDependenciesEnabled = true;
 	/** Test mode for the debug of hybrid iterative dependencies. */
 	bool bHybridIterativeDebug = false;
-	bool bHasBlockedOnAssetRegistry = false;
+	bool bFirstCookInThisProcessInitialized = false;
+	bool bFirstCookInThisProcess = true;
 	bool bHasDeferredInitializeCookOnTheFly = false;
 	bool bImportBehaviorCallbackInstalled = false;
 
diff --git a/Engine/Source/Editor/UnrealEd/Private/CookOnTheFlyServer.cpp b/Engine/Source/Editor/UnrealEd/Private/CookOnTheFlyServer.cpp
--- a/Engine/Source/Editor/UnrealEd/Private/CookOnTheFlyServer.cpp
+++ b/Engine/Source/Editor/UnrealEd/Private/CookOnTheFlyServer.cpp
@@ -490,7 +490,7 @@ void UCookOnTheFlyServer::Tick(float DeltaTime)
 
 	check(IsCookingInEditor());
 
-	if (IsCookByTheBookMode() && !IsCookByTheBookRunning() && !GIsSlowTask)
+	if (IsCookByTheBookMode() && !IsCookByTheBookRunning() && !GIsSlowTask && IsCookFlagSet(ECookInitializationFlags::BuildDDCInBackground))
 	{
 		// if we are in the editor then precache some stuff ;)
 		TArray<const ITargetPlatform*> CacheTargetPlatforms;
@@ -502,11 +502,6 @@ void UCookOnTheFlyServer::Tick(float DeltaTime)
 		}
 		if (CacheTargetPlatforms.Num() > 0)
 		{
-			// early out all the stuff we don't care about 
-			if (!IsCookFlagSet(ECookInitializationFlags::BuildDDCInBackground))
-			{
-				return;
-			}
 			TickPrecacheObjectsForPlatforms(0.001, CacheTargetPlatforms);
 		}
 	}
@@ -3486,6 +3481,8 @@ void UCookOnTheFlyServer::TickPrecacheObjectsForPlatforms(const float TimeSlice,
 	if (LastUpdateTick > 50 ||
 		((CachedMaterialsToCacheArray.Num() == 0) && (CachedTexturesToCacheArray.Num() == 0)))
 	{
+		CachedMaterialsToCacheArray.Reset();
+		CachedTexturesToCacheArray.Reset();
 		LastUpdateTick = 0;
 		TArray<UObject*> Materials;
 		GetObjectsOfClass(UMaterial::StaticClass(), Materials, true);
@@ -4612,6 +4609,7 @@ void UCookOnTheFlyServer::Initialize( ECookMode::Type DesiredCookMode, ECookInit
 	ExternalRequests = MakeUnique<UE::Cook::FExternalRequests>();
 	PackageTracker = MakeUnique<UE::Cook::FPackageTracker>(*PackageDatas.Get());
 	DiffModeHelper = MakeUnique<FDiffModeCookServerUtils>();
+	AssetRegistry = IAssetRegistry::Get();
 
 	BuildDefinitions.Reset(new UE::Cook::FBuildDefinitions());
 	UE::Cook::InitializeTls();
@@ -5915,10 +5913,6 @@ void UCookOnTheFlyServer::GenerateAssetRegistry()
 {
 	TRACE_CPUPROFILER_EVENT_SCOPE(UCookOnTheFlyServer::GenerateAssetRegistry);
 
-	// Cache asset registry for later
-	FAssetRegistryModule& AssetRegistryModule = FModuleManager::LoadModuleChecked<FAssetRegistryModule>(TEXT("AssetRegistry"));
-	AssetRegistry = &AssetRegistryModule.Get();
-
 	// Mark package as dirty for the last ones saved
 	for (FName AssetFilename : ModifiedAssetFilenames)
 	{
@@ -5960,11 +5954,10 @@ void UCookOnTheFlyServer::GenerateAssetRegistry()
 
 void UCookOnTheFlyServer::BlockOnAssetRegistry()
 {
-	if (bHasBlockedOnAssetRegistry)
+	if (!bFirstCookInThisProcess)
 	{
 		return;
 	}
-	bHasBlockedOnAssetRegistry = true;
 	if (ShouldPopulateFullAssetRegistry())
 	{
 		// Trigger or waitfor completion the primary AssetRegistry scan.
@@ -6582,11 +6575,6 @@ FString UCookOnTheFlyServer::GetSandboxDirectory( const FString& PlatformName )
 
 	Result.ReplaceInline(TEXT("[Platform]"), *PlatformName);
 
-	/*if ( IsCookingDLC() )
-	{
-		check( IsCookByTheBookRunning() );
-		Result.ReplaceInline(TEXT("/Cooked/"), *FString::Printf(TEXT("/CookedDLC_%s/"), *CookByTheBookOptions->DlcName) );
-	}*/
 	return Result;
 }
 
@@ -6627,11 +6615,6 @@ FString UCookOnTheFlyServer::ConvertToFullSandboxPath( const FString &FileName,
 		Result = SandboxFile->ConvertToAbsolutePathForExternalAppForRead(*FileName);
 	}
 
-	/*if ( IsCookingDLC() )
-	{
-		check( IsCookByTheBookRunning() );
-		Result.ReplaceInline(TEXT("/Cooked/"), *FString::Printf(TEXT("/CookedDLC_%s/"), *CookByTheBookOptions->DlcName) );
-	}*/
 	return Result;
 }
 
@@ -6711,6 +6694,28 @@ void UCookOnTheFlyServer::InitShaderCodeLibrary(void)
     }
 }
 
+
+void UCookOnTheFlyServer::RegisterShaderChunkDataGenerator()
+{
+	// add shader library and PSO cache chunkers
+	const UProjectPackagingSettings* const PackagingSettings = GetDefault<UProjectPackagingSettings>();
+	FString LibraryName = GetProjectShaderLibraryName();
+	if (PackagingSettings->bShareMaterialShaderCode && IsUsingShaderCodeLibrary())
+	{
+		for (const ITargetPlatform* TargetPlatform : PlatformManager->GetSessionPlatforms())
+		{
+			FAssetRegistryGenerator& RegistryGenerator = *(PlatformManager->GetPlatformData(TargetPlatform)->RegistryGenerator);
+			RegistryGenerator.RegisterChunkDataGenerator(MakeShared<FShaderLibraryChunkDataGenerator>(TargetPlatform));
+			RegistryGenerator.RegisterChunkDataGenerator(MakeShared<FPipelineCacheChunkDataGenerator>(TargetPlatform, LibraryName));
+		}
+	}
+}
+
+FString UCookOnTheFlyServer::GetProjectShaderLibraryName() const
+{
+	return !IsCookingDLC() ? FApp::GetProjectName() : CookByTheBookOptions->DlcName;
+}
+
 static FString GenerateShaderCodeLibraryName(FString const& Name, bool bIsIterateSharedBuild)
 {
 	FString ActualName = (!bIsIterateSharedBuild) ? Name : Name + TEXT("_SC");
@@ -6988,8 +6993,7 @@ void UCookOnTheFlyServer::CookByTheBookFinished()
 
 	const UProjectPackagingSettings* const PackagingSettings = GetDefault<UProjectPackagingSettings>();
 	bool const bCacheShaderLibraries = IsUsingShaderCodeLibrary();
-	// note: must agree with the name used in StartCookByTheBook
-	FString LibraryName = !IsCookingDLC() ? FApp::GetProjectName() : CookByTheBookOptions->DlcName;
+	FString LibraryName = GetProjectShaderLibraryName();
 
 	{
 		// Save modified asset registry with all streaming chunk info generated during cook
@@ -7010,6 +7014,9 @@ void UCookOnTheFlyServer::CookByTheBookFinished()
 			UE_SCOPED_HIERARCHICAL_COOKTIMER(SavingAssetRegistry);
 			SCOPED_BOOT_TIMING("SavingAssetRegistry");
 
+			RegisterLocalizationChunkDataGenerator();
+			RegisterShaderChunkDataGenerator();
+
 			// if we are cooking DLC, the DevelopmentAR isn't needed - it's used when making DLC against shipping, so there's no need to make it
 			// again, as we don't make DLC against DLC (but allow an override just in case)
 			bool bSaveDevelopmentAssetRegistry = !FParse::Param(FCommandLine::Get(), TEXT("NoSaveDevAR"));
@@ -7506,6 +7513,18 @@ void UCookOnTheFlyServer::SetBeginCookConfigSettings()
 	bHybridIterativeEnabled &= !IsCookingDLC();
 	// HybridIterative uses TargetDomain storage of dependencies which is only implemented in ZenStore
 	bHybridIterativeEnabled &= IsUsingZenStore();
+	if (!bFirstCookInThisProcessInitialized)
+	{
+		// This is the first cook; set bFirstCookInThisProcess=true for the entire cook until SetBeginCookConfigSettings is called to mark the second cook
+		bFirstCookInThisProcessInitialized = true;
+		bFirstCookInThisProcess = true;
+	}
+	else
+	{
+		// We have cooked before; set bFirstCookInThisProcess=false
+		bFirstCookInThisProcess = false;
+	}
+
 }
 
 void UCookOnTheFlyServer::BeginCookSandbox(TConstArrayView<const ITargetPlatform*> TargetPlatforms)
@@ -7920,28 +7939,7 @@ void UCookOnTheFlyServer::StartCookByTheBook( const FCookByTheBookStartupOptions
 			LocalizedPackageNames.AddUnique(LocalizedPackageName);
 		}
 
-		// Get the list of localization targets to chunk, and remove any targets that we've been asked not to stage
-		TArray<FString> LocalizationTargetsToChunk = PackagingSettings->LocalizationTargetsToChunk;
-		{
-			TArray<FString> BlacklistLocalizationTargets;
-			GConfig->GetArray(TEXT("Staging"), TEXT("BlacklistLocalizationTargets"), BlacklistLocalizationTargets, GGameIni);
-			if (BlacklistLocalizationTargets.Num() > 0)
-			{
-				LocalizationTargetsToChunk.RemoveAll([&BlacklistLocalizationTargets](const FString& InLocalizationTarget)
-				{
-					return BlacklistLocalizationTargets.Contains(InLocalizationTarget);
-				});
-			}
-		}
-
-		if (LocalizationTargetsToChunk.Num() > 0 && AllCulturesToCook.Num() > 0)
-		{
-			for (const ITargetPlatform* TargetPlatform : TargetPlatforms)
-			{
-				FAssetRegistryGenerator& RegistryGenerator = *(PlatformManager->GetPlatformData(TargetPlatform)->RegistryGenerator);
-				RegistryGenerator.RegisterChunkDataGenerator(MakeShared<FLocalizationChunkDataGenerator>(RegistryGenerator.GetPakchunkIndex(PackagingSettings->LocalizationTargetCatchAllChunkId), LocalizationTargetsToChunk, AllCulturesToCook));
-			}
-		}
+		CookByTheBookOptions->AllCulturesToCook = MoveTemp(AllCulturesToCook);
 	}
 
 	PackageTracker->NeverCookPackageList.Empty();
@@ -8123,19 +8121,6 @@ void UCookOnTheFlyServer::StartCookByTheBook( const FCookByTheBookStartupOptions
 		}
 	}
 	
-	// add shader library and PSO cache chunkers
-	// note: library name must agree with the name used in CookByTheBookFinished
-	FString LibraryName = !IsCookingDLC() ? FApp::GetProjectName() : CookByTheBookOptions->DlcName;
-	if (PackagingSettings->bShareMaterialShaderCode && IsUsingShaderCodeLibrary())
-	{
-		for (const ITargetPlatform* TargetPlatform : TargetPlatforms)
-		{
-			FAssetRegistryGenerator& RegistryGenerator = *(PlatformManager->GetPlatformData(TargetPlatform)->RegistryGenerator);
-			RegistryGenerator.RegisterChunkDataGenerator(MakeShared<FShaderLibraryChunkDataGenerator>(TargetPlatform));
-			RegistryGenerator.RegisterChunkDataGenerator(MakeShared<FPipelineCacheChunkDataGenerator>(TargetPlatform, LibraryName));
-		}
-	}
-
 	TSet<FName> StartupSoftObjectPackages;
 	if (!IsCookByTheBookMode() || !CookByTheBookOptions->bSkipSoftReferences)
 	{
@@ -8278,6 +8263,7 @@ void UCookOnTheFlyServer::StartCookByTheBook( const FCookByTheBookStartupOptions
 
 	// Open the shader code library for the current project or the current DLC pack, depending on which we are cooking
 	{
+		FString LibraryName = GetProjectShaderLibraryName();
 		if (LibraryName.Len() > 0)
 		{
 			OpenShaderLibrary(LibraryName);
@@ -9294,5 +9280,39 @@ void UCookOnTheFlyServer::PropertyImportBehaviorCallback(const FObjectImport& Im
 	}
 }
 
+void UCookOnTheFlyServer::RegisterLocalizationChunkDataGenerator()
+{
+	if (!CookByTheBookOptions)
+	{
+		return;
+	}
+
+	// Get the list of localization targets to chunk, and remove any targets that we've been asked not to stage
+	const UProjectPackagingSettings* const PackagingSettings = GetDefault<UProjectPackagingSettings>();
+	TArray<FString> LocalizationTargetsToChunk = PackagingSettings->LocalizationTargetsToChunk;
+	{
+		TArray<FString> BlocklistLocalizationTargets;
+		GConfig->GetArray(TEXT("Staging"), TEXT("BlacklistLocalizationTargets"), BlocklistLocalizationTargets, GGameIni);
+		if (BlocklistLocalizationTargets.Num() > 0)
+		{
+			LocalizationTargetsToChunk.RemoveAll([&BlocklistLocalizationTargets](const FString& InLocalizationTarget)
+				{
+					return BlocklistLocalizationTargets.Contains(InLocalizationTarget);
+				});
+		}
+	}
+
+	if (LocalizationTargetsToChunk.Num() > 0 && CookByTheBookOptions->AllCulturesToCook.Num() > 0)
+	{
+		for (const ITargetPlatform* TargetPlatform : PlatformManager->GetSessionPlatforms())
+		{
+			FAssetRegistryGenerator& RegistryGenerator = *(PlatformManager->GetPlatformData(TargetPlatform)->RegistryGenerator);
+			TSharedRef<FLocalizationChunkDataGenerator> LocalizationGenerator =
+				MakeShared<FLocalizationChunkDataGenerator>(RegistryGenerator.GetPakchunkIndex(PackagingSettings->LocalizationTargetCatchAllChunkId),
+					LocalizationTargetsToChunk, CookByTheBookOptions->AllCulturesToCook);
+			RegistryGenerator.RegisterChunkDataGenerator(MoveTemp(LocalizationGenerator));
+		}
+	}
+}
 
 #undef LOCTEXT_NAMESPACE
diff --git a/Engine/Source/Editor/UnrealEd/Private/Cooker/CookTypes.h b/Engine/Source/Editor/UnrealEd/Private/Cooker/CookTypes.h
--- a/Engine/Source/Editor/UnrealEd/Private/Cooker/CookTypes.h
+++ b/Engine/Source/Editor/UnrealEd/Private/Cooker/CookTypes.h
@@ -276,6 +276,8 @@ public:
 
 	/** Mapping from source packages to their localized variants (based on the culture list in FCookByTheBookStartupOptions) */
 	TMap<FName, TArray<FName>>		SourceToLocalizedPackageVariants;
+	/** List of all the cultures (e.g. "en") that need to be cooked */
+	TArray<FString>					AllCulturesToCook;
 };
 
 void LogCookerMessage(const FString& MessageText, EMessageSeverity::Type Severity);
