EditorPerf: Require statistics mentioned in EditorPerfStats.csv (allowlist) to exist in the data

### Features

Prior to this change, EditorPerf can sliently fail when statistics in the csv aren't found in the data; causing the telemetry data to be missing. This could be because asset or configuration changes modified the test routine. This change causes this missing data to print errors and fail the commandlet. This makes maintenance easier because it prevents likely incorrect data from being pushed to telemetry.

Add -alltelemetry flag for unfiltered telemetry. This flag is for scale testing and local development purproses.

### Testing

Verified to work properly with invalid statistics in the csv, and corrected statistics in the csv
Expected results running with and without -alltelemetry (see SummarizeArgs in the build graph)

### Tags

#rnx

diff --git a/Engine/Source/Editor/UnrealEd/Private/Commandlets/SummarizeTraceCommandlet.cpp b/Engine/Source/Editor/UnrealEd/Private/Commandlets/SummarizeTraceCommandlet.cpp
--- a/Engine/Source/Editor/UnrealEd/Private/Commandlets/SummarizeTraceCommandlet.cpp
+++ b/Engine/Source/Editor/UnrealEd/Private/Commandlets/SummarizeTraceCommandlet.cpp
@@ -1713,6 +1713,7 @@ int32 USummarizeTraceCommandlet::Main(const FString& CmdLineParams)
 		UE_LOG(LogSummarizeTrace, Log, TEXT("Options:"));
 		UE_LOG(LogSummarizeTrace, Log, TEXT(" Required: -inputfile=<utrace path>   (The utrace you wish to process)"));
 		UE_LOG(LogSummarizeTrace, Log, TEXT(" Optional: -testname=<string>         (Test name to use in telemetry csv)"));
+		UE_LOG(LogSummarizeTrace, Log, TEXT(" Optional: -alltelemetry              (Dump all data to telemetry csv)"));
 		return 0;
 	}
 
@@ -1727,22 +1728,27 @@ int32 USummarizeTraceCommandlet::Main(const FString& CmdLineParams)
 		return 1;
 	}
 
+	bool bAllTelemetry = FParse::Param(*CmdLineParams, TEXT("alltelemetry"));
+
 	// load the stats file to know which event name and statistic name to generate in the telemetry csv
 	// the telemetry csv is ingested completely, so this just highlights specific data elements we want to track
 	TMultiMap<FString, StatisticDefinition> NameToDefinitionMap;
 	TSet<FString> CpuScopeNamesWithWildcards;
-	FString GlobalStatisticsFileName = FPaths::RootDir() / TEXT("Engine") / TEXT("Build") / TEXT("EditorPerfStats.csv");
-	if (FPaths::FileExists(GlobalStatisticsFileName))
-	{
-		UE_LOG(LogSummarizeTrace, Display, TEXT("Loading global statistics from %s"), *GlobalStatisticsFileName);
-		bool bCSVOk = StatisticDefinition::LoadFromCSV(GlobalStatisticsFileName, NameToDefinitionMap, CpuScopeNamesWithWildcards);
-		check(bCSVOk);
-	}
-	FString ProjectStatisticsFileName = FPaths::ProjectDir() / TEXT("Build") / TEXT("EditorPerfStats.csv");
-	if (FPaths::FileExists(ProjectStatisticsFileName))
+	if (!bAllTelemetry)
 	{
-		UE_LOG(LogSummarizeTrace, Display, TEXT("Loading project statistics from %s"), *ProjectStatisticsFileName);
-		bool bCSVOk = StatisticDefinition::LoadFromCSV(ProjectStatisticsFileName, NameToDefinitionMap, CpuScopeNamesWithWildcards);
+		FString GlobalStatisticsFileName = FPaths::RootDir() / TEXT("Engine") / TEXT("Build") / TEXT("EditorPerfStats.csv");
+		if (FPaths::FileExists(GlobalStatisticsFileName))
+		{
+			UE_LOG(LogSummarizeTrace, Display, TEXT("Loading global statistics from %s"), *GlobalStatisticsFileName);
+			bool bCSVOk = StatisticDefinition::LoadFromCSV(GlobalStatisticsFileName, NameToDefinitionMap, CpuScopeNamesWithWildcards);
+			check(bCSVOk);
+		}
+		FString ProjectStatisticsFileName = FPaths::ProjectDir() / TEXT("Build") / TEXT("EditorPerfStats.csv");
+		if (FPaths::FileExists(ProjectStatisticsFileName))
+		{
+			UE_LOG(LogSummarizeTrace, Display, TEXT("Loading project statistics from %s"), *ProjectStatisticsFileName);
+			bool bCSVOk = StatisticDefinition::LoadFromCSV(ProjectStatisticsFileName, NameToDefinitionMap, CpuScopeNamesWithWildcards);
+		}
 	}
 
 	bool bFound;
@@ -1972,33 +1978,43 @@ int32 USummarizeTraceCommandlet::Main(const FString& CmdLineParams)
 		CsvHandle = nullptr;
 	}
 
-	// if we were asked to generate a telemetry file, generate it
-	if (!NameToDefinitionMap.IsEmpty())
-	{
-		FString TelemetryCsvFileName = TraceFileBasename + TEXT("Telemetry");
-		TelemetryCsvFileName = FPaths::Combine(TracePath, FPaths::SetExtension(TelemetryCsvFileName, "csv"));
+	FString TelemetryCsvFileName = TraceFileBasename + TEXT("Telemetry");
+	TelemetryCsvFileName = FPaths::Combine(TracePath, FPaths::SetExtension(TelemetryCsvFileName, "csv"));
 
-		// override the test name
-		FString TestName = TraceFileBasename;
-		FParse::Value(*CmdLineParams, TEXT("testname="), TestName, true);
+	// override the test name
+	FString TestName = TraceFileBasename;
+	FParse::Value(*CmdLineParams, TEXT("testname="), TestName, true);
 
-		TArray<TelemetryDefinition> TelemetryData;
-		{
-			TArray<StatisticDefinition> Statistics;
+	TArray<TelemetryDefinition> TelemetryData;
+	TSet<FString> ResolvedStatistics;
+	{
+		TArray<StatisticDefinition> Statistics;
 
-			// resolve scopes to telemetry
-			for (const FSummarizeScope& Scope : SortedScopes)
+		// resolve scopes to telemetry
+		for (const FSummarizeScope& Scope : SortedScopes)
+		{
+			if (!IsCsvSafeString(Scope.Name))
 			{
-				if (!IsCsvSafeString(Scope.Name))
-				{
-					continue;
-				}
+				continue;
+			}
 
+			if (bAllTelemetry)
+			{
+				TelemetryData.Add(TelemetryDefinition(TestName, Scope.Name, TEXT("Count"), Scope.GetValue(TEXT("Count")), TEXT("Count")));
+				TelemetryData.Add(TelemetryDefinition(TestName, Scope.Name, TEXT("TotalDurationSeconds"), Scope.GetValue(TEXT("TotalDurationSeconds")), TEXT("Seconds")));
+				TelemetryData.Add(TelemetryDefinition(TestName, Scope.Name, TEXT("MinDurationSeconds"), Scope.GetValue(TEXT("MinDurationSeconds")), TEXT("Seconds")));
+				TelemetryData.Add(TelemetryDefinition(TestName, Scope.Name, TEXT("MaxDurationSeconds"), Scope.GetValue(TEXT("MaxDurationSeconds")), TEXT("Seconds")));
+				TelemetryData.Add(TelemetryDefinition(TestName, Scope.Name, TEXT("MeanDurationSeconds"), Scope.GetValue(TEXT("MeanDurationSeconds")), TEXT("Seconds")));
+				TelemetryData.Add(TelemetryDefinition(TestName, Scope.Name, TEXT("DeviationDurationSeconds"), Scope.GetValue(TEXT("DeviationDurationSeconds")), TEXT("Seconds")));
+			}
+			else
+			{
 				// Is that scope summary desired in the output, using an exact name match?
 				NameToDefinitionMap.MultiFind(Scope.Name, Statistics, true);
 				for (const StatisticDefinition& Statistic : Statistics)
 				{
 					TelemetryData.Add(TelemetryDefinition(TestName, Statistic.TelemetryContext, Statistic.TelemetryDataPoint, Statistic.TelemetryUnit, Scope.GetValue(Statistic.Statistic)));
+					ResolvedStatistics.Add(Statistic.Name);
 				}
 				Statistics.Reset();
 
@@ -2016,47 +2032,72 @@ int32 USummarizeTraceCommandlet::Main(const FString& CmdLineParams)
 							// case, it is a one to many relationship (1 pattern => * matches).
 							const FString& DataPoint = Scope.Name;
 							TelemetryData.Add(TelemetryDefinition(TestName, Statistic.TelemetryContext, DataPoint, Statistic.TelemetryUnit, Scope.GetValue(Statistic.Statistic)));
+							ResolvedStatistics.Add(Statistic.Name);
 						}
 						Statistics.Reset();
 					}
 				}
 			}
+		}
 
-			// resolve counters to telemetry
-			for (const TMap<uint16, FSummarizeCountersAnalyzer::FCounter>::ElementType& Counter : CountersAnalyzer.Counters)
+		// resolve counters to telemetry
+		for (const TMap<uint16, FSummarizeCountersAnalyzer::FCounter>::ElementType& Counter : CountersAnalyzer.Counters)
+		{
+			if (!IsCsvSafeString(Counter.Value.Name))
 			{
-				if (!IsCsvSafeString(Counter.Value.Name))
-				{
-					continue;
-				}
+				continue;
+			}
 
+			if (bAllTelemetry)
+			{
+				TelemetryData.Add(TelemetryDefinition(TestName, Counter.Value.Name, TEXT("Count"), Counter.Value.GetValue(), TEXT("Count")));
+			}
+			else
+			{
 				NameToDefinitionMap.MultiFind(Counter.Value.Name, Statistics, true);
 				ensure(Statistics.Num() <= 1); // there should only be one, the counter value
 				for (const StatisticDefinition& Statistic : Statistics)
 				{
 					TelemetryData.Add(TelemetryDefinition(TestName, Statistic.TelemetryContext, Statistic.TelemetryDataPoint, Statistic.TelemetryUnit, Counter.Value.GetValue()));
+					ResolvedStatistics.Add(Statistic.Name);
 				}
-				Statistics.Reset();
+				Statistics.Reset();			
 			}
+		}
 
-			// resolve bookmarks to telemetry
-			for (const TMap<FString, FSummarizeBookmark>::ElementType& Bookmark : BookmarksAnalyzer.Bookmarks)
+		// resolve bookmarks to telemetry
+		for (const TMap<FString, FSummarizeBookmark>::ElementType& Bookmark : BookmarksAnalyzer.Bookmarks)
+		{
+			if (!IsCsvSafeString(Bookmark.Value.Name))
 			{
-				if (!IsCsvSafeString(Bookmark.Value.Name))
-				{
-					continue;
-				}
+				continue;
+			}
 
+			if (bAllTelemetry)
+			{
+				TelemetryData.Add(TelemetryDefinition(TestName, Bookmark.Value.Name, TEXT("Count"), Bookmark.Value.GetValue(TEXT("Count")), TEXT("Count")));
+				TelemetryData.Add(TelemetryDefinition(TestName, Bookmark.Value.Name, TEXT("TotalDurationSeconds"), Bookmark.Value.GetValue(TEXT("TotalDurationSeconds")), TEXT("Seconds")));
+				TelemetryData.Add(TelemetryDefinition(TestName, Bookmark.Value.Name, TEXT("MinDurationSeconds"), Bookmark.Value.GetValue(TEXT("MinDurationSeconds")), TEXT("Seconds")));
+				TelemetryData.Add(TelemetryDefinition(TestName, Bookmark.Value.Name, TEXT("MaxDurationSeconds"), Bookmark.Value.GetValue(TEXT("MaxDurationSeconds")), TEXT("Seconds")));
+				TelemetryData.Add(TelemetryDefinition(TestName, Bookmark.Value.Name, TEXT("MeanDurationSeconds"), Bookmark.Value.GetValue(TEXT("MeanDurationSeconds")), TEXT("Seconds")));
+				TelemetryData.Add(TelemetryDefinition(TestName, Bookmark.Value.Name, TEXT("DeviationDurationSeconds"), Bookmark.Value.GetValue(TEXT("DeviationDurationSeconds")), TEXT("Seconds")));
+			}
+			else
+			{
 				NameToDefinitionMap.MultiFind(Bookmark.Value.Name, Statistics, true);
 				ensure(Statistics.Num() <= 1); // there should only be one, the bookmark itself
 				for (const StatisticDefinition& Statistic : Statistics)
 				{
 					TelemetryData.Add(TelemetryDefinition(TestName, Statistic.TelemetryContext, Statistic.TelemetryDataPoint, Statistic.TelemetryUnit, Bookmark.Value.GetValue(Statistic.Statistic)));
+					ResolvedStatistics.Add(Statistic.Name);
 				}
-				Statistics.Reset();
+				Statistics.Reset();			
 			}
 		}
+	}
 
+	if (!bAllTelemetry)
+	{
 		// compare vs. baseline telemetry file, if it exists
 		// note this does assume that the tracefile basename is directly comparable to a file in the baseline folder
 		FString BaselineTelemetryCsvFilePath = FPaths::Combine(FPaths::EngineDir(), TEXT("Build"), TEXT("Baseline"), FPaths::SetExtension(TraceFileBasename + TEXT("Telemetry"), "csv"));
@@ -2154,28 +2195,44 @@ int32 USummarizeTraceCommandlet::Main(const FString& CmdLineParams)
 			}
 		}
 
-		UE_LOG(LogSummarizeTrace, Display, TEXT("Writing telemetry to %s..."), *TelemetryCsvFileName);
-		IFileHandle* TelemetryCsvHandle = FPlatformFileManager::Get().GetPlatformFile().OpenWrite(*TelemetryCsvFileName);
-		if (TelemetryCsvHandle)
+		// check for references to statistics desired for telemetry that are unresolved
+		for (const FString& StatisticName : ResolvedStatistics)
 		{
-			// no newline, see row printfs
-			WriteUTF8(TelemetryCsvHandle, FString::Printf(TEXT("TestName,Context,DataPoint,Unit,Measurement,Baseline,")));
-			for (const TelemetryDefinition& Telemetry : TelemetryData)
-			{
-				// note newline is at the front of every data line to prevent final extraneous newline, per customary for csv
-				WriteUTF8(TelemetryCsvHandle, FString::Printf(TEXT("\n%s,%s,%s,%s,%s,%s,"), *Telemetry.TestName, *Telemetry.Context, *Telemetry.DataPoint, *Telemetry.Unit, *Telemetry.Measurement, *Telemetry.Baseline));
-			}
+			NameToDefinitionMap.Remove(StatisticName);
+		}
 
-			TelemetryCsvHandle->Flush();
-			delete TelemetryCsvHandle;
-			TelemetryCsvHandle = nullptr;
+		for (const TPair<FString, StatisticDefinition>& Statistic : NameToDefinitionMap)
+		{
+			UE_LOG(LogSummarizeTrace, Error, TEXT("Failed to find resolve telemety data for statistic \"%s\""), *Statistic.Value.Name);
 		}
-		else
+
+		if (!NameToDefinitionMap.IsEmpty())
 		{
-			UE_LOG(LogSummarizeTrace, Error, TEXT("Unable to open telemetry csv '%s' for write"), *TelemetryCsvFileName);
 			return 1;
 		}
 	}
 
+	UE_LOG(LogSummarizeTrace, Display, TEXT("Writing telemetry to %s..."), *TelemetryCsvFileName);
+	IFileHandle* TelemetryCsvHandle = FPlatformFileManager::Get().GetPlatformFile().OpenWrite(*TelemetryCsvFileName);
+	if (TelemetryCsvHandle)
+	{
+		// no newline, see row printfs
+		WriteUTF8(TelemetryCsvHandle, FString::Printf(TEXT("TestName,Context,DataPoint,Unit,Measurement,Baseline,")));
+		for (const TelemetryDefinition& Telemetry : TelemetryData)
+		{
+			// note newline is at the front of every data line to prevent final extraneous newline, per customary for csv
+			WriteUTF8(TelemetryCsvHandle, FString::Printf(TEXT("\n%s,%s,%s,%s,%s,%s,"), *Telemetry.TestName, *Telemetry.Context, *Telemetry.DataPoint, *Telemetry.Unit, *Telemetry.Measurement, *Telemetry.Baseline));
+		}
+
+		TelemetryCsvHandle->Flush();
+		delete TelemetryCsvHandle;
+		TelemetryCsvHandle = nullptr;
+	}
+	else
+	{
+		UE_LOG(LogSummarizeTrace, Error, TEXT("Unable to open telemetry csv '%s' for write"), *TelemetryCsvFileName);
+		return 1;
+	}
+
 	return 0;
 }
\ No newline at end of file
