Cook on the fly now respects the `-IgnoreIniSettingsOutOfDate` cook flag.

- Store result in bIgnoreIniSettingsOutOfDate

diff --git a/Engine/Source/Editor/UnrealEd/Classes/Commandlets/CookCommandlet.h b/Engine/Source/Editor/UnrealEd/Classes/Commandlets/CookCommandlet.h
--- a/Engine/Source/Editor/UnrealEd/Classes/Commandlets/CookCommandlet.h
+++ b/Engine/Source/Editor/UnrealEd/Classes/Commandlets/CookCommandlet.h
@@ -57,6 +57,8 @@ class UCookCommandlet
 	bool bPartialGC;
 	/** Do not cook any shaders.  Shader maps will be empty */
 	bool bNoShaderCooking;
+	/** Ignore ini settings out of date. */
+	bool bIgnoreIniSettingsOutOfDate;
 	/** All commandline tokens */
 	TArray<FString> Tokens;
 	/** All commandline switches */
diff --git a/Engine/Source/Editor/UnrealEd/Private/Commandlets/CookCommandlet.cpp b/Engine/Source/Editor/UnrealEd/Private/Commandlets/CookCommandlet.cpp
--- a/Engine/Source/Editor/UnrealEd/Private/Commandlets/CookCommandlet.cpp
+++ b/Engine/Source/Editor/UnrealEd/Private/Commandlets/CookCommandlet.cpp
@@ -422,6 +422,7 @@ bool UCookCommandlet::CookOnTheFly( FGuid InstanceId, int32 Timeout, bool bForce
 	CookFlags |= bSkipEditorContent ? ECookInitializationFlags::SkipEditorContent : ECookInitializationFlags::None;
 	CookFlags |= bUnversioned ? ECookInitializationFlags::Unversioned : ECookInitializationFlags::None;
 	CookFlags |= bCookEditorOptional ? ECookInitializationFlags::CookEditorOptional : ECookInitializationFlags::None;
+	CookFlags |= bIgnoreIniSettingsOutOfDate || CookerSettings->bIgnoreIniSettingsOutOfDateForIteration ? ECookInitializationFlags::IgnoreIniSettingsOutOfDate : ECookInitializationFlags::None;
 	CookOnTheFlyServer->Initialize( ECookMode::CookOnTheFly, CookFlags );
 
 	UCookOnTheFlyServer::FCookOnTheFlyOptions CookOnTheFlyStartupOptions;
@@ -642,6 +643,7 @@ int32 UCookCommandlet::Main(const FString& CmdLineParams)
 	ShowErrorCount = !Switches.Contains(TEXT("DIFFONLY"));
 	ShowProgress = !Switches.Contains(TEXT("DIFFONLY"));
 	bNoShaderCooking = bCookOnTheFly; // Do not cook any shaders into the shader maps. Always true if we are running w/ cook on the fly
+	bIgnoreIniSettingsOutOfDate = Switches.Contains(TEXT("IgnoreIniSettingsOutOfDate"));
 
 	COOK_STAT(DetailedCookStats::CookProject = FApp::GetProjectName());
 
@@ -754,7 +756,7 @@ bool UCookCommandlet::CookByTheBook( const TArray<ITargetPlatform*>& Platforms)
 	bool bTestCook = Switches.Contains(TEXT("TestCook"));
 	CookFlags |= bTestCook ? ECookInitializationFlags::TestCook : ECookInitializationFlags::None;
 	CookFlags |= Switches.Contains(TEXT("LogDebugInfo")) ? ECookInitializationFlags::LogDebugInfo : ECookInitializationFlags::None;
-	CookFlags |= Switches.Contains(TEXT("IgnoreIniSettingsOutOfDate")) || CookerSettings->bIgnoreIniSettingsOutOfDateForIteration ? ECookInitializationFlags::IgnoreIniSettingsOutOfDate : ECookInitializationFlags::None;
+	CookFlags |= bIgnoreIniSettingsOutOfDate || CookerSettings->bIgnoreIniSettingsOutOfDateForIteration ? ECookInitializationFlags::IgnoreIniSettingsOutOfDate : ECookInitializationFlags::None;
 	CookFlags |= Switches.Contains(TEXT("IgnoreScriptPackagesOutOfDate")) || CookerSettings->bIgnoreScriptPackagesOutOfDateForIteration ? ECookInitializationFlags::IgnoreScriptPackagesOutOfDate : ECookInitializationFlags::None;
 
 	//////////////////////////////////////////////////////////////////////////
