CookOnTheFlyServer - StartCookByTheBook refactoring part 1.
1) Add CookSettings metadata to the cook directory so we can invalidate iterative cooks when a cook parameter has changed to a value incompatible with the previous cook.
2) Split BeginCookSandbox into two parts so we can calculate whether the build is iterative earlier and use it in functions that need to run before deleting the previous cook results.

#rnx

diff --git a/Engine/Source/Editor/UnrealEd/Classes/CookOnTheSide/CookOnTheFlyServer.h b/Engine/Source/Editor/UnrealEd/Classes/CookOnTheSide/CookOnTheFlyServer.h
--- a/Engine/Source/Editor/UnrealEd/Classes/CookOnTheSide/CookOnTheFlyServer.h
+++ b/Engine/Source/Editor/UnrealEd/Classes/CookOnTheSide/CookOnTheFlyServer.h
@@ -31,6 +31,7 @@ class IAssetRegistry;
 class ICookedPackageWriter;
 class IPlugin;
 class ITargetPlatform;
+struct FBeginCookContext;
 struct FPropertyChangedEvent;
 
 namespace UE::LinkerLoad { enum class EImportBehavior : uint8; }
@@ -237,6 +238,7 @@ private:
 
 	ECookInitializationFlags CookFlags = ECookInitializationFlags::None;
 	TUniquePtr<class FSandboxPlatformFile> SandboxFile;
+	FString SandboxFileOutputDirectory;
 	TUniquePtr<FAsyncIODelete> AsyncIODelete; // Helper for deleting the old cook directory asynchronously
 	bool bIsSavingPackage = false; // used to stop recursive mark package dirty functions
 	/** Set to true during CookOnTheFly if a plugin is calling RequestPackage and we should therefore not make assumptions about when platforms are done cooking */
@@ -713,6 +715,10 @@ private:
 
 	//////////////////////////////////////////////////////////////////////////
 	// cook by the book specific functions
+	/** Construct the on-stack data for StartCook functions, based on the arguments to StartCookByTheBook */
+	FBeginCookContext SetCookByTheBookOptions(const FCookByTheBookStartupOptions& StartupOptions);
+	/** Set the PlatformManager's session platforms and finish filling out the BeginContext for StartCookByTheBook */
+	void SelectSessionPlatforms(FBeginCookContext& BeginContext);
 	/**
 	* Collect all the files which need to be cooked for a cook by the book session
 	*/
@@ -821,6 +827,8 @@ private:
 	//////////////////////////////////////////////////////////////////////////
 	// cook on the fly specific functions
 
+	/** Construct the on-stack data for StartCook functions, based on the arguments when adding a COTF platform */
+	FBeginCookContext CreateAddPlatformContext(ITargetPlatform* TargetPlatform);
 	void GetCookOnTheFlyUnsolicitedFiles(const ITargetPlatform* TargetPlatform, const FString& PlatformName, TArray<FString>& UnsolicitedFiles, const FString& Filename, bool bIsCookable);
 
 	//////////////////////////////////////////////////////////////////////////
@@ -843,10 +851,10 @@ private:
 	bool LoadPackageForCooking(UE::Cook::FPackageData& PackageData, UPackage*& OutPackage, FString* LoadFromFileName = nullptr,
 		UE::Cook::FPackageData* ReportingPackageData = nullptr);
 
-	/**
-	* Initialize the sandbox for a new cook session with @param TargetPlatforms
-	*/
-	void BeginCookSandbox(TConstArrayView<const ITargetPlatform*> TargetPlatforms);
+	/** Read information about the previous cook from disk and set whether the current cook is iterative based on previous cook, config, and commandline. */
+	void SetBeginCookIterativeFlags(FBeginCookContext& BeginContext);
+	/** Initialize the sandbox for a new cook session */
+	void BeginCookSandbox(FBeginCookContext& BeginContext);
 
 	/** Set parameters that rely on config and CookByTheBook settings. */
 	void SetBeginCookConfigSettings();
@@ -878,11 +886,6 @@ private:
 	void DiscoverPlatformSpecificNeverCookPackages(
 		const TArrayView<const ITargetPlatform* const>& TargetPlatforms, const TArray<FString>& UBTPlatformStrings);
 
-	/**
-	* Clean up the sandbox
-	*/
-	void TermSandbox();
-
 	/**
 	* GetDependentPackages
 	* get package dependencies according to the asset registry
@@ -971,6 +974,14 @@ private:
 	*/
 	bool GetCookedIniVersionStrings( const ITargetPlatform* TargetPlatform, FIniSettingContainer& IniVersionStrings, TMap<FString, FString>& AdditionalStrings ) const;
 
+	/** Test the CurrentCookSettings against the previous cooksettings to decide if we have to wipe the previous cook even when running iteratively. */
+	bool ArePreviousCookSettingsCompatible(const TMap<FName, FString>& CurrentCookSettings, const ITargetPlatform* TargetPlatform) const;
+	/** Save the CurrentCookSettings into the output directory. */
+	void SaveCookSettings(const TMap<FName, FString>& CurrentCookSettings, const ITargetPlatform* TargetPlatform);
+	/** Populate a map suitable for saving as an ini with the current value of the Cook settings that need to be tested for compatibility. */
+	TMap<FName, FString> CalculateCookSettingStrings() const;
+	/** Get the path to the CookSettings metadata file for the given platform and current output directory. */
+	FString GetCookSettingsFileName(const ITargetPlatform* TargetPlatform) const;
 
 	/**
 	* Convert a path to a full sandbox path
diff --git a/Engine/Source/Editor/UnrealEd/Private/CookOnTheFlyServer.cpp b/Engine/Source/Editor/UnrealEd/Private/CookOnTheFlyServer.cpp
--- a/Engine/Source/Editor/UnrealEd/Private/CookOnTheFlyServer.cpp
+++ b/Engine/Source/Editor/UnrealEd/Private/CookOnTheFlyServer.cpp
@@ -8,6 +8,7 @@
 
 #include "Algo/AnyOf.h"
 #include "Algo/Find.h"
+#include "Algo/Unique.h"
 #include "AssetCompilingManager.h"
 #include "AssetRegistry/AssetRegistryModule.h"
 #include "AssetRegistry/AssetRegistryState.h"
@@ -535,6 +536,8 @@ bool UCookOnTheFlyServer::StartCookOnTheFly(FCookOnTheFlyOptions InCookOnTheFlyO
 #endif
 
 	CookOnTheFlyOptions = MoveTemp(InCookOnTheFlyOptions);
+	GenerateAssetRegistry();
+	CreateSandboxFile();
 	CookOnTheFlyServerInterface = MakeUnique<UCookOnTheFlyServer::FCookOnTheFlyServerInterface>(*this);
 	ExternalRequests->CookRequestEvent = FPlatformProcess::GetSynchEventFromPool();
 
@@ -548,8 +551,6 @@ bool UCookOnTheFlyServer::StartCookOnTheFly(FCookOnTheFlyOptions InCookOnTheFlyO
 	PlatformManager->SetArePlatformsPrepopulated(true);
 
 	SetBeginCookConfigSettings();
-	CreateSandboxFile();
-	GenerateAssetRegistry();
 	for (FName NeverCookPackage : GetNeverCookPackageFileNames(TArrayView<const FString>()))
 	{
 		PackageTracker->NeverCookPackageList.Add(NeverCookPackage);
@@ -621,13 +622,14 @@ void UCookOnTheFlyServer::AddCookOnTheFlyPlatformFromGameThread(ITargetPlatform*
 		return;
 	}
 
-	TArrayView<ITargetPlatform* const> NewTargetPlatforms(&TargetPlatform,1);
+	FBeginCookContext BeginContext = CreateAddPlatformContext(TargetPlatform);
 
-	RefreshPlatformAssetRegistries(NewTargetPlatforms);
-	FindOrCreateSaveContexts(NewTargetPlatforms);
-	BeginCookSandbox(NewTargetPlatforms);
-	InitializeTargetPlatforms(NewTargetPlatforms);
-	InitializeShadersForCookOnTheFly(NewTargetPlatforms);
+	RefreshPlatformAssetRegistries(BeginContext.TargetPlatforms);
+	FindOrCreateSaveContexts(BeginContext.TargetPlatforms);
+	SetBeginCookIterativeFlags(BeginContext);
+	BeginCookSandbox(BeginContext);
+	InitializeTargetPlatforms(BeginContext.TargetPlatforms);
+	InitializeShadersForCookOnTheFly(BeginContext.TargetPlatforms);
 
 	// When cooking on the fly the full registry is saved at the beginning
 	// in cook by the book asset registry is saved after the cook is finished
@@ -5500,6 +5502,67 @@ void UCookOnTheFlyServer::ProcessAccessedIniSettings(const FConfigFile* Config,
 	}
 }
 
+static const TCHAR* TEXT_CookSettings(TEXT("CookSettings"));
+
+TMap<FName, FString> UCookOnTheFlyServer::CalculateCookSettingStrings() const
+{
+	TMap<FName, FString> CookSettingStrings;
+	const FName NAME_CookMode(TEXT("CookMode"));
+
+	CookSettingStrings.Add(FName(TEXT("Version")), TEXT("C7C76F79"));
+	if (CookByTheBookOptions)
+	{
+		CookSettingStrings.Add(NAME_CookMode, TEXT("CookByTheBook"));
+		CookSettingStrings.Add(FName(TEXT("DLCName")), CookByTheBookOptions->DlcName);
+	}
+	else
+	{
+		CookSettingStrings.Add(NAME_CookMode, TEXT("CookOnTheFly"));
+	}
+	return CookSettingStrings;
+}
+
+FString UCookOnTheFlyServer::GetCookSettingsFileName(const ITargetPlatform* TargetPlatform) const
+{
+	FString CookedSettingsIni = FPaths::ProjectDir() / TEXT("Metadata") / TEXT("CookedSettings.txt");
+	return ConvertToFullSandboxPath(*CookedSettingsIni, true, TargetPlatform->PlatformName());
+}
+
+bool UCookOnTheFlyServer::ArePreviousCookSettingsCompatible(const TMap<FName, FString>& CurrentCookSettings, const ITargetPlatform* TargetPlatform) const
+{
+	FConfigFile ConfigFile;
+	ConfigFile.Read(GetCookSettingsFileName(TargetPlatform));
+
+	const FConfigSection* CookSettings = ConfigFile.Find(TEXT_CookSettings);
+	if (CookSettings == nullptr)
+	{
+		return false;
+	}
+
+	for (const TPair<FName, FString>& CurrentSetting : CurrentCookSettings)
+	{
+		const FConfigValue* PreviousSetting = CookSettings->Find(CurrentSetting.Key);
+		if (!PreviousSetting || PreviousSetting->GetValue() != CurrentSetting.Value)
+		{
+			return false;
+		}
+	}
+
+	return true;
+}
+
+void UCookOnTheFlyServer::SaveCookSettings(const TMap<FName, FString>& CurrentCookSettings, const ITargetPlatform* TargetPlatform)
+{
+	FConfigFile ConfigFile;
+	FConfigSection& SavedSettings = *ConfigFile.FindOrAddSection(TEXT_CookSettings);
+	for (const TPair<FName, FString>& CurrentSetting : CurrentCookSettings)
+	{
+		SavedSettings.Add(CurrentSetting.Key, CurrentSetting.Value);
+	}
+	ConfigFile.Dirty = true; // Writing to a section does not set the dirty flag, so set it manually to make Write work
+	ConfigFile.Write(GetCookSettingsFileName(TargetPlatform));
+}
+
 bool UCookOnTheFlyServer::IniSettingsOutOfDate(const ITargetPlatform* TargetPlatform) const
 {
 	TGuardValue<bool> A(GSuppressProcessConfigSettings, true);
@@ -7318,7 +7381,6 @@ void UCookOnTheFlyServer::CancelCookByTheBook()
 		CancelAllQueues();
 
 		CookByTheBookOptions->bRunning = false;
-		SandboxFile = nullptr;
 
 		PrintFinishStats();
 	} 
@@ -7343,12 +7405,15 @@ void UCookOnTheFlyServer::StopAndClearCookedData()
 
 void UCookOnTheFlyServer::ClearAllCookedData()
 {
+	checkf(!IsInSession(), TEXT("We do not handle removing SessionPlatforms, so ClearAllCookedData must not be called while in a cook session"));
+
 	// if we are going to clear the cooked packages it is conceivable that we will recook the packages which we just cooked 
 	// that means it's also conceivable that we will recook the same package which currently has an outstanding async write request
 	UPackage::WaitForAsyncFileWrites();
 
 	PackageTracker->UnsolicitedCookedPackages.Empty();
 	PackageDatas->ClearCookedPlatforms();
+	ClearPackageStoreContexts();
 }
 
 void UCookOnTheFlyServer::CancelAllQueues()
@@ -7419,6 +7484,11 @@ void UCookOnTheFlyServer::ClearPlatformCookedData(const ITargetPlatform* TargetP
 	{
 		return;
 	}
+	if (!SandboxFile)
+	{
+		// We cannot get the PackageWriter without it, and we do not have anything to clear if it has not been created
+		return;
+	}
 	ResetCook({ TPair<const ITargetPlatform*,bool>{TargetPlatform, true /* bResetResults */}});
 
 	FindOrCreatePackageWriter(TargetPlatform).RemoveCookedPackages();
@@ -7491,19 +7561,31 @@ void UCookOnTheFlyServer::OnTargetPlatformChangedSupportedFormats(const ITargetP
 
 void UCookOnTheFlyServer::CreateSandboxFile()
 {
-	// initialize the sandbox file after determining if we are cooking dlc
-	// Local sandbox file wrapper. This will be used to handle path conversions,
-	// but will not be used to actually write/read files so we can safely
-	// use [Platform] token in the sandbox directory name and then replace it
-	// with the actual platform name.
-	check(SandboxFile == nullptr);
-	SandboxFile = FSandboxPlatformFile::Create(false);
-
-	// Output directory override.	
+	// Output directory override. This directory depends on whether we are cooking dlc, so we cannot
+	// create the sandbox until after StartCookByTheBook or StartCookOnTheFly
 	FString OutputDirectory = GetOutputDirectoryOverride();
+	check(!OutputDirectory.IsEmpty());
+	check((SandboxFile == nullptr) == SandboxFileOutputDirectory.IsEmpty());
 
-	// Use SandboxFile to do path conversion to properly handle sandbox paths (outside of standard paths in particular).
+	if (SandboxFile)
+	{
+		if (SandboxFileOutputDirectory == OutputDirectory)
+		{
+			return;
+		}
+		ClearAllCookedData(); // Does not delete files on disk, only deletes in-memory data
+		SandboxFile.Reset();
+	}
+
+	// Local sandbox file wrapper. This will be used to handle path conversions, but will not be used to actually
+	// write/read files so we can safely use [Platform] token in the sandbox directory name and then replace it
+	// with the actual platform name.
+	// Filename lookups in the cooker must Use this SandboxFile to do path conversion to properly handle sandbox paths
+	// (outside of standard paths in particular).
+	SandboxFile = FSandboxPlatformFile::Create(false);
 	SandboxFile->Initialize(&FPlatformFileManager::Get().GetPlatformFile(), *FString::Printf(TEXT("-sandbox=\"%s\""), *OutputDirectory));
+	SandboxFileOutputDirectory = OutputDirectory;
+
 }
 
 void UCookOnTheFlyServer::SetBeginCookConfigSettings()
@@ -7527,119 +7609,141 @@ void UCookOnTheFlyServer::SetBeginCookConfigSettings()
 
 }
 
-void UCookOnTheFlyServer::BeginCookSandbox(TConstArrayView<const ITargetPlatform*> TargetPlatforms)
+void UCookOnTheFlyServer::SetBeginCookIterativeFlags(FBeginCookContext& BeginContext)
 {
-#if OUTPUT_COOKTIMING
-	double CleanSandboxTime = 0.0;
-#endif
-	{
-		UE_SCOPED_HIERARCHICAL_COOKTIMER_AND_DURATION(CleanSandbox, CleanSandboxTime);
+	const bool bIsDiffOnly = FParse::Param(FCommandLine::Get(), TEXT("DIFFONLY"));
+	const bool bIterative = !FParse::Param(FCommandLine::Get(), TEXT("fullcook")) && (bHybridIterativeEnabled || IsCookFlagSet(ECookInitializationFlags::Iterative));
+	const bool bIsSharedIterativeCook = IsCookFlagSet(ECookInitializationFlags::IterateSharedBuild);
 
-		if (SandboxFile == nullptr)
+	for (FBeginCookContextPlatform& PlatformContext : BeginContext.PlatformContexts)
+	{
+		const ITargetPlatform* TargetPlatform = PlatformContext.TargetPlatform;
+		UE::Cook::FPlatformData* PlatformData = PlatformContext.PlatformData;
+		ICookedPackageWriter& PackageWriter = FindOrCreatePackageWriter(TargetPlatform);
+		bool bIterateSharedBuild = false;
+		if (bIterative && bIsSharedIterativeCook && !PlatformData->bIsSandboxInitialized)
 		{
-			CreateSandboxFile();
-		}
+			// see if the shared build is newer then the current cooked content in the local directory
+			FString SharedCookedAssetRegistry = FPaths::Combine(*FPaths::ProjectSavedDir(), TEXT("SharedIterativeBuild"),
+				*TargetPlatform->PlatformName(), TEXT("Metadata"), GetDevelopmentAssetRegistryFilename());
 
-		const bool bIsDiffOnly = FParse::Param(FCommandLine::Get(), TEXT("DIFFONLY"));
-		const bool bIterative = !FParse::Param(FCommandLine::Get(), TEXT("fullcook")) && (bHybridIterativeEnabled || IsCookFlagSet(ECookInitializationFlags::Iterative));
-		const bool bIsSharedIterativeCook = IsCookFlagSet(ECookInitializationFlags::IterateSharedBuild);
-		TArray<TPair<const ITargetPlatform*,bool>, TInlineAllocator<ExpectedMaxNumPlatforms>> ResetPlatforms;
-		TArray<const ITargetPlatform*, TInlineAllocator<ExpectedMaxNumPlatforms>> PopulatePlatforms;
-		TArray<const ITargetPlatform*, TInlineAllocator<ExpectedMaxNumPlatforms>> AlreadyCookedPlatforms;
+			FDateTime PreviousLocalCookedBuild = PackageWriter.GetPreviousCookTime();
+			FDateTime PreviousSharedCookedBuild = IFileManager::Get().GetTimeStamp(*SharedCookedAssetRegistry);
+			if (PreviousSharedCookedBuild != FDateTime::MinValue() &&
+				PreviousSharedCookedBuild >= PreviousLocalCookedBuild)
+			{
+				// copy the ini settings from the shared cooked build. 
+				const FString SharedCookedIniFile = FPaths::Combine(*FPaths::ProjectSavedDir(), TEXT("SharedIterativeBuild"),
+					*TargetPlatform->PlatformName(), TEXT("Metadata"), TEXT("CookedIniVersion.txt"));
+				FString SandboxCookedIniFile = FPaths::ProjectDir() / TEXT("Metadata") / TEXT("CookedIniVersion.txt");
+				SandboxCookedIniFile = ConvertToFullSandboxPath(*SandboxCookedIniFile, true, *TargetPlatform->PlatformName());
+				IFileManager::Get().Copy(*SandboxCookedIniFile, *SharedCookedIniFile);
+				bIterateSharedBuild = true;
+				UE_LOG(LogCook, Display, TEXT("Shared iterative build is newer then local cooked build, iteratively cooking from shared build."));
+			}
+			else
+			{
+				UE_LOG(LogCook, Display, TEXT("Local cook is newer then shared cooked build, iteratively cooking from local build."));
+			}
+		}
+		PlatformContext.CurrentCookSettings = CalculateCookSettingStrings();
+		const bool bIsIniSettingsOutOfDate = IniSettingsOutOfDate(TargetPlatform); // needs to be executed for side effects even if non-iterative
+		PlatformContext.bHasMemoryResults = PlatformData->bIsSandboxInitialized;
 
-		for (const ITargetPlatform* TargetPlatform : TargetPlatforms)
+		if (bIsDiffOnly)
 		{
-			UE::Cook::FPlatformData* PlatformData = PlatformManager->GetPlatformData(TargetPlatform);
-			ICookedPackageWriter& PackageWriter = FindOrCreatePackageWriter(TargetPlatform);
-			bool bIterateSharedBuild = false;
-			if (bIterative && bIsSharedIterativeCook && !PlatformData->bIsSandboxInitialized)
+			UE_LOG(LogCook, Display, TEXT("Keeping cooked content for platform %s for DiffOnly"), *TargetPlatform->PlatformName());
+			// When looking for deterministic cooking differences in cooked packages, don't delete the packages on disk
+			PlatformContext.bFullBuild = false;
+			PlatformContext.bClearMemoryResults = true;
+			PlatformContext.bPopulateMemoryResultsFromDiskResults = false;
+			PlatformContext.bIterateSharedBuild = false;
+		}
+		else
+		{
+			bool bIterativeAllowed = true;
+			if (!bIterative && !PlatformData->bIsSandboxInitialized)
 			{
-				// see if the shared build is newer then the current cooked content in the local directory
-				FString SharedCookedAssetRegistry = FPaths::Combine(*FPaths::ProjectSavedDir(), TEXT("SharedIterativeBuild"),
-					*TargetPlatform->PlatformName(), TEXT("Metadata"), GetDevelopmentAssetRegistryFilename());
-				
-				FDateTime PreviousLocalCookedBuild = PackageWriter.GetPreviousCookTime();
-				FDateTime PreviousSharedCookedBuild = IFileManager::Get().GetTimeStamp(*SharedCookedAssetRegistry);
-				if (PreviousSharedCookedBuild != FDateTime::MinValue() &&
-					PreviousSharedCookedBuild >= PreviousLocalCookedBuild)
-				{
-					// copy the ini settings from the shared cooked build. 
-					const FString SharedCookedIniFile = FPaths::Combine(*FPaths::ProjectSavedDir(), TEXT("SharedIterativeBuild"),
-						*TargetPlatform->PlatformName(), TEXT("Metadata"), TEXT("CookedIniVersion.txt"));
-					FString SandboxCookedIniFile = FPaths::ProjectDir() / TEXT("Metadata") / TEXT("CookedIniVersion.txt");
-					SandboxCookedIniFile = ConvertToFullSandboxPath(*SandboxCookedIniFile, true);
-					SandboxCookedIniFile.ReplaceInline(TEXT("[Platform]"), *TargetPlatform->PlatformName());
-					IFileManager::Get().Copy(*SandboxCookedIniFile, *SharedCookedIniFile);
-					bIterateSharedBuild = true;
-					UE_LOG(LogCook, Display, TEXT("Shared iterative build is newer then local cooked build, iteratively cooking from shared build."));
-				}
-				else
-				{
-					UE_LOG(LogCook, Display, TEXT("Local cook is newer then shared cooked build, iteratively cooking from local build."));
-				}
+				UE_LOG(LogCook, Display, TEXT("Clearing all cooked content for platform %s"), *TargetPlatform->PlatformName());
+				bIterativeAllowed = false;
 			}
-			const bool bIsIniSettingsOutOfDate = IniSettingsOutOfDate(TargetPlatform); // needs to be executed for side effects even if non-iterative
-
-			bool bShouldClearCookedContent = true;
-			if (bIsDiffOnly)
+			else if (!ArePreviousCookSettingsCompatible(PlatformContext.CurrentCookSettings, TargetPlatform))
 			{
-				// When looking for deterministic cooking differences in cooked packages, don't delete the packages on disk
-				bShouldClearCookedContent = false;
+				UE_LOG(LogCook, Display, TEXT("Cook invalidated for platform %s because cook DLC settings have changed, clearing all cooked content"), *TargetPlatform->PlatformName());
+				bIterativeAllowed = false;
 			}
-			else if (bIterative || PlatformData->bIsSandboxInitialized)
+			else if (bIsIniSettingsOutOfDate)
 			{
-				if (!bIsIniSettingsOutOfDate)
+				if (!IsCookFlagSet(ECookInitializationFlags::IgnoreIniSettingsOutOfDate))
 				{
-					bShouldClearCookedContent = false;
+					UE_LOG(LogCook, Display, TEXT("Cook invalidated for platform %s ini settings don't match from last cook, clearing all cooked content"), *TargetPlatform->PlatformName());
+					bIterativeAllowed = false;
 				}
 				else
 				{
-					if (!IsCookFlagSet(ECookInitializationFlags::IgnoreIniSettingsOutOfDate))
-					{
-						UE_LOG(LogCook, Display, TEXT("Cook invalidated for platform %s ini settings don't match from last cook, clearing all cooked content"), *TargetPlatform->PlatformName());
-						bShouldClearCookedContent = true;
-					}
-					else
-					{
-						UE_LOG(LogCook, Display, TEXT("Inisettings were out of date for platform %s but we are going with it anyway because IgnoreIniSettingsOutOfDate is set"), *TargetPlatform->PlatformName());
-						bShouldClearCookedContent = false;
-					}
+					UE_LOG(LogCook, Display, TEXT("Inisettings were out of date for platform %s but we are going with it anyway because IgnoreIniSettingsOutOfDate is set"), *TargetPlatform->PlatformName());
+					bIterativeAllowed = true;
 				}
 			}
 			else
 			{
-				// In non-iterative cooks we will be replacing every cooked package and so should wipe the cooked directory
-				UE_LOG(LogCook, Display, TEXT("Clearing all cooked content for platform %s"), *TargetPlatform->PlatformName());
-				bShouldClearCookedContent = true;
+				UE_LOG(LogCook, Display, TEXT("Keeping cooked content for platform %s and cooking iteratively"), *TargetPlatform->PlatformName());
+				bIterativeAllowed = true;
 			}
 
-			ICookedPackageWriter::FCookInfo CookInfo;
-			CookInfo.CookMode = IsCookOnTheFlyMode() ? ICookedPackageWriter::FCookInfo::CookOnTheFlyMode : ICookedPackageWriter::FCookInfo::CookByTheBookMode;
-			CookInfo.bFullBuild = false;
-			CookInfo.bIterateSharedBuild = bIterateSharedBuild && !bShouldClearCookedContent;
-			bool bResetResults = false;
-			if (bShouldClearCookedContent)
+			if (bIterativeAllowed)
 			{
-				CookInfo.bFullBuild = true;
-				bResetResults = true;
+				PlatformContext.bFullBuild = false;
+				PlatformContext.bClearMemoryResults = false;
+				PlatformContext.bPopulateMemoryResultsFromDiskResults = !PlatformContext.bHasMemoryResults;
+				PlatformContext.bIterateSharedBuild = bIterateSharedBuild;
 			}
-			else if (!PlatformData->bIsSandboxInitialized)
+			else
 			{
-				bResetResults = true;
-				// Reset but don't populate if we are looking for deterministic cooking differences; start from an empty list of cooked packages
-				if (!bIsDiffOnly)
-				{
-					PopulatePlatforms.Add(TargetPlatform);
-				}
+				PlatformContext.bFullBuild = true;
+				PlatformContext.bClearMemoryResults = true;
+				PlatformContext.bPopulateMemoryResultsFromDiskResults = false;
+				PlatformContext.bIterateSharedBuild = false;
 			}
-			else
+		}
+		PlatformData->bFullBuild = PlatformContext.bFullBuild;
+	}
+}
+
+void UCookOnTheFlyServer::BeginCookSandbox(FBeginCookContext& BeginContext)
+{
+#if OUTPUT_COOKTIMING
+	double CleanSandboxTime = 0.0;
+#endif
+	{
+		UE_SCOPED_HIERARCHICAL_COOKTIMER_AND_DURATION(CleanSandbox, CleanSandboxTime);
+		TArray<TPair<const ITargetPlatform*, bool>, TInlineAllocator<ExpectedMaxNumPlatforms>> ResetPlatforms;
+		TArray<const ITargetPlatform*, TInlineAllocator<ExpectedMaxNumPlatforms>> PopulatePlatforms;
+		TArray<const ITargetPlatform*, TInlineAllocator<ExpectedMaxNumPlatforms>> AlreadyCookedPlatforms;
+		for (FBeginCookContextPlatform& PlatformContext : BeginContext.PlatformContexts)
+		{
+			const ITargetPlatform* TargetPlatform = PlatformContext.TargetPlatform;
+			UE::Cook::FPlatformData* PlatformData = PlatformContext.PlatformData;
+			ICookedPackageWriter& PackageWriter = FindOrCreatePackageWriter(TargetPlatform);
+			ICookedPackageWriter::FCookInfo CookInfo;
+			CookInfo.CookMode = IsCookOnTheFlyMode() ? ICookedPackageWriter::FCookInfo::CookOnTheFlyMode : ICookedPackageWriter::FCookInfo::CookByTheBookMode;
+			CookInfo.bFullBuild = PlatformContext.bFullBuild;
+			CookInfo.bIterateSharedBuild = PlatformContext.bIterateSharedBuild;
+			PackageWriter.Initialize(CookInfo);
+
+			if (PlatformContext.bPopulateMemoryResultsFromDiskResults)
+			{
+				PopulatePlatforms.Add(TargetPlatform);
+			}
+			else if (PlatformContext.bHasMemoryResults && !PlatformContext.bClearMemoryResults)
 			{
 				AlreadyCookedPlatforms.Add(TargetPlatform);
 			}
-			PackageWriter.Initialize(CookInfo);
-			PlatformData->bFullBuild = CookInfo.bFullBuild;
+			bool bResultResults = PlatformContext.bHasMemoryResults && PlatformContext.bClearMemoryResults;
+			ResetPlatforms.Emplace(TargetPlatform, bResultResults);
+			SaveCookSettings(PlatformContext.CurrentCookSettings, TargetPlatform);
+
 			PlatformData->bIsSandboxInitialized = true;
-			ResetPlatforms.Emplace(TargetPlatform, bResetResults);
 		}
 
 		ResetCook(ResetPlatforms);
@@ -7664,7 +7768,7 @@ void UCookOnTheFlyServer::BeginCookSandbox(TConstArrayView<const ITargetPlatform
 
 #if OUTPUT_COOKTIMING
 	FString PlatformNames;
-	for (const ITargetPlatform* Target : TargetPlatforms)
+	for (const ITargetPlatform* Target : BeginContext.TargetPlatforms)
 	{
 		PlatformNames += Target->PlatformName() + TEXT(" ");
 	}
@@ -7676,11 +7780,7 @@ void UCookOnTheFlyServer::BeginCookSandbox(TConstArrayView<const ITargetPlatform
 UE::Cook::FCookSavePackageContext* UCookOnTheFlyServer::CreateSaveContext(const ITargetPlatform* TargetPlatform)
 {
 	using namespace UE::Cook;
-
-	if (SandboxFile == nullptr)
-	{
-		CreateSandboxFile();
-	}
+	checkf(SandboxFile, TEXT("SaveContexts cannot be created until after CreateSandboxFile has been called from StartCookByTheBook or StartCookOnTheFly"));
 
 	const FString RootPathSandbox = ConvertToFullSandboxPath(FPaths::RootDir(), true);
 	FString MetadataPathSandbox;
@@ -7816,12 +7916,6 @@ void UCookOnTheFlyServer::DiscoverPlatformSpecificNeverCookPackages(
 	}
 }
 
-void UCookOnTheFlyServer::TermSandbox()
-{
-	ClearAllCookedData();
-	SandboxFile = nullptr;
-}
-
 void UCookOnTheFlyServer::StartCookByTheBook( const FCookByTheBookStartupOptions& CookByTheBookStartupOptions )
 {
 	UE_SCOPED_COOKTIMER(StartCookByTheBook);
@@ -7856,37 +7950,28 @@ void UCookOnTheFlyServer::StartCookByTheBook( const FCookByTheBookStartupOptions
 	CookByTheBookOptions->bDlcLoadMainAssetRegistry = !!(CookOptions & ECookByTheBookOptions::DlcLoadMainAssetRegistry);
 	CookByTheBookOptions->bErrorOnEngineContentUse = CookByTheBookStartupOptions.bErrorOnEngineContentUse;
 
-	// if we are going to change the state of dlc, we need to clean out our package filename cache (the generated filename cache is dependent on this key). This has to happen later on, but we want to set the DLC State earlier.
-	const bool bDlcStateChanged = CookByTheBookOptions->DlcName != TEXT("") && CookByTheBookOptions->DlcName != DLCName;
 	CookByTheBookOptions->DlcName = DLCName;
 	if (CookByTheBookOptions->bSkipHardReferences && !CookByTheBookOptions->bSkipSoftReferences)
 	{
 		UE_LOG(LogCook, Warning, TEXT("Setting bSkipSoftReferences to true since bSkipHardReferences is true and skipping hard references requires skipping soft references."));
 		CookByTheBookOptions->bSkipSoftReferences = true;
 	}
-	SetBeginCookConfigSettings();
 	COOK_STAT(UE::SavePackageUtilities::ResetCookStats());
+	FBeginCookContext BeginContext = SetCookByTheBookOptions(CookByTheBookStartupOptions);
 
 	GenerateAssetRegistry();
 	BlockOnAssetRegistry();
+	CreateSandboxFile();
+	SetBeginCookConfigSettings();
 	PackageDatas->BeginCook();
 	if (!IsCookingInEditor())
 	{
 		FCoreUObjectDelegates::PackageCreatedForLoad.AddUObject(this, &UCookOnTheFlyServer::MaybeMarkPackageAsAlreadyLoaded);
 	}
 
-	// SelectSessionPlatforms does not check for uniqueness and non-null, and we rely on those properties for performance, so ensure it here before calling SelectSessionPlatforms
-	TArray<ITargetPlatform*> TargetPlatforms;
-	TargetPlatforms.Reserve(CookByTheBookStartupOptions.TargetPlatforms.Num());
-	for (ITargetPlatform* TargetPlatform : CookByTheBookStartupOptions.TargetPlatforms)
-	{
-		if (TargetPlatform)
-		{
-			TargetPlatforms.AddUnique(TargetPlatform);
-		}
-	}
-	PlatformManager->SelectSessionPlatforms(TargetPlatforms);
-	bPackageFilterDirty = true;
+	SelectSessionPlatforms(BeginContext);
+	SetBeginCookIterativeFlags(BeginContext);
+	TArray<ITargetPlatform*> TargetPlatforms = BeginContext.TargetPlatforms;
 	check(PlatformManager->GetSessionPlatforms().Num() == TargetPlatforms.Num());
 
 	// We want to set bRunning = true as early as possible, but it implies that session platforms have been selected so this is the earliest point we can set it
@@ -7962,17 +8047,11 @@ void UCookOnTheFlyServer::StartCookByTheBook( const FCookByTheBookStartupOptions
 		DiscoverPlatformSpecificNeverCookPackages(TargetPlatforms, UBTPlatformStrings);
 	}
 
-	if (bDlcStateChanged)
-	{
-		// If we changed the DLC State earlier on, we must clear out the package name cache
-		TermSandbox();
-	}
-	
 	FindOrCreateSaveContexts(TargetPlatforms);
 	// TODO: Fix the elements of StartCookByTheBook that rely on BeginCookSandbox being called early, and remove this call to BeginCookSandbox
 	if (!bHybridIterativeEnabled)
 	{
-		BeginCookSandbox(TargetPlatforms); // This will either delete the sandbox or iteratively clean it
+		BeginCookSandbox(BeginContext); // This will either delete the sandbox or iteratively clean it
 	}
 	InitializeTargetPlatforms(TargetPlatforms);
 
@@ -8248,7 +8327,7 @@ void UCookOnTheFlyServer::StartCookByTheBook( const FCookByTheBookStartupOptions
 
 	if (bHybridIterativeEnabled)
 	{
-		BeginCookSandbox(TargetPlatforms);
+		BeginCookSandbox(BeginContext);
 	}
 
 	// don't resave the global shader map files in dlc
@@ -8284,6 +8363,51 @@ void UCookOnTheFlyServer::StartCookByTheBook( const FCookByTheBookStartupOptions
 	InitializePollables();
 }
 
+FBeginCookContext UCookOnTheFlyServer::SetCookByTheBookOptions(const FCookByTheBookStartupOptions& StartupOptions)
+{
+	FBeginCookContext BeginContext;
+
+	BeginContext.TargetPlatforms = StartupOptions.TargetPlatforms;
+	Algo::Sort(BeginContext.TargetPlatforms);
+	BeginContext.TargetPlatforms.SetNum(Algo::Unique(BeginContext.TargetPlatforms));
+
+	BeginContext.PlatformContexts.SetNum(BeginContext.TargetPlatforms.Num());
+	for (int32 Index = 0; Index < BeginContext.TargetPlatforms.Num(); ++Index)
+	{
+		BeginContext.PlatformContexts[Index].TargetPlatform = BeginContext.TargetPlatforms[Index];
+		// PlatformContext.PlatformData is currently null and is set in SelectSessionPlatforms
+	}
+
+	return BeginContext;
+}
+
+FBeginCookContext UCookOnTheFlyServer::CreateAddPlatformContext(ITargetPlatform* TargetPlatform)
+{
+	FBeginCookContext BeginContext;
+
+	BeginContext.TargetPlatforms.Add(TargetPlatform);
+
+	FBeginCookContextPlatform& PlatformContext = BeginContext.PlatformContexts.Emplace_GetRef();
+	PlatformContext.TargetPlatform = TargetPlatform;
+	PlatformContext.PlatformData = &PlatformManager->CreatePlatformData(TargetPlatform);
+
+	return BeginContext;
+}
+
+void UCookOnTheFlyServer::SelectSessionPlatforms(FBeginCookContext& BeginContext)
+{
+	PlatformManager->SelectSessionPlatforms(BeginContext.TargetPlatforms);
+	if (PackageTracker->HasBeenConsumed())
+	{
+		bPackageFilterDirty = true;
+	}
+	FindOrCreateSaveContexts(BeginContext.TargetPlatforms);
+	for (FBeginCookContextPlatform& PlatformContext : BeginContext.PlatformContexts)
+	{
+		PlatformContext.PlatformData = PlatformManager->GetPlatformData(PlatformContext.TargetPlatform);
+	}
+}
+
 TArray<FName> UCookOnTheFlyServer::GetNeverCookPackageFileNames(TArrayView<const FString> ExtraNeverCookDirectories)
 {
 	TArray<FString> NeverCookDirectories(ExtraNeverCookDirectories);
diff --git a/Engine/Source/Editor/UnrealEd/Private/Cooker/CookTypes.h b/Engine/Source/Editor/UnrealEd/Private/Cooker/CookTypes.h
--- a/Engine/Source/Editor/UnrealEd/Private/Cooker/CookTypes.h
+++ b/Engine/Source/Editor/UnrealEd/Private/Cooker/CookTypes.h
@@ -64,6 +64,8 @@ class TFastPointerSet : public TSet<KeyType, TFastPointerSetKeyFuncs<KeyType>, S
 namespace UE::Cook
 {
 	struct FPackageData;
+	struct FPlatformData;
+
 	/** A function that is called when a requested package finishes cooking (when successful, failed, or skipped) */
 	typedef TUniqueFunction<void(FPackageData*)> FCompletionCallback;
 
@@ -280,6 +282,35 @@ public:
 	TArray<FString>					AllCulturesToCook;
 };
 
+/** Helper struct for FBeginCookContext; holds the context data for each platform being cooked */
+struct FBeginCookContextPlatform
+{
+	ITargetPlatform* TargetPlatform = nullptr;
+	UE::Cook::FPlatformData* PlatformData = nullptr;
+	TMap<FName, FString> CurrentCookSettings;
+
+	/** If true, we are deleting all old results from disk and rebuilding every package. If false, we are building iteratively. */
+	bool bFullBuild = false;
+	/** If true, a cook has already been run in the current process and we still have results from it. */
+	bool bHasMemoryResults = false;
+	/** If true, we should delete the in-memory results from an earlier cook in the same process, if we have any. */
+	bool bClearMemoryResults = false;
+	/** If true, we should load results that previous cooks left on disk into the current cook's results; this is one way to cook iteratively. */
+	bool bPopulateMemoryResultsFromDiskResults = false;
+	/** If true we are cooking iteratively, from results in a shared build (e.g. from buildfarm) rather than from our previous cook. */
+	bool bIterateSharedBuild = false;
+};
+
+/** Data held on the stack and shared with multiple subfunctions when running StartCookByTheBook or StartCookOnTheFly */
+struct FBeginCookContext
+{
+	/** List of the platforms we are building, with startup context data about each one */
+	TArray<FBeginCookContextPlatform> PlatformContexts;
+	/** The list of platforms by themselves, for passing to functions that need just a list of platforms */
+	TArray<ITargetPlatform*> TargetPlatforms;
+};
+
+
 void LogCookerMessage(const FString& MessageText, EMessageSeverity::Type Severity);
 LLM_DECLARE_TAG(Cooker);
 
diff --git a/Engine/Source/Editor/UnrealEd/Private/Cooker/PackageTracker.cpp b/Engine/Source/Editor/UnrealEd/Private/Cooker/PackageTracker.cpp
--- a/Engine/Source/Editor/UnrealEd/Private/Cooker/PackageTracker.cpp
+++ b/Engine/Source/Editor/UnrealEd/Private/Cooker/PackageTracker.cpp
@@ -89,7 +89,15 @@ namespace UE::Cook
 
 	TMap<UPackage*, FInstigator> FPackageTracker::GetNewPackages()
 	{
-		return MoveTemp(NewPackages);
+		TMap<UPackage*, FInstigator> Result = MoveTemp(NewPackages);
+		NewPackages.Reset();
+		bHasBeenConsumed = true;
+		return Result;
+	}
+
+	bool FPackageTracker::HasBeenConsumed() const
+	{
+		return bHasBeenConsumed;
 	}
 
 	void FPackageTracker::NotifyUObjectCreated(const class UObjectBase* Object, int32 Index)
diff --git a/Engine/Source/Editor/UnrealEd/Private/Cooker/PackageTracker.h b/Engine/Source/Editor/UnrealEd/Private/Cooker/PackageTracker.h
--- a/Engine/Source/Editor/UnrealEd/Private/Cooker/PackageTracker.h
+++ b/Engine/Source/Editor/UnrealEd/Private/Cooker/PackageTracker.h
@@ -174,9 +174,12 @@ namespace Cook
 		FPackageTracker(FPackageDatas& InPackageDatas);
 		~FPackageTracker();
 
-		/* Returns all packages that have been loaded since the last time GetNewPackages was called */
+		/** Returns all packages that have been loaded since the last time GetNewPackages was called */
 		TMap<UPackage*, FInstigator> GetNewPackages();
 
+		/** Report whether GetNewPackages has ever been called */
+		bool HasBeenConsumed() const;
+
 		virtual void NotifyUObjectCreated(const class UObjectBase* Object, int32 Index) override;
 		virtual void NotifyUObjectDeleted(const class UObjectBase* Object, int32 Index) override;
 		virtual void OnUObjectArrayShutdown() override;
@@ -205,6 +208,8 @@ namespace Cook
 		FThreadSafeSet<FName> NeverCookPackageList;
 		FThreadSafeSet<FName> UncookedEditorOnlyPackages; // set of packages that have been rejected due to being referenced by editor-only properties
 		TFastPointerMap<const ITargetPlatform*, TSet<FName>> PlatformSpecificNeverCookPackages;
+
+		bool bHasBeenConsumed = false;
 	}; 
 }
 }
