CookOnTheFlyServer - StartCookByTheBook refactoring part 2.
1) Remove GenerateAssetRegistry, it is no longer needed.
2) Factor shadercode out of StartCookByTheBook into a subfunction and move the portion that occurs after cleaning the sandbox to the end of StartCookByTheBook to give its asynchronous operations more time to complete.
3) Move CancelCookByTheBook out of the scoped cooktimer so it can clear the cook timers.
4) Move CookOnTheFly status logs into UpdateDisplay.
#rnx

diff --git a/Engine/Source/Editor/UnrealEd/Classes/CookOnTheSide/CookOnTheFlyServer.h b/Engine/Source/Editor/UnrealEd/Classes/CookOnTheSide/CookOnTheFlyServer.h
--- a/Engine/Source/Editor/UnrealEd/Classes/CookOnTheSide/CookOnTheFlyServer.h
+++ b/Engine/Source/Editor/UnrealEd/Classes/CookOnTheSide/CookOnTheFlyServer.h
@@ -49,7 +49,7 @@ enum class ECookInitializationFlags
 	IncludeServerMaps =							0x00000080, // should we include the server maps when cooking
 	UseSerializationForPackageDependencies =	0x00000100, // should we use the serialization code path for generating package dependencies (old method will be deprecated)
 	BuildDDCInBackground =						0x00000200, // build ddc content in background while the editor is running (only valid for modes which are in editor IsCookingInEditor())
-	GeneratedAssetRegistry =					0x00000400, // have we generated asset registry yet
+	// unused =									0x00000400,
 	OutputVerboseCookerWarnings =				0x00000800, // output additional cooker warnings about content issues
 	EnablePartialGC =							0x00001000, // mark up with an object flag objects which are in packages which we are about to use or in the middle of using, this means we can gc more often but only gc stuff which we have finished with
 	TestCook =									0x00002000, // test the cooker garbage collection process and cooking (cooker will never end just keep testing).
@@ -713,6 +713,9 @@ private:
 
 	bool IsUsingLegacyCookOnTheFlyScheduling() const;
 
+	/** Update accumulators of editor data and consume editor changes since the last cook when starting cooks in the editor. */
+	void BeginCookEditorSystems();
+
 	//////////////////////////////////////////////////////////////////////////
 	// cook by the book specific functions
 	/** Construct the on-stack data for StartCook functions, based on the arguments to StartCookByTheBook */
@@ -746,10 +749,10 @@ private:
 	/** Return the name to use for the project's global shader library */
 	FString GetProjectShaderLibraryName() const;
 
-	/**
-	* Invokes the necessary FShaderCodeLibrary functions to start cooking the shader code library.
-	*/
-	void InitShaderCodeLibrary(void);
+	/** Invokes the necessary FShaderCodeLibrary functions to start cooking the shader code library. */
+	void BeginCookStartShaderCodeLibrary(FBeginCookContext& BeginContext);
+	/** Finishes async operations from ShaderCodeLibraryBeginCookStart, saves initial data to disk, and opens the library for the rest of the cook */
+	void BeginCookFinishShaderCodeLibrary(FBeginCookContext& BeginContext);
     
 	/**
 	 * Opens Global shader library. Global shaderlib isn't split into chunks nor associated with the assets, so it a special case
@@ -1091,8 +1094,6 @@ private:
 	 */
 	void PopulateCookedPackages(const TConstArrayView<const ITargetPlatform*> TargetPlatforms);
 
-	/** Loads the platform-independent asset registry for use by the cooker */
-	void GenerateAssetRegistry();
 	/** Waits for the AssetRegistry to complete so that we know any missing assets are missing on disk */
 	void BlockOnAssetRegistry();
 	/** Setup necessary only once for CookOnTheFly, but that are not required until the first request. */
diff --git a/Engine/Source/Editor/UnrealEd/Private/CookOnTheFlyServer.cpp b/Engine/Source/Editor/UnrealEd/Private/CookOnTheFlyServer.cpp
--- a/Engine/Source/Editor/UnrealEd/Private/CookOnTheFlyServer.cpp
+++ b/Engine/Source/Editor/UnrealEd/Private/CookOnTheFlyServer.cpp
@@ -536,7 +536,6 @@ bool UCookOnTheFlyServer::StartCookOnTheFly(FCookOnTheFlyOptions InCookOnTheFlyO
 #endif
 
 	CookOnTheFlyOptions = MoveTemp(InCookOnTheFlyOptions);
-	GenerateAssetRegistry();
 	CreateSandboxFile();
 	CookOnTheFlyServerInterface = MakeUnique<UCookOnTheFlyServer::FCookOnTheFlyServerInterface>(*this);
 	ExternalRequests->CookRequestEvent = FPlatformProcess::GetSynchEventFromPool();
@@ -579,6 +578,7 @@ bool UCookOnTheFlyServer::StartCookOnTheFly(FCookOnTheFlyOptions InCookOnTheFlyO
 		ServerOptions.Port = CookOnTheFlyOptions.bBindAnyPort ? 0 : -1;
 		CookOnTheFlyRequestManager = UE::Cook::MakeNetworkFileCookOnTheFlyRequestManager(*CookOnTheFlyServerInterface, MoveTemp(ServerOptions));
 	}
+	BeginCookEditorSystems();
 
 	const bool bInitialized = CookOnTheFlyRequestManager->Initialize();
 	return bInitialized;
@@ -613,8 +613,6 @@ void UCookOnTheFlyServer::InitializeShadersForCookOnTheFly(const TArrayView<ITar
 
 void UCookOnTheFlyServer::AddCookOnTheFlyPlatformFromGameThread(ITargetPlatform* TargetPlatform)
 {
-	check(!!(CookFlags & ECookInitializationFlags::GeneratedAssetRegistry)); // GenerateAssetRegistry should have been called in StartNetworkFileServer
-
 	UE::Cook::FPlatformData* PlatformData = PlatformManager->GetPlatformData(TargetPlatform);
 	check(PlatformData != nullptr); // should have been checked by the caller
 	if (PlatformData->bIsSandboxInitialized)
@@ -624,12 +622,12 @@ void UCookOnTheFlyServer::AddCookOnTheFlyPlatformFromGameThread(ITargetPlatform*
 
 	FBeginCookContext BeginContext = CreateAddPlatformContext(TargetPlatform);
 
-	RefreshPlatformAssetRegistries(BeginContext.TargetPlatforms);
 	FindOrCreateSaveContexts(BeginContext.TargetPlatforms);
 	SetBeginCookIterativeFlags(BeginContext);
 	BeginCookSandbox(BeginContext);
 	InitializeTargetPlatforms(BeginContext.TargetPlatforms);
 	InitializeShadersForCookOnTheFly(BeginContext.TargetPlatforms);
+	RefreshPlatformAssetRegistries(BeginContext.TargetPlatforms);
 
 	// When cooking on the fly the full registry is saved at the beginning
 	// in cook by the book asset registry is saved after the cook is finished
@@ -1202,9 +1200,10 @@ uint32 UCookOnTheFlyServer::TickCookOnTheSide(const float TimeSlice, uint32 &Coo
 
 	UE::Cook::FTickStackData StackData(TimeSlice, IsRealtimeMode(), TickFlags, InMaxNumPackagesToSave);
 	bool bCookComplete = false;
+	bool bCookCancelled = false;
 
 	{
-		UE_SCOPED_HIERARCHICAL_COOKTIMER(TickCookOnTheSide); // Make sure no UE_SCOPED_HIERARCHICAL_COOKTIMERs are around CookByTheBookFinishes, as that function deletes memory for them
+		UE_SCOPED_HIERARCHICAL_COOKTIMER(TickCookOnTheSide); // Make sure no UE_SCOPED_HIERARCHICAL_COOKTIMERs are around CookByTheBookFinishes or CancelCookByTheBook, as those functions delete memory for them
 
 		bool bContinueTick = true;
 		while (bContinueTick && (!IsEngineExitRequested() || CurrentCookMode == ECookMode::CookByTheBook))
@@ -1264,7 +1263,7 @@ uint32 UCookOnTheFlyServer::TickCookOnTheSide(const float TimeSlice, uint32 &Coo
 				bCookComplete = true;
 				break;
 			case ECookAction::Cancel:
-				CancelCookByTheBook();
+				bCookCancelled = true;
 				bContinueTick = false;
 				break;
 			default:
@@ -1274,21 +1273,9 @@ uint32 UCookOnTheFlyServer::TickCookOnTheSide(const float TimeSlice, uint32 &Coo
 		}
 	}
 
-	if (IsCookOnTheFlyMode() && (IsCookingInEditor() == false))
+	if (bCookCancelled)
 	{
-		static uint64 LastNumLoadedAndSaved = 0;
-		static int32 TickCounter = 0;
-		++TickCounter;
-		// dump stats every 50 ticks or so, but only if there is new data
-		if (TickCounter > 50)
-		{
-			if (StatLoadedPackageCount + StatSavedPackageCount != LastNumLoadedAndSaved)
-			{
-				LastNumLoadedAndSaved = StatLoadedPackageCount + StatSavedPackageCount;
-				DumpStats();
-			}
-			TickCounter = 0;
-		}
+		CancelCookByTheBook();
 	}
 
 	if (CookByTheBookOptions)
@@ -1557,6 +1544,17 @@ void UCookOnTheFlyServer::UpdateDisplay(ECookTickFlags TickFlags, bool bForceDis
 			TEXT("Cook Diagnostics: OpenFileHandles=%d, VirtualMemory=%dMiB"),
 			OpenFileHandles, FPlatformMemory::GetStats().UsedVirtual / 1024 / 1024);
 		LastDiagnosticsDisplayTime = CurrentTime;
+
+		if (IsCookOnTheFlyMode() && (IsCookingInEditor() == false))
+		{
+			// Dump stats in CookOnTheFly, but only if there is new data
+			static uint64 LastNumLoadedAndSaved = 0;
+			if (StatLoadedPackageCount + StatSavedPackageCount != LastNumLoadedAndSaved)
+			{
+				LastNumLoadedAndSaved = StatLoadedPackageCount + StatSavedPackageCount;
+				DumpStats();
+			}
+		}
 	}
 }
 
@@ -5972,55 +5970,14 @@ void DumpAssetRegistryForCooker(IAssetRegistry* AssetRegistry)
 #endif
 
 
-void UCookOnTheFlyServer::GenerateAssetRegistry()
-{
-	TRACE_CPUPROFILER_EVENT_SCOPE(UCookOnTheFlyServer::GenerateAssetRegistry);
-
-	// Mark package as dirty for the last ones saved
-	for (FName AssetFilename : ModifiedAssetFilenames)
-	{
-		const FString AssetPathOnDisk = AssetFilename.ToString();
-		if (FPaths::FileExists(AssetPathOnDisk))
-		{
-			const FString PackageName = FPackageName::FilenameToLongPackageName(AssetPathOnDisk);
-			FSoftObjectPath SoftPackage(PackageName);
-			if (UPackage* Package = Cast<UPackage>(SoftPackage.ResolveObject()))
-			{
-				MarkPackageDirtyForCooker(Package, true);
-			}
-		}
-	}
-
-	if (!!(CookFlags & ECookInitializationFlags::GeneratedAssetRegistry))
-	{
-		UE_LOG(LogCook, Display, TEXT("Updating asset registry"));
-
-		// Force a rescan of modified package files
-		TArray<FString> ModifiedPackageFileList;
-
-		for (FName ModifiedPackage : ModifiedAssetFilenames)
-		{
-			ModifiedPackageFileList.Add(ModifiedPackage.ToString());
-		}
-
-		AssetRegistry->ScanModifiedAssetFiles(ModifiedPackageFileList);
-	}
-	else
-	{
-		CookFlags |= ECookInitializationFlags::GeneratedAssetRegistry;
-		UE_LOG(LogCook, Display, TEXT("Creating asset registry"));
-
-		ModifiedAssetFilenames.Reset();
-		UE::Cook::FPackageDatas::OnAssetRegistryGenerated(*AssetRegistry);
-	}
-}
-
 void UCookOnTheFlyServer::BlockOnAssetRegistry()
 {
 	if (!bFirstCookInThisProcess)
 	{
 		return;
 	}
+	TRACE_CPUPROFILER_EVENT_SCOPE(UCookOnTheFlyServer::BlockOnAssetRegistry);
+	UE_LOG(LogCook, Display, TEXT("Waiting for Asset Registry"));
 	if (ShouldPopulateFullAssetRegistry())
 	{
 		// Trigger or waitfor completion the primary AssetRegistry scan.
@@ -6043,6 +6000,7 @@ void UCookOnTheFlyServer::BlockOnAssetRegistry()
 		ScanPaths.Add(FString::Printf(TEXT("/%s/"), *CookByTheBookOptions->DlcName));
 		AssetRegistry->ScanPathsSynchronous(ScanPaths);
 	}
+	UE::Cook::FPackageDatas::OnAssetRegistryGenerated(*AssetRegistry);
 
 #if ASSET_REGISTRY_STATE_DUMPING_ENABLED
 	if (FParse::Param(FCommandLine::Get(), TEXT("DumpAssetRegistry")))
@@ -6709,7 +6667,7 @@ const FString UCookOnTheFlyServer::GetCookedAssetRegistryFilename(const FString&
 	return CookedAssetRegistryFilename;
 }
 
-void UCookOnTheFlyServer::InitShaderCodeLibrary(void)
+void UCookOnTheFlyServer::BeginCookStartShaderCodeLibrary(FBeginCookContext& BeginContext)
 {
     const UProjectPackagingSettings* const PackagingSettings = GetDefault<UProjectPackagingSettings>();
 	bool const bCacheShaderLibraries = IsUsingShaderCodeLibrary();
@@ -6722,7 +6680,7 @@ void UCookOnTheFlyServer::InitShaderCodeLibrary(void)
 		GConfig->GetBool(TEXT("DevOptions.Shaders"), TEXT("NeedsShaderStableKeys"), bAllPlatformsNeedStableKeys, GEngineIni);
 		GConfig->GetBool(TEXT("DevOptions.Shaders"), TEXT("bNeedsShaderStableKeys"), bAllPlatformsNeedStableKeys, GEngineIni);
 
-        for (const ITargetPlatform* TargetPlatform : PlatformManager->GetSessionPlatforms())
+        for (const ITargetPlatform* TargetPlatform : BeginContext.TargetPlatforms)
         {
 			// Find out if this platform requires stable shader keys, by reading the platform setting file.
 			// Stable shader keys are needed if we are going to create a PSO cache.
@@ -6755,6 +6713,29 @@ void UCookOnTheFlyServer::InitShaderCodeLibrary(void)
 			}
         }
     }
+
+	CleanShaderCodeLibraries();
+}
+
+void UCookOnTheFlyServer::BeginCookFinishShaderCodeLibrary(FBeginCookContext& BeginContext)
+{
+	check(BeginContext.StartupOptions && CookByTheBookOptions); // CookByTheBook only for now
+	// don't resave the global shader map files in dlc
+	if (!IsCookingDLC() && !EnumHasAnyFlags(BeginContext.StartupOptions->CookOptions, ECookByTheBookOptions::ForceDisableSaveGlobalShaders))
+	{
+		OpenGlobalShaderLibrary();
+
+		SaveGlobalShaderMapFiles(BeginContext.TargetPlatforms);
+
+		SaveAndCloseGlobalShaderLibrary();
+	}
+
+	// Open the shader code library for the current project or the current DLC pack, depending on which we are cooking
+	FString LibraryName = GetProjectShaderLibraryName();
+	if (LibraryName.Len() > 0)
+	{
+		OpenShaderLibrary(LibraryName);
+	}
 }
 
 
@@ -7007,14 +6988,13 @@ void UCookOnTheFlyServer::SaveShaderLibrary(const ITargetPlatform* TargetPlatfor
 void UCookOnTheFlyServer::CleanShaderCodeLibraries()
 {
 	const UProjectPackagingSettings* const PackagingSettings = GetDefault<UProjectPackagingSettings>();
-    bool const bCacheShaderLibraries = IsUsingShaderCodeLibrary();
-	ITargetPlatformManagerModule& TPM = GetTargetPlatformManagerRef();
-	bool bIterativeCook = IsCookFlagSet(ECookInitializationFlags::Iterative) ||	PackageDatas->GetNumCooked() != 0;
+	bool const bCacheShaderLibraries = IsUsingShaderCodeLibrary();
 
-	// If not iterative then clean up our temporary files
-	if (bCacheShaderLibraries && PackagingSettings->bShareMaterialShaderCode && !bIterativeCook)
+	for (const ITargetPlatform* TargetPlatform : PlatformManager->GetSessionPlatforms())
 	{
-		for (const ITargetPlatform* TargetPlatform : PlatformManager->GetSessionPlatforms())
+		UE::Cook::FPlatformData* PlatformData = PlatformManager->GetPlatformData(TargetPlatform);
+		// If this is a full non-iterative build then clean up our temporary files
+		if (bCacheShaderLibraries && PackagingSettings->bShareMaterialShaderCode && PlatformData->bFullBuild)
 		{
 			TArray<FName> ShaderFormats;
 			TargetPlatform->GetAllTargetedShaderFormats(ShaderFormats);
@@ -7959,7 +7939,6 @@ void UCookOnTheFlyServer::StartCookByTheBook( const FCookByTheBookStartupOptions
 	COOK_STAT(UE::SavePackageUtilities::ResetCookStats());
 	FBeginCookContext BeginContext = SetCookByTheBookOptions(CookByTheBookStartupOptions);
 
-	GenerateAssetRegistry();
 	BlockOnAssetRegistry();
 	CreateSandboxFile();
 	SetBeginCookConfigSettings();
@@ -8070,9 +8049,7 @@ void UCookOnTheFlyServer::StartCookByTheBook( const FCookByTheBookStartupOptions
 		}
 	}
 	
-	// start shader code library cooking
-	InitShaderCodeLibrary();
-	CleanShaderCodeLibraries();
+	BeginCookStartShaderCodeLibrary(BeginContext); // start shader code library cooking asynchronously; we block on it later
 	
 	if ( IsCookingDLC() )
 	{
@@ -8329,25 +8306,7 @@ void UCookOnTheFlyServer::StartCookByTheBook( const FCookByTheBookStartupOptions
 	{
 		BeginCookSandbox(BeginContext);
 	}
-
-	// don't resave the global shader map files in dlc
-	if (!IsCookingDLC() && !(CookByTheBookStartupOptions.CookOptions & ECookByTheBookOptions::ForceDisableSaveGlobalShaders))
-	{
-		OpenGlobalShaderLibrary();
-
-		SaveGlobalShaderMapFiles(TargetPlatforms);
-
-		SaveAndCloseGlobalShaderLibrary();
-	}
-
-	// Open the shader code library for the current project or the current DLC pack, depending on which we are cooking
-	{
-		FString LibraryName = GetProjectShaderLibraryName();
-		if (LibraryName.Len() > 0)
-		{
-			OpenShaderLibrary(LibraryName);
-		}
-	}
+	BeginCookEditorSystems();
 
 	if (bHybridIterativeDebug)
 	{
@@ -8361,12 +8320,15 @@ void UCookOnTheFlyServer::StartCookByTheBook( const FCookByTheBookStartupOptions
 	}
 
 	InitializePollables();
+
+	BeginCookFinishShaderCodeLibrary(BeginContext);
 }
 
 FBeginCookContext UCookOnTheFlyServer::SetCookByTheBookOptions(const FCookByTheBookStartupOptions& StartupOptions)
 {
 	FBeginCookContext BeginContext;
 
+	BeginContext.StartupOptions = &StartupOptions;
 	BeginContext.TargetPlatforms = StartupOptions.TargetPlatforms;
 	Algo::Sort(BeginContext.TargetPlatforms);
 	BeginContext.TargetPlatforms.SetNum(Algo::Unique(BeginContext.TargetPlatforms));
@@ -8408,6 +8370,27 @@ void UCookOnTheFlyServer::SelectSessionPlatforms(FBeginCookContext& BeginContext
 	}
 }
 
+void UCookOnTheFlyServer::BeginCookEditorSystems()
+{
+	if (!IsCookingInEditor())
+	{
+		return;
+	}
+
+	// Notify AssetRegistry to update itself for any saved packages
+	if (!bFirstCookInThisProcess)
+	{
+		// Force a rescan of modified package files
+		TArray<FString> ModifiedPackageFileList;
+		for (FName ModifiedPackage : ModifiedAssetFilenames)
+		{
+			ModifiedPackageFileList.Add(ModifiedPackage.ToString());
+		}
+		AssetRegistry->ScanModifiedAssetFiles(ModifiedPackageFileList);
+	}
+	ModifiedAssetFilenames.Empty();
+}
+
 TArray<FName> UCookOnTheFlyServer::GetNeverCookPackageFileNames(TArrayView<const FString> ExtraNeverCookDirectories)
 {
 	TArray<FString> NeverCookDirectories(ExtraNeverCookDirectories);
diff --git a/Engine/Source/Editor/UnrealEd/Private/Cooker/CookTypes.h b/Engine/Source/Editor/UnrealEd/Private/Cooker/CookTypes.h
--- a/Engine/Source/Editor/UnrealEd/Private/Cooker/CookTypes.h
+++ b/Engine/Source/Editor/UnrealEd/Private/Cooker/CookTypes.h
@@ -304,6 +304,7 @@ struct FBeginCookContextPlatform
 /** Data held on the stack and shared with multiple subfunctions when running StartCookByTheBook or StartCookOnTheFly */
 struct FBeginCookContext
 {
+	const UCookOnTheFlyServer::FCookByTheBookStartupOptions* StartupOptions = nullptr;
 	/** List of the platforms we are building, with startup context data about each one */
 	TArray<FBeginCookContextPlatform> PlatformContexts;
 	/** The list of platforms by themselves, for passing to functions that need just a list of platforms */
