Fix a nullptr issue when running in -game -server using the unrealEditor.exe
GEditor might be null

diff --git a/Engine/Source/Editor/UnrealEd/Private/PackageTools.cpp b/Engine/Source/Editor/UnrealEd/Private/PackageTools.cpp
--- a/Engine/Source/Editor/UnrealEd/Private/PackageTools.cpp
+++ b/Engine/Source/Editor/UnrealEd/Private/PackageTools.cpp
@@ -379,42 +379,44 @@ UPackageTools::UPackageTools(const FObjectInitializer& ObjectInitializer)
 
 			OutErrorMessage = FText::Format( NSLOCTEXT("UnrealEd", "UnloadDirtyPackagesList", "The following assets have been modified and cannot be unloaded:{DirtyPackages}\nSaving these assets will allow them to be unloaded."), Args );
 		}
-
-		if (UWorld* EditorWorld = GEditor->GetEditorWorldContext().World())
+		if (GEditor)
 		{
-			// Is the currently loaded world being unloaded? If so, we just reset the current world.
-			// We also need to skip the build data package as that will also be destroyed by the call to CreateNewMapForEditing.
-			if (PackagesToUnload.Contains(EditorWorld->GetOutermost()))
+			if (UWorld* EditorWorld = GEditor->GetEditorWorldContext().World())
 			{
-				// Remove the world package from the unload list
-				PackagesToUnload.Remove(EditorWorld->GetOutermost());
-
-				// Remove the level build data package from the unload list as creating a new map will unload build data for the current world
-				for (int32 LevelIndex = 0; LevelIndex < EditorWorld->GetNumLevels(); ++LevelIndex)
+				// Is the currently loaded world being unloaded? If so, we just reset the current world.
+				// We also need to skip the build data package as that will also be destroyed by the call to CreateNewMapForEditing.
+				if (PackagesToUnload.Contains(EditorWorld->GetOutermost()))
 				{
-					ULevel* Level = EditorWorld->GetLevel(LevelIndex);
-					if (Level->MapBuildData)
+					// Remove the world package from the unload list
+					PackagesToUnload.Remove(EditorWorld->GetOutermost());
+
+					// Remove the level build data package from the unload list as creating a new map will unload build data for the current world
+					for (int32 LevelIndex = 0; LevelIndex < EditorWorld->GetNumLevels(); ++LevelIndex)
 					{
-						PackagesToUnload.Remove(Level->MapBuildData->GetOutermost());
+						ULevel* Level = EditorWorld->GetLevel(LevelIndex);
+						if (Level->MapBuildData)
+						{
+							PackagesToUnload.Remove(Level->MapBuildData->GetOutermost());
+						}
 					}
-				}
 
-				// Remove any streaming levels from the unload list as creating a new map will unload streaming levels for the current world
-				for (ULevelStreaming* EditorStreamingLevel : EditorWorld->GetStreamingLevels())
-				{
-					if (EditorStreamingLevel->IsLevelLoaded())
+					// Remove any streaming levels from the unload list as creating a new map will unload streaming levels for the current world
+					for (ULevelStreaming* EditorStreamingLevel : EditorWorld->GetStreamingLevels())
 					{
-						UPackage* EditorStreamingLevelPackage = EditorStreamingLevel->GetLoadedLevel()->GetOutermost();
-						PackagesToUnload.Remove(EditorStreamingLevelPackage);
+						if (EditorStreamingLevel->IsLevelLoaded())
+						{
+							UPackage* EditorStreamingLevelPackage = EditorStreamingLevel->GetLoadedLevel()->GetOutermost();
+							PackagesToUnload.Remove(EditorStreamingLevelPackage);
+						}
 					}
-				}
 
-				// Unload the current world
-				GEditor->CreateNewMapForEditing();
+					// Unload the current world
+					GEditor->CreateNewMapForEditing();
+				}
 			}
 		}
 
-		if (PackagesToUnload.Num() > 0)
+		if (PackagesToUnload.Num() > 0 && GEditor)
 		{
 			const FScopedBusyCursor BusyCursor;
 
