Fixed 'Create' button being disabled when creating a new animation blueprint when switching from the template tab
Fixed 'Parent Class' extendable area extending outside of the dialog box

diff --git a/Engine/Source/Editor/UnrealEd/Private/Factories/AnimBlueprintFactory.cpp b/Engine/Source/Editor/UnrealEd/Private/Factories/AnimBlueprintFactory.cpp
--- a/Engine/Source/Editor/UnrealEd/Private/Factories/AnimBlueprintFactory.cpp
+++ b/Engine/Source/Editor/UnrealEd/Private/Factories/AnimBlueprintFactory.cpp
@@ -107,10 +107,6 @@ public:
 									.OnValueChanged_Lambda([this](bool bInNewValue)
 									{
 										bTemplate = bInNewValue;
-										if(bTemplate)
-										{
-											TargetSkeleton = FAssetData(); 
-										}
 										RefreshSkeletonPicker();
 									})
 									.Value_Lambda([this](){ return bTemplate; })	
@@ -208,8 +204,7 @@ public:
 
 		TSharedRef<SWindow> Window = SNew(SWindow)
 		.Title( LOCTEXT("CreateAnimBlueprintOptions", "Create Animation Blueprint") )
-		.ClientSize(FVector2D(500, 500))
-		.SizingRule(ESizingRule::FixedSize)
+		.SizingRule(ESizingRule::Autosized)
 		.SupportsMinimize(false)
 		.SupportsMaximize(false)
 		[
@@ -270,6 +265,7 @@ private:
 			.ToolTipText(LOCTEXT("ParentClass_Tooltip", "Optionally choose a parent class for your Animation Blueprint"))
 			.Padding(10.0f)
 			.InitiallyCollapsed(true)
+			.MaxHeight(200)
 			.HeaderContent()
 			[
 				SNew(SBox)
@@ -347,6 +343,11 @@ private:
 	/** Handler for when ok is clicked */
 	FReply OkClicked()
 	{
+		if (bTemplate)
+		{
+			TargetSkeleton = FAssetData();
+		}
+
 		if ( AnimBlueprintFactory.IsValid() )
 		{
 			AnimBlueprintFactory->BlueprintType = BPTYPE_Normal;
