EditorDomain saving and UWorld: Add a tag that informs editor domain saves to keep the UWorld initialized, because UEditorEngine::Map_Load requires that it either be not-yet-initialized or that it is still fully initialized.
#rnx

diff --git a/Engine/Source/Editor/UnrealEd/Private/EditorEngine.cpp b/Engine/Source/Editor/UnrealEd/Private/EditorEngine.cpp
--- a/Engine/Source/Editor/UnrealEd/Private/EditorEngine.cpp
+++ b/Engine/Source/Editor/UnrealEd/Private/EditorEngine.cpp
@@ -4375,6 +4375,20 @@ FSavePackageResultStruct UEditorEngine::Save(UPackage* InOuter, UObject* InAsset
 		{
 			bInitializedPhysicsSceneForSave = InitializePhysicsSceneForSaveIfNecessary(World, bForceInitializedWorld);
 
+			// bForceInitialized=true marks that We Called InitWorld and need to call CleanupWorld
+			// but if we are saving during Load and the caller of LoadPackage requested that we keep it initialized, we should keep it initialized.
+			if (bForceInitializedWorld)
+			{
+				// TODO: If we ever add a way to set bInitialized=false on a UWorld, then in future saves after setting it back to false
+				// we will still parse the KeepInitializedDuringLoadTag here and not CleanupWorld. If we set bInitialized=false on
+				// a UWorld, we need to clear KeepInitializedDuringLoadTag from its InstancingContext.
+				FLinkerLoad* LinkerLoad = World->GetPackage()->GetLinker();
+				if (bForceInitializedWorld && LinkerLoad && LinkerLoad->GetInstancingContext().HasTag(UWorld::KeepInitializedDuringLoadTag))
+				{
+					bForceInitializedWorld = false;
+				}
+			}
+
 			PRAGMA_DISABLE_DEPRECATION_WARNINGS;
 			OnPreSaveWorld(ObjectSaveContext.SaveFlags, World);
 			PRAGMA_ENABLE_DEPRECATION_WARNINGS;
diff --git a/Engine/Source/Editor/UnrealEd/Private/EditorServer.cpp b/Engine/Source/Editor/UnrealEd/Private/EditorServer.cpp
--- a/Engine/Source/Editor/UnrealEd/Private/EditorServer.cpp
+++ b/Engine/Source/Editor/UnrealEd/Private/EditorServer.cpp
@@ -2652,6 +2652,9 @@ bool UEditorEngine::Map_Load(const TCHAR* Str, FOutputDevice& Ar)
 				}
 
 				
+				FLinkerInstancingContext WorldPackageInstancingContext;
+				WorldPackageInstancingContext.AddTag(UWorld::KeepInitializedDuringLoadTag);
+
 				UPackage* WorldPackage;
 				// Load startup maps and templates into new outermost packages so that the Save function in the editor won't overwrite the original
 				if (bIsLoadingMapTemplate)
@@ -2668,7 +2671,7 @@ bool UEditorEngine::Map_Load(const TCHAR* Str, FOutputDevice& Ar)
 					//now load the map into the package created above
 					const FName WorldPackageFName = WorldPackage->GetFName();
 					UWorld::WorldTypePreLoadMap.FindOrAdd(WorldPackageFName) = EWorldType::Editor;
-					WorldPackage = LoadPackage( WorldPackage, *LongTempFname, LoadFlags );
+					WorldPackage = LoadPackage( WorldPackage, *LongTempFname, LoadFlags, nullptr /* InReaderOverride */, &WorldPackageInstancingContext);
 					WorldPackage->SetPackageFlags(PKG_NewlyCreated);
 					UWorld::WorldTypePreLoadMap.Remove(WorldPackageFName);
 				}
@@ -2691,7 +2694,7 @@ bool UEditorEngine::Map_Load(const TCHAR* Str, FOutputDevice& Ar)
 						//Load the map normally into a new package
 						const FName WorldPackageFName = FName(*LongTempFname);
 						UWorld::WorldTypePreLoadMap.FindOrAdd(WorldPackageFName) = EWorldType::Editor;
-						WorldPackage = LoadPackage( NULL, *LongTempFname, LoadFlags);
+						WorldPackage = LoadPackage( NULL, *LongTempFname, LoadFlags, nullptr /* InReaderOverride */, &WorldPackageInstancingContext);
 						UWorld::WorldTypePreLoadMap.Remove(WorldPackageFName);
 					}
 				}
@@ -2737,10 +2740,20 @@ bool UEditorEngine::Map_Load(const TCHAR* Str, FOutputDevice& Ar)
 				int32 FeatureLevelIndex = (int32)GMaxRHIFeatureLevel;
 				FParse::Value(Str, TEXT("FEATURELEVEL="), FeatureLevelIndex);
 				FeatureLevelIndex = FMath::Clamp(FeatureLevelIndex, 0, (int32)ERHIFeatureLevel::Num);
+				ERHIFeatureLevel::Type FeatureLevel = (ERHIFeatureLevel::Type)FeatureLevelIndex;
 
-				check(!World->bIsWorldInitialized);
-				World->FeatureLevel = (ERHIFeatureLevel::Type)FeatureLevelIndex;
-				World->InitWorld( GetEditorWorldInitializationValues() );
+				if (World->bIsWorldInitialized)
+				{
+					// We do not handle reinitializing. If the World was initialized during Load, we require that it is still fully initialized
+					// This should have been guaranteed by UWorld::KeepInitializedDuringLoadTag
+					check(World->IsInitializedAndNeedsCleanup());
+					World->ChangeFeatureLevel(FeatureLevel);
+				}
+				else
+				{
+					World->FeatureLevel = FeatureLevel;
+					World->InitWorld(GetEditorWorldInitializationValues());
+				}
 				
 				SlowTask.EnterProgressFrame(20, FText::Format( LOCTEXT( "LoadingMapStatus_Initializing", "Loading map: {0}... (Initializing world)" ), FText::FromString(MapFileName) ));
 				{
diff --git a/Engine/Source/Runtime/Engine/Classes/Engine/World.h b/Engine/Source/Runtime/Engine/Classes/Engine/World.h
--- a/Engine/Source/Runtime/Engine/Classes/Engine/World.h
+++ b/Engine/Source/Runtime/Engine/Classes/Engine/World.h
@@ -3083,6 +3083,18 @@ public:
 	 * Initializes the world, associates the persistent level and sets the proper zones.
 	 */
 	void InitWorld(const InitializationValues IVS = InitializationValues());
+#if WITH_EDITOR
+	/**
+	 * InitWorld usually has to be balanced with CleanupWorld. If the KeepInitializedDuringLoadTag LinkerInstancingContext tag is present,
+	 * operations that need to call InitWorld during the Load of the World's package should break that rule and not call CleanupWorld.
+	 */
+	static const FName KeepInitializedDuringLoadTag;
+	/**
+	 * Returns whether InitWorld has been called without yet calling CleanupWorld.
+	 * Note this is different from bIsWorldInitialized flag, which is set to true InitWorld and never set back to false.
+	 */
+	bool IsInitializedAndNeedsCleanup() const { return bInitializedAndNeedsCleanup; }
+#endif
 
 	/**
 	 * Initializes a newly created world.
diff --git a/Engine/Source/Runtime/Engine/Private/World.cpp b/Engine/Source/Runtime/Engine/Private/World.cpp
--- a/Engine/Source/Runtime/Engine/Private/World.cpp
+++ b/Engine/Source/Runtime/Engine/Private/World.cpp
@@ -1908,6 +1908,10 @@ void UWorld::InitWorld(const InitializationValues IVS)
 	BroadcastLevelsChanged();
 }
 
+#if WITH_EDITOR
+const FName UWorld::KeepInitializedDuringLoadTag(TEXT("KeepInitializedDuringLoadTag"));
+#endif
+
 void UWorld::ConditionallyCreateDefaultLevelCollections()
 {
 	LevelCollections.Reserve((int32)ELevelCollectionType::MAX);
