Added CookLabel, HordeJobId, HordeStepId, HordeBatchId support to analytics

diff --git a/Engine/Source/Editor/UnrealEd/Private/Commandlets/CookCommandlet.cpp b/Engine/Source/Editor/UnrealEd/Private/Commandlets/CookCommandlet.cpp
--- a/Engine/Source/Editor/UnrealEd/Private/Commandlets/CookCommandlet.cpp
+++ b/Engine/Source/Editor/UnrealEd/Private/Commandlets/CookCommandlet.cpp
@@ -78,6 +78,10 @@ namespace DetailedCookStats
 	double TickLoopProcessDeferredCommandsTimeSec = 0.0;
 	double TickLoopTickCommandletStatsTimeSec = 0.0;
 	double TickLoopFlushRenderingCommandsTimeSec = 0.0;
+	bool IsCookAll = false;
+	bool IsCookOnTheFly = false;
+	bool IsIterativeCook = false;
+	bool IsUnversioned = false;
 
 	FCookStatsManager::FAutoRegisterCallback RegisterCookStats([](FCookStatsManager::AddStatFuncRef AddStat)
 	{
@@ -103,6 +107,11 @@ namespace DetailedCookStats
 		ADD_COOK_STAT_FLT(" 0. 1. 8", TargetPlatforms);
 		ADD_COOK_STAT_FLT(" 0. 1. 9", CookProject);
 		ADD_COOK_STAT_FLT(" 0. 1. 10", CookCultures);
+		ADD_COOK_STAT_FLT(" 0. 1. 11", IsCookAll);
+		ADD_COOK_STAT_FLT(" 0. 1. 13", IsCookOnTheFly);
+		ADD_COOK_STAT_FLT(" 0. 1. 14", IsIterativeCook);
+		ADD_COOK_STAT_FLT(" 0. 1. 15", IsUnversioned);
+		
 		#undef ADD_COOK_STAT_FLT
 	});
 
@@ -110,6 +119,7 @@ namespace DetailedCookStats
 	{
 		if (FStudioAnalytics::IsAvailable())
 		{
+
 			// convert filtered stats directly to an analytics event
 			TArray<FAnalyticsEventAttribute> StatAttrs;
 
@@ -124,6 +134,15 @@ namespace DetailedCookStats
 				}
 			};
 
+			// Append any cook label specified on the command-line
+			FString Label;
+
+			if (FParse::Value(FCommandLine::Get(), TEXT("CookLabel="), Label))
+			{
+				// Append a Name/Value pair 
+				StatAttrs.Emplace("Label", Label);
+			}
+			
 			// Now actually grab the stats 
 			FCookStatsManager::LogCookStats(SendCookStatsToAnalytics);
 
@@ -645,6 +664,11 @@ int32 UCookCommandlet::Main(const FString& CmdLineParams)
 	bNoShaderCooking = bCookOnTheFly; // Do not cook any shaders into the shader maps. Always true if we are running w/ cook on the fly
 	bIgnoreIniSettingsOutOfDate = Switches.Contains(TEXT("IgnoreIniSettingsOutOfDate"));
 
+	COOK_STAT(DetailedCookStats::IsCookAll = bCookAll);
+	COOK_STAT(DetailedCookStats::IsCookOnTheFly = bCookOnTheFly);
+	COOK_STAT(DetailedCookStats::IsIterativeCook = bIterativeCooking);
+	COOK_STAT(DetailedCookStats::IsUnversioned = bUnversioned);
+
 	COOK_STAT(DetailedCookStats::CookProject = FApp::GetProjectName());
 
 	COOK_STAT(FParse::Value(*Params, TEXT("CookCultures="), DetailedCookStats::CookCultures));
