Fix regression crash on App exit because of Initialized Inactive worlds (with initialized subsystems) that haven't been cleaned up

#rnx

diff --git a/Engine/Source/Editor/UnrealEd/Classes/Editor/EditorEngine.h b/Engine/Source/Editor/UnrealEd/Classes/Editor/EditorEngine.h
--- a/Engine/Source/Editor/UnrealEd/Classes/Editor/EditorEngine.h
+++ b/Engine/Source/Editor/UnrealEd/Classes/Editor/EditorEngine.h
@@ -739,7 +739,6 @@ public:
 	void BroadcastObjectReimported(UObject* InObject);
 
 	//~ Begin UObject Interface.
-	virtual void BeginDestroy() override;
 	virtual void FinishDestroy() override;
 	virtual void PostEditChangeProperty(FPropertyChangedEvent& PropertyChangedEvent) override;
 	static void AddReferencedObjects(UObject* InThis, FReferenceCollector& Collector);
@@ -748,6 +747,7 @@ public:
 	//~ Begin UEngine Interface.
 public:
 	virtual void Init(IEngineLoop* InEngineLoop) override;
+	virtual void PreExit() override;
 	virtual float GetMaxTickRate(float DeltaTime, bool bAllowFrameRateSmoothing = true) const override;
 	virtual void Tick(float DeltaSeconds, bool bIdleMode) override;
 	virtual bool ShouldDrawBrushWireframe(AActor* InActor) override;
diff --git a/Engine/Source/Editor/UnrealEd/Private/EditorEngine.cpp b/Engine/Source/Editor/UnrealEd/Private/EditorEngine.cpp
--- a/Engine/Source/Editor/UnrealEd/Private/EditorEngine.cpp
+++ b/Engine/Source/Editor/UnrealEd/Private/EditorEngine.cpp
@@ -1238,21 +1238,31 @@ void UEditorEngine::BroadcastObjectReimported(UObject* InObject)
 	GetEditorSubsystem<UImportSubsystem>()->BroadcastAssetReimport(InObject);
 }
 
-void UEditorEngine::BeginDestroy()
+void UEditorEngine::PreExit()
 {
 	if (!HasAnyFlags(RF_ClassDefaultObject))
 	{
-		// UWorld::FinishDestroy clears GWorld if this == GWorld, so we need to do any cleanup on GWorld
-		// before any calls to FinishDestroy are made, in case UEditorEngine and UWorld are GC'd in the same
-		// pass, which happens on editor shutdown. Doing it here in BeginDestroy is early enough.
-		UWorld* World = GWorld;
-		if (World != NULL)
+		// Cleanup GWorld before final GC Purge
+		if (UWorld* World = GWorld)
 		{
 			World->ClearWorldComponents();
 			World->CleanupWorld();
 		}
+		
+		// Cleanup worlds that were initialized through UEditorEngine::InitializeNewlyCreatedInactiveWorld before final GC Purge
+		for (TObjectIterator<UWorld> It; It; ++It)
+		{
+			if (UWorld* World = *It; World && World->WorldType == EWorldType::Inactive && World->IsInitializedAndNeedsCleanup())
+			{
+				// GWorld shouldn't be an Inactive World
+				check(World != GWorld);
+				World->ClearWorldComponents();
+				World->CleanupWorld();
+			}
+		}
 	}
-	Super::BeginDestroy();
+
+	Super::PreExit();
 }
 
 void UEditorEngine::FinishDestroy()
