Minor analytics tweaks

diff --git a/Engine/Source/Editor/UnrealEd/Private/Commandlets/CookCommandlet.cpp b/Engine/Source/Editor/UnrealEd/Private/Commandlets/CookCommandlet.cpp
--- a/Engine/Source/Editor/UnrealEd/Private/Commandlets/CookCommandlet.cpp
+++ b/Engine/Source/Editor/UnrealEd/Private/Commandlets/CookCommandlet.cpp
@@ -137,11 +137,10 @@ namespace DetailedCookStats
 			// Append any cook label specified on the command-line
 			FString Label;
 
-			if (FParse::Value(FCommandLine::Get(), TEXT("CookLabel="), Label))
-			{
-				// Append a Name/Value pair 
-				StatAttrs.Emplace("Label", Label);
-			}
+			FParse::Value(FCommandLine::Get(), TEXT("CookLabel="), Label);
+			
+			// Append a Name/Value pair 
+			StatAttrs.Emplace("Label", Label);
 			
 			// Now actually grab the stats 
 			FCookStatsManager::LogCookStats(SendCookStatsToAnalytics);
