Minor tweak to cook label capture

diff --git a/Engine/Source/Editor/UnrealEd/Private/Commandlets/CookCommandlet.cpp b/Engine/Source/Editor/UnrealEd/Private/Commandlets/CookCommandlet.cpp
--- a/Engine/Source/Editor/UnrealEd/Private/Commandlets/CookCommandlet.cpp
+++ b/Engine/Source/Editor/UnrealEd/Private/Commandlets/CookCommandlet.cpp
@@ -60,6 +60,7 @@ namespace DetailedCookStats
 {
 	FString CookProject;
 	FString CookCultures;
+	FString CookLabel;
 	FString TargetPlatforms;
 	double CookWallTimeSec = 0.0;
 	double StartupWallTimeSec = 0.0;
@@ -83,6 +84,7 @@ namespace DetailedCookStats
 	bool IsIterativeCook = false;
 	bool IsUnversioned = false;
 
+
 	FCookStatsManager::FAutoRegisterCallback RegisterCookStats([](FCookStatsManager::AddStatFuncRef AddStat)
 	{
 		const FString StatName(TEXT("Cook.Profile"));
@@ -111,6 +113,7 @@ namespace DetailedCookStats
 		ADD_COOK_STAT_FLT(" 0. 1. 13", IsCookOnTheFly);
 		ADD_COOK_STAT_FLT(" 0. 1. 14", IsIterativeCook);
 		ADD_COOK_STAT_FLT(" 0. 1. 15", IsUnversioned);
+		ADD_COOK_STAT_FLT(" 0. 1. 16", CookLabel);
 		
 		#undef ADD_COOK_STAT_FLT
 	});
@@ -134,14 +137,6 @@ namespace DetailedCookStats
 				}
 			};
 
-			// Append any cook label specified on the command-line
-			FString Label;
-
-			FParse::Value(FCommandLine::Get(), TEXT("CookLabel="), Label);
-			
-			// Append a Name/Value pair 
-			StatAttrs.Emplace("Label", Label);
-			
 			// Now actually grab the stats 
 			FCookStatsManager::LogCookStats(SendCookStatsToAnalytics);
 
@@ -671,6 +666,7 @@ int32 UCookCommandlet::Main(const FString& CmdLineParams)
 	COOK_STAT(DetailedCookStats::CookProject = FApp::GetProjectName());
 
 	COOK_STAT(FParse::Value(*Params, TEXT("CookCultures="), DetailedCookStats::CookCultures));
+	COOK_STAT(FParse::Value(*Params, TEXT("CookLabel="), DetailedCookStats::CookLabel));
 	
 	ITargetPlatformManagerModule& TPM = GetTargetPlatformManagerRef();
 	if ( bCookOnTheFly )
