PlasticSCM: Fix "Delete" operation not getting called for source control other than Perforce (Contributed by [REDACTED_PERSON])

#9097

diff --git a/Engine/Source/Editor/UnrealEd/Private/ObjectTools.cpp b/Engine/Source/Editor/UnrealEd/Private/ObjectTools.cpp
--- a/Engine/Source/Editor/UnrealEd/Private/ObjectTools.cpp
+++ b/Engine/Source/Editor/UnrealEd/Private/ObjectTools.cpp
@@ -2387,7 +2387,7 @@ namespace ObjectTools
 						// The file was open for add and reverted, this leaves the file on disk so here we delete it
 						IFileManager::Get().Delete(*PackageFilename);
 					}
-					else if ((!bIsCheckedOut) && (!IFileManager::Get().IsReadOnly(*PackageFilename)))
+					else if (!bIsCheckedOut && SourceControlProvider.UsesLocalReadOnlyState() && !IFileManager::Get().IsReadOnly(*PackageFilename))
 					{
 						// The file was made writable and deleted, so we delete it on disk.
 						IFileManager::Get().Delete(*PackageFilename);
