Second pass on Editor Context

diff --git a/Engine/Source/Editor/DataLayerEditor/Private/DataLayer/DataLayerEditorSubsystem.cpp b/Engine/Source/Editor/DataLayerEditor/Private/DataLayer/DataLayerEditorSubsystem.cpp
--- a/Engine/Source/Editor/DataLayerEditor/Private/DataLayer/DataLayerEditorSubsystem.cpp
+++ b/Engine/Source/Editor/DataLayerEditor/Private/DataLayer/DataLayerEditorSubsystem.cpp
@@ -209,7 +209,7 @@ bool UDataLayerEditorSubsystem::GetActorEditorContextDisplayInfo(UWorld* InWorld
 	AWorldDataLayers* WorldDataLayer = InWorld->GetWorldDataLayers();
 	if (WorldDataLayer && !WorldDataLayer->GetActorEditorContextDataLayers().IsEmpty())
 	{
-		OutDiplayInfo.Title = TEXT("Current Data Layers");
+		OutDiplayInfo.Title = TEXT("Data Layers");
 		OutDiplayInfo.Brush = FEditorStyle::GetBrush(TEXT("DataLayer.Editor"));
 		return true;
 	}
@@ -231,10 +231,12 @@ TSharedRef<SWidget> UDataLayerEditorSubsystem::GetActorEditorContextWidget(UWorl
 				+ SHorizontalBox::Slot()
 				.AutoWidth()
 				.Padding(0.0f, 1.0f, 1.0f, 1.0f)
+				.VAlign(VAlign_Center)
 				[
 					SNew(SImage)
 					.ColorAndOpacity(DataLayer->GetDebugColor())
-					.Image(FAppStyle::Get().GetBrush("Level.ColorIcon"))
+					.Image(FAppStyle::Get().GetBrush("DataLayer.ColorIcon"))
+					.DesiredSizeOverride(FVector2D(8,8))
 				]
 				+ SHorizontalBox::Slot()
 				.AutoWidth()
diff --git a/Engine/Source/Editor/DataLayerEditor/Private/DataLayer/SDataLayerTreeLabel.cpp b/Engine/Source/Editor/DataLayerEditor/Private/DataLayer/SDataLayerTreeLabel.cpp
--- a/Engine/Source/Editor/DataLayerEditor/Private/DataLayer/SDataLayerTreeLabel.cpp
+++ b/Engine/Source/Editor/DataLayerEditor/Private/DataLayer/SDataLayerTreeLabel.cpp
@@ -19,6 +19,7 @@ void SDataLayerTreeLabel::Construct(const FArguments& InArgs, FDataLayerTreeItem
 	TreeItemPtr = StaticCastSharedRef<FDataLayerTreeItem>(DataLayerItem.AsShared());
 	DataLayerPtr = DataLayerItem.GetDataLayer();
 	HighlightText = SceneOutliner.GetFilterHighlightText();
+	bInEditingMode = false;
 
 	TSharedPtr<SInlineEditableTextBlock> InlineTextBlock;
 
@@ -35,6 +36,8 @@ void SDataLayerTreeLabel::Construct(const FArguments& InArgs, FDataLayerTreeItem
 		.ColorAndOpacity(this, &SDataLayerTreeLabel::GetForegroundColor)
 		.OnTextCommitted(this, &SDataLayerTreeLabel::OnLabelCommitted)
 		.OnVerifyTextChanged(this, &SDataLayerTreeLabel::OnVerifyItemLabelChanged)
+		.OnEnterEditingMode(this, &SDataLayerTreeLabel::OnEnterEditingMode)
+		.OnExitEditingMode(this, &SDataLayerTreeLabel::OnExitEditingMode)
 		.IsSelected(FIsSelected::CreateSP(&InRow, &STableRow<FSceneOutlinerTreeItemPtr>::IsSelectedExclusively))
 		.IsReadOnly_Lambda([Item = DataLayerItem.AsShared(), this]()
 		{
@@ -63,21 +66,6 @@ void SDataLayerTreeLabel::Construct(const FArguments& InArgs, FDataLayerTreeItem
 		SNew(SHorizontalBox)
 		+ SHorizontalBox::Slot()
 		.AutoWidth()
-		.HAlign(HAlign_Right)
-		.VAlign(VAlign_Center)
-		[
-			SNew(SBox)
-			.WidthOverride(FSceneOutlinerDefaultTreeItemMetrics::IconSize())
-			.HeightOverride(FSceneOutlinerDefaultTreeItemMetrics::IconSize())
-			[
-				SNew(SImage)
-				.Visibility_Lambda([this]() { return IsInActorEditorContext() ? EVisibility::Visible : EVisibility::Collapsed; })
-				.Image(FEditorStyle::Get().GetBrush(TEXT("SceneOutliner.MarkedAsCurrent")))
-				.ColorAndOpacity(FLinearColor::White)
-			]
-		]
-		+ SHorizontalBox::Slot()
-		.AutoWidth()
 		.VAlign(VAlign_Center)
 		.Padding(FSceneOutlinerDefaultTreeItemMetrics::IconPadding())
 		[
@@ -142,22 +130,30 @@ FText SDataLayerTreeLabel::GetDisplayText() const
 {
 	const UDataLayerInstance* DataLayer = DataLayerPtr.Get();
 	bool bIsDataLayerActive = false;
-	FText DataLayerRuntimeStateText = FText::GetEmpty();
-	if (DataLayer && DataLayer->IsRuntime() && DataLayer->GetWorld()->IsPlayInEditor())
+	FText SuffixText = FText::GetEmpty();
+	if (!bInEditingMode)
 	{
-		const UDataLayerSubsystem* DataLayerSubsystem = DataLayer->GetWorld()->GetSubsystem<UDataLayerSubsystem>();
-		DataLayerRuntimeStateText = FText::Format(LOCTEXT("DataLayerRuntimeState", " ({0})"), FTextStringHelper::CreateFromBuffer(GetDataLayerRuntimeStateName(DataLayerSubsystem->GetDataLayerEffectiveRuntimeState(DataLayer))));
+		if (DataLayer && DataLayer->IsRuntime() && DataLayer->GetWorld()->IsPlayInEditor())
+		{
+			const UDataLayerSubsystem* DataLayerSubsystem = DataLayer->GetWorld()->GetSubsystem<UDataLayerSubsystem>();
+			SuffixText = FText::Format(LOCTEXT("DataLayerRuntimeState", " ({0})"), FTextStringHelper::CreateFromBuffer(GetDataLayerRuntimeStateName(DataLayerSubsystem->GetDataLayerEffectiveRuntimeState(DataLayer))));
+		}
+		else if (IsInActorEditorContext())
+		{
+			SuffixText = FText(LOCTEXT("IsCurrentSuffix", " (Current)"));
+		}
 	}
-	
+
 	static const FText DataLayerDeleted = LOCTEXT("DataLayerLabelForMissingDataLayer", "(Deleted Data Layer)");
-	return DataLayer ? FText::Format(LOCTEXT("DataLayerDisplayText", "{0}{1}"), FText::FromString(DataLayer->GetDataLayerShortName()), DataLayerRuntimeStateText) : DataLayerDeleted;
+	return DataLayer ? FText::Format(LOCTEXT("DataLayerDisplayText", "{0}{1}"), FText::FromString(DataLayer->GetDataLayerShortName()), SuffixText) : DataLayerDeleted;
 }
 
 FText SDataLayerTreeLabel::GetTooltipText() const
 {
 	if (const FSceneOutlinerTreeItemPtr TreeItem = TreeItemPtr.Pin())
 	{
-		return FText::FromString(TreeItem->GetDisplayString());
+		FText Description = IsInActorEditorContext() ? LOCTEXT("DataLayerIsCurrentDescription", "This Data Layer is part of Current Data Layers. New actors will attempt to be added to this Data Layer.") : FText::GetEmpty();
+		return FText::Format(LOCTEXT("DataLayerTooltipText", "{0}\n{1}"), FText::FromString(TreeItem->GetDisplayString()), Description);
 	}
 
 	return FText();
@@ -242,6 +238,10 @@ FSlateColor SDataLayerTreeLabel::GetForegroundColor() const
 	{
 		return FAppStyle::Get().GetSlateColor("Colors.AccentBlue");
 	}
+	else if (IsInActorEditorContext())
+	{
+		return FAppStyle::Get().GetSlateColor("Colors.AccentGreen");
+	}
 	return FSlateColor::UseForeground();
 }
 
@@ -291,4 +291,14 @@ void SDataLayerTreeLabel::OnLabelCommitted(const FText& InLabel, ETextCommit::Ty
 	}
 }
 
+void SDataLayerTreeLabel::OnEnterEditingMode()
+{
+	bInEditingMode = true;
+}
+
+void SDataLayerTreeLabel::OnExitEditingMode()
+{
+	bInEditingMode = false;
+}
+
 #undef LOCTEXT_NAMESPACE 
\ No newline at end of file
diff --git a/Engine/Source/Editor/DataLayerEditor/Private/DataLayer/SDataLayerTreeLabel.h b/Engine/Source/Editor/DataLayerEditor/Private/DataLayer/SDataLayerTreeLabel.h
--- a/Engine/Source/Editor/DataLayerEditor/Private/DataLayer/SDataLayerTreeLabel.h
+++ b/Engine/Source/Editor/DataLayerEditor/Private/DataLayer/SDataLayerTreeLabel.h
@@ -34,8 +34,11 @@ private:
 	void OnLabelCommitted(const FText& InLabel, ETextCommit::Type InCommitInfo);
 	bool ShouldBeHighlighted() const;
 	bool IsInActorEditorContext() const;
+	void OnEnterEditingMode();
+	void OnExitEditingMode();
 
 	TWeakPtr<FDataLayerTreeItem> TreeItemPtr;
 	TWeakObjectPtr<UDataLayerInstance> DataLayerPtr;
 	TAttribute<FText> HighlightText;
+	bool bInEditingMode;
 };
\ No newline at end of file
diff --git a/Engine/Source/Editor/EditorStyle/Private/StarshipStyle.cpp b/Engine/Source/Editor/EditorStyle/Private/StarshipStyle.cpp
--- a/Engine/Source/Editor/EditorStyle/Private/StarshipStyle.cpp
+++ b/Engine/Source/Editor/EditorStyle/Private/StarshipStyle.cpp
@@ -1202,6 +1202,7 @@ void FStarshipEditorStyle::FStyle::SetupGeneralStyles()
 		);
 		Set("DataLayerBrowser.LabelFont", DEFAULT_FONT("Regular", 9));
 		Set("DataLayerBrowser.LabelFontBold", DEFAULT_FONT("Bold", 10));
+		Set("DataLayer.ColorIcon", new FSlateBoxBrush(NAME_None, 8.0f / 32.0f, FStyleColors::White));
 	}
 
 	// Derived Data
@@ -1234,7 +1235,6 @@ void FStarshipEditorStyle::FStyle::SetupGeneralStyles()
 		Set( "SceneOutliner.FolderClosed",  new CORE_IMAGE_BRUSH_SVG("Starship/Common/folder-closed", Icon16x16, FStyleColors::AccentFolder ) );
 		Set( "SceneOutliner.FolderOpen",    new CORE_IMAGE_BRUSH_SVG("Starship/Common/folder-open", Icon16x16, FStyleColors::AccentFolder ) );
 		Set( "SceneOutliner.World", 		new CORE_IMAGE_BRUSH_SVG("Starship/Common/world", Icon16x16 ) );
-		Set( "SceneOutliner.MarkedAsCurrent", new IMAGE_BRUSH("Icons/MarkedAsCurrent", Icon12x12));
 		Set( "SceneOutliner.ChangedItemHighlight", new FSlateRoundedBoxBrush( FStyleColors::White, 1.0) );
 
 		Set( "SceneOutliner.TableViewRow", FTableRowStyle(NormalTableRowStyle));
diff --git a/Engine/Source/Editor/LevelEditor/Private/LevelEditorSubsystem.cpp b/Engine/Source/Editor/LevelEditor/Private/LevelEditorSubsystem.cpp
--- a/Engine/Source/Editor/LevelEditor/Private/LevelEditorSubsystem.cpp
+++ b/Engine/Source/Editor/LevelEditor/Private/LevelEditorSubsystem.cpp
@@ -688,6 +688,7 @@ public:
 				// Current Level
 				SNew(SVerticalBox)
 				+ SVerticalBox::Slot()
+				.Padding(0.0f, 0.0f, 2.0f, 0.0f)
 				.AutoHeight()
 				[
 					SNew(STextBlock)
@@ -828,7 +829,7 @@ bool ULevelEditorSubsystem::GetActorEditorContextDisplayInfo(UWorld* InWorld, FA
 	const bool bIsVisible = InWorld && (InWorld->GetCurrentLevel()->OwningWorld->GetLevels().Num() > 1) && (!InWorld->IsPartitionedWorld() || (InWorld->GetCurrentLevel() != InWorld->PersistentLevel));
 	if (bIsVisible)
 	{
-		OutDiplayInfo.Title = TEXT("Current Level");
+		OutDiplayInfo.Title = TEXT("Level");
 		OutDiplayInfo.Brush = FSlateIconFinder::FindIconBrushForClass(UWorld::StaticClass());
 		return true;
 	}
diff --git a/Engine/Source/Editor/SceneOutliner/Private/ActorFolderTreeItem.cpp b/Engine/Source/Editor/SceneOutliner/Private/ActorFolderTreeItem.cpp
--- a/Engine/Source/Editor/SceneOutliner/Private/ActorFolderTreeItem.cpp
+++ b/Engine/Source/Editor/SceneOutliner/Private/ActorFolderTreeItem.cpp
@@ -35,10 +35,13 @@ struct SActorFolderTreeLabel : FSceneOutlinerCommonLabelData, public SCompoundWi
 
 		TSharedPtr<SInlineEditableTextBlock> InlineTextBlock = SNew(SInlineEditableTextBlock)
 			.Text(this, &SActorFolderTreeLabel::GetDisplayText)
+			.ToolTipText(this, &SActorFolderTreeLabel::GetTooltipText)
 			.HighlightText(SceneOutliner.GetFilterHighlightText())
 			.ColorAndOpacity(this, &SActorFolderTreeLabel::GetForegroundColor)
 			.OnTextCommitted(this, &SActorFolderTreeLabel::OnLabelCommitted)
 			.OnVerifyTextChanged(this, &SActorFolderTreeLabel::OnVerifyItemLabelChanged)
+			.OnEnterEditingMode(this, &SActorFolderTreeLabel::OnEnterEditingMode)
+			.OnExitEditingMode(this, &SActorFolderTreeLabel::OnExitEditingMode)
 			.IsSelected(FIsSelected::CreateSP(&InRow, &STableRow<FSceneOutlinerTreeItemPtr>::IsSelectedExclusively))
 			.IsReadOnly_Lambda([Item = FolderItem.AsShared(), this]()
 		{
@@ -55,21 +58,6 @@ struct SActorFolderTreeLabel : FSceneOutlinerCommonLabelData, public SCompoundWi
 				SNew(SHorizontalBox)
 				+ SHorizontalBox::Slot()
 				.AutoWidth()
-				.HAlign(HAlign_Right)
-				.VAlign(VAlign_Center)
-				[
-					SNew(SBox)
-					.WidthOverride(FSceneOutlinerDefaultTreeItemMetrics::IconSize())
-					.HeightOverride(FSceneOutlinerDefaultTreeItemMetrics::IconSize())
-					[
-						SNew(SImage)
-						.Visibility_Lambda([this]() { return IsInActorEditorContext() ? EVisibility::Visible : EVisibility::Collapsed; })
-						.Image(FEditorStyle::Get().GetBrush(TEXT("SceneOutliner.MarkedAsCurrent")))
-						.ColorAndOpacity(FSlateColor::UseForeground())
-					]
-				]
-				+ SHorizontalBox::Slot()
-				.AutoWidth()
 				.VAlign(VAlign_Center)
 				.Padding(FSceneOutlinerDefaultTreeItemMetrics::IconPadding())
 				[
@@ -105,7 +93,27 @@ private:
 	FText GetDisplayText() const
 	{
 		auto Folder = TreeItemPtr.Pin();
-		return Folder.IsValid() ? FText::FromString(Folder->GetDisplayString()) : FText();
+		if (Folder.IsValid())
+		{
+			if (!bInEditingMode)
+			{
+				FText IsCurrentSuffixText = IsInActorEditorContext() ? FText(LOCTEXT("IsCurrentSuffix", " (Current)")) : FText::GetEmpty();
+				return FText::Format(LOCTEXT("LevelInstanceDisplay", "{0}{1}"), FText::FromString(Folder->GetDisplayString()), IsCurrentSuffixText);
+			}
+			return FText::FromString(Folder->GetDisplayString());
+		}
+		return FText();
+	}
+
+	FText GetTooltipText() const
+	{
+		if (const FSceneOutlinerTreeItemPtr TreeItem = TreeItemPtr.Pin())
+		{
+			FText Description = IsInActorEditorContext() ? LOCTEXT("ActorFolderIsCurrentDescription", "This Folder is set as Current Folder. New actors will be added to this Folder.") : FText::GetEmpty();
+			return FText::Format(LOCTEXT("DataLayerTooltipText", "{0}\n{1}"), FText::FromString(TreeItem->GetDisplayString()), Description);
+		}
+
+		return FText();
 	}
 
 	const FSlateBrush* GetIcon() const
@@ -133,6 +141,11 @@ private:
 			return BaseColor.GetValue();
 		}
 
+		if (IsInActorEditorContext())
+		{
+			return FAppStyle::Get().GetSlateColor("Colors.AccentGreen");
+		}
+
 		return FSlateColor::UseForeground();
 	}
 
@@ -223,6 +236,18 @@ private:
 			}
 		}
 	}
+
+	void OnEnterEditingMode()
+	{
+		bInEditingMode = true;
+	}
+
+	void OnExitEditingMode()
+	{
+		bInEditingMode = false;
+	}
+
+	bool bInEditingMode = false;
 };
 
 FActorFolderTreeItem::FActorFolderTreeItem(const FFolder& InFolder, const TWeakObjectPtr<UWorld>& InWorld)
diff --git a/Engine/Source/Editor/SceneOutliner/Private/ActorTreeItem.cpp b/Engine/Source/Editor/SceneOutliner/Private/ActorTreeItem.cpp
--- a/Engine/Source/Editor/SceneOutliner/Private/ActorTreeItem.cpp
+++ b/Engine/Source/Editor/SceneOutliner/Private/ActorTreeItem.cpp
@@ -122,12 +122,11 @@ private:
 			const AActor* Actor = ActorPtr.Get();
 			if (const ILevelInstanceInterface* LevelInstance = Cast<ILevelInstanceInterface>(Actor))
 			{
-				if (LevelInstance->IsDirty() && !bInEditingMode)
+				if (!bInEditingMode)
 				{
-					FFormatNamedArguments Args;
-					Args.Add(TEXT("ActorLabel"), FText::FromString(TreeItem->GetDisplayString()));
-					Args.Add(TEXT("EditTag"), LOCTEXT("EditingLevelInstanceLabel", "*"));
-					return FText::Format(LOCTEXT("LevelInstanceDisplay", "{ActorLabel}{EditTag}"), Args);
+					FText DirtySuffixText = LevelInstance->IsDirty() ? FText(LOCTEXT("IsDirtySuffix", "*")) : FText::GetEmpty();
+					FText IsCurrentSuffixText = LevelInstance->GetLoadedLevel() && LevelInstance->GetLoadedLevel()->IsCurrentLevel() ? FText(LOCTEXT("IsCurrentSuffix", " (Current)")) : FText::GetEmpty();
+					return FText::Format(LOCTEXT("LevelInstanceDisplay", "{0}{1}{2}"), FText::FromString(TreeItem->GetDisplayString()), DirtySuffixText, IsCurrentSuffixText);
 				}
 			}
 			return FText::FromString(TreeItem->GetDisplayString());
diff --git a/Engine/Source/Editor/UnrealEd/Private/EditorActorFolders.cpp b/Engine/Source/Editor/UnrealEd/Private/EditorActorFolders.cpp
--- a/Engine/Source/Editor/UnrealEd/Private/EditorActorFolders.cpp
+++ b/Engine/Source/Editor/UnrealEd/Private/EditorActorFolders.cpp
@@ -578,7 +578,7 @@ bool FActorFolders::GetActorEditorContextDisplayInfo(UWorld* InWorld, FActorEdit
 	const FFolder Folder = GetActorEditorContextFolder(*InWorld);
 	if (!Folder.IsNone())
 	{
-		OutDiplayInfo.Title = TEXT("Current Actor Folder");
+		OutDiplayInfo.Title = TEXT("Actor Folder");
 		OutDiplayInfo.Brush = FEditorStyle::GetBrush(TEXT("SceneOutliner.FolderClosed"));
 		return true;
 	}
diff --git a/Engine/Source/Editor/UnrealEd/Private/SActorEditorContext.cpp b/Engine/Source/Editor/UnrealEd/Private/SActorEditorContext.cpp
--- a/Engine/Source/Editor/UnrealEd/Private/SActorEditorContext.cpp
+++ b/Engine/Source/Editor/UnrealEd/Private/SActorEditorContext.cpp
@@ -14,6 +14,7 @@
 void SActorEditorContext::Construct(const FArguments& InArgs)
 {
 	World = (InArgs._World);
+	bIsContextExpanded = true;
 	GEditor->GetEditorWorldContext().AddRef(World);
 	UActorEditorContextSubsystem::Get()->OnActorEditorContextSubsystemChanged().AddSP(this, &SActorEditorContext::Rebuild);
 	FEditorDelegates::MapChange.AddSP(this, &SActorEditorContext::OnEditorMapChange);
@@ -32,26 +33,93 @@ void SActorEditorContext::Rebuild()
 	TArray<IActorEditorContextClient*> Clients = UActorEditorContextSubsystem::Get()->GetDisplayableClients();
 	if (Clients.Num() > 0 && World)
 	{
+		TSharedPtr<SHorizontalBox> HBox;
 		TSharedPtr<SVerticalBox> VBox;
 		ChildSlot.Padding(2, 2, 2, 2)
 		[
 			SAssignNew(VBox, SVerticalBox)
 		];
 
+		VBox->AddSlot()
+			.AutoHeight()
+			.Padding(0, 2)
+			[
+				SNew(SVerticalBox)
+				+ SVerticalBox::Slot()
+				.AutoHeight()
+				[
+					SNew(SHorizontalBox)
+					+ SHorizontalBox::Slot()
+					.AutoWidth()
+					.Padding(2.0f, 1.0f, 2.0f, 1.0f)
+					[
+						SNew(STextBlock)
+						.Text(FText::FromString(TEXT("Current Context")))
+						.ShadowOffset(FVector2D(1, 1))
+						.Font(FCoreStyle::GetDefaultFontStyle("Bold", 10))
+						.ColorAndOpacity(FLinearColor::White)
+					]
+					+ SHorizontalBox::Slot()
+					.AutoWidth()
+					.HAlign(HAlign_Right)
+					[
+						SAssignNew(HBox, SHorizontalBox)
+						.Visibility_Lambda([this]()
+						{
+							return !bIsContextExpanded ? EVisibility::Visible : EVisibility::Collapsed;
+						})
+					]
+					+ SHorizontalBox::Slot()					
+					.HAlign(HAlign_Right)
+					.VAlign(VAlign_Center)
+					.Padding(2.0f, 1.0f, 2.0f, 1.0f)
+					[
+						SNew(SButton)
+						.Cursor(EMouseCursor::Default)
+						.ButtonStyle(FAppStyle::Get(), "NoBorder")
+						.ContentPadding(0)
+						.OnClicked_Lambda([this]()
+						{
+							bIsContextExpanded = !bIsContextExpanded;
+							return FReply::Handled();
+						})
+						.Content()
+						[
+							SNew(SImage)
+							.Image_Lambda([this]() { return bIsContextExpanded ? FEditorStyle::GetBrush("ContentBrowser.SortDown") : FEditorStyle::GetBrush("ContentBrowser.SortUp"); })
+							.ColorAndOpacity(FSlateColor::UseForeground())
+						]
+					]
+				]
+			];
+
 		for (IActorEditorContextClient* Client : Clients)
 		{
 			FActorEditorContextClientDisplayInfo Info;
 			Client->GetActorEditorContextDisplayInfo(World, Info);
 
+			HBox->AddSlot()
+			[
+				SNew(SHorizontalBox)
+				+ SHorizontalBox::Slot()
+				.HAlign(HAlign_Right)
+				.Padding(2.0f, 1.0f, 2.0f, 1.0f)
+				[
+					SNew(SImage)
+					.Image(Info.Brush ? Info.Brush : FStyleDefaults::GetNoBrush())
+					.DesiredSizeOverride(FVector2D(16, 16))
+				]
+			];
+
 			VBox->AddSlot()
-				.AutoHeight()
-				.Padding(0, 2)
+			.AutoHeight()
+			.Padding(0, 2)
 			[
 				SNew(SBorder)
 				.Visibility_Lambda([Client, this]()
 				{ 
 					FActorEditorContextClientDisplayInfo Info;
-					return Client->GetActorEditorContextDisplayInfo(World, Info) ? EVisibility::Visible : EVisibility::Collapsed; 
+					return bIsContextExpanded && Client->GetActorEditorContextDisplayInfo(World, Info) ? EVisibility::Visible : EVisibility::Collapsed;
 				})
 				.BorderImage(FCoreStyle::Get().GetBrush("Docking.Sidebar.Border"))
 				.Content()
@@ -63,14 +131,15 @@ void SActorEditorContext::Rebuild()
 						SNew(SHorizontalBox)
 						+ SHorizontalBox::Slot()
 						.AutoWidth()
-						.Padding(4.0f, 1.0f, 2.0f, 1.0f)
+						.Padding(2.0f, 1.0f, 2.0f, 1.0f)
 						[
 							SNew(SImage)
 							.Image(Info.Brush ? Info.Brush : FStyleDefaults::GetNoBrush())
+							.DesiredSizeOverride(FVector2D(16, 16))
 						]
 						+ SHorizontalBox::Slot()
 						.AutoWidth()
-						.Padding(4.0f, 1.0f, 2.0f, 1.0f)
+						.Padding(2.0f, 1.0f, 2.0f, 1.0f)
 						[
 							SNew(STextBlock)
 							.Text(FText::FromString(Info.Title))
@@ -79,15 +148,16 @@ void SActorEditorContext::Rebuild()
 							.ColorAndOpacity(FLinearColor::White)
 						]
 						+ SHorizontalBox::Slot()
-						.AutoWidth()
+						.HAlign(HAlign_Right)
+						.VAlign(VAlign_Center)
+						.FillWidth(1.f)
 						[
 							SNew(SButton)
-							.HAlign(HAlign_Right)
 							.Cursor(EMouseCursor::Default)
 							.ButtonStyle(FAppStyle::Get(), "SimpleButton")
 							.ContentPadding(0)
 							.Visibility_Lambda([Client, this]() { return (Client && Client->CanResetContext(World)) ? EVisibility::Visible : EVisibility::Collapsed; })
-							.ToolTipText(FText::Format(LOCTEXT("ResetActorEditorContextTooltip", "Reset {0}"), FText::FromString(Info.Title)))
+							.ToolTipText(FText::Format(LOCTEXT("ResetActorEditorContextTooltip", "Reset Current {0}"), FText::FromString(Info.Title)))
 							.OnClicked_Lambda([Client, this]()
 							{
 								if (Client && Client->CanResetContext(World))
@@ -106,7 +176,8 @@ void SActorEditorContext::Rebuild()
 					]
 					+ SVerticalBox::Slot()
 					.AutoHeight()
-					.Padding(8, 2, 0, 2)
+					// Left padding details (2 + 2 + 16 + 2) : 2 + 2 = Info.Brush left/right padding, + 16 = Info.Brush width, + 2 = Info.Title left padding
+					.Padding(2 + 2 + 16 + 2, 2, 0, 2)
 					[
 						World ? Client->GetActorEditorContextWidget(World) : SNullWidget::NullWidget
 					]
diff --git a/Engine/Source/Editor/UnrealEd/Public/SActorEditorContext.h b/Engine/Source/Editor/UnrealEd/Public/SActorEditorContext.h
--- a/Engine/Source/Editor/UnrealEd/Public/SActorEditorContext.h
+++ b/Engine/Source/Editor/UnrealEd/Public/SActorEditorContext.h
@@ -19,5 +19,6 @@ private:
 	void OnEditorMapChange(uint32 MapChangeFlags = 0) { Rebuild(); }
 	void Rebuild();
 	
+	bool bIsContextExpanded;
 	UWorld* World;
 };
\ No newline at end of file
