EditorFactories fix duplicate loc text keys

diff --git a/Engine/Source/Editor/UnrealEd/Private/Factories/EditorFactories.cpp b/Engine/Source/Editor/UnrealEd/Private/Factories/EditorFactories.cpp
--- a/Engine/Source/Editor/UnrealEd/Private/Factories/EditorFactories.cpp
+++ b/Engine/Source/Editor/UnrealEd/Private/Factories/EditorFactories.cpp
@@ -4579,7 +4579,7 @@ bool UTextureFactory::IsImportResolutionValid(int64 Width, int64 Height, bool bA
 	if (Width == 0 || Height == 0 )
 	{
 		Warn->Log(ELogVerbosity::Error, *FText::Format(
-			NSLOCTEXT("UnrealEd", "Warning_TextureSizeTooLargeOrInvalid", "Texture has zero width or height"),
+			NSLOCTEXT("UnrealEd", "Warning_TextureSizeZero", "Texture has zero width or height"),
 			FText::AsNumber(FMath::Square(MaximumSupportedVirtualTextureResolution))
 			).ToString());
 
@@ -4591,7 +4591,7 @@ bool UTextureFactory::IsImportResolutionValid(int64 Width, int64 Height, bool bA
 	if (Width < 0 || Height < 0 || Width > MAX_int32 || Height > MAX_int32)
 	{
 		Warn->Log(ELogVerbosity::Error, *FText::Format(
-			NSLOCTEXT("UnrealEd", "Warning_TextureSizeTooLargeOrInvalid", "Texture is too large to import or it has an invalid resolution. The current maximun is {0} pixels"),
+			NSLOCTEXT("UnrealEd", "Warning_TextureSizeTooLargeOrInvalid", "Texture is too large to import or it has an invalid resolution. The current maximum is {0} pixels"),
 			FText::AsNumber(FMath::Square(MaximumSupportedVirtualTextureResolution))
 			).ToString());
 
@@ -4604,7 +4604,7 @@ bool UTextureFactory::IsImportResolutionValid(int64 Width, int64 Height, bool bA
 	if ( Width * Height > MAX_int32)
 	{
 		Warn->Log(ELogVerbosity::Error, *FText::Format(
-			NSLOCTEXT("UnrealEd", "Warning_TextureSizeTooLargeOrInvalid", "Texture is too large to import or it has an invalid resolution. The current maximun is {0} pixels"),
+			NSLOCTEXT("UnrealEd", "Warning_TextureSizeTooLargeOrInvalid", "Texture is too large to import or it has an invalid resolution. The current maximum is {0} pixels"),
 			FText::AsNumber(FMath::Square(MaximumSupportedVirtualTextureResolution))
 			).ToString());
 
@@ -4614,7 +4614,7 @@ bool UTextureFactory::IsImportResolutionValid(int64 Width, int64 Height, bool bA
 	if ( (Width * Height) > FMath::Square(MaximumSupportedVirtualTextureResolution))
 	{
 		Warn->Log(ELogVerbosity::Error, *FText::Format(
-			NSLOCTEXT("UnrealEd", "Warning_TextureSizeTooLarge", "Texture is too large to import. The current maximun is {0} pixels"),
+			NSLOCTEXT("UnrealEd", "Warning_TextureSizeTooLargeOrInvalid", "Texture is too large to import or it has an invalid resolution. The current maximum is {0} pixels"),
 			FText::AsNumber(FMath::Square(MaximumSupportedVirtualTextureResolution))
 			).ToString());
 
