Cooker: Add PKG_CookGenerated to mark packages that are generated during the cook and do not exist in the workspace domain.
#rnx

diff --git a/Engine/Source/Editor/UnrealEd/Private/Cooker/CookPackageData.cpp b/Engine/Source/Editor/UnrealEd/Private/Cooker/CookPackageData.cpp
--- a/Engine/Source/Editor/UnrealEd/Private/Cooker/CookPackageData.cpp
+++ b/Engine/Source/Editor/UnrealEd/Private/Cooker/CookPackageData.cpp
@@ -1302,6 +1302,7 @@ UPackage* FGeneratorPackage::CreateGeneratedUPackage(FGeneratorPackage::FGenerat
 	GeneratedPackage->SetGuid(OwnerPackage->GetGuid());
 	PRAGMA_ENABLE_DEPRECATION_WARNINGS
 	GeneratedPackage->SetPersistentGuid(OwnerPackage->GetPersistentGuid());
+	GeneratedPackage->SetPackageFlags(PKG_CookGenerated);
 	GeneratedStruct.bHasCreatedPackage = true;
 	return GeneratedPackage;
 }
diff --git a/Engine/Source/Editor/UnrealEd/Private/Cooker/LooseCookedPackageWriter.cpp b/Engine/Source/Editor/UnrealEd/Private/Cooker/LooseCookedPackageWriter.cpp
--- a/Engine/Source/Editor/UnrealEd/Private/Cooker/LooseCookedPackageWriter.cpp
+++ b/Engine/Source/Editor/UnrealEd/Private/Cooker/LooseCookedPackageWriter.cpp
@@ -517,7 +517,23 @@ TUniquePtr<FAssetRegistryState> FLooseCookedPackageWriter::LoadPreviousAssetRegi
 		{
 			FName PackageName = Pair.Key;
 			const FName UncookedFilename = PackageDatas.GetFileNameByPackageName(PackageName);
+
+			bool bNoLongerExistsInEditor = false;
 			if (UncookedFilename.IsNone())
+			{
+				// Script and generated packages do not exist uncooked
+				// Check that the package is not an exception before removing from cooked
+				bool bIsCookedOnly = FPackageName::IsScriptPackage(WriteToString<256>(PackageName));
+				if (!bIsCookedOnly)
+				{
+					for (const FAssetData* AssetData : PreviousState->GetAssetsByPackageName(PackageName))
+					{
+						bIsCookedOnly |= !!(AssetData->PackageFlags & PKG_CookGenerated);
+					}
+				}
+				bNoLongerExistsInEditor = !bIsCookedOnly;
+			}
+			if (bNoLongerExistsInEditor)
 			{
 				// Remove package from both disk and registry
 				// Keep its RemoveFromDisk entry
diff --git a/Engine/Source/Programs/Shared/EpicGames.Core/UnrealEngineTypes.cs b/Engine/Source/Programs/Shared/EpicGames.Core/UnrealEngineTypes.cs
--- a/Engine/Source/Programs/Shared/EpicGames.Core/UnrealEngineTypes.cs
+++ b/Engine/Source/Programs/Shared/EpicGames.Core/UnrealEngineTypes.cs
@@ -851,7 +851,11 @@ namespace EpicGames.Core
 		// Unused = 0x01000000,	
 		// Unused = 0x02000000,	
 		// Unused = 0x04000000,
-		// Unused = 0x08000000,	
+
+		/// <summary>
+		/// This package This package was generated by the cooker and does not exist in the WorkspaceDomain
+		/// </summary>
+		CookGenerated = 0x08000000,
 
 		/// <summary>
 		/// This package should resolve dynamic imports from its export at runtime.
diff --git a/Engine/Source/Runtime/CoreUObject/Public/UObject/ObjectMacros.h b/Engine/Source/Runtime/CoreUObject/Public/UObject/ObjectMacros.h
--- a/Engine/Source/Runtime/CoreUObject/Public/UObject/ObjectMacros.h
+++ b/Engine/Source/Runtime/CoreUObject/Public/UObject/ObjectMacros.h
@@ -137,7 +137,7 @@ enum EPackageFlags
 //	PKG_Unused						= 0x01000000,	
 //	PKG_Unused						= 0x02000000,	
 //	PKG_Unused						= 0x04000000,
-//	PKG_Unused						= 0x08000000,	
+	PKG_CookGenerated				= 0x08000000,	///< This package was generated by the cooker and does not exist in the WorkspaceDomain
 	PKG_DynamicImports				= 0x10000000,	///< This package should resolve dynamic imports from its export at runtime.
 	PKG_RuntimeGenerated			= 0x20000000,	///< This package contains elements that are runtime generated, and may not follow standard loading order rules
 	PKG_ReloadingForCooker			= 0x40000000,   ///< This package is reloading in the cooker, try to avoid getting data we will never need. We won't save this package.
