CookOnTheFlyServer - StartCookByTheBook refactoring part 3.
1) Split StartCookByTheBook into subfunctions.
2) Rename some functions to match the new naming convention I've made for StartCookByTheBook helpers.
3) Make StartCookByTheBook and StartCookOnTheFly/AddCookOnTheFlyPlatform share subfunctions rather than duplicate code. Make Cancel and CookByTheBookFinished share subfunctions.
4) Remove up front allocation of MapDependencyGraphs; we can construct them on demand in CookByTheBookFinished.
#rnx

diff --git a/Engine/Source/Editor/UnrealEd/Classes/CookOnTheSide/CookOnTheFlyServer.h b/Engine/Source/Editor/UnrealEd/Classes/CookOnTheSide/CookOnTheFlyServer.h
--- a/Engine/Source/Editor/UnrealEd/Classes/CookOnTheSide/CookOnTheFlyServer.h
+++ b/Engine/Source/Editor/UnrealEd/Classes/CookOnTheSide/CookOnTheFlyServer.h
@@ -715,6 +715,7 @@ private:
 
 	/** Update accumulators of editor data and consume editor changes since the last cook when starting cooks in the editor. */
 	void BeginCookEditorSystems();
+	void BeginCookPackageWriters(FBeginCookContext& BeginContext);
 
 	//////////////////////////////////////////////////////////////////////////
 	// cook by the book specific functions
@@ -722,6 +723,15 @@ private:
 	FBeginCookContext SetCookByTheBookOptions(const FCookByTheBookStartupOptions& StartupOptions);
 	/** Set the PlatformManager's session platforms and finish filling out the BeginContext for StartCookByTheBook */
 	void SelectSessionPlatforms(FBeginCookContext& BeginContext);
+	/** CollectFilesToCook and add them as external requests. */
+	void GenerateInitialRequests(FBeginCookContext& BeginContext);
+	/** If cooking DLC, initialize cooked PackageDatas for all of the packages cooked by the base release. */
+	void RecordDLCPackagesFromBaseGame(FBeginCookContext& BeginContext);
+	void RegisterCookByTheBookDelegates();
+	void UnregisterCookByTheBookDelegates();
+	/** Start the collection of EDL diagnostics for the cook, if applicable. */
+	void BeginCookEDLCookInfo(FBeginCookContext& BeginContext);
+
 	/**
 	* Collect all the files which need to be cooked for a cook by the book session
 	*/
@@ -809,27 +819,21 @@ private:
 	bool GetAllPackageFilenamesFromAssetRegistry( const FString& AssetRegistryPath, bool bVerifyPackagesExist,
 		 bool bSkipUncookedPackages, TArray<UE::Cook::FConstructPackageData>& OutPackageDatas) const;
 
-	/**
-	* BuildMapDependencyGraph
-	* builds a map of dependencies from maps
-	* 
-	* @param TargetPlatform the platform we want to build a map for
-	*/
-	void BuildMapDependencyGraph(const ITargetPlatform* TargetPlatform);
+	/** Build a map of the package dependencies used by each Map Package. */
+	TMap<FName, TSet<FName>> BuildMapDependencyGraph(const ITargetPlatform* TargetPlatform);
 
-	/**
-	* WriteMapDependencyGraph
-	* write a previously built map dependency graph out to the sandbox directory for a platform
-	*
-	* @param TargetPlatform the platform we want to save out the dependency graph for
-	*/
-	void WriteMapDependencyGraph(const ITargetPlatform* TargetPlatform);
+	/** Write a MapDependencyGraph to a metadata file in the sandbox for the given platform. */
+	void WriteMapDependencyGraph(const ITargetPlatform* TargetPlatform, TMap<FName, TSet<FName>>& MapDependencyGraph);
 
+	/** Find localization dependencies for all packages, used to add required localization files as soft references. */
+	void GenerateLocalizationReferences(TConstArrayView<FString> CookCultures);
 	void RegisterLocalizationChunkDataGenerator();
 
 	//////////////////////////////////////////////////////////////////////////
 	// cook on the fly specific functions
 
+	/** Construct the on-stack data for StartCook functions, based on the arguments to StartCookOnTheFly */
+	FBeginCookContext CreateBeginCookOnTheFlyContext();
 	/** Construct the on-stack data for StartCook functions, based on the arguments when adding a COTF platform */
 	FBeginCookContext CreateAddPlatformContext(ITargetPlatform* TargetPlatform);
 	void GetCookOnTheFlyUnsolicitedFiles(const ITargetPlatform* TargetPlatform, const FString& PlatformName, TArray<FString>& UnsolicitedFiles, const FString& Filename, bool bIsCookable);
@@ -860,7 +864,9 @@ private:
 	void BeginCookSandbox(FBeginCookContext& BeginContext);
 
 	/** Set parameters that rely on config and CookByTheBook settings. */
-	void SetBeginCookConfigSettings();
+	void SetBeginCookConfigSettings(FBeginCookContext& BeginContext);
+	/** Initialize NeverCookPackageList from packaging settings and platform-specific sources. */
+	void SetNeverCookPackageConfigSettings(FBeginCookContext& BeginContext);
 
 	/**
 	* Finalize the package store
@@ -872,10 +878,8 @@ private:
 	*/
 	void ClearPackageStoreContexts();
 
-	/**
-	* Initialize platforms in @param NewTargetPlatforms
-	*/
-	void InitializeTargetPlatforms(const TArrayView<ITargetPlatform* const>& NewTargetPlatforms);
+	/** Give TargetPlatforms a chance to take action at the beginning of a cook. */
+	void BeginCookTargetPlatforms(const TArrayView<ITargetPlatform* const>& NewTargetPlatforms);
 
 	/**
 	* Initialize shaders for the specified platforms when running cook on the fly.
diff --git a/Engine/Source/Editor/UnrealEd/Private/CookOnTheFlyServer.cpp b/Engine/Source/Editor/UnrealEd/Private/CookOnTheFlyServer.cpp
--- a/Engine/Source/Editor/UnrealEd/Private/CookOnTheFlyServer.cpp
+++ b/Engine/Source/Editor/UnrealEd/Private/CookOnTheFlyServer.cpp
@@ -6,6 +6,7 @@
 
 #include "CookOnTheSide/CookOnTheFlyServer.h"
 
+#include "Algo/AllOf.h"
 #include "Algo/AnyOf.h"
 #include "Algo/Find.h"
 #include "Algo/Unique.h"
@@ -536,7 +537,9 @@ bool UCookOnTheFlyServer::StartCookOnTheFly(FCookOnTheFlyOptions InCookOnTheFlyO
 #endif
 
 	CookOnTheFlyOptions = MoveTemp(InCookOnTheFlyOptions);
+	FBeginCookContext BeginContext = CreateBeginCookOnTheFlyContext();
 	CreateSandboxFile();
+
 	CookOnTheFlyServerInterface = MakeUnique<UCookOnTheFlyServer::FCookOnTheFlyServerInterface>(*this);
 	ExternalRequests->CookRequestEvent = FPlatformProcess::GetSynchEventFromPool();
 
@@ -549,11 +552,7 @@ bool UCookOnTheFlyServer::StartCookOnTheFly(FCookOnTheFlyOptions InCookOnTheFlyO
 	}
 	PlatformManager->SetArePlatformsPrepopulated(true);
 
-	SetBeginCookConfigSettings();
-	for (FName NeverCookPackage : GetNeverCookPackageFileNames(TArrayView<const FString>()))
-	{
-		PackageTracker->NeverCookPackageList.Add(NeverCookPackage);
-	}
+	SetBeginCookConfigSettings(BeginContext);
 
 	GRedirectCollector.OnStartupPackageLoadComplete();
 
@@ -595,7 +594,6 @@ void UCookOnTheFlyServer::CookOnTheFlyDeferredInitialize()
 	}
 	bHasDeferredInitializeCookOnTheFly = true;
 	BlockOnAssetRegistry();
-	PackageDatas->BeginCook();
 }
 
 void UCookOnTheFlyServer::InitializeShadersForCookOnTheFly(const TArrayView<ITargetPlatform* const>& NewTargetPlatforms)
@@ -622,24 +620,32 @@ void UCookOnTheFlyServer::AddCookOnTheFlyPlatformFromGameThread(ITargetPlatform*
 
 	FBeginCookContext BeginContext = CreateAddPlatformContext(TargetPlatform);
 
+	// Initialize systems and settings that the rest of AddCookOnTheFlyPlatformFromGameThread depends on
+	// Functions in this section are ordered and can depend on the functions before them
 	FindOrCreateSaveContexts(BeginContext.TargetPlatforms);
 	SetBeginCookIterativeFlags(BeginContext);
+
+	// Initialize systems referenced by later stages or that need to start early for async performance
+	// Functions in this section must not need to read/write the SandboxDirectory or MemoryCookedPackages
+	// Functions in this section are not dependent upon each other and can be ordered arbitrarily or for async performance
+	RefreshPlatformAssetRegistries(BeginContext.TargetPlatforms);
+
+	// Clear the sandbox directory, or preserve it and populate iterative cooks
+	// Clear in-memory CookedPackages, or preserve them and cook iteratively in-process
 	BeginCookSandbox(BeginContext);
-	InitializeTargetPlatforms(BeginContext.TargetPlatforms);
+
+	// Initialize systems that nothing in AddCookOnTheFlyPlatformFromGameThread references
+	// Functions in this section are not dependent upon each other and can be ordered arbitrarily or for async performance
+	BeginCookPackageWriters(BeginContext);
+	BeginCookTargetPlatforms(BeginContext.TargetPlatforms);
 	InitializeShadersForCookOnTheFly(BeginContext.TargetPlatforms);
-	RefreshPlatformAssetRegistries(BeginContext.TargetPlatforms);
 
-	// When cooking on the fly the full registry is saved at the beginning
-	// in cook by the book asset registry is saved after the cook is finished
-	FAssetRegistryGenerator& Generator = *PlatformData->RegistryGenerator;
-	Generator.SaveAssetRegistry(GetSandboxAssetRegistryFilename(), true);
+	// Execute save operations that are done in CookByTheBookFinished for CBTB, but that CookOnTheFly does at the beginning since it does not have a definite end
 	check(PlatformData->bIsSandboxInitialized); // This should have been set by BeginCookSandbox, and it is what we use to determine whether a platform has been initialized
-
+	PlatformData->RegistryGenerator->SaveAssetRegistry(GetSandboxAssetRegistryFilename(), true);
 	// This will miss settings that are accessed during the cook
 	// TODO: A better way of handling ini settings
 	SaveCurrentIniSettings(TargetPlatform);
-
-	FindOrCreatePackageWriter(TargetPlatform).BeginCook();
 }
 
 void UCookOnTheFlyServer::OnTargetPlatformsInvalidated()
@@ -647,10 +653,6 @@ void UCookOnTheFlyServer::OnTargetPlatformsInvalidated()
 	check(IsInGameThread());
 	TMap<ITargetPlatform*, ITargetPlatform*> Remap = PlatformManager->RemapTargetPlatforms();
 
-	if (CookByTheBookOptions)
-	{
-		RemapMapKeys(CookByTheBookOptions->MapDependencyGraphs, Remap);
-	}
 	PackageDatas->RemapTargetPlatforms(Remap);
 	PackageTracker->RemapTargetPlatforms(Remap);
 	ExternalRequests->RemapTargetPlatforms(Remap);
@@ -7236,30 +7238,16 @@ void UCookOnTheFlyServer::CookByTheBookFinished()
 	if (CookByTheBookOptions->bGenerateDependenciesForMaps)
 	{
 		UE_SCOPED_HIERARCHICAL_COOKTIMER(GenerateMapDependencies);
-		for (auto& MapDependencyGraphIt : CookByTheBookOptions->MapDependencyGraphs)
+		for (const ITargetPlatform* Platform : PlatformManager->GetSessionPlatforms())
 		{
-			BuildMapDependencyGraph(MapDependencyGraphIt.Key);
-			WriteMapDependencyGraph(MapDependencyGraphIt.Key);
+			TMap<FName, TSet<FName>> MapDependencyGraph = BuildMapDependencyGraph(Platform);
+			WriteMapDependencyGraph(Platform, MapDependencyGraph);
 		}
 	}
 
 	FinalizePackageStore();
-
-	CookByTheBookOptions->BasedOnReleaseCookedPackages.Empty();
-	CookByTheBookOptions->bRunning = false;
-	CookByTheBookOptions->bFullLoadAndSave = false;
-
-	if (!IsCookingInEditor())
-	{
-		FCoreUObjectDelegates::PackageCreatedForLoad.RemoveAll(this);
-	}
 	ShutdownCookSession();
 
-	PrintFinishStats();
-
-	OutputHierarchyTimers();
-	ClearHierarchyTimers();
-
 	UE_LOG(LogCook, Display, TEXT("Done!"));
 }
 
@@ -7272,6 +7260,21 @@ void UCookOnTheFlyServer::ShutdownCookSession()
 	}
 
 	ConditionalUninstallImportBehaviorCallback();
+
+	if (IsCookByTheBookMode())
+	{
+		check(CookByTheBookOptions);
+		CookByTheBookOptions->BasedOnReleaseCookedPackages.Empty();
+		CookByTheBookOptions->bRunning = false;
+		CookByTheBookOptions->bFullLoadAndSave = false;
+
+		UnregisterCookByTheBookDelegates();
+
+		PrintFinishStats();
+
+		OutputHierarchyTimers();
+		ClearHierarchyTimers();
+	}
 }
 
 void UCookOnTheFlyServer::PrintFinishStats()
@@ -7286,9 +7289,9 @@ void UCookOnTheFlyServer::PrintFinishStats()
 		UE::SavePackageUtilities::GetNumPackagesSaved(), DetailedCookStats::NumPackagesIterativelySkipped, PackageDatas->GetNumCooked()));
 }
 
-void UCookOnTheFlyServer::BuildMapDependencyGraph(const ITargetPlatform* TargetPlatform)
+TMap<FName, TSet<FName>> UCookOnTheFlyServer::BuildMapDependencyGraph(const ITargetPlatform* TargetPlatform)
 {
-	auto& MapDependencyGraph = CookByTheBookOptions->MapDependencyGraphs.FindChecked(TargetPlatform);
+	TMap<FName, TSet<FName>> MapDependencyGraph;
 
 	TArray<UE::Cook::FPackageData*> PlatformCookedPackages;
 	PackageDatas->GetCookedPackagesForPlatform(TargetPlatform, PlatformCookedPackages, /* include failed */ true, /* include successful */ true);
@@ -7313,12 +7316,11 @@ void UCookOnTheFlyServer::BuildMapDependencyGraph(const ITargetPlatform* TargetP
 
 		MapDependencyGraph.Add(Name, DependentPackages);
 	}
+	return MapDependencyGraph;
 }
 
-void UCookOnTheFlyServer::WriteMapDependencyGraph(const ITargetPlatform* TargetPlatform)
+void UCookOnTheFlyServer::WriteMapDependencyGraph(const ITargetPlatform* TargetPlatform, TMap<FName, TSet<FName>>& MapDependencyGraph)
 {
-	auto& MapDependencyGraph = CookByTheBookOptions->MapDependencyGraphs.FindChecked(TargetPlatform);
-
 	FString MapDependencyGraphFile = FPaths::ProjectDir() / TEXT("MapDependencyGraph.json");
 	// dump dependency graph. 
 	FString DependencyString;
@@ -7360,9 +7362,7 @@ void UCookOnTheFlyServer::CancelCookByTheBook()
 
 		CancelAllQueues();
 
-		CookByTheBookOptions->bRunning = false;
-
-		PrintFinishStats();
+		ShutdownCookSession();
 	} 
 }
 
@@ -7568,13 +7568,17 @@ void UCookOnTheFlyServer::CreateSandboxFile()
 
 }
 
-void UCookOnTheFlyServer::SetBeginCookConfigSettings()
+void UCookOnTheFlyServer::SetBeginCookConfigSettings(FBeginCookContext& BeginContext)
 {
 	GConfig->GetBool(TEXT("CookSettings"), TEXT("HybridIterativeEnabled"), bHybridIterativeEnabled, GEditorIni);
 	// TODO: HybridIterative is not yet implemented for DLC
 	bHybridIterativeEnabled &= !IsCookingDLC();
 	// HybridIterative uses TargetDomain storage of dependencies which is only implemented in ZenStore
 	bHybridIterativeEnabled &= IsUsingZenStore();
+
+	PackageDatas->SetBeginCookConfigSettings();
+
+	SetNeverCookPackageConfigSettings(BeginContext);
 	if (!bFirstCookInThisProcessInitialized)
 	{
 		// This is the first cook; set bFirstCookInThisProcess=true for the entire cook until SetBeginCookConfigSettings is called to mark the second cook
@@ -7586,7 +7590,35 @@ void UCookOnTheFlyServer::SetBeginCookConfigSettings()
 		// We have cooked before; set bFirstCookInThisProcess=false
 		bFirstCookInThisProcess = false;
 	}
+}
 
+void UCookOnTheFlyServer::SetNeverCookPackageConfigSettings(FBeginCookContext& BeginContext)
+{
+	PackageTracker->NeverCookPackageList.Empty();
+	TArrayView<const FString> ExtraNeverCookDirectories;
+	if (BeginContext.StartupOptions)
+	{
+		ExtraNeverCookDirectories = BeginContext.StartupOptions->NeverCookDirectories;
+	}
+	for (FName NeverCookPackage : GetNeverCookPackageFileNames(ExtraNeverCookDirectories))
+	{
+		PackageTracker->NeverCookPackageList.Add(NeverCookPackage);
+	}
+
+	// use temp list of UBT platform strings to discover PlatformSpecificNeverCookPackages
+	if (BeginContext.TargetPlatforms.Num())
+	{
+		TArray<FString> UBTPlatformStrings;
+		UBTPlatformStrings.Reserve(BeginContext.TargetPlatforms.Num());
+		for (const ITargetPlatform* Platform : BeginContext.TargetPlatforms)
+		{
+			FString UBTPlatformName;
+			Platform->GetPlatformInfo().UBTPlatformName.ToString(UBTPlatformName);
+			UBTPlatformStrings.Emplace(MoveTemp(UBTPlatformName));
+		}
+
+		DiscoverPlatformSpecificNeverCookPackages(BeginContext.TargetPlatforms, UBTPlatformStrings);
+	}
 }
 
 void UCookOnTheFlyServer::SetBeginCookIterativeFlags(FBeginCookContext& BeginContext)
@@ -7830,7 +7862,7 @@ void UCookOnTheFlyServer::ClearPackageStoreContexts()
 	SavePackageContexts.Empty();
 }
 
-void UCookOnTheFlyServer::InitializeTargetPlatforms(const TArrayView<ITargetPlatform* const>& NewTargetPlatforms)
+void UCookOnTheFlyServer::BeginCookTargetPlatforms(const TArrayView<ITargetPlatform* const>& NewTargetPlatforms)
 {
 	//allow each platform to update its internals before cooking
 	for (ITargetPlatform* TargetPlatform : NewTargetPlatforms)
@@ -7900,285 +7932,107 @@ void UCookOnTheFlyServer::StartCookByTheBook( const FCookByTheBookStartupOptions
 {
 	UE_SCOPED_COOKTIMER(StartCookByTheBook);
 	LLM_SCOPE_BYTAG(Cooker);
+	check(IsInGameThread());
+	check(IsCookByTheBookMode());
 
-	const TArray<FString>& CookMaps = CookByTheBookStartupOptions.CookMaps;
-	const TArray<FString>& CookDirectories = CookByTheBookStartupOptions.CookDirectories;
-	const TArray<FString>& IniMapSections = CookByTheBookStartupOptions.IniMapSections;
-	const ECookByTheBookOptions& CookOptions = CookByTheBookStartupOptions.CookOptions;
-	const FString& DLCName = CookByTheBookStartupOptions.DLCName;
+	// Initialize systems and settings that the rest of StartCookByTheBook depends on
+	// Functions in this section are ordered and can depend on the functions before them
+	FBeginCookContext BeginContext = SetCookByTheBookOptions(CookByTheBookStartupOptions);
+	BlockOnAssetRegistry();
+	CreateSandboxFile();
+	SetBeginCookConfigSettings(BeginContext);
+	SelectSessionPlatforms(BeginContext);
+	SetBeginCookIterativeFlags(BeginContext);
+	CookByTheBookOptions->bRunning = true;
 
-	const FString& CreateReleaseVersion = CookByTheBookStartupOptions.CreateReleaseVersion;
-	const FString& BasedOnReleaseVersion = CookByTheBookStartupOptions.BasedOnReleaseVersion;
+	// Initialize systems referenced by later stages or that need to start early for async performance
+	// Functions in this section must not need to read/write the SandboxDirectory or MemoryCookedPackages
+	// Functions in this section are not dependent upon each other and can be ordered arbitrarily or for async performance
+	BeginCookStartShaderCodeLibrary(BeginContext); // start shader code library cooking asynchronously; we block on it later
+	RefreshPlatformAssetRegistries(BeginContext.TargetPlatforms); // Required by BeginCookSandbox stage
 
-	check( IsInGameThread() );
-	check( IsCookByTheBookMode() );
+	// Clear the sandbox directory, or preserve it and populate iterative cooks
+	// Clear in-memory CookedPackages, or preserve them and cook iteratively in-process
+	BeginCookSandbox(BeginContext);
+
+	// Initialize systems that nothing in StartCookByTheBook references
+	// Functions in this section are not dependent upon each other and can be ordered arbitrarily or for async performance
+	BeginCookEditorSystems();
+	BeginCookEDLCookInfo(BeginContext);
+	BeginCookPackageWriters(BeginContext);
+	BeginCookTargetPlatforms(BeginContext.TargetPlatforms);
+	GenerateInitialRequests(BeginContext);
+	GenerateLocalizationReferences(BeginContext.StartupOptions->CookCultures);
+	InitializePollables();
+	RecordDLCPackagesFromBaseGame(BeginContext);
+	RegisterCookByTheBookDelegates();
+
+	BeginCookFinishShaderCodeLibrary(BeginContext);
+}
 
-	//force precache objects to refresh themselves before cooking anything
-	LastUpdateTick = INT_MAX;
+FBeginCookContext UCookOnTheFlyServer::SetCookByTheBookOptions(const FCookByTheBookStartupOptions& StartupOptions)
+{
+	FBeginCookContext BeginContext;
 
+	BeginContext.StartupOptions = &StartupOptions;
+	const ECookByTheBookOptions& CookOptions = StartupOptions.CookOptions;
 	CookByTheBookOptions->bCancel = false;
 	CookByTheBookOptions->CookTime = 0.0f;
 	CookByTheBookOptions->CookStartTime = FPlatformTime::Seconds();
-	CookByTheBookOptions->bGenerateStreamingInstallManifests = CookByTheBookStartupOptions.bGenerateStreamingInstallManifests;
-	CookByTheBookOptions->bGenerateDependenciesForMaps = CookByTheBookStartupOptions.bGenerateDependenciesForMaps;
-	CookByTheBookOptions->CreateReleaseVersion = CreateReleaseVersion;
+	CookByTheBookOptions->bGenerateStreamingInstallManifests = StartupOptions.bGenerateStreamingInstallManifests;
+	CookByTheBookOptions->bGenerateDependenciesForMaps = StartupOptions.bGenerateDependenciesForMaps;
+	CookByTheBookOptions->CreateReleaseVersion = StartupOptions.CreateReleaseVersion;
 	CookByTheBookOptions->bSkipHardReferences = !!(CookOptions & ECookByTheBookOptions::SkipHardReferences);
 	CookByTheBookOptions->bSkipSoftReferences = !!(CookOptions & ECookByTheBookOptions::SkipSoftReferences);
 	CookByTheBookOptions->bFullLoadAndSave = !!(CookOptions & ECookByTheBookOptions::FullLoadAndSave);
 	CookByTheBookOptions->bZenStore = !!(CookOptions & ECookByTheBookOptions::ZenStore);
 	CookByTheBookOptions->bCookAgainstFixedBase = !!(CookOptions & ECookByTheBookOptions::CookAgainstFixedBase);
 	CookByTheBookOptions->bDlcLoadMainAssetRegistry = !!(CookOptions & ECookByTheBookOptions::DlcLoadMainAssetRegistry);
-	CookByTheBookOptions->bErrorOnEngineContentUse = CookByTheBookStartupOptions.bErrorOnEngineContentUse;
-
-	CookByTheBookOptions->DlcName = DLCName;
+	CookByTheBookOptions->bErrorOnEngineContentUse = StartupOptions.bErrorOnEngineContentUse;
+	CookByTheBookOptions->DlcName = StartupOptions.DLCName;
 	if (CookByTheBookOptions->bSkipHardReferences && !CookByTheBookOptions->bSkipSoftReferences)
 	{
 		UE_LOG(LogCook, Warning, TEXT("Setting bSkipSoftReferences to true since bSkipHardReferences is true and skipping hard references requires skipping soft references."));
 		CookByTheBookOptions->bSkipSoftReferences = true;
 	}
-	COOK_STAT(UE::SavePackageUtilities::ResetCookStats());
-	FBeginCookContext BeginContext = SetCookByTheBookOptions(CookByTheBookStartupOptions);
-
-	BlockOnAssetRegistry();
-	CreateSandboxFile();
-	SetBeginCookConfigSettings();
-	PackageDatas->BeginCook();
-	if (!IsCookingInEditor())
-	{
-		FCoreUObjectDelegates::PackageCreatedForLoad.AddUObject(this, &UCookOnTheFlyServer::MaybeMarkPackageAsAlreadyLoaded);
-	}
-
-	SelectSessionPlatforms(BeginContext);
-	SetBeginCookIterativeFlags(BeginContext);
-	TArray<ITargetPlatform*> TargetPlatforms = BeginContext.TargetPlatforms;
-	check(PlatformManager->GetSessionPlatforms().Num() == TargetPlatforms.Num());
-
-	// We want to set bRunning = true as early as possible, but it implies that session platforms have been selected so this is the earliest point we can set it
-	CookByTheBookOptions->bRunning = true;
-
-	RefreshPlatformAssetRegistries(TargetPlatforms);
-
-	const UProjectPackagingSettings* const PackagingSettings = GetDefault<UProjectPackagingSettings>();
-
-	// Find all the localized packages and map them back to their source package
-	{
-		TArray<FString> AllCulturesToCook = CookByTheBookStartupOptions.CookCultures;
-		for (const FString& CultureName : CookByTheBookStartupOptions.CookCultures)
-		{
-			const TArray<FString> PrioritizedCultureNames = FInternationalization::Get().GetPrioritizedCultureNames(CultureName);
-			for (const FString& PrioritizedCultureName : PrioritizedCultureNames)
-			{
-				AllCulturesToCook.AddUnique(PrioritizedCultureName);
-			}
-		}
-		AllCulturesToCook.Sort();
-
-		UE_LOG(LogCook, Display, TEXT("Discovering localized assets for cultures: %s"), *FString::Join(AllCulturesToCook, TEXT(", ")));
-
-		TArray<FString> RootPaths;
-		FPackageName::QueryRootContentPaths(RootPaths);
-
-		FARFilter Filter;
-		Filter.bRecursivePaths = true;
-		Filter.bIncludeOnlyOnDiskAssets = false;
-		Filter.PackagePaths.Reserve(AllCulturesToCook.Num() * RootPaths.Num());
-		for (const FString& RootPath : RootPaths)
-		{
-			for (const FString& CultureName : AllCulturesToCook)
-			{
-				FString LocalizedPackagePath = RootPath / TEXT("L10N") / CultureName;
-				Filter.PackagePaths.Add(*LocalizedPackagePath);
-			}
-		}
-
-		TArray<FAssetData> AssetDataForCultures;
-		AssetRegistry->GetAssets(Filter, AssetDataForCultures);
-
-		for (const FAssetData& AssetData : AssetDataForCultures)
-		{
-			const FName LocalizedPackageName = AssetData.PackageName;
-			const FName SourcePackageName = *FPackageName::GetSourcePackagePath(LocalizedPackageName.ToString());
-
-			TArray<FName>& LocalizedPackageNames = CookByTheBookOptions->SourceToLocalizedPackageVariants.FindOrAdd(SourcePackageName);
-			LocalizedPackageNames.AddUnique(LocalizedPackageName);
-		}
-
-		CookByTheBookOptions->AllCulturesToCook = MoveTemp(AllCulturesToCook);
-	}
-
-	PackageTracker->NeverCookPackageList.Empty();
-	for (FName NeverCookPackage : GetNeverCookPackageFileNames(CookByTheBookStartupOptions.NeverCookDirectories))
-	{
-		PackageTracker->NeverCookPackageList.Add(NeverCookPackage);
-	}
 
-	// use temp list of UBT platform strings to discover PlatformSpecificNeverCookPackages
-	{
-		TArray<FString> UBTPlatformStrings;
-		UBTPlatformStrings.Reserve(TargetPlatforms.Num());
-		for (const ITargetPlatform* Platform : TargetPlatforms)
-		{
-			FString UBTPlatformName;
-			Platform->GetPlatformInfo().UBTPlatformName.ToString(UBTPlatformName);
-			UBTPlatformStrings.Emplace(MoveTemp(UBTPlatformName));
-		}
-
-		DiscoverPlatformSpecificNeverCookPackages(TargetPlatforms, UBTPlatformStrings);
-	}
-
-	FindOrCreateSaveContexts(TargetPlatforms);
-	// TODO: Fix the elements of StartCookByTheBook that rely on BeginCookSandbox being called early, and remove this call to BeginCookSandbox
-	if (!bHybridIterativeEnabled)
-	{
-		BeginCookSandbox(BeginContext); // This will either delete the sandbox or iteratively clean it
-	}
-	InitializeTargetPlatforms(TargetPlatforms);
-
-	if (CurrentCookMode == ECookMode::CookByTheBook && !IsCookFlagSet(ECookInitializationFlags::Iterative))
-	{
-		UE::SavePackageUtilities::StartSavingEDLCookInfoForVerification();
-	}
+	BeginContext.TargetPlatforms = StartupOptions.TargetPlatforms;
+	Algo::Sort(BeginContext.TargetPlatforms);
+	BeginContext.TargetPlatforms.SetNum(Algo::Unique(BeginContext.TargetPlatforms));
 
+	BeginContext.PlatformContexts.SetNum(BeginContext.TargetPlatforms.Num());
+	for (int32 Index = 0; Index < BeginContext.TargetPlatforms.Num(); ++Index)
 	{
-		if (CookByTheBookOptions->bGenerateDependenciesForMaps)
-		{
-			for (const ITargetPlatform* Platform : TargetPlatforms)
-			{
-				CookByTheBookOptions->MapDependencyGraphs.Add(Platform);
-			}
-		}
+		BeginContext.PlatformContexts[Index].TargetPlatform = BeginContext.TargetPlatforms[Index];
+		// PlatformContext.PlatformData is currently null and is set in SelectSessionPlatforms
 	}
-	
-	BeginCookStartShaderCodeLibrary(BeginContext); // start shader code library cooking asynchronously; we block on it later
-	
-	if ( IsCookingDLC() )
-	{
-		// If we're cooking against a fixed base, we don't need to verify the packages exist on disk, we simply want to use the Release Data 
-		const bool bVerifyPackagesExist = !IsCookingAgainstFixedBase();
-		const bool bReevaluateUncookedPackages = !!(CookOptions & ECookByTheBookOptions::DlcReevaluateUncookedAssets);
-
-		// if we are cooking dlc we must be based on a release version cook
-		check( !BasedOnReleaseVersion.IsEmpty() );
-
-		auto ReadDevelopmentAssetRegistry = [this, &BasedOnReleaseVersion, bVerifyPackagesExist, bReevaluateUncookedPackages]
-			(TArray<UE::Cook::FConstructPackageData>& OutPackageList, const FString& InPlatformName)
-		{
-			TArray<FString> AttemptedNames;
-			FString OriginalSandboxRegistryFilename = GetBasedOnReleaseVersionAssetRegistryPath(BasedOnReleaseVersion, InPlatformName ) / TEXT("Metadata") / GetDevelopmentAssetRegistryFilename();
-			AttemptedNames.Add(OriginalSandboxRegistryFilename);
-
-			// if this check fails probably because the asset registry can't be found or read
-			bool bSucceeded = GetAllPackageFilenamesFromAssetRegistry(OriginalSandboxRegistryFilename, bVerifyPackagesExist, bReevaluateUncookedPackages, OutPackageList);
-			if (!bSucceeded)
-			{
-				OriginalSandboxRegistryFilename = GetBasedOnReleaseVersionAssetRegistryPath(BasedOnReleaseVersion, InPlatformName) / GetAssetRegistryFilename();
-				AttemptedNames.Add(OriginalSandboxRegistryFilename);
-				bSucceeded = GetAllPackageFilenamesFromAssetRegistry(OriginalSandboxRegistryFilename, bVerifyPackagesExist, bReevaluateUncookedPackages, OutPackageList);
-			}
 
-			if (!bSucceeded)
-			{
-				const PlatformInfo::FTargetPlatformInfo* PlatformInfo = PlatformInfo::FindPlatformInfo(*InPlatformName);
-				if (PlatformInfo)
-				{
-					for (const PlatformInfo::FTargetPlatformInfo* PlatformFlavor : PlatformInfo->Flavors)
-					{
-						OriginalSandboxRegistryFilename = GetBasedOnReleaseVersionAssetRegistryPath(BasedOnReleaseVersion, PlatformFlavor->Name.ToString()) / GetAssetRegistryFilename();
-						AttemptedNames.Add(OriginalSandboxRegistryFilename);
-						bSucceeded = GetAllPackageFilenamesFromAssetRegistry(OriginalSandboxRegistryFilename, bVerifyPackagesExist, bReevaluateUncookedPackages, OutPackageList);
-						if (bSucceeded)
-						{
-							break;
-						}
-					}
-				}
-			}
-
-			if (bSucceeded)
-			{
-				UE_LOG(LogCook, Log, TEXT("Loaded assetregistry: %s"), *OriginalSandboxRegistryFilename);
-			}
-			else
-			{
-				UE_LOG(LogCook, Log, TEXT("Failed to load DevelopmentAssetRegistry for platform %s. Attempted the following names:\n%s"), *InPlatformName, *FString::Join(AttemptedNames, TEXT("\n")));
-			}
-			return bSucceeded;
-		};
+	return BeginContext;
+}
 
-		TArray<UE::Cook::FConstructPackageData> OverridePackageList;
-		FString DevelopmentAssetRegistryPlatformOverride;
-		const bool bUsingDevRegistryOverride = FParse::Value(FCommandLine::Get(), TEXT("DevelopmentAssetRegistryPlatformOverride="), DevelopmentAssetRegistryPlatformOverride);
-		if (bUsingDevRegistryOverride)
-		{
-			// Read the contents of the asset registry for the overriden platform. We'll use this for all requested platforms so we can just keep one copy of it here
-			bool bReadSucceeded = ReadDevelopmentAssetRegistry(OverridePackageList, *DevelopmentAssetRegistryPlatformOverride);
-			if (!bReadSucceeded || OverridePackageList.Num() == 0)
-			{
-				UE_LOG(LogCook, Fatal, TEXT("%s based-on AssetRegistry file %s for DevelopmentAssetRegistryPlatformOverride %s. ")
-					TEXT("When cooking DLC, if DevelopmentAssetRegistryPlatformOverride is specified %s is expected to exist under Release/<override> and contain some valid data. Terminating the cook."),
-					!bReadSucceeded ? TEXT("Could not find") : TEXT("Empty"),
-					*(GetBasedOnReleaseVersionAssetRegistryPath(BasedOnReleaseVersion, DevelopmentAssetRegistryPlatformOverride) / TEXT("Metadata") / GetAssetRegistryFilename()),
-					*DevelopmentAssetRegistryPlatformOverride, *GetAssetRegistryFilename());
-			}
-		}
+FBeginCookContext UCookOnTheFlyServer::CreateBeginCookOnTheFlyContext()
+{
+	return FBeginCookContext();
+}
 
-		for ( const ITargetPlatform* TargetPlatform: TargetPlatforms )
-		{
-			SCOPED_BOOT_TIMING("AddCookedPlatforms");
-			TArray<UE::Cook::FConstructPackageData> PackageList;
-			FString PlatformNameString = TargetPlatform->PlatformName();
-			FName PlatformName(*PlatformNameString);
+FBeginCookContext UCookOnTheFlyServer::CreateAddPlatformContext(ITargetPlatform* TargetPlatform)
+{
+	FBeginCookContext BeginContext;
 
-			if (!bUsingDevRegistryOverride)
-			{
-				bool bReadSucceeded = ReadDevelopmentAssetRegistry(PackageList, PlatformNameString);
-				if (!bReadSucceeded)
-				{
-					UE_LOG(LogCook, Fatal, TEXT("Could not find based-on AssetRegistry file %s for platform %s. ")
-						TEXT("When cooking DLC, %s is expected to exist Release/<platform> for each platform being cooked. (Or use DevelopmentAssetRegistryPlatformOverride=<PlatformName> to specify an override platform that all platforms should use to find the %s file). Terminating the cook."),
-						*(GetBasedOnReleaseVersionAssetRegistryPath(BasedOnReleaseVersion, PlatformNameString) / TEXT("Metadata") / GetAssetRegistryFilename()),
-						*PlatformNameString, *GetAssetRegistryFilename(), *GetAssetRegistryFilename());
-				}
-			}
+	BeginContext.TargetPlatforms.Add(TargetPlatform);
 
-			TArray<UE::Cook::FConstructPackageData>& ActivePackageList = OverridePackageList.Num() > 0 ? OverridePackageList : PackageList;
-			if (ActivePackageList.Num() > 0)
-			{
-				SCOPED_BOOT_TIMING("AddPackageDataByFileNamesForPlatform");
-				PackageDatas->AddExistingPackageDatasForPlatform(ActivePackageList, TargetPlatform);
-			}
+	FBeginCookContextPlatform& PlatformContext = BeginContext.PlatformContexts.Emplace_GetRef();
+	PlatformContext.TargetPlatform = TargetPlatform;
+	PlatformContext.PlatformData = &PlatformManager->CreatePlatformData(TargetPlatform);
 
-			TArray<FName>& PlatformBasedPackages = CookByTheBookOptions->BasedOnReleaseCookedPackages.FindOrAdd(PlatformName);
-			PlatformBasedPackages.Reset(ActivePackageList.Num());
-			for (UE::Cook::FConstructPackageData& PackageData : ActivePackageList)
-			{
-				PlatformBasedPackages.Add(PackageData.NormalizedFileName);
-			}
-		}
+	return BeginContext;
+}
 
-		FString ExtraReleaseVersionAssetsFile;
-		const bool bUsingExtraReleaseVersionAssets = FParse::Value(FCommandLine::Get(), TEXT("ExtraReleaseVersionAssets="), ExtraReleaseVersionAssetsFile);
-		if (bUsingExtraReleaseVersionAssets)
-		{
-			// read AssetPaths out of the file and add them as already-cooked PackageDatas
-			TArray<FString> OutAssetPaths;
-			FFileHelper::LoadFileToStringArray(OutAssetPaths, *ExtraReleaseVersionAssetsFile);
-			for (const FString& AssetPath : OutAssetPaths)
-			{
-				if (UE::Cook::FPackageData* PackageData = PackageDatas->TryAddPackageDataByFileName(FName(*AssetPath)))
-				{
-					PackageData->SetPlatformsCooked(TargetPlatforms, true /* Succeeded */);
-				}
-				else
-				{
-					UE_LOG(LogCook, Error, TEXT("Failed to resolve package data for ExtraReleaseVersionAsset [%s]"), *AssetPath);
-				}
-			}
-		}
-	}
-	
+void UCookOnTheFlyServer::GenerateInitialRequests(FBeginCookContext& BeginContext)
+{
+	TArray<ITargetPlatform*>& TargetPlatforms = BeginContext.TargetPlatforms;
 	TSet<FName> StartupSoftObjectPackages;
-	if (!IsCookByTheBookMode() || !CookByTheBookOptions->bSkipSoftReferences)
+	if (!CookByTheBookOptions->bSkipSoftReferences)
 	{
 		// Get the list of soft references, for both empty package and all startup packages
 		GRedirectCollector.ProcessSoftObjectPathPackageList(NAME_None, false, StartupSoftObjectPackages);
@@ -8201,9 +8055,13 @@ void UCookOnTheFlyServer::StartCookByTheBook( const FCookByTheBookStartupOptions
 			StartupSoftObjectPackages.Remove(AssetName);
 		}
 	}
-	
+
 	TArray<FName> FilesInPath;
 	TMap<FName, UE::Cook::FInstigator> FilesInPathInstigators;
+	const TArray<FString>& CookMaps = BeginContext.StartupOptions->CookMaps;
+	const TArray<FString>& CookDirectories = BeginContext.StartupOptions->CookDirectories;
+	const TArray<FString>& IniMapSections = BeginContext.StartupOptions->IniMapSections;
+	ECookByTheBookOptions CookOptions = BeginContext.StartupOptions->CookOptions;
 	CollectFilesToCook(FilesInPath, FilesInPathInstigators, CookMaps, CookDirectories, IniMapSections, CookOptions, TargetPlatforms, GameDefaultObjects);
 
 	// Add soft/hard startup references after collecting requested files and handling empty requests
@@ -8230,13 +8088,13 @@ void UCookOnTheFlyServer::StartCookByTheBook( const FCookByTheBookStartupOptions
 				UE::Cook::EInstigator::StartupSoftObjectPath);
 		}
 	}
-	
+
 	if (FilesInPath.Num() == 0)
 	{
 		LogCookerMessage(FString::Printf(TEXT("No files found to cook.")), EMessageSeverity::Warning);
 	}
 
-	if (FParse::Param(FCommandLine::Get(), TEXT("RANDOMPACKAGEORDER")) || 
+	if (FParse::Param(FCommandLine::Get(), TEXT("RANDOMPACKAGEORDER")) ||
 		(FParse::Param(FCommandLine::Get(), TEXT("DIFFONLY")) && !FParse::Param(FCommandLine::Get(), TEXT("DIFFNORANDCOOK"))))
 	{
 		UE_LOG(LogCook, Log, TEXT("Randomizing package order."));
@@ -8260,7 +8118,7 @@ void UCookOnTheFlyServer::StartCookByTheBook( const FCookByTheBookStartupOptions
 		}
 
 		const FName PackageFileFName = PackageDatas->GetFileNameByPackageName(PackageName);
-		
+
 		if (!PackageFileFName.IsNone())
 		{
 			UE::Cook::FInstigator& Instigator = FilesInPathInstigators.FindChecked(PackageName);
@@ -8268,92 +8126,180 @@ void UCookOnTheFlyServer::StartCookByTheBook( const FCookByTheBookStartupOptions
 		}
 		else if (!FLinkerLoad::IsKnownMissingPackage(PackageName))
 		{
-			LogCookerMessage( FString::Printf(TEXT("Unable to find package for cooking %s"), *PackageName.ToString()),
-				EMessageSeverity::Warning );
+			LogCookerMessage(FString::Printf(TEXT("Unable to find package for cooking %s"), *PackageName.ToString()),
+				EMessageSeverity::Warning);
 		}
 	}
 
-	if (!IsCookingDLC())
+	const FString& CreateReleaseVersion = BeginContext.StartupOptions->CreateReleaseVersion;
+	const FString& BasedOnReleaseVersion = BeginContext.StartupOptions->BasedOnReleaseVersion;
+	if (!IsCookingDLC() && !BasedOnReleaseVersion.IsEmpty())
 	{
-		// if we are not cooking dlc then basedOnRelease version just needs to make sure that we cook all the packages which are in the previous release (as well as the new ones)
-		if ( !BasedOnReleaseVersion.IsEmpty() )
-		{
-			// if we are based on a release and we are not cooking dlc then we should always be creating a new one (note that we could be creating the same one we are based on).
-			// note that we might erroneously enter here if we are generating a patch instead and we accidentally passed in BasedOnReleaseVersion to the cooker instead of to unrealpak
-			UE_CLOG(CreateReleaseVersion.IsEmpty(), LogCook, Fatal, TEXT("-BasedOnReleaseVersion must be used together with either -dlcname or -CreateReleaseVersion."));
+		// if we are based on a release and we are not cooking dlc then we should always be creating a new one (note that we could be creating the same one we are based on).
+		// note that we might erroneously enter here if we are generating a patch instead and we accidentally passed in BasedOnReleaseVersion to the cooker instead of to unrealpak
+		UE_CLOG(CreateReleaseVersion.IsEmpty(), LogCook, Fatal, TEXT("-BasedOnReleaseVersion must be used together with either -dlcname or -CreateReleaseVersion."));
 
-			for ( const ITargetPlatform* TargetPlatform : TargetPlatforms )
-			{
-				// if we are based of a cook and we are creating a new one we need to make sure that at least all the old packages are cooked as well as the new ones
-				FString OriginalAssetRegistryPath = GetBasedOnReleaseVersionAssetRegistryPath( BasedOnReleaseVersion, TargetPlatform->PlatformName() ) / GetAssetRegistryFilename();
+		// if we are creating a new Release then we need cook all the packages which are in the previous release (as well as the new ones)
+		for (const ITargetPlatform* TargetPlatform : TargetPlatforms)
+		{
+			// if we are based of a cook and we are creating a new one we need to make sure that at least all the old packages are cooked as well as the new ones
+			FString OriginalAssetRegistryPath = GetBasedOnReleaseVersionAssetRegistryPath(BasedOnReleaseVersion, TargetPlatform->PlatformName()) / GetAssetRegistryFilename();
 
-				TArray<UE::Cook::FConstructPackageData> BasedOnReleaseDatas;
-				verify( GetAllPackageFilenamesFromAssetRegistry(OriginalAssetRegistryPath, true, false, BasedOnReleaseDatas) );
+			TArray<UE::Cook::FConstructPackageData> BasedOnReleaseDatas;
+			verify(GetAllPackageFilenamesFromAssetRegistry(OriginalAssetRegistryPath, true, false, BasedOnReleaseDatas));
 
-				TArray<const ITargetPlatform*, TInlineAllocator<1>> RequestPlatforms;
-				RequestPlatforms.Add(TargetPlatform);
-				for ( const UE::Cook::FConstructPackageData& PackageData : BasedOnReleaseDatas)
-				{
-					ExternalRequests->EnqueueUnique(
-						UE::Cook::FFilePlatformRequest(PackageData.NormalizedFileName,
-							UE::Cook::EInstigator::PreviousAssetRegistry, RequestPlatforms));
-				}
+			TArray<const ITargetPlatform*, TInlineAllocator<1>> RequestPlatforms;
+			RequestPlatforms.Add(TargetPlatform);
+			for (const UE::Cook::FConstructPackageData& PackageData : BasedOnReleaseDatas)
+			{
+				ExternalRequests->EnqueueUnique(
+					UE::Cook::FFilePlatformRequest(PackageData.NormalizedFileName,
+						UE::Cook::EInstigator::PreviousAssetRegistry, RequestPlatforms));
 			}
 		}
 	}
+}
 
-	if (bHybridIterativeEnabled)
+void UCookOnTheFlyServer::RecordDLCPackagesFromBaseGame(FBeginCookContext& BeginContext)
+{
+	if (!IsCookingDLC())
 	{
-		BeginCookSandbox(BeginContext);
+		return;
 	}
-	BeginCookEditorSystems();
 
-	if (bHybridIterativeDebug)
+	const ECookByTheBookOptions& CookOptions = BeginContext.StartupOptions->CookOptions;
+	const FString& BasedOnReleaseVersion = BeginContext.StartupOptions->BasedOnReleaseVersion;
+
+	// If we're cooking against a fixed base, we don't need to verify the packages exist on disk, we simply want to use the Release Data 
+	const bool bVerifyPackagesExist = !IsCookingAgainstFixedBase();
+	const bool bReevaluateUncookedPackages = !!(CookOptions & ECookByTheBookOptions::DlcReevaluateUncookedAssets);
+
+	// if we are cooking dlc we must be based on a release version cook
+	check(!BasedOnReleaseVersion.IsEmpty());
+
+	auto ReadDevelopmentAssetRegistry = [this, &BasedOnReleaseVersion, bVerifyPackagesExist, bReevaluateUncookedPackages]
+	(TArray<UE::Cook::FConstructPackageData>& OutPackageList, const FString& InPlatformName)
 	{
-		// Discoveries during the processing of the initial cluster are expected, so LogDiscoveredPackages must be off.
-		PackageDatas->SetLogDiscoveredPackages(false);
-	}
+		TArray<FString> AttemptedNames;
+		FString OriginalSandboxRegistryFilename = GetBasedOnReleaseVersionAssetRegistryPath(BasedOnReleaseVersion, InPlatformName) / TEXT("Metadata") / GetDevelopmentAssetRegistryFilename();
+		AttemptedNames.Add(OriginalSandboxRegistryFilename);
 
-	for (const ITargetPlatform* TargetPlatform : TargetPlatforms)
+		// if this check fails probably because the asset registry can't be found or read
+		bool bSucceeded = GetAllPackageFilenamesFromAssetRegistry(OriginalSandboxRegistryFilename, bVerifyPackagesExist, bReevaluateUncookedPackages, OutPackageList);
+		if (!bSucceeded)
+		{
+			OriginalSandboxRegistryFilename = GetBasedOnReleaseVersionAssetRegistryPath(BasedOnReleaseVersion, InPlatformName) / GetAssetRegistryFilename();
+			AttemptedNames.Add(OriginalSandboxRegistryFilename);
+			bSucceeded = GetAllPackageFilenamesFromAssetRegistry(OriginalSandboxRegistryFilename, bVerifyPackagesExist, bReevaluateUncookedPackages, OutPackageList);
+		}
+
+		if (!bSucceeded)
+		{
+			const PlatformInfo::FTargetPlatformInfo* PlatformInfo = PlatformInfo::FindPlatformInfo(*InPlatformName);
+			if (PlatformInfo)
+			{
+				for (const PlatformInfo::FTargetPlatformInfo* PlatformFlavor : PlatformInfo->Flavors)
+				{
+					OriginalSandboxRegistryFilename = GetBasedOnReleaseVersionAssetRegistryPath(BasedOnReleaseVersion, PlatformFlavor->Name.ToString()) / GetAssetRegistryFilename();
+					AttemptedNames.Add(OriginalSandboxRegistryFilename);
+					bSucceeded = GetAllPackageFilenamesFromAssetRegistry(OriginalSandboxRegistryFilename, bVerifyPackagesExist, bReevaluateUncookedPackages, OutPackageList);
+					if (bSucceeded)
+					{
+						break;
+					}
+				}
+			}
+		}
+
+		if (bSucceeded)
+		{
+			UE_LOG(LogCook, Log, TEXT("Loaded assetregistry: %s"), *OriginalSandboxRegistryFilename);
+		}
+		else
+		{
+			UE_LOG(LogCook, Log, TEXT("Failed to load DevelopmentAssetRegistry for platform %s. Attempted the following names:\n%s"), *InPlatformName, *FString::Join(AttemptedNames, TEXT("\n")));
+		}
+		return bSucceeded;
+	};
+
+	TArray<UE::Cook::FConstructPackageData> OverridePackageList;
+	FString DevelopmentAssetRegistryPlatformOverride;
+	const bool bUsingDevRegistryOverride = FParse::Value(FCommandLine::Get(), TEXT("DevelopmentAssetRegistryPlatformOverride="), DevelopmentAssetRegistryPlatformOverride);
+	if (bUsingDevRegistryOverride)
 	{
-		FindOrCreatePackageWriter(TargetPlatform).BeginCook();
+		// Read the contents of the asset registry for the overriden platform. We'll use this for all requested platforms so we can just keep one copy of it here
+		bool bReadSucceeded = ReadDevelopmentAssetRegistry(OverridePackageList, *DevelopmentAssetRegistryPlatformOverride);
+		if (!bReadSucceeded || OverridePackageList.Num() == 0)
+		{
+			UE_LOG(LogCook, Fatal, TEXT("%s based-on AssetRegistry file %s for DevelopmentAssetRegistryPlatformOverride %s. ")
+				TEXT("When cooking DLC, if DevelopmentAssetRegistryPlatformOverride is specified %s is expected to exist under Release/<override> and contain some valid data. Terminating the cook."),
+				!bReadSucceeded ? TEXT("Could not find") : TEXT("Empty"),
+				*(GetBasedOnReleaseVersionAssetRegistryPath(BasedOnReleaseVersion, DevelopmentAssetRegistryPlatformOverride) / TEXT("Metadata") / GetAssetRegistryFilename()),
+				*DevelopmentAssetRegistryPlatformOverride, *GetAssetRegistryFilename());
+		}
 	}
 
-	InitializePollables();
+	for (const ITargetPlatform* TargetPlatform : BeginContext.TargetPlatforms)
+	{
+		SCOPED_BOOT_TIMING("AddCookedPlatforms");
+		TArray<UE::Cook::FConstructPackageData> PackageList;
+		FString PlatformNameString = TargetPlatform->PlatformName();
+		FName PlatformName(*PlatformNameString);
 
-	BeginCookFinishShaderCodeLibrary(BeginContext);
-}
+		if (!bUsingDevRegistryOverride)
+		{
+			bool bReadSucceeded = ReadDevelopmentAssetRegistry(PackageList, PlatformNameString);
+			if (!bReadSucceeded)
+			{
+				UE_LOG(LogCook, Fatal, TEXT("Could not find based-on AssetRegistry file %s for platform %s. ")
+					TEXT("When cooking DLC, %s is expected to exist Release/<platform> for each platform being cooked. (Or use DevelopmentAssetRegistryPlatformOverride=<PlatformName> to specify an override platform that all platforms should use to find the %s file). Terminating the cook."),
+					*(GetBasedOnReleaseVersionAssetRegistryPath(BasedOnReleaseVersion, PlatformNameString) / TEXT("Metadata") / GetAssetRegistryFilename()),
+					*PlatformNameString, *GetAssetRegistryFilename(), *GetAssetRegistryFilename());
+			}
+		}
 
-FBeginCookContext UCookOnTheFlyServer::SetCookByTheBookOptions(const FCookByTheBookStartupOptions& StartupOptions)
-{
-	FBeginCookContext BeginContext;
+		TArray<UE::Cook::FConstructPackageData>& ActivePackageList = OverridePackageList.Num() > 0 ? OverridePackageList : PackageList;
+		if (ActivePackageList.Num() > 0)
+		{
+			SCOPED_BOOT_TIMING("AddPackageDataByFileNamesForPlatform");
+			PackageDatas->AddExistingPackageDatasForPlatform(ActivePackageList, TargetPlatform);
+		}
 
-	BeginContext.StartupOptions = &StartupOptions;
-	BeginContext.TargetPlatforms = StartupOptions.TargetPlatforms;
-	Algo::Sort(BeginContext.TargetPlatforms);
-	BeginContext.TargetPlatforms.SetNum(Algo::Unique(BeginContext.TargetPlatforms));
+		TArray<FName>& PlatformBasedPackages = CookByTheBookOptions->BasedOnReleaseCookedPackages.FindOrAdd(PlatformName);
+		PlatformBasedPackages.Reset(ActivePackageList.Num());
+		for (UE::Cook::FConstructPackageData& PackageData : ActivePackageList)
+		{
+			PlatformBasedPackages.Add(PackageData.NormalizedFileName);
+		}
+	}
 
-	BeginContext.PlatformContexts.SetNum(BeginContext.TargetPlatforms.Num());
-	for (int32 Index = 0; Index < BeginContext.TargetPlatforms.Num(); ++Index)
+	FString ExtraReleaseVersionAssetsFile;
+	const bool bUsingExtraReleaseVersionAssets = FParse::Value(FCommandLine::Get(), TEXT("ExtraReleaseVersionAssets="), ExtraReleaseVersionAssetsFile);
+	if (bUsingExtraReleaseVersionAssets)
 	{
-		BeginContext.PlatformContexts[Index].TargetPlatform = BeginContext.TargetPlatforms[Index];
-		// PlatformContext.PlatformData is currently null and is set in SelectSessionPlatforms
+		// read AssetPaths out of the file and add them as already-cooked PackageDatas
+		TArray<FString> OutAssetPaths;
+		FFileHelper::LoadFileToStringArray(OutAssetPaths, *ExtraReleaseVersionAssetsFile);
+		for (const FString& AssetPath : OutAssetPaths)
+		{
+			if (UE::Cook::FPackageData* PackageData = PackageDatas->TryAddPackageDataByFileName(FName(*AssetPath)))
+			{
+				PackageData->SetPlatformsCooked(BeginContext.TargetPlatforms, true /* Succeeded */);
+			}
+			else
+			{
+				UE_LOG(LogCook, Error, TEXT("Failed to resolve package data for ExtraReleaseVersionAsset [%s]"), *AssetPath);
+			}
+		}
 	}
-
-	return BeginContext;
 }
 
-FBeginCookContext UCookOnTheFlyServer::CreateAddPlatformContext(ITargetPlatform* TargetPlatform)
+void UCookOnTheFlyServer::BeginCookPackageWriters(FBeginCookContext& BeginContext)
 {
-	FBeginCookContext BeginContext;
-
-	BeginContext.TargetPlatforms.Add(TargetPlatform);
-
-	FBeginCookContextPlatform& PlatformContext = BeginContext.PlatformContexts.Emplace_GetRef();
-	PlatformContext.TargetPlatform = TargetPlatform;
-	PlatformContext.PlatformData = &PlatformManager->CreatePlatformData(TargetPlatform);
-
-	return BeginContext;
+	for (const ITargetPlatform* TargetPlatform : BeginContext.TargetPlatforms)
+	{
+		FindOrCreatePackageWriter(TargetPlatform).BeginCook();
+	}
 }
 
 void UCookOnTheFlyServer::SelectSessionPlatforms(FBeginCookContext& BeginContext)
@@ -8377,6 +8323,14 @@ void UCookOnTheFlyServer::BeginCookEditorSystems()
 		return;
 	}
 
+	if (IsCookByTheBookMode())
+	{
+		//force precache objects to refresh themselves before cooking anything
+		LastUpdateTick = INT_MAX;
+
+		COOK_STAT(UE::SavePackageUtilities::ResetCookStats());
+	}
+
 	// Notify AssetRegistry to update itself for any saved packages
 	if (!bFirstCookInThisProcess)
 	{
@@ -8391,6 +8345,39 @@ void UCookOnTheFlyServer::BeginCookEditorSystems()
 	ModifiedAssetFilenames.Empty();
 }
 
+void UCookOnTheFlyServer::BeginCookEDLCookInfo(FBeginCookContext& BeginContext)
+{
+	if (IsCookingInEditor())
+	{
+		return;
+	}
+	if (!Algo::AllOf(BeginContext.PlatformContexts, [this](const FBeginCookContextPlatform& PlatformContext)
+		{
+			return PlatformContext.PlatformData->bFullBuild;
+		}))
+	{
+		return;
+	}
+	UE::SavePackageUtilities::StartSavingEDLCookInfoForVerification();
+}
+
+void UCookOnTheFlyServer::RegisterCookByTheBookDelegates()
+{
+	if (!IsCookingInEditor())
+	{
+		FCoreUObjectDelegates::PackageCreatedForLoad.AddUObject(this, &UCookOnTheFlyServer::MaybeMarkPackageAsAlreadyLoaded);
+	}
+}
+
+void UCookOnTheFlyServer::UnregisterCookByTheBookDelegates()
+{
+	if (!IsCookingInEditor())
+	{
+		FCoreUObjectDelegates::PackageCreatedForLoad.RemoveAll(this);
+	}
+}
+
+
 TArray<FName> UCookOnTheFlyServer::GetNeverCookPackageFileNames(TArrayView<const FString> ExtraNeverCookDirectories)
 {
 	TArray<FString> NeverCookDirectories(ExtraNeverCookDirectories);
@@ -9387,6 +9374,62 @@ void UCookOnTheFlyServer::PropertyImportBehaviorCallback(const FObjectImport& Im
 	}
 }
 
+void UCookOnTheFlyServer::GenerateLocalizationReferences(TConstArrayView<FString> CookCultures)
+{
+	if (!CookByTheBookOptions)
+	{
+		return;
+	}
+	CookByTheBookOptions->SourceToLocalizedPackageVariants.Reset();
+	CookByTheBookOptions->AllCulturesToCook.Reset();
+
+	const UProjectPackagingSettings* const PackagingSettings = GetDefault<UProjectPackagingSettings>();
+
+	// Find all the localized packages and map them back to their source package
+	TArray<FString> AllCulturesToCook(CookCultures);
+	for (const FString& CultureName : CookCultures)
+	{
+		const TArray<FString> PrioritizedCultureNames = FInternationalization::Get().GetPrioritizedCultureNames(CultureName);
+		for (const FString& PrioritizedCultureName : PrioritizedCultureNames)
+		{
+			AllCulturesToCook.AddUnique(PrioritizedCultureName);
+		}
+	}
+	AllCulturesToCook.Sort();
+
+	UE_LOG(LogCook, Display, TEXT("Discovering localized assets for cultures: %s"), *FString::Join(AllCulturesToCook, TEXT(", ")));
+
+	TArray<FString> RootPaths;
+	FPackageName::QueryRootContentPaths(RootPaths);
+
+	FARFilter Filter;
+	Filter.bRecursivePaths = true;
+	Filter.bIncludeOnlyOnDiskAssets = false;
+	Filter.PackagePaths.Reserve(AllCulturesToCook.Num() * RootPaths.Num());
+	for (const FString& RootPath : RootPaths)
+	{
+		for (const FString& CultureName : AllCulturesToCook)
+		{
+			FString LocalizedPackagePath = RootPath / TEXT("L10N") / CultureName;
+			Filter.PackagePaths.Add(*LocalizedPackagePath);
+		}
+	}
+
+	TArray<FAssetData> AssetDataForCultures;
+	AssetRegistry->GetAssets(Filter, AssetDataForCultures);
+
+	for (const FAssetData& AssetData : AssetDataForCultures)
+	{
+		const FName LocalizedPackageName = AssetData.PackageName;
+		const FName SourcePackageName = *FPackageName::GetSourcePackagePath(LocalizedPackageName.ToString());
+
+		TArray<FName>& LocalizedPackageNames = CookByTheBookOptions->SourceToLocalizedPackageVariants.FindOrAdd(SourcePackageName);
+		LocalizedPackageNames.AddUnique(LocalizedPackageName);
+	}
+
+	CookByTheBookOptions->AllCulturesToCook = MoveTemp(AllCulturesToCook);
+}
+
 void UCookOnTheFlyServer::RegisterLocalizationChunkDataGenerator()
 {
 	if (!CookByTheBookOptions)
diff --git a/Engine/Source/Editor/UnrealEd/Private/Cooker/CookPackageData.cpp b/Engine/Source/Editor/UnrealEd/Private/Cooker/CookPackageData.cpp
--- a/Engine/Source/Editor/UnrealEd/Private/Cooker/CookPackageData.cpp
+++ b/Engine/Source/Editor/UnrealEd/Private/Cooker/CookPackageData.cpp
@@ -1596,7 +1596,7 @@ FPackageDatas::FPackageDatas(UCookOnTheFlyServer& InCookOnTheFlyServer)
 {
 }
 
-void FPackageDatas::BeginCook()
+void FPackageDatas::SetBeginCookConfigSettings()
 {
 	FString FileOrPackageName;
 	ShowInstigatorPackageData = nullptr;
@@ -1620,6 +1620,9 @@ void FPackageDatas::BeginCook()
 			}
 		}
 	}
+
+	// Discoveries during the processing of the initial cluster are expected, so LogDiscoveredPackages must be off.
+	SetLogDiscoveredPackages(false);
 }
 
 FPackageDatas::~FPackageDatas()
diff --git a/Engine/Source/Editor/UnrealEd/Private/Cooker/CookPackageData.h b/Engine/Source/Editor/UnrealEd/Private/Cooker/CookPackageData.h
--- a/Engine/Source/Editor/UnrealEd/Private/Cooker/CookPackageData.h
+++ b/Engine/Source/Editor/UnrealEd/Private/Cooker/CookPackageData.h
@@ -872,8 +872,8 @@ public:
 	/** Called when the initial AssetRegistry search is done and it can be used to determine package existence. */
 	static void OnAssetRegistryGenerated(IAssetRegistry& InAssetRegistry);
 
-	/** Called each time BeginCook is called, for deferred intialization steps. */
-	void BeginCook();
+	/** Called each time BeginCook is called, to initialize settings from config */
+	void SetBeginCookConfigSettings();
 
 	/** FGCObject interface function - return a debug name describing this FGCObject. */
 	virtual FString GetReferencerName() const override;
diff --git a/Engine/Source/Editor/UnrealEd/Private/Cooker/CookTypes.h b/Engine/Source/Editor/UnrealEd/Private/Cooker/CookTypes.h
--- a/Engine/Source/Editor/UnrealEd/Private/Cooker/CookTypes.h
+++ b/Engine/Source/Editor/UnrealEd/Private/Cooker/CookTypes.h
@@ -256,9 +256,6 @@ public:
 	/** Create a release from this manifest and store it in the releases directory for this cgame */
 	FString							CreateReleaseVersion;
 
-	/** Dependency graph of maps as root objects. */
-	TFastPointerMap<const ITargetPlatform*, TMap<FName, TSet<FName>>> MapDependencyGraphs;
-
 	/** If we are based on a release version of the game this is the set of packages which were cooked in that release. Map from platform name to list of uncooked package filenames */
 	TMap<FName, TArray<FName>>			BasedOnReleaseCookedPackages;
 
