add support for plugin content to FAutomationEditorCommonUtils::ConvertPackagePathToAssetPath() #preflight 6270ceceae03caabf70ed2ea

diff --git a/Engine/Source/Editor/UnrealEd/Private/Tests/AutomationEditorCommon.cpp b/Engine/Source/Editor/UnrealEd/Private/Tests/AutomationEditorCommon.cpp
--- a/Engine/Source/Editor/UnrealEd/Private/Tests/AutomationEditorCommon.cpp
+++ b/Engine/Source/Editor/UnrealEd/Private/Tests/AutomationEditorCommon.cpp
@@ -533,6 +533,8 @@ FString FAutomationEditorCommonUtils::ConvertPackagePathToAssetPath(const FStrin
 	const FString Filename = FPaths::ConvertRelativePathToFull(PackagePath);
 	FString EngineFileName = Filename;
 	FString GameFileName = Filename;
+	FString ProjectPluginFileName = Filename;
+	FString EnginePluginFileName = Filename;
 	if (FPaths::MakePathRelativeTo(EngineFileName, *FPaths::EngineContentDir()) && !EngineFileName.Contains(TEXT("../")))
 	{
 		const FString ShortName = FPaths::GetBaseFilename(EngineFileName);
@@ -547,6 +549,22 @@ FString FAutomationEditorCommonUtils::ConvertPackagePathToAssetPath(const FStrin
 		const FString AssetName = FString::Printf(TEXT("/Game/%s/%s.%s"), *PathName, *ShortName, *ShortName);
 		return AssetName;
 	}
+	else if (FPaths::MakePathRelativeTo(ProjectPluginFileName, *FPaths::ProjectPluginsDir()) && !ProjectPluginFileName.Contains(TEXT("../")))
+	{
+		const FString ShortName = FPaths::GetBaseFilename(ProjectPluginFileName);
+		const FString FullPathName = FPaths::GetPath(ProjectPluginFileName);
+		const FString CleanedPathName = FullPathName.Replace(TEXT("Content/"), TEXT(""));
+		const FString AssetName = FString::Printf(TEXT("/%s/%s.%s"), *CleanedPathName, *ShortName, *ShortName);
+		return AssetName;
+	}
+	else if (FPaths::MakePathRelativeTo(EnginePluginFileName, *FPaths::EnginePluginsDir()) && !EnginePluginFileName.Contains(TEXT("../")))
+	{
+		const FString ShortName = FPaths::GetBaseFilename(EnginePluginFileName);
+		const FString FullPathName = FPaths::GetPath(EnginePluginFileName);
+		const FString CleanedPathName = FullPathName.Replace(TEXT("Content/"), TEXT(""));
+		const FString AssetName = FString::Printf(TEXT("/%s/%s.%s"), *CleanedPathName, *ShortName, *ShortName);
+		return AssetName;
+	}
 	else
 	{
 		UE_LOG(LogAutomationEditorCommon, Error, TEXT("PackagePath (%s) is invalid for the current project"), *PackagePath);
