Fixed a typo in an ensure: "exanded" --> "expanded".

#UE5 #RNX #NoReleaseNotes

diff --git a/Engine/Source/Editor/UnrealEd/Private/Kismet2/KismetDebugUtilities.cpp b/Engine/Source/Editor/UnrealEd/Private/Kismet2/KismetDebugUtilities.cpp
--- a/Engine/Source/Editor/UnrealEd/Private/Kismet2/KismetDebugUtilities.cpp
+++ b/Engine/Source/Editor/UnrealEd/Private/Kismet2/KismetDebugUtilities.cpp
@@ -2126,7 +2126,7 @@ FPropertyInstanceInfo::FPropertyInstanceInfo(FPropertyInstance PropertyInstance)
 		return;
 	}
 
-	ensureMsgf(false, TEXT("Failed to identify property type. This function may need to be exanded to include it: %s"), *ResolvedProperty->GetClass()->GetName());
+	ensureMsgf(false, TEXT("Failed to identify property type. This function may need to be expanded to include it: %s"), *ResolvedProperty->GetClass()->GetName());
 }
 
 TSharedPtr<FPropertyInstanceInfo> FPropertyInstanceInfo::FindOrMake(FPropertyInstance PropertyInstance,
