EditorEngine::OnPreSaveWorld: do not CollectGarbage from EditorDomain saves. Keep the existing do-not-GC for cook saves, but extend it to cook saves when cooking in editor. Keep the existing special behavior (and do-not-GC) for autosaves. Autosaves cannot occur in commandlet, so we can simplify the branching paths somewhat.
#rnx

diff --git a/Engine/Source/Editor/UnrealEd/Private/EditorEngine.cpp b/Engine/Source/Editor/UnrealEd/Private/EditorEngine.cpp
--- a/Engine/Source/Editor/UnrealEd/Private/EditorEngine.cpp
+++ b/Engine/Source/Editor/UnrealEd/Private/EditorEngine.cpp
@@ -4549,51 +4549,48 @@ void UEditorEngine::OnPreSaveWorld(UWorld* World, FObjectPreSaveContext ObjectSa
 	// Update cull distance volumes (and associated primitives).
 	World->UpdateCullDistanceVolumes();
 
-	if ( !IsRunningCommandlet() )
+	const bool bAutosave = (ObjectSaveContext.GetSaveFlags() & SAVE_FromAutosave) != 0;
+	if (bAutosave)
 	{
-		const bool bAutosave = (ObjectSaveContext.GetSaveFlags() & SAVE_FromAutosave) != 0;
-		if ( bAutosave )
+		// Temporarily flag packages saved under a PIE filename as PKG_PlayInEditor for serialization so loading
+		// them will have the flag set. We need to undo this as the object flagged isn't actually the PIE package, 
+		// but rather only the loaded one will be.
+		// PIE prefix detected, mark package.
+		if (World->GetName().StartsWith(PLAYWORLD_PACKAGE_PREFIX))
 		{
-			// Temporarily flag packages saved under a PIE filename as PKG_PlayInEditor for serialization so loading
-			// them will have the flag set. We need to undo this as the object flagged isn't actually the PIE package, 
-			// but rather only the loaded one will be.
-			// PIE prefix detected, mark package.
-			if( World->GetName().StartsWith( PLAYWORLD_PACKAGE_PREFIX ) )
-			{
-				World->GetOutermost()->SetPackageFlags(PKG_PlayInEditor);
-			}
+			World->GetOutermost()->SetPackageFlags(PKG_PlayInEditor);
 		}
-		else
-		{
-			// Normal non-pie and non-autosave codepath
-			FWorldContext &EditorContext = GetEditorWorldContext();
+	}
+	else if ( !IsRunningCommandlet() && !ObjectSaveContext.IsProceduralSave())
+	{
+		// A user-initiated save in the editor
+		FWorldContext &EditorContext = GetEditorWorldContext();
 
-			// Check that this world is GWorld to avoid stomping on the saved views of sub-levels.
-			if ( World == EditorContext.World() )
+		// Check that this world is GWorld to avoid stomping on the saved views of sub-levels.
+		if ( World == EditorContext.World() )
+		{
+			if( FModuleManager::Get().IsModuleLoaded("LevelEditor") )
 			{
-				if( FModuleManager::Get().IsModuleLoaded("LevelEditor") )
-				{
-					FLevelEditorModule& LevelEditor = FModuleManager::GetModuleChecked<FLevelEditorModule>("LevelEditor");
+				FLevelEditorModule& LevelEditor = FModuleManager::GetModuleChecked<FLevelEditorModule>("LevelEditor");
 
-					// Notify slate level editors of the map change
-					LevelEditor.BroadcastMapChanged( World, EMapChangeType::SaveMap );
-				}
+				// Notify slate level editors of the map change
+				LevelEditor.BroadcastMapChanged( World, EMapChangeType::SaveMap );
 			}
+		}
 
-			// Shrink model and clean up deleted actors.
-			// Don't do this when autosaving or PIE saving so that actor adds can still undo.
-			World->ShrinkLevel();
+		// Shrink model and clean up deleted actors.
+		// Don't do this when autosaving or PIE saving so that actor adds can still undo.
+		World->ShrinkLevel();
 
-			{
-				FScopedSlowTask SlowTask(0, FText::Format(NSLOCTEXT("UnrealEd", "SavingMapStatus_CollectingGarbage", "Saving map: {0}... (Collecting garbage)"), FText::FromString(World->GetName())));
-				// NULL empty or "invalid" entries (e.g. IsPendingKill()) in actors array.
-				CollectGarbage( GARBAGE_COLLECTION_KEEPFLAGS );
-			}
-			
-			// Compact and sort actors array to remove empty entries.
-			// Don't do this when autosaving or PIE saving so that actor adds can still undo.
-			World->PersistentLevel->SortActorList();
+		{
+			FScopedSlowTask SlowTask(0, FText::Format(NSLOCTEXT("UnrealEd", "SavingMapStatus_CollectingGarbage", "Saving map: {0}... (Collecting garbage)"), FText::FromString(World->GetName())));
+			// NULL empty or "invalid" entries (e.g. IsPendingKill()) in actors array.
+			CollectGarbage( GARBAGE_COLLECTION_KEEPFLAGS );
 		}
+			
+		// Compact and sort actors array to remove empty entries.
+		// Don't do this when autosaving or PIE saving so that actor adds can still undo.
+		World->PersistentLevel->SortActorList();
 	}
 
 	// Move level position closer to world origin
