Cooker: Fix assertion when cooking cooked editor DLC; SessionPlatforms are not yet available from the PlatformManager so we need to pass in the BeginCookContext to get them.
#rnx

diff --git a/Engine/Source/Editor/UnrealEd/Classes/CookOnTheSide/CookOnTheFlyServer.h b/Engine/Source/Editor/UnrealEd/Classes/CookOnTheSide/CookOnTheFlyServer.h
--- a/Engine/Source/Editor/UnrealEd/Classes/CookOnTheSide/CookOnTheFlyServer.h
+++ b/Engine/Source/Editor/UnrealEd/Classes/CookOnTheSide/CookOnTheFlyServer.h
@@ -1088,9 +1088,9 @@ private:
 
 
 	/** Create sandbox file in directory using current settings supplied */
-	void CreateSandboxFile();
+	void CreateSandboxFile(FBeginCookContext& BeginContext);
 	/** Gets the output directory respecting any command line overrides */
-	FString GetOutputDirectoryOverride() const;
+	FString GetOutputDirectoryOverride(FBeginCookContext& BeginContext) const;
 
 	/**
 	 * Populate cooked packages from the PackageWriter's previous manifest and assetregistry of cooked output.
diff --git a/Engine/Source/Editor/UnrealEd/Private/CookOnTheFlyServer.cpp b/Engine/Source/Editor/UnrealEd/Private/CookOnTheFlyServer.cpp
--- a/Engine/Source/Editor/UnrealEd/Private/CookOnTheFlyServer.cpp
+++ b/Engine/Source/Editor/UnrealEd/Private/CookOnTheFlyServer.cpp
@@ -538,7 +538,7 @@ bool UCookOnTheFlyServer::StartCookOnTheFly(FCookOnTheFlyOptions InCookOnTheFlyO
 
 	CookOnTheFlyOptions = MoveTemp(InCookOnTheFlyOptions);
 	FBeginCookContext BeginContext = CreateBeginCookOnTheFlyContext();
-	CreateSandboxFile();
+	CreateSandboxFile(BeginContext);
 
 	CookOnTheFlyServerInterface = MakeUnique<UCookOnTheFlyServer::FCookOnTheFlyServerInterface>(*this);
 	ExternalRequests->CookRequestEvent = FPlatformProcess::GetSynchEventFromPool();
@@ -4970,7 +4970,7 @@ uint32 UCookOnTheFlyServer::NumConnections() const
 	return Result;
 }
 
-FString UCookOnTheFlyServer::GetOutputDirectoryOverride() const
+FString UCookOnTheFlyServer::GetOutputDirectoryOverride(FBeginCookContext& BeginContext) const
 {
 	FString OutputDirectory = OutputDirectoryOverride;
 	// Output directory override.	
@@ -4999,7 +4999,7 @@ FString UCookOnTheFlyServer::GetOutputDirectoryOverride() const
 		// Output directory needs to contain [Platform] token to be able to cook for multiple targets.
 		if ( IsCookByTheBookMode() )
 		{
-			checkf(PlatformManager->GetSessionPlatforms().Num() == 1,
+			checkf(BeginContext.TargetPlatforms.Num() == 1,
 				TEXT("If OutputDirectoryOverride is provided when cooking multiple platforms, it must include [Platform] in the text, to be replaced with the name of each of the requested Platforms.") );
 		}
 		else
@@ -7539,11 +7539,11 @@ void UCookOnTheFlyServer::OnTargetPlatformChangedSupportedFormats(const ITargetP
 	}
 }
 
-void UCookOnTheFlyServer::CreateSandboxFile()
+void UCookOnTheFlyServer::CreateSandboxFile(FBeginCookContext& BeginContext)
 {
 	// Output directory override. This directory depends on whether we are cooking dlc, so we cannot
 	// create the sandbox until after StartCookByTheBook or StartCookOnTheFly
-	FString OutputDirectory = GetOutputDirectoryOverride();
+	FString OutputDirectory = GetOutputDirectoryOverride(BeginContext);
 	check(!OutputDirectory.IsEmpty());
 	check((SandboxFile == nullptr) == SandboxFileOutputDirectory.IsEmpty());
 
@@ -7939,7 +7939,7 @@ void UCookOnTheFlyServer::StartCookByTheBook( const FCookByTheBookStartupOptions
 	// Functions in this section are ordered and can depend on the functions before them
 	FBeginCookContext BeginContext = SetCookByTheBookOptions(CookByTheBookStartupOptions);
 	BlockOnAssetRegistry();
-	CreateSandboxFile();
+	CreateSandboxFile(BeginContext);
 	SetBeginCookConfigSettings(BeginContext);
 	SelectSessionPlatforms(BeginContext);
 	SetBeginCookIterativeFlags(BeginContext);
