open BogusMapName exits PIE instead of opening the default map

diff --git a/Engine/Source/Editor/UnrealEd/Classes/Editor/EditorEngine.h b/Engine/Source/Editor/UnrealEd/Classes/Editor/EditorEngine.h
--- a/Engine/Source/Editor/UnrealEd/Classes/Editor/EditorEngine.h
+++ b/Engine/Source/Editor/UnrealEd/Classes/Editor/EditorEngine.h
@@ -3183,8 +3183,6 @@ public:
 	UE_DEPRECATED(4.25, "Use RequestPlaySession and StartQueuedPlaySessionRequest instead.")
 	void AutomationPlayUsingLauncher(const FString& InLauncherDeviceId);	
 
-	virtual void HandleTravelFailure(UWorld* InWorld, ETravelFailure::Type FailureType, const FString& ErrorString);
-
 	void AutomationLoadMap(const FString& MapName, bool bForceReload, FString* OutError);
 
 	/** This function should be called to notify the editor that new materials were added to our scene or some materials were modified */
diff --git a/Engine/Source/Editor/UnrealEd/Private/EditorEngine.cpp b/Engine/Source/Editor/UnrealEd/Private/EditorEngine.cpp
--- a/Engine/Source/Editor/UnrealEd/Private/EditorEngine.cpp
+++ b/Engine/Source/Editor/UnrealEd/Private/EditorEngine.cpp
@@ -7364,29 +7364,6 @@ void FActorLabelUtilities::RenameExistingActor(AActor* Actor, const FString& New
 	}
 }
 
-void UEditorEngine::HandleTravelFailure(UWorld* InWorld, ETravelFailure::Type FailureType, const FString& ErrorString)
-{
-	if (InWorld && InWorld->IsPlayInEditor())
-	{
-		// Default behavior will try to fall back to default map and potentially throw a fatal error if that fails.
-		// Rather than bringing down the whole editor if this happens during a PIE session, just throw a warning and abort the PIE session.
-		{
-			FFormatNamedArguments Arguments;
-			Arguments.Add(TEXT("FailureType"), FText::FromString(ETravelFailure::ToString(FailureType)));
-			Arguments.Add(TEXT("ErrorString"), FText::FromString(ErrorString));
-			FText ErrorMsg = FText::Format(LOCTEXT("PIETravelFailure", "TravelFailure: {FailureType}, Reason for Failure: '{ErrorString}'. Shutting down PIE."), Arguments);
-			UE_LOG(LogNet, Warning, TEXT("%s"), *ErrorMsg.ToString());
-			FMessageLog("PIE").Warning(ErrorMsg);
-		}
-
-		RequestEndPlayMap();
-	}
-	else
-	{
-		Super::HandleTravelFailure(InWorld, FailureType, ErrorString);
-	}
-}
-
 void UEditorEngine::AutomationLoadMap(const FString& MapName, bool bForceReload, FString* OutError)
 {
 #if WITH_AUTOMATION_TESTS
diff --git a/Engine/Source/Runtime/Engine/Private/URL.cpp b/Engine/Source/Runtime/Engine/Private/URL.cpp
--- a/Engine/Source/Runtime/Engine/Private/URL.cpp
+++ b/Engine/Source/Runtime/Engine/Private/URL.cpp
@@ -510,7 +510,7 @@ FURL::FURL( FURL* Base, const TCHAR* TextURL, ETravelType Type )
 			if (!bFoundMap)
 			{
 				// can't find file, invalidate and bail
-				UE_CLOG(MapNameError.ToString().Len() > 0, LogLongPackageNames, Warning, TEXT("URL: %s: %s"), *URLStr, *MapNameError.ToString());
+				UE_LOG(LogLongPackageNames, Warning, TEXT("Can't Find URL: %s: %s. Invalidating and reverting to Default URL."), *URLStr, *MapNameError.ToString());
 				*this = FURL();
 				Valid = 0;
 			}
