Speculative fix crash when changing maps while the current map is unsaved.

diff --git a/Engine/Source/Editor/UnrealEd/Private/EditorViewportClient.cpp b/Engine/Source/Editor/UnrealEd/Private/EditorViewportClient.cpp
--- a/Engine/Source/Editor/UnrealEd/Private/EditorViewportClient.cpp
+++ b/Engine/Source/Editor/UnrealEd/Private/EditorViewportClient.cpp
@@ -3863,6 +3863,12 @@ void FEditorViewportClient::Draw(FViewport* InViewport, FCanvas* Canvas)
 		Time = World->GetTime();
 	}
 
+	// Early out if we are changing maps in editor as there is no reason to render the scene and it may not even be valid (For unsaved maps)
+	if (World->IsPreparingMapChange())
+	{
+		return;
+	}
+
 	// Allow HMD to modify the view later, just before rendering
 	const bool bStereoRendering = GEngine->IsStereoscopic3D( InViewport );
 	FCanvas* DebugCanvas = Viewport->GetDebugCanvas();
