UE-149122 - TMap Values in User Structs do not convert to real numbers

The underlying cause is connected to inconsistent usage of serialization methods among BP-related types.

Both FEdGraphPinType and FStructVariableDescription have a FEdGraphTerminalType field type. FEdGraphTerminalType will only perform versioning fixup in operator<<. This implies that fixup can only be achieved when an outer type uses manual serialization and calls << on the type, which FEdGraphPinType does. However, FStructVariableDescription does *not* have a Serialize function, and expects tagged property serialization to fixup everything in its FEdGraphTerminalType member, which never happens.

Since this is needed as a hotfix, the easiest fix is to update FStructVariableDescription::PostSerialize to perform the fixup for FEdGraphTerminalType, which is what FEdGraphTerminalType already does in operator<<.

diff --git a/Engine/Source/Editor/UnrealEd/Private/UserDefinedStructEditorData.cpp b/Engine/Source/Editor/UnrealEd/Private/UserDefinedStructEditorData.cpp
--- a/Engine/Source/Editor/UnrealEd/Private/UserDefinedStructEditorData.cpp
+++ b/Engine/Source/Editor/UnrealEd/Private/UserDefinedStructEditorData.cpp
@@ -10,6 +10,7 @@
 #include "Kismet2/BlueprintEditorUtils.h"
 #include "Blueprint/BlueprintSupport.h"
 #include "Editor.h"
+#include "EdGraphSchema_K2.h"
 
 #define LOCTEXT_NAMESPACE "UserDefinedStructEditorData"
 
@@ -35,13 +36,22 @@ void FStructVariableDescription::PostSerialize(const FArchive& Ar)
 
 	bool FixupPinCategories =
 		Ar.IsLoading() &&
-		(Ar.CustomVer(FUE5ReleaseStreamObjectVersion::GUID) < FUE5ReleaseStreamObjectVersion::BlueprintPinsUseRealNumbers) &&
-		((Category == TEXT("double")) || (Category == TEXT("float")));
+		((Category == UEdGraphSchema_K2::PC_Double) || (Category == UEdGraphSchema_K2::PC_Float));
 
 	if (FixupPinCategories)
 	{
-		Category = TEXT("real");
-		SubCategory = TEXT("double");
+		Category = UEdGraphSchema_K2::PC_Real;
+		SubCategory = UEdGraphSchema_K2::PC_Double;
+	}
+
+	bool FixupPinValueCategories =
+		Ar.IsLoading() &&
+		((PinValueType.TerminalCategory == UEdGraphSchema_K2::PC_Double) || (PinValueType.TerminalCategory == UEdGraphSchema_K2::PC_Float));
+
+	if (FixupPinValueCategories)
+	{
+		PinValueType.TerminalCategory = UEdGraphSchema_K2::PC_Real;
+		PinValueType.TerminalSubCategory = UEdGraphSchema_K2::PC_Double;
 	}
 }
 
