LevelInstance: Edit on null WorldAsset will create a new level

diff --git a/Engine/Source/Editor/LevelInstanceEditor/Private/LevelInstanceEditorModule.cpp b/Engine/Source/Editor/LevelInstanceEditor/Private/LevelInstanceEditorModule.cpp
--- a/Engine/Source/Editor/LevelInstanceEditor/Private/LevelInstanceEditorModule.cpp
+++ b/Engine/Source/Editor/LevelInstanceEditor/Private/LevelInstanceEditorModule.cpp
@@ -209,7 +209,7 @@ namespace LevelInstanceMenuUtils
 				});
 
 				FToolMenuSection& Section = CreateLevelSection(Menu);
-				Section.AddMenuEntry(NAME_None, LOCTEXT("LevelInstanceSetCurrent", "Set Current Level"), TAttribute<FText>(), FSlateIcon(), LevelInstanceSetCurrentAction);
+				Section.AddMenuEntry(TEXT("LevelInstanceSetCurrent"), LOCTEXT("LevelInstanceSetCurrent", "Set Current Level"), TAttribute<FText>(), FSlateIcon(), LevelInstanceSetCurrentAction);
 			}
 		}
 	}
@@ -251,7 +251,7 @@ namespace LevelInstanceMenuUtils
 					});
 
 				FToolMenuSection& Section = CreateLevelSection(Menu);
-				Section.AddMenuEntry(NAME_None, LOCTEXT("LevelInstanceMoveSelectionTo", "Move Selection to"), TAttribute<FText>(), FSlateIcon(), LevelInstanceMoveSelectionAction);
+				Section.AddMenuEntry(TEXT("LevelInstanceMoveSelectionTo"), LOCTEXT("LevelInstanceMoveSelectionTo", "Move Selection to"), TAttribute<FText>(), FSlateIcon(), LevelInstanceMoveSelectionAction);
 
 			}
 		}
@@ -327,7 +327,7 @@ namespace LevelInstanceMenuUtils
 				if (GetDefault<UEditorExperimentalSettings>()->bLevelInstance)
 				{
 					Section.AddMenuEntry(
-						NAME_None,
+						TEXT("CreateLevelInstance"),
 						FText::Format(LOCTEXT("CreateFromSelectionLabel", "Create {0}..."), StaticEnum<ELevelInstanceCreationType>()->GetDisplayNameTextByValue((int64)ELevelInstanceCreationType::LevelInstance)),
 						TAttribute<FText>(),
 						FSlateIcon(FAppStyle::GetAppStyleSetName(), "ClassIcon.LevelInstance"),
@@ -337,7 +337,7 @@ namespace LevelInstanceMenuUtils
 				if (GetDefault<UEditorExperimentalSettings>()->bPackedLevelActor)
 				{
 					Section.AddMenuEntry(
-						NAME_None,
+						TEXT("CreatePackedLevelBlueprint"),
 						FText::Format(LOCTEXT("CreateFromSelectionLabel", "Create {0}..."), StaticEnum<ELevelInstanceCreationType>()->GetDisplayNameTextByValue((int64)ELevelInstanceCreationType::PackedLevelActor)),
 						TAttribute<FText>(),
 						FSlateIcon(FAppStyle::GetAppStyleSetName(), "ClassIcon.PackedLevelActor"),
diff --git a/Engine/Source/Editor/UnrealEd/Private/EditorLevelUtils.cpp b/Engine/Source/Editor/UnrealEd/Private/EditorLevelUtils.cpp
--- a/Engine/Source/Editor/UnrealEd/Private/EditorLevelUtils.cpp
+++ b/Engine/Source/Editor/UnrealEd/Private/EditorLevelUtils.cpp
@@ -686,7 +686,7 @@ UWorld* GetWorldForNewStreamingLevel(UWorld* InTemplateWorld, bool bIsPartitione
 }
 };
 
-ULevelStreaming* UEditorLevelUtils::CreateNewStreamingLevelForWorld(UWorld& InWorld, TSubclassOf<ULevelStreaming> LevelStreamingClass, const FString& DefaultFilename /* = TEXT( "" ) */, bool bMoveSelectedActorsIntoNewLevel /* = false */, UWorld* InTemplateWorld /* = nullptr */, bool bInUseSaveAs /*= true*/, TFunction<void(ULevel*)> InPreSaveLevelOperation /* = TFunction<void(ULevel*)>()*/)
+ULevelStreaming* UEditorLevelUtils::CreateNewStreamingLevelForWorld(UWorld& InWorld, TSubclassOf<ULevelStreaming> LevelStreamingClass, const FString& DefaultFilename /* = TEXT( "" ) */, bool bMoveSelectedActorsIntoNewLevel /* = false */, UWorld* InTemplateWorld /* = nullptr */, bool bInUseSaveAs /*= true*/, TFunction<void(ULevel*)> InPreSaveLevelOperation /* = TFunction<void(ULevel*)>()*/, const FTransform& InTransform /* = FTransform::Identity */)
 {
 	TArray<AActor*> ActorsToMove;
 	if (bMoveSelectedActorsIntoNewLevel)
@@ -701,10 +701,10 @@ ULevelStreaming* UEditorLevelUtils::CreateNewStreamingLevelForWorld(UWorld& InWo
 		}
 	}
 
-	return CreateNewStreamingLevelForWorld(InWorld, LevelStreamingClass, /*bUseExternalActors=*/false, DefaultFilename, &ActorsToMove, InTemplateWorld, bInUseSaveAs, /*bIsPartitioned=*/false, InPreSaveLevelOperation);
+	return CreateNewStreamingLevelForWorld(InWorld, LevelStreamingClass, /*bUseExternalActors=*/false, DefaultFilename, &ActorsToMove, InTemplateWorld, bInUseSaveAs, /*bIsPartitioned=*/false, InPreSaveLevelOperation, InTransform);
 }
 
-ULevelStreaming* UEditorLevelUtils::CreateNewStreamingLevelForWorld(UWorld& InWorld, TSubclassOf<ULevelStreaming> LevelStreamingClass, bool bUseExternalActors, const FString& DefaultFilename /* = TEXT("") */, const TArray<AActor*>* ActorsToMove /* = nullptr */, UWorld* InTemplateWorld /* = nullptr */, bool bInUseSaveAs /*= true*/, bool bIsPartitioned /*= false*/, TFunction<void(ULevel*)> InPreSaveLevelOperation /* = TFunction<void(ULevel*)>()*/)
+ULevelStreaming* UEditorLevelUtils::CreateNewStreamingLevelForWorld(UWorld& InWorld, TSubclassOf<ULevelStreaming> LevelStreamingClass, bool bUseExternalActors, const FString& DefaultFilename /* = TEXT("") */, const TArray<AActor*>* ActorsToMove /* = nullptr */, UWorld* InTemplateWorld /* = nullptr */, bool bInUseSaveAs /*= true*/, bool bIsPartitioned /*= false*/, TFunction<void(ULevel*)> InPreSaveLevelOperation /* = TFunction<void(ULevel*)>()*/, const FTransform& InTransform /* = FTransform::Identity */)
 {
 	// Editor modes cannot be active when any level saving occurs.
 	if (!IsRunningCommandlet())
@@ -763,7 +763,7 @@ ULevelStreaming* UEditorLevelUtils::CreateNewStreamingLevelForWorld(UWorld& InWo
 	ULevel* NewLevel = nullptr;
 	if (bNewWorldSaved)
 	{
-		NewStreamingLevel = AddLevelToWorld(WorldToAddLevelTo, *NewPackageName, LevelStreamingClass);
+		NewStreamingLevel = AddLevelToWorld(WorldToAddLevelTo, *NewPackageName, LevelStreamingClass, InTransform);
 		if (NewStreamingLevel != nullptr)
 		{
 			NewLevel = NewStreamingLevel->GetLoadedLevel();
diff --git a/Engine/Source/Editor/UnrealEd/Public/EditorLevelUtils.h b/Engine/Source/Editor/UnrealEd/Public/EditorLevelUtils.h
--- a/Engine/Source/Editor/UnrealEd/Public/EditorLevelUtils.h
+++ b/Engine/Source/Editor/UnrealEd/Public/EditorLevelUtils.h
@@ -145,10 +145,11 @@ public:
 	 * @param	InTemplateWorld						If valid, the new level will be a copy of the template world.
 	 * @param	bInUseSaveAs						If true, show SaveAs dialog instead of Save with DefaultFilename
 	 * @param	InPreSaveLevelOperation				Optional function to call before saving the created level
+	 * @param	InTransform							The transform to apply to the streaming level.
 	 * 
 	 * @return	Returns the newly created level, or NULL on failure
 	 */
-	static UNREALED_API ULevelStreaming* CreateNewStreamingLevelForWorld(UWorld& World, TSubclassOf<ULevelStreaming> LevelStreamingClass, const FString& DefaultFilename = TEXT(""), bool bMoveSelectedActorsIntoNewLevel = false, UWorld* InTemplateWorld = nullptr, bool bInUseSaveAs = true, TFunction<void(ULevel*)> InPreSaveLevelOperation = TFunction<void(ULevel*)>());
+	static UNREALED_API ULevelStreaming* CreateNewStreamingLevelForWorld(UWorld& World, TSubclassOf<ULevelStreaming> LevelStreamingClass, const FString& DefaultFilename = TEXT(""), bool bMoveSelectedActorsIntoNewLevel = false, UWorld* InTemplateWorld = nullptr, bool bInUseSaveAs = true, TFunction<void(ULevel*)> InPreSaveLevelOperation = TFunction<void(ULevel*)>(), const FTransform& InTransform = FTransform::Identity);
 
 	/**
 	 * Creates a new streaming level and adds it to a world
@@ -162,10 +163,11 @@ public:
 	 * @param	bInUseSaveAs						If true, show SaveAs dialog instead of Save with DefaultFilename
 	 * @param	bIsPartitioned						If level should be partitioned (has precedence over bUseExternalActors).
 	 * @param	InPreSaveLevelOperation				Optional function to call before saving the created level
+	 * @param	InTransform							The transform to apply to the streaming level.
 	 *
 	 * @return	Returns the newly created level, or NULL on failure
 	 */
-	static UNREALED_API ULevelStreaming* CreateNewStreamingLevelForWorld(UWorld& World, TSubclassOf<ULevelStreaming> LevelStreamingClass, bool bUseExternalActors, const FString& DefaultFilename, const TArray<AActor*>* ActorsToMove = nullptr, UWorld* InTemplateWorld = nullptr, bool bInUseSaveAs = true, bool bIsPartitioned = false, TFunction<void(ULevel*)> InPreSaveLevelOperation = TFunction<void(ULevel*)>());
+	static UNREALED_API ULevelStreaming* CreateNewStreamingLevelForWorld(UWorld& World, TSubclassOf<ULevelStreaming> LevelStreamingClass, bool bUseExternalActors, const FString& DefaultFilename, const TArray<AActor*>* ActorsToMove = nullptr, UWorld* InTemplateWorld = nullptr, bool bInUseSaveAs = true, bool bIsPartitioned = false, TFunction<void(ULevel*)> InPreSaveLevelOperation = TFunction<void(ULevel*)>(), const FTransform& InTransform = FTransform::Identity);
 
 	/**
 	 * Adds the named level packages to the world.  Does nothing if all the levels already exist in the world.
diff --git a/Engine/Source/Runtime/Engine/Private/LevelInstance/LevelInstanceActorImpl.cpp b/Engine/Source/Runtime/Engine/Private/LevelInstance/LevelInstanceActorImpl.cpp
--- a/Engine/Source/Runtime/Engine/Private/LevelInstance/LevelInstanceActorImpl.cpp
+++ b/Engine/Source/Runtime/Engine/Private/LevelInstance/LevelInstanceActorImpl.cpp
@@ -351,6 +351,12 @@ void FLevelInstanceActorImpl::PushLevelInstanceEditingStateToProxies(bool bInEdi
 
 void FLevelInstanceActorImpl::CheckForErrors()
 {
+	AActor* Actor = CastChecked<AActor>(LevelInstance);
+	if (Actor->IsTemplate())
+	{
+		return;
+	}
+
 	TArray<TPair<FText, TSoftObjectPtr<UWorld>>> LoopInfo;
 	const ILevelInstanceInterface* LoopStart = nullptr;
 	if (!ULevelInstanceSubsystem::CheckForLoop(LevelInstance, LevelInstance->GetWorldAsset(), & LoopInfo, &LoopStart))
diff --git a/Engine/Source/Runtime/Engine/Private/LevelInstance/LevelInstanceEditorLevelStreaming.cpp b/Engine/Source/Runtime/Engine/Private/LevelInstance/LevelInstanceEditorLevelStreaming.cpp
--- a/Engine/Source/Runtime/Engine/Private/LevelInstance/LevelInstanceEditorLevelStreaming.cpp
+++ b/Engine/Source/Runtime/Engine/Private/LevelInstance/LevelInstanceEditorLevelStreaming.cpp
@@ -62,7 +62,23 @@ ULevelStreamingLevelInstanceEditor* ULevelStreamingLevelInstanceEditor::Load(ILe
 	UWorld* CurrentWorld = LevelInstanceActor->GetWorld();
 
 	TGuardValue<FLevelInstanceID> GuardEditLevelInstanceID(EditLevelInstanceID, LevelInstance->GetLevelInstanceID());
-	if (ULevelStreamingLevelInstanceEditor* LevelStreaming = Cast<ULevelStreamingLevelInstanceEditor>(EditorLevelUtils::AddLevelToWorld(CurrentWorld, *LevelInstance->GetWorldAssetPackage(), ULevelStreamingLevelInstanceEditor::StaticClass(), LevelInstanceActor->GetTransform())))
+	ULevelStreamingLevelInstanceEditor* LevelStreaming = nullptr;
+	// If Asset is null we can create a level here
+	if (LevelInstance->GetWorldAsset().IsNull())
+	{
+		LevelStreaming = Cast<ULevelStreamingLevelInstanceEditor>(EditorLevelUtils::CreateNewStreamingLevelForWorld(*CurrentWorld, ULevelStreamingLevelInstanceEditor::StaticClass(), TEXT(""), false, nullptr, true, TFunction<void(ULevel*)>(), LevelInstanceActor->GetTransform()));
+		if (LevelStreaming)
+		{
+			LevelInstanceActor->Modify();
+			LevelInstance->SetWorldAsset(LevelStreaming->GetLoadedLevel()->GetTypedOuter<UWorld>());
+		}
+	}
+	else
+	{
+		LevelStreaming = Cast<ULevelStreamingLevelInstanceEditor>(EditorLevelUtils::AddLevelToWorld(CurrentWorld, *LevelInstance->GetWorldAssetPackage(), ULevelStreamingLevelInstanceEditor::StaticClass(), LevelInstanceActor->GetTransform()));
+	}
+		
+	if (LevelStreaming)
 	{
 		check(LevelStreaming);
 		check(LevelStreaming->LevelInstanceID == LevelInstance->GetLevelInstanceID());
diff --git a/Engine/Source/Runtime/Engine/Private/LevelInstance/LevelInstanceSubsystem.cpp b/Engine/Source/Runtime/Engine/Private/LevelInstance/LevelInstanceSubsystem.cpp
--- a/Engine/Source/Runtime/Engine/Private/LevelInstance/LevelInstanceSubsystem.cpp
+++ b/Engine/Source/Runtime/Engine/Private/LevelInstance/LevelInstanceSubsystem.cpp
@@ -1546,41 +1546,35 @@ bool ULevelInstanceSubsystem::CanEditLevelInstance(const ILevelInstanceInterface
 		}
 	}
 	
-	if (!LevelInstance->IsWorldAssetValid())
+	if (LevelInstance->IsWorldAssetValid())
 	{
-		if (OutReason)
+		if (GetWorld()->PersistentLevel->GetPackage()->GetName() == LevelInstance->GetWorldAssetPackage())
 		{
-			*OutReason = LOCTEXT("CanEditLevelInstanceDirtyInvalid", "Level Instance path is invalid.");
-		}
-		return false;
-	}
-
-	if (GetWorld()->PersistentLevel->GetPackage()->GetName() == LevelInstance->GetWorldAssetPackage())
-	{
-		if (OutReason)
-		{
-			*OutReason = FText::Format(LOCTEXT("CanEditLevelInstancePersistentLevel", "The Persistent level and the Level Instance are the same ({0})."), FText::FromString(LevelInstance->GetWorldAssetPackage()));
+			if (OutReason)
+			{
+				*OutReason = FText::Format(LOCTEXT("CanEditLevelInstancePersistentLevel", "The Persistent level and the Level Instance are the same ({0})."), FText::FromString(LevelInstance->GetWorldAssetPackage()));
+			}
+			return false;
 		}
-		return false;
-	}
 
-	if (FLevelUtils::FindStreamingLevel(GetWorld(), *LevelInstance->GetWorldAssetPackage()))
-	{
-		if (OutReason)
+		if (FLevelUtils::FindStreamingLevel(GetWorld(), *LevelInstance->GetWorldAssetPackage()))
 		{
-			*OutReason = FText::Format(LOCTEXT("CanEditLevelInstanceAlreadyExists", "The same level was added to world outside of Level Instances ({0})."), FText::FromString(LevelInstance->GetWorldAssetPackage()));
+			if (OutReason)
+			{
+				*OutReason = FText::Format(LOCTEXT("CanEditLevelInstanceAlreadyExists", "The same level was added to world outside of Level Instances ({0})."), FText::FromString(LevelInstance->GetWorldAssetPackage()));
+			}
+			return false;
 		}
-		return false;
-	}
 
-	FPackagePath WorldAssetPath;
-	if (!FPackagePath::TryFromPackageName(LevelInstance->GetWorldAssetPackage(), WorldAssetPath) || !FPackageName::DoesPackageExist(WorldAssetPath))
-	{
-		if (OutReason)
+		FPackagePath WorldAssetPath;
+		if (!FPackagePath::TryFromPackageName(LevelInstance->GetWorldAssetPackage(), WorldAssetPath) || !FPackageName::DoesPackageExist(WorldAssetPath))
 		{
-			*OutReason = FText::Format(LOCTEXT("CanEditLevelInstanceInvalidAsset", "Level Instance asset is invalid ({0})."), FText::FromString(LevelInstance->GetWorldAssetPackage()));
+			if (OutReason)
+			{
+				*OutReason = FText::Format(LOCTEXT("CanEditLevelInstanceInvalidAsset", "Level Instance asset is invalid ({0})."), FText::FromString(LevelInstance->GetWorldAssetPackage()));
+			}
+			return false;
 		}
-		return false;
 	}
 
 	return true;
diff --git a/Engine/Source/Runtime/Engine/Public/LevelInstance/LevelInstanceActorGuid.h b/Engine/Source/Runtime/Engine/Public/LevelInstance/LevelInstanceActorGuid.h
--- a/Engine/Source/Runtime/Engine/Public/LevelInstance/LevelInstanceActorGuid.h
+++ b/Engine/Source/Runtime/Engine/Public/LevelInstance/LevelInstanceActorGuid.h
@@ -11,7 +11,7 @@ class AActor;
 /**
  * Helper struct that allows serializing the ActorGuid for runtime use.
  */
-struct FLevelInstanceActorGuid
+struct ENGINE_API FLevelInstanceActorGuid
 {
 	// Exists only to support 'FVTableHelper' Actor constructors
 	FLevelInstanceActorGuid() : FLevelInstanceActorGuid(nullptr) {}
@@ -27,5 +27,5 @@ struct FLevelInstanceActorGuid
 	TObjectPtr<AActor> Actor = nullptr;
 	FGuid ActorGuid;
 
-	friend FArchive& operator<<(FArchive& Ar, FLevelInstanceActorGuid& LevelInstanceActorGuid);
+	ENGINE_API friend FArchive& operator<<(FArchive& Ar, FLevelInstanceActorGuid& LevelInstanceActorGuid);
 };
\ No newline at end of file
