Changed FBX Level exporter to export actor labels instead of actor names
Since UE 5.0, actor names and labels aren't kept in sync. Actor names tend to be generic and are not chosen by users so exporting labels makes more sense for DCCs.

diff --git a/Engine/Plugins/Experimental/Interchange/Runtime/Source/Parsers/Fbx/Private/FbxHelper.cpp b/Engine/Plugins/Experimental/Interchange/Runtime/Source/Parsers/Fbx/Private/FbxHelper.cpp
--- a/Engine/Plugins/Experimental/Interchange/Runtime/Source/Parsers/Fbx/Private/FbxHelper.cpp
+++ b/Engine/Plugins/Experimental/Interchange/Runtime/Source/Parsers/Fbx/Private/FbxHelper.cpp
@@ -16,6 +16,11 @@ namespace UE
 		{
 			FString FFbxHelper::GetMeshName(FbxGeometryBase* Mesh)
 			{
+				if (!Mesh)
+				{
+					return {};
+				}
+
 				FString DefaultPrefix;
 				if (Mesh->GetAttributeType() == FbxNodeAttribute::eMesh)
 				{
diff --git a/Engine/Source/Editor/UnrealEd/Private/EditorExporters.cpp b/Engine/Source/Editor/UnrealEd/Private/EditorExporters.cpp
--- a/Engine/Source/Editor/UnrealEd/Private/EditorExporters.cpp
+++ b/Engine/Source/Editor/UnrealEd/Private/EditorExporters.cpp
@@ -1687,7 +1687,43 @@ bool ULevelExporterFBX::ExportBinary( UObject* Object, const TCHAR* Type, FArchi
 				Exporter->ExportBSP(World->GetModel(), true);
 			}
 
-			INodeNameAdapter NodeNameAdapter;
+			class FLevelExporterFBXNodeNameAdapter : public INodeNameAdapter
+			{
+			public:
+				virtual FString GetActorNodeName(const AActor* InActor) override
+				{
+					if (ActorNames.Contains(InActor))
+					{
+						return ActorNames[InActor].ToString();
+					}
+
+					FString ActorNodeName = InActor->GetActorLabel();
+
+					if (UniqueActorNames.Contains(ActorNodeName))
+					{
+						FString ActorNodeNamePrefix = ActorNodeName;
+						int32 ActorNodeNameIndex = 1;
+						FActorLabelUtilities::SplitActorLabel(ActorNodeNamePrefix, ActorNodeNameIndex);
+
+						do
+						{
+							ActorNodeName = FString::Printf(TEXT("%s%d"), *ActorNodeNamePrefix, ++ActorNodeNameIndex);
+						}
+						while (UniqueActorNames.Contains(ActorNodeName));
+					}
+
+					UniqueActorNames.Add(ActorNodeName);
+					ActorNames.Add(InActor) = FName(ActorNodeName);
+
+					return ActorNodeName;
+				}
+
+			private:
+				TSet<FString> UniqueActorNames;
+				TMap<const AActor*, FName> ActorNames;
+			};
+
+			FLevelExporterFBXNodeNameAdapter NodeNameAdapter;
 
 			Exporter->ExportLevelMesh(Level, bSelectedOnly, NodeNameAdapter);
 
