LooseCookedPackageWriter: Fix bug where generated packages marked for keeping are still being removed from disk.
#rnx

diff --git a/Engine/Source/Editor/UnrealEd/Private/Cooker/LooseCookedPackageWriter.cpp b/Engine/Source/Editor/UnrealEd/Private/Cooker/LooseCookedPackageWriter.cpp
--- a/Engine/Source/Editor/UnrealEd/Private/Cooker/LooseCookedPackageWriter.cpp
+++ b/Engine/Source/Editor/UnrealEd/Private/Cooker/LooseCookedPackageWriter.cpp
@@ -516,22 +516,28 @@ TUniquePtr<FAssetRegistryState> FLooseCookedPackageWriter::LoadPreviousAssetRegi
 		for (const TPair<FName, const FAssetPackageData*>& Pair : PreviousState->GetAssetPackageDataMap())
 		{
 			FName PackageName = Pair.Key;
-			const FName UncookedFilename = PackageDatas.GetFileNameByPackageName(PackageName);
+			FName UncookedFilename = PackageDatas.GetFileNameByPackageName(PackageName);
 
 			bool bNoLongerExistsInEditor = false;
+			bool bIsScriptPackage = FPackageName::IsScriptPackage(WriteToString<256>(PackageName));
 			if (UncookedFilename.IsNone())
 			{
 				// Script and generated packages do not exist uncooked
 				// Check that the package is not an exception before removing from cooked
-				bool bIsCookedOnly = FPackageName::IsScriptPackage(WriteToString<256>(PackageName));
+				bool bIsCookedOnly = bIsScriptPackage;
+				bool bIsMap = false;
 				if (!bIsCookedOnly)
 				{
 					for (const FAssetData* AssetData : PreviousState->GetAssetsByPackageName(PackageName))
 					{
 						bIsCookedOnly |= !!(AssetData->PackageFlags & PKG_CookGenerated);
+						bIsMap |= !!(AssetData->PackageFlags & PKG_ContainsMap);
 					}
 				}
 				bNoLongerExistsInEditor = !bIsCookedOnly;
+				UncookedFilename = UE::Cook::FPackageDatas::LookupFileNameOnDisk(PackageName,
+					false /* bRequireExists */, bIsMap);
+
 			}
 			if (bNoLongerExistsInEditor)
 			{
@@ -544,7 +550,7 @@ TUniquePtr<FAssetRegistryState> FLooseCookedPackageWriter::LoadPreviousAssetRegi
 				// Keep package on disk if it exists. Keep package in registry if it exists on disk or was a FailedSave.
 				bool bExistsOnDisk = (RemoveFromDisk.Remove(UncookedFilename) == 1); // Remove its RemoveFromDisk entry
 				const FAssetPackageData* PackageData = Pair.Value;
-				if (!bExistsOnDisk && PackageData->DiskSize >= 0)
+				if (!bExistsOnDisk && PackageData->DiskSize >= 0 && !bIsScriptPackage)
 				{
 					// Add RemoveFromRegistry entry if it didn't exist on disk and is a SuccessfulSave package
 					RemoveFromRegistry.Add(PackageName);
