Enable placement mode

diff --git a/Engine/Source/Editor/UnrealEd/Classes/Settings/LevelEditorMiscSettings.h b/Engine/Source/Editor/UnrealEd/Classes/Settings/LevelEditorMiscSettings.h
--- a/Engine/Source/Editor/UnrealEd/Classes/Settings/LevelEditorMiscSettings.h
+++ b/Engine/Source/Editor/UnrealEd/Classes/Settings/LevelEditorMiscSettings.h
@@ -47,6 +47,10 @@ public:
 	UPROPERTY(EditAnywhere, config, AdvancedDisplay, Category = Editing)
 	bool bEnableExperimentalLevelEditor = false;
 
+	/** If enabled, the experimental asset placement mode will be available in the level editor */
+	UPROPERTY(EditAnywhere, config, AdvancedDisplay, Category = Editing)
+	bool bEnableAssetPlacementMode = true;
+
 	/** If enabled, will avoid relabeling actors in UUnrealEdEngine::edactPasteSelected */
 	UPROPERTY(EditAnywhere, config, Category = Editing, meta = (DisplayName = "Avoid Actor Relabel on Paste Selected"))
 	uint32 bAvoidRelabelOnPasteSelected:1;
