PR #9141: [Localization Text Gathering] Fixed parser - incorrect number of arguments for LOCTABLE_FROMFILE_GAME (Contributed by [REDACTED_PERSON])

diff --git a/Engine/Source/Editor/UnrealEd/Classes/Commandlets/GatherTextFromSourceCommandlet.h b/Engine/Source/Editor/UnrealEd/Classes/Commandlets/GatherTextFromSourceCommandlet.h
--- a/Engine/Source/Editor/UnrealEd/Classes/Commandlets/GatherTextFromSourceCommandlet.h
+++ b/Engine/Source/Editor/UnrealEd/Classes/Commandlets/GatherTextFromSourceCommandlet.h
@@ -445,7 +445,7 @@ private:
 	class FStringTableFromFileMacroDescriptor : public FMacroDescriptor
 	{
 	public:
-		FStringTableFromFileMacroDescriptor(FString InName, FString InRootPath) : FMacroDescriptor(MoveTemp(InName), 2), RootPath(MoveTemp(InRootPath)) {}
+		FStringTableFromFileMacroDescriptor(FString InName, FString InRootPath) : FMacroDescriptor(MoveTemp(InName), 3), RootPath(MoveTemp(InRootPath)) {}
 
 		virtual void TryParse(const FString& Text, FSourceFileParseContext& Context) const override;
 
