Fix a bug that WP minimap won't capture LevelInstance since it never get chance to UpdateStreamingState to actually load the actors in that LevelInstance before been captured.

diff --git a/Engine/Source/Editor/UnrealEd/Private/WorldPartition/WorldPartitionMiniMapBuilder.cpp b/Engine/Source/Editor/UnrealEd/Private/WorldPartition/WorldPartitionMiniMapBuilder.cpp
--- a/Engine/Source/Editor/UnrealEd/Private/WorldPartition/WorldPartitionMiniMapBuilder.cpp
+++ b/Engine/Source/Editor/UnrealEd/Private/WorldPartition/WorldPartitionMiniMapBuilder.cpp
@@ -28,6 +28,7 @@
 #include "WorldPartition/WorldPartitionHandle.h"
 #include "WorldPartition/WorldPartitionMiniMap.h"
 #include "WorldPartition/WorldPartitionMiniMapHelper.h"
+#include "LevelInstance/LevelInstanceSubsystem.h"
 
 DEFINE_LOG_CATEGORY_STATIC(LogWorldPartitionMiniMapBuilder, All, All);
 
@@ -84,6 +85,11 @@ bool UWorldPartitionMiniMapBuilder::RunInternal(UWorld* World, const FCellInfo&
 	FString TextureName = FString::Format(TEXT("MinimapTile_{0}_{1}_{2}"), { InCellInfo.Location.X, InCellInfo.Location.Y, InCellInfo.Location.Z });
 
 	UTexture2D* TileTexture = nullptr;
+	check(World != nullptr);
+	if (ULevelInstanceSubsystem* LevelInstanceSubsystem = World->GetSubsystem<ULevelInstanceSubsystem>())
+	{
+		LevelInstanceSubsystem->UpdateStreamingState();
+	}
 	FWorldPartitionMiniMapHelper::CaptureBoundsMiniMapToTexture(World, WorldMiniMap, WorldMiniMap->MiniMapTileSize, TileTexture, TextureName, InCellInfo.Bounds);
 
 	FMinimapTile MinimapTile;
