Fix scaling issues when copy/pasting actor hierarchies:
- Fix ParentActor->IsSelected() not being valid anymore (Use OutPastedActors.Contains instead)
- Avoid Teleport if Offset is Zero

diff --git a/Engine/Source/Editor/UnrealEd/Private/EditorActor.cpp b/Engine/Source/Editor/UnrealEd/Private/EditorActor.cpp
--- a/Engine/Source/Editor/UnrealEd/Private/EditorActor.cpp
+++ b/Engine/Source/Editor/UnrealEd/Private/EditorActor.cpp
@@ -362,14 +362,24 @@ void UUnrealEdEngine::PasteActors(TArray<AActor*>& OutPastedActors, UWorld* InWo
 	ULayersSubsystem* LayersSubsystem = GEditor->GetEditorSubsystem<ULayersSubsystem>();
 	for (AActor* Actor : OutPastedActors)
 	{
-		// We only want to offset the location if this actor is the root of a selected attachment hierarchy
-		// Offsetting children of an attachment hierarchy would cause them to drift away from the node they're attached to
-		// as the offset would effectively get applied twice
-		const AActor* const ParentActor = Actor->GetAttachParentActor();
-		const FVector& ActorLocationOffset = (ParentActor && ParentActor->IsSelected()) ? FVector::ZeroVector : LocationOffset;
-
-		// Offset the actor's location.
-		Actor->TeleportTo(Actor->GetActorLocation() + ActorLocationOffset, Actor->GetActorRotation(), false, true);
+		if (!LocationOffset.IsZero())
+		{
+			// We only want to offset the location if this actor is the root of a selected attachment hierarchy
+			// Offsetting children of an attachment hierarchy would cause them to drift away from the node they're attached to
+			// as the offset would effectively get applied twice
+			AActor* ParentActor = Actor->GetAttachParentActor();
+			if (!ParentActor)
+			{
+				Actor->TeleportTo(Actor->GetActorLocation() + LocationOffset, Actor->GetActorRotation(), false, true);
+			}
+			else if(!OutPastedActors.Contains(ParentActor))
+			{
+				FName SocketName = Actor->GetAttachParentSocketName();
+				Actor->DetachFromActor(FDetachmentTransformRules::KeepWorldTransform);
+				Actor->TeleportTo(Actor->GetActorLocation() + LocationOffset, Actor->GetActorRotation(), false, true);
+				Actor->AttachToActor(ParentActor, FAttachmentTransformRules::KeepWorldTransform, SocketName);
+			}
+		}
 
 		if (!GetDefault<ULevelEditorMiscSettings>()->bAvoidRelabelOnPasteSelected)
 		{
