Minor tweak to cook label capture

diff --git a/Engine/Source/Editor/UnrealEd/Private/Commandlets/CookCommandlet.cpp b/Engine/Source/Editor/UnrealEd/Private/Commandlets/CookCommandlet.cpp
--- a/Engine/Source/Editor/UnrealEd/Private/Commandlets/CookCommandlet.cpp
+++ b/Engine/Source/Editor/UnrealEd/Private/Commandlets/CookCommandlet.cpp
@@ -110,10 +110,10 @@ namespace DetailedCookStats
 		ADD_COOK_STAT_FLT(" 0. 1. 9", CookProject);
 		ADD_COOK_STAT_FLT(" 0. 1. 10", CookCultures);
 		ADD_COOK_STAT_FLT(" 0. 1. 11", IsCookAll);
-		ADD_COOK_STAT_FLT(" 0. 1. 13", IsCookOnTheFly);
-		ADD_COOK_STAT_FLT(" 0. 1. 14", IsIterativeCook);
-		ADD_COOK_STAT_FLT(" 0. 1. 15", IsUnversioned);
-		ADD_COOK_STAT_FLT(" 0. 1. 16", CookLabel);
+		ADD_COOK_STAT_FLT(" 0. 1. 12", IsCookOnTheFly);
+		ADD_COOK_STAT_FLT(" 0. 1. 13", IsIterativeCook);
+		ADD_COOK_STAT_FLT(" 0. 1. 14", IsUnversioned);
+		ADD_COOK_STAT_FLT(" 0. 1. 15", CookLabel);
 		
 		#undef ADD_COOK_STAT_FLT
 	});
