Made the deletion of an asset remove the public flag from the object.
This fix a crash when deleting asset that was already supposed to be deleted.

diff --git a/Engine/Source/Editor/UnrealEd/Private/ObjectTools.cpp b/Engine/Source/Editor/UnrealEd/Private/ObjectTools.cpp
--- a/Engine/Source/Editor/UnrealEd/Private/ObjectTools.cpp
+++ b/Engine/Source/Editor/UnrealEd/Private/ObjectTools.cpp
@@ -2916,12 +2916,12 @@ namespace ObjectTools
 		// Mark its package as dirty as we're going to delete it.
 		ObjectToDelete->MarkPackageDirty();
 
-		// Remove standalone flag so garbage collection can delete the object.
-		ObjectToDelete->ClearFlags( RF_Standalone );
-
-		// Notify the asset registry
+		// Notify the asset registry. This done before the removal of the flags otherwise the content browser will ignore the update.
 		FAssetRegistryModule::AssetDeleted( ObjectToDelete );
 
+		// Remove standalone flag so garbage collection can delete the object and public flag so that the object is no longer considered to be an asset
+		ObjectToDelete->ClearFlags(RF_Standalone | RF_Public);
+
 		return true;
 	}
 
