[BlueprintEditorUtils] Take into account components when finding references to an actor.

diff --git a/Engine/Source/Editor/UnrealEd/Private/Kismet2/BlueprintEditorUtils.cpp b/Engine/Source/Editor/UnrealEd/Private/Kismet2/BlueprintEditorUtils.cpp
--- a/Engine/Source/Editor/UnrealEd/Private/Kismet2/BlueprintEditorUtils.cpp
+++ b/Engine/Source/Editor/UnrealEd/Private/Kismet2/BlueprintEditorUtils.cpp
@@ -128,7 +128,10 @@
 #include "AssetRegistry/AssetRegistryModule.h"
 #include "Misc/FeedbackContext.h"
 
-#include "AssetRegistry/AssetRegistryModule.h"
+#include "UObject/FastReferenceCollector.h"
+#include "Elements/Framework/TypedElementRegistry.h"
+#include "Elements/Interfaces/TypedElementHierarchyInterface.h"
+#include "Elements/Interfaces/TypedElementObjectInterface.h"
 #include "Misc/MessageDialog.h"
 
 #define LOCTEXT_NAMESPACE "Blueprint"
@@ -8170,73 +8173,108 @@ void FBlueprintEditorUtils::FindActorsThatReferenceActor( AActor* InActor, TArra
 			}
 		}
 	}
-}
+};
 
-void FBlueprintEditorUtils::GetActorReferenceMap(UWorld* InWorld, TArray<UClass*>& InClassesToIgnore, TMap<AActor*, TArray<AActor*> >& OutReferencingActors)
+class FActorMapReferenceProcessor : public FSimpleReferenceProcessorBase
 {
-	TRACE_CPUPROFILER_EVENT_SCOPE(FBlueprintEditorUtils::GetActorReferenceMap);
-
-	TArray<AActor*> Actors;
-
-	// Iterate all actors in the same world as InActor
-	for (FActorIterator ActorIt(InWorld); ActorIt; ++ActorIt)
+	TArray<UObject*> PotentiallyReferencedActors;
+	TMap<AActor*, TArray<AActor*>>& ReferencingActors;
+public:
+	FActorMapReferenceProcessor(UWorld* InWorld, TArray<UObject*>& OutPotentialReferencerObjects, const TArray<UClass*>& ClassesToIgnore, TMap<AActor*, TArray<AActor*>>& ReferencingActors)
+		: ReferencingActors(ReferencingActors)
 	{
-		AActor* CurrentActor = *ActorIt;
-		if (CurrentActor)
+		// Collect all actors in the world
+		for (FActorIterator ActorIt(InWorld); ActorIt; ++ActorIt)
 		{
-			bool bShouldIgnore = false;
-
-			// Ignore Actors if they are of a type we were instructed to ignore.
-			for (int32 IgnoreIndex = 0; IgnoreIndex < InClassesToIgnore.Num() && !bShouldIgnore; IgnoreIndex++)
+			if (AActor* CurrentActor = *ActorIt)
 			{
-				if (CurrentActor->IsA(InClassesToIgnore[IgnoreIndex]))
+				bool bShouldIgnore = false;
+				// Ignore actors if they belong to a class that's being ignored
+				for (UClass* ClassToIgnore : ClassesToIgnore)
 				{
-					bShouldIgnore = true;
+					if (CurrentActor->IsA(ClassToIgnore))
+					{
+						bShouldIgnore = true;
+						break;
+					}
 				}
-			}
+				OutPotentialReferencerObjects.Add(Cast<UObject>(CurrentActor));
 
-			if (!bShouldIgnore)
-			{
-				Actors.Add(CurrentActor);
+				// Collect all child elements of the actor and add them as potential referencer objects
+				TArray<FTypedElementHandle> ChildElementHandles;
+				UTypedElementRegistry* Registry = UTypedElementRegistry::GetInstance();
+				TTypedElement<ITypedElementHierarchyInterface> ElementHierarchyHandle = Registry->GetElement<ITypedElementHierarchyInterface>(UEngineElementsLibrary::AcquireEditorActorElementHandle(CurrentActor));
+				ElementHierarchyHandle.GetChildElements(ChildElementHandles, true);
+				// This is intentionally constructed in a way that will recursively get all child elements, by adding
+				// to the array of handles while iterating it. Don't try to change to a range-based loop
+				for (int i = 0; i < ChildElementHandles.Num(); ++i)
+				{
+					FTypedElementHandle ChildElementHandle = ChildElementHandles[i];
+					if (ITypedElementHierarchyInterface* ChildElementHierarchyInterface = Registry->GetElementInterface<ITypedElementHierarchyInterface>(ChildElementHandle))
+					{
+						ChildElementHierarchyInterface->GetChildElements(ChildElementHandle, ChildElementHandles, true);
+					}
+					if (ITypedElementObjectInterface* ChildElementObjectInterface = Registry->GetElementInterface<ITypedElementObjectInterface>(ChildElementHandle))
+					{
+						if (UObject* ChildObject = ChildElementObjectInterface->GetObject(ChildElementHandle))
+						{
+							OutPotentialReferencerObjects.Add(ChildObject);
+						}
+					}
+				}
+				if (bShouldIgnore)
+				{
+					continue;
+				}
+				PotentiallyReferencedActors.Add(CurrentActor);
 			}
 		}
 	}
-
-	TArray<TArray<UObject*>> ActorsReferences;
-	ActorsReferences.AddDefaulted(Actors.Num());
-
-	ParallelFor(Actors.Num(), [&Actors = std::as_const(Actors), &ActorsReferences](int32 Index) {
-		AActor* Actor(Actors[Index]);
-		TArray<UObject*>& References(ActorsReferences[Index]);
-
-		// Find all references
-		FReferenceFinder Finder(References);
-		Finder.FindReferences(Actor);
-
-		// Null out references that aren't actors
-		for (int32 ReferenceIndex = 0; ReferenceIndex < References.Num(); ReferenceIndex++)
+	FORCEINLINE_DEBUGGABLE void HandleTokenStreamObjectReference(FGCArrayStruct& ObjectsToSerializeStruct, UObject* ReferencingObject, UObject*& Object, const int32 TokenIndex, const EGCTokenType TokenType, bool bAllowReferenceElimination)
+	{
+		if (!ReferencingObject)
 		{
-			UObject* Reference = References[ReferenceIndex];
-			if (Reference && !Reference->IsA(AActor::StaticClass()))
-			{
-				References[ReferenceIndex] = nullptr;
-			}
+			ReferencingObject = ObjectsToSerializeStruct.GetReferencingObject();
+		}
+		if (!Object || !ReferencingObject || Object == ReferencingObject || !Object->IsA<AActor>())
+		{
+			return;
 		}
-	});
 
-	for (int ActorIndex=0; ActorIndex <ActorsReferences.Num(); ActorIndex++)
-	{
-		TArray<UObject*>& References(ActorsReferences[ActorIndex]);
-		for (int32 ReferenceIndex=0; ReferenceIndex<References.Num(); ReferenceIndex++)
+		AActor* Actor = CastChecked<AActor>(Object);
+		if (!PotentiallyReferencedActors.Contains(Actor))
 		{
-			// This is a sparse array as we may have nulled out non-actors
-			UObject* Reference = References[ReferenceIndex];
-			if (Reference)
-			{
-				OutReferencingActors.FindOrAdd(Cast<AActor>(Reference)).Add(Actors[ActorIndex]);
-			}
+			return;
+		}
+		// Ignore references from child objects
+		if (ReferencingObject->IsInOuter(Object))
+		{
+			return;
+		}
+		// The object itself if it's an actor, or the actor that contains it (if that exists)
+		if (AActor* ReferencingActor = ReferencingObject->IsA<AActor>() ? CastChecked<AActor>(ReferencingObject) : ReferencingObject->GetTypedOuter<AActor>())
+		{
+			// Don't record more than one reference from the same actor
+			ReferencingActors.FindOrAdd(Actor).AddUnique(ReferencingActor);
 		}
 	}
+};
+
+void FBlueprintEditorUtils::GetActorReferenceMap(UWorld* InWorld, TArray<UClass*>& InClassesToIgnore, TMap<AActor*, TArray<AActor*>>& OutReferencingActors)
+{
+	TRACE_CPUPROFILER_EVENT_SCOPE(FBlueprintEditorUtils::GetActorReferenceMap);
+	
+	FGCArrayStruct ArrayStruct;
+	
+	FActorMapReferenceProcessor Processor(InWorld, ArrayStruct.ObjectsToSerialize, InClassesToIgnore, OutReferencingActors);
+	TFastReferenceCollector<
+		FActorMapReferenceProcessor, 
+		TDefaultReferenceCollector<FActorMapReferenceProcessor>, 
+		FGCArrayPool, 
+		EFastReferenceCollectorOptions::AutogenerateTokenStream | EFastReferenceCollectorOptions::ProcessNoOpTokens
+	> ReferenceCollector(Processor, FGCArrayPool::Get());
+	
+	ReferenceCollector.CollectReferences(ArrayStruct);
 }
 
 void FBlueprintEditorUtils::FixLevelScriptActorBindings(ALevelScriptActor* LevelScriptActor, const ULevelScriptBlueprint* ScriptBlueprint)
