Add "Emulate Stereo" Editor Play option for standalone games.

- Also passthrough -emulatestereo if the editor was started like this.

diff --git a/Engine/Source/Editor/UnrealEd/Classes/Settings/LevelEditorPlaySettings.h b/Engine/Source/Editor/UnrealEd/Classes/Settings/LevelEditorPlaySettings.h
--- a/Engine/Source/Editor/UnrealEd/Classes/Settings/LevelEditorPlaySettings.h
+++ b/Engine/Source/Editor/UnrealEd/Classes/Settings/LevelEditorPlaySettings.h
@@ -251,6 +251,10 @@ public:
 	UPROPERTY(config, EditAnywhere, Category = PlayInEditor, meta = (ToolTip = "Whether or not the editor is minimized on non-VR PIE"))
 	bool bShouldMinimizeEditorOnNonVRPIE;
 
+	/** Whether we should emulate stereo (helps checking VR rendering issues). */
+	UPROPERTY(config, EditAnywhere, Category = PlayInStandaloneGame)
+	bool bEmulateStereo;
+
 	/** Whether to automatically recompile blueprints on PIE */
 	UPROPERTY(config, EditAnywhere, Category=PlayInEditor, meta=(ToolTip="Automatically recompile blueprints used by the current level when initiating a Play In Editor session"))
 	bool AutoRecompileBlueprints;
diff --git a/Engine/Source/Editor/UnrealEd/Private/PlayLevelNewProcess.cpp b/Engine/Source/Editor/UnrealEd/Private/PlayLevelNewProcess.cpp
--- a/Engine/Source/Editor/UnrealEd/Private/PlayLevelNewProcess.cpp
+++ b/Engine/Source/Editor/UnrealEd/Private/PlayLevelNewProcess.cpp
@@ -181,6 +181,12 @@ void UEditorEngine::LaunchNewProcess(const FRequestPlaySessionParams& InParams,
 		GLog->CategorizedLogf(FName("LogHMD"), ELogVerbosity::Warning, TEXT("Standalone Game VR not supported, please use VR Preview."));
 	}
 
+	// if we had -emulatestereo on the commandline, also pass it to the new process
+	if (InParams.EditorPlaySettings->bEmulateStereo || FParse::Param(FCommandLine::Get(), TEXT("emulatestereo")))
+	{
+		CommandLine += TEXT(" -emulatestereo");
+	}
+
 	// Allow disabling the sound in the new clients.
 	if (InParams.EditorPlaySettings->DisableStandaloneSound)
 	{
