Unity compile fix: Change consolecommand registration variable to static and mangle its name to prevent collision with other global variables.
#rnx

diff --git a/Engine/Source/Editor/UnrealEd/Private/ReferenceInfoUtils.cpp b/Engine/Source/Editor/UnrealEd/Private/ReferenceInfoUtils.cpp
--- a/Engine/Source/Editor/UnrealEd/Private/ReferenceInfoUtils.cpp
+++ b/Engine/Source/Editor/UnrealEd/Private/ReferenceInfoUtils.cpp
@@ -48,7 +48,7 @@ void ExecuteReferenceInfo(const TArray<FString>& Args, UWorld* InWorld )
 	ReferenceInfoUtils::GenerateOutput(InWorld, Depth, bShowDefault, bShowScript);
 }
 
-FAutoConsoleCommandWithWorldAndArgs ReferenceInfo(
+static FAutoConsoleCommandWithWorldAndArgs ActorReferenceInfoCVar(
 	TEXT("ReferenceInfo"),
 	TEXT("Outputs reference info for selected actors to a log file. Syntax is: ReferenceInfo [-depth=<depth value>] [-nodefault] [-noscript]"),
 	FConsoleCommandWithWorldAndArgsDelegate::CreateStatic(ExecuteReferenceInfo)
