Catch up placement mode and API to UE5

# Editgrated CLs 18420733 18589986 18654845 18661572 19001860 19093239 19244734

#rnx

diff --git a/Engine/Source/Editor/UnrealEd/Private/Subsystems/AssetEditorSubsystem.cpp b/Engine/Source/Editor/UnrealEd/Private/Subsystems/AssetEditorSubsystem.cpp
--- a/Engine/Source/Editor/UnrealEd/Private/Subsystems/AssetEditorSubsystem.cpp
+++ b/Engine/Source/Editor/UnrealEd/Private/Subsystems/AssetEditorSubsystem.cpp
@@ -1031,6 +1031,12 @@ void UAssetEditorSubsystem::UnregisterEditorModes()
 	EditorModes.Empty();
 }
 
+void UAssetEditorSubsystem::OnSMInstanceElementsEnabled()
+{
+	// Let the modes know that SM instance elements may have been enabled or disabled and update state accordingly
+	OnEditorModesChanged().Broadcast();
+}
+
 FNamePermissionList& UAssetEditorSubsystem::GetAllowedEditorModes()
 {
 	return AllowedEditorModes;
