World Partition: deprecate all flavors of FWorldPartitionHelpers::ForEachActorWithLoading in favor of the version which takes a FForEachActorWithLoadingParams.

#rnx

diff --git a/Engine/Source/Editor/UnrealEd/Private/Commandlets/WorldPartitionDataLayerToAssetCommandlet.cpp b/Engine/Source/Editor/UnrealEd/Private/Commandlets/WorldPartitionDataLayerToAssetCommandlet.cpp
--- a/Engine/Source/Editor/UnrealEd/Private/Commandlets/WorldPartitionDataLayerToAssetCommandlet.cpp
+++ b/Engine/Source/Editor/UnrealEd/Private/Commandlets/WorldPartitionDataLayerToAssetCommandlet.cpp
@@ -476,50 +476,49 @@ bool UDataLayerToAssetCommandlet::RemapActorDataLayersToAssets(TStrongObjectPtr<
 	UE_LOG(LogDataLayerToAssetCommandlet, Log, TEXT("Starting Actor Data Layer Remapping To Data Layer Asset. This can take a while."));
 
 	uint32 ErrorCount = 0;
-	FWorldPartitionHelpers::ForEachActorWithLoading(MainWorld->GetWorldPartition(), AActor::StaticClass(),
-		[&ErrorCount, &CommandletContext, this](const FWorldPartitionActorDesc* ActorDesc)
+	FWorldPartitionHelpers::ForEachActorWithLoading(MainWorld->GetWorldPartition(), [&ErrorCount, &CommandletContext, this](const FWorldPartitionActorDesc* ActorDesc)
+	{
+		uint32 ActorConversionErrors = 0;
+		if (AActor* Actor = ActorDesc->GetActor())
 		{
-			uint32 ActorConversionErrors = 0;
-			if (AActor* Actor = ActorDesc->GetActor())
-			{
-				ActorConversionErrors += RemapDataLayersAssetsFromPreviousConversions(CommandletContext, Actor);
-				ActorConversionErrors += RemapActorDataLayers(CommandletContext, Actor);
+			ActorConversionErrors += RemapDataLayersAssetsFromPreviousConversions(CommandletContext, Actor);
+			ActorConversionErrors += RemapActorDataLayers(CommandletContext, Actor);
 
-				if (!PerformAdditionalActorConversions(CommandletContext, Actor))
-				{
-					ActorConversionErrors++;
-				}
+			if (!PerformAdditionalActorConversions(CommandletContext, Actor))
+			{
+				ActorConversionErrors++;
+			}
 
-				if (ActorConversionErrors == 0 && bPerformSavePackages)
+			if (ActorConversionErrors == 0 && bPerformSavePackages)
+			{
+				if (bPerformSavePackages && Actor->GetExternalPackage()->IsDirty())
 				{
-					if (bPerformSavePackages && Actor->GetExternalPackage()->IsDirty())
+					if (!WorldPartitionCommandletHelpers::CheckoutSaveAdd(Actor->GetExternalPackage(), PackageHelper))
 					{
-						if (!WorldPartitionCommandletHelpers::CheckoutSaveAdd(Actor->GetExternalPackage(), PackageHelper))
-						{
-							ActorConversionErrors++;
-						}
+						ActorConversionErrors++;
 					}
 				}
 			}
-			else
+		}
+		else
+		{
+			const TArray<FName>& ActDescDataLayers = ActorDesc->GetDataLayerInstanceNames();
+			if (!ActDescDataLayers.IsEmpty())
 			{
-				const TArray<FName>& ActDescDataLayers = ActorDesc->GetDataLayerInstanceNames();
-				if (!ActDescDataLayers.IsEmpty())
-				{
-					FString DataLayerString = FString::JoinBy(ActDescDataLayers, TEXT(", "), [](const FName& DataLayerName) { return DataLayerName.ToString(); });
+				FString DataLayerString = FString::JoinBy(ActDescDataLayers, TEXT(", "), [](const FName& DataLayerName) { return DataLayerName.ToString(); });
 
-					UE_LOG(LogDataLayerToAssetCommandlet, Error, TEXT("Actor %s failed to load. Its data layers %s will not be remapped to a data layer asset."),
-						*ActorDesc->GetActorName().ToString(), *DataLayerString);
-					if (!bIgnoreActorLoadingErrors)
-					{
-						ActorConversionErrors++;
-					}
+				UE_LOG(LogDataLayerToAssetCommandlet, Error, TEXT("Actor %s failed to load. Its data layers %s will not be remapped to a data layer asset."),
+					*ActorDesc->GetActorName().ToString(), *DataLayerString);
+				if (!bIgnoreActorLoadingErrors)
+				{
+					ActorConversionErrors++;
 				}
 			}
+		}
 
-			ErrorCount += ActorConversionErrors;
-			return true;
-		});
+		ErrorCount += ActorConversionErrors;
+		return true;
+	});
 
 	UE_LOG(LogDataLayerToAssetCommandlet, Log, TEXT("Actor Data Layer Remapping To Data Layer Asset Done."));
 
diff --git a/Engine/Source/Editor/UnrealEd/Private/WorldPartition/WorldPartitionFoliageBuilder.cpp b/Engine/Source/Editor/UnrealEd/Private/WorldPartition/WorldPartitionFoliageBuilder.cpp
--- a/Engine/Source/Editor/UnrealEd/Private/WorldPartition/WorldPartitionFoliageBuilder.cpp
+++ b/Engine/Source/Editor/UnrealEd/Private/WorldPartition/WorldPartitionFoliageBuilder.cpp
@@ -55,13 +55,6 @@ bool UWorldPartitionFoliageBuilder::RunInternal(UWorld* World, const FCellInfo&
 	FPackageSourceControlHelper SCCHelper;
 	TArray<FGuid> ExistingActors;
 	
-	// Snapshot existing ActorDescs
-	FWorldPartitionHelpers::ForEachActorDesc(WorldPartition, AInstancedFoliageActor::StaticClass(), [&ExistingActors](const FWorldPartitionActorDesc* ActorDesc)
-	{
-		ExistingActors.Add(ActorDesc->GetGuid());
-		return true;
-	});
-		
 	TMap<UActorPartitionSubsystem::FCellCoord, FGuid> NewActors;
 	UActorPartitionSubsystem* ActorPartitionSubsystem = World->GetSubsystem<UActorPartitionSubsystem>();
 	int32 NumInstances = 0;
@@ -75,8 +68,13 @@ bool UWorldPartitionFoliageBuilder::RunInternal(UWorld* World, const FCellInfo&
 	// - Save new/reload IFAs that were modified
 	// - Unload the Existing IFA and all the new IFAs
 	// - Repeat for next Existing IFA
-	FWorldPartitionHelpers::ForEachActorWithLoading(WorldPartition, ExistingActors, [this, World, WorldPartition, ActorPartitionSubsystem, &SCCHelper, &NewActors, &NumInstances, &NumInstancesProcessed](const FWorldPartitionActorDesc* ActorDesc)
+	FWorldPartitionHelpers::FForEachActorWithLoadingParams ForEachActorWithLoadingParams;
+	ForEachActorWithLoadingParams.ActorClass = AInstancedFoliageActor::StaticClass();
+
+	FWorldPartitionHelpers::ForEachActorWithLoading(WorldPartition, [this, World, WorldPartition, ActorPartitionSubsystem, &SCCHelper, &NewActors, &NumInstances, &NumInstancesProcessed, &ExistingActors](const FWorldPartitionActorDesc* ActorDesc)
 	{
+		ExistingActors.Add(ActorDesc->GetGuid());
+
 		TMap<UFoliageType*, TArray<FFoliageInstance>> FoliageToAdd;
 		FBox InstanceBounds(ForceInit);
 				
@@ -201,7 +199,7 @@ bool UWorldPartitionFoliageBuilder::RunInternal(UWorld* World, const FCellInfo&
         }
 					
 		return true;
-	});
+	}, ForEachActorWithLoadingParams);
 
 	check(NumInstances == NumInstancesProcessed);
 
diff --git a/Engine/Source/Editor/UnrealEd/Private/WorldPartition/WorldPartitionResaveActorsBuilder.cpp b/Engine/Source/Editor/UnrealEd/Private/WorldPartition/WorldPartitionResaveActorsBuilder.cpp
--- a/Engine/Source/Editor/UnrealEd/Private/WorldPartition/WorldPartitionResaveActorsBuilder.cpp
+++ b/Engine/Source/Editor/UnrealEd/Private/WorldPartition/WorldPartitionResaveActorsBuilder.cpp
@@ -273,7 +273,7 @@ bool UWorldPartitionResaveActorsBuilder::RunInternal(UWorld* World, const FCellI
 	{
 		TArray<UPackage*> PackagesToSave;
 
-		FWorldPartitionHelpers::FForEachActorWithLoadiongParams ForEachActorWithLoadingParams;
+		FWorldPartitionHelpers::FForEachActorWithLoadingParams ForEachActorWithLoadingParams;
 
 		ForEachActorWithLoadingParams.ActorClass = ActorClass;
 
@@ -299,7 +299,7 @@ bool UWorldPartitionResaveActorsBuilder::RunInternal(UWorld* World, const FCellI
 			PackagesToSave.Empty();
 		};
 
-		FWorldPartitionHelpers::ForEachActorWithLoading(WorldPartition, ForEachActorWithLoadingParams, [this, &PackagesToDelete, &PackagesToSave](const FWorldPartitionActorDesc* ActorDesc)
+		FWorldPartitionHelpers::ForEachActorWithLoading(WorldPartition, [this, &PackagesToDelete, &PackagesToSave](const FWorldPartitionActorDesc* ActorDesc)
 		{
 			AActor* Actor = ActorDesc->GetActor();
 
@@ -345,7 +345,7 @@ bool UWorldPartitionResaveActorsBuilder::RunInternal(UWorld* World, const FCellI
 			}
 
 			return true;
-		});
+		}, ForEachActorWithLoadingParams);
 	}
 
 	if (!bReportOnly)
diff --git a/Engine/Source/Runtime/Engine/Private/WorldPartition/WorldPartitionHelpers.cpp b/Engine/Source/Runtime/Engine/Private/WorldPartition/WorldPartitionHelpers.cpp
--- a/Engine/Source/Runtime/Engine/Private/WorldPartition/WorldPartitionHelpers.cpp
+++ b/Engine/Source/Runtime/Engine/Private/WorldPartition/WorldPartitionHelpers.cpp
@@ -102,7 +102,7 @@ namespace WorldPartitionHelpers
 	}
 }
 
-FWorldPartitionHelpers::FForEachActorWithLoadiongParams::FForEachActorWithLoadiongParams()
+FWorldPartitionHelpers::FForEachActorWithLoadingParams::FForEachActorWithLoadingParams()
 	: bGCPerActor(false)
 	, ActorClass(AActor::StaticClass())
 {}
@@ -145,7 +145,7 @@ void FWorldPartitionHelpers::ForEachActorWithLoading(UWorldPartition* WorldParti
 	DoCollectGarbage();
 }
 
-void FWorldPartitionHelpers::ForEachActorWithLoading(UWorldPartition* WorldPartition, const FForEachActorWithLoadiongParams& Params, TFunctionRef<bool(const FWorldPartitionActorDesc*)> Func)
+void FWorldPartitionHelpers::ForEachActorWithLoading(UWorldPartition* WorldPartition, TFunctionRef<bool(const FWorldPartitionActorDesc*)> Func, const FForEachActorWithLoadingParams& Params)
 {
 	TMap<FGuid, FWorldPartitionReference> ActorReferences;
 
diff --git a/Engine/Source/Runtime/Engine/Public/WorldPartition/WorldPartitionHelpers.h b/Engine/Source/Runtime/Engine/Public/WorldPartition/WorldPartitionHelpers.h
--- a/Engine/Source/Runtime/Engine/Public/WorldPartition/WorldPartitionHelpers.h
+++ b/Engine/Source/Runtime/Engine/Public/WorldPartition/WorldPartitionHelpers.h
@@ -31,9 +31,9 @@ public:
 	static void ForEachActorDesc(UWorldPartition* WorldPartition, TSubclassOf<AActor> ActorClass, TFunctionRef<bool(const FWorldPartitionActorDesc*)> Func);
 
 	/* Struct of optional parameters passed to foreach actordesc functions. */
-	struct ENGINE_API FForEachActorWithLoadiongParams
+	struct ENGINE_API FForEachActorWithLoadingParams
 	{
-		FForEachActorWithLoadiongParams();
+		FForEachActorWithLoadingParams();
 
 		/* Perform a garbage collection per-actor, useful to test if the caller properly handle GCs. */
 		bool bGCPerActor;
@@ -48,15 +48,13 @@ public:
 		TUniqueFunction<void()> OnPreGarbageCollect;
 	};
 
-	template<class ActorClass = AActor>
-	static void ForEachActorWithLoading(UWorldPartition* WorldPartition, TFunctionRef<bool(const FWorldPartitionActorDesc*)> Func)
-	{
-		ForEachActorWithLoading(WorldPartition, ActorClass::StaticClass(), Func);
-	}
-
+	UE_DEPRECATED(5.1, "ForEachActorWithLoading is deprecated, ForEachActorWithLoading with FForEachActorWithLoadingParams should be used instead.")
 	static void ForEachActorWithLoading(UWorldPartition* WorldPartition, TSubclassOf<AActor> ActorClass, TFunctionRef<bool(const FWorldPartitionActorDesc*)> Func, TFunctionRef<void()> OnReleasingActorReferences = [](){}, bool bGCPerActor = false);
+
+	UE_DEPRECATED(5.1, "ForEachActorWithLoading is deprecated, ForEachActorWithLoading with FForEachActorWithLoadingParams should be used instead.")
 	static void ForEachActorWithLoading(UWorldPartition* WorldPartition, const TArray<FGuid>& ActorGuids, TFunctionRef<bool(const FWorldPartitionActorDesc*)> Func, TFunctionRef<void()> OnReleasingActorReferences = [](){}, bool bGCPerActor = false);
-	static void ForEachActorWithLoading(UWorldPartition* WorldPartition, const FForEachActorWithLoadiongParams& Params, TFunctionRef<bool(const FWorldPartitionActorDesc*)> Func);
+
+	static void ForEachActorWithLoading(UWorldPartition* WorldPartition, TFunctionRef<bool(const FWorldPartitionActorDesc*)> Func, const FForEachActorWithLoadingParams& Params = FForEachActorWithLoadingParams());
 	
 	static bool HasExceededMaxMemory();
 	static void DoCollectGarbage();
