World Partition: fix actor tags filteriing in resave actors builder commandlet due to last minute changes.

#rnx

diff --git a/Engine/Source/Editor/UnrealEd/Private/WorldPartition/WorldPartitionResaveActorsBuilder.cpp b/Engine/Source/Editor/UnrealEd/Private/WorldPartition/WorldPartitionResaveActorsBuilder.cpp
--- a/Engine/Source/Editor/UnrealEd/Private/WorldPartition/WorldPartitionResaveActorsBuilder.cpp
+++ b/Engine/Source/Editor/UnrealEd/Private/WorldPartition/WorldPartitionResaveActorsBuilder.cpp
@@ -272,14 +272,14 @@ bool UWorldPartitionResaveActorsBuilder::RunInternal(UWorld* World, const FCellI
 	else
 	{
 		TArray<UPackage*> PackagesToSave;
+		TSet<FName> ActorTagsSet(ActorTags);
 
 		FWorldPartitionHelpers::FForEachActorWithLoadingParams ForEachActorWithLoadingParams;
 
 		ForEachActorWithLoadingParams.ActorClass = ActorClass;
 
 		if (ActorTags.Num())
-		{
-			TSet<FName> ActorTagsSet(ActorTags);
+		{			
 			ForEachActorWithLoadingParams.FilterActorDesc = [this, &ActorTagsSet](const FWorldPartitionActorDesc* ActorDesc) -> bool
 			{
 				for (const FName& ActorTag : ActorDesc->GetTags())
