World Partition Convert Commandlet: automatically set actors referenced by the Level Script as non-spatially loaded to avoid warnings when opening the map.

#rnx

diff --git a/Engine/Source/Editor/UnrealEd/Private/Commandlets/WorldPartitionConvertCommandlet.cpp b/Engine/Source/Editor/UnrealEd/Private/Commandlets/WorldPartitionConvertCommandlet.cpp
--- a/Engine/Source/Editor/UnrealEd/Private/Commandlets/WorldPartitionConvertCommandlet.cpp
+++ b/Engine/Source/Editor/UnrealEd/Private/Commandlets/WorldPartitionConvertCommandlet.cpp
@@ -1081,6 +1081,23 @@ int32 UWorldPartitionConvertCommandlet::Main(const FString& Params)
 			}
 		}
 
+		if (bMainLevel)
+		{
+			if (LevelHasLevelScriptBlueprint(Level))
+			{
+				ULevelScriptBlueprint* LevelScriptBlueprint = Level->GetLevelScriptBlueprint(true);
+				TArray<AActor*> LevelScriptActorReferences = ActorsReferencesUtils::GetActorReferences(LevelScriptBlueprint);
+
+				for (AActor* LevelScriptActorReference : LevelScriptActorReferences)
+				{
+					if (LevelScriptActorReference->GetIsSpatiallyLoaded() && LevelScriptActorReference->CanChangeIsSpatiallyLoadedFlag())
+					{
+						LevelScriptActorReference->SetIsSpatiallyLoaded(false);
+					}
+				}
+			}
+		}
+
 		// do loop after as it may modify Level->Actors
 		if (IFAs.Num())
 		{
