clarify use of GetAllTextureFormats vs GetTextureFormats
GetAllTextureFormats is only used by COTF for version checking
could deprecate then remove it entirely

diff --git a/Engine/Source/Developer/TargetPlatform/Private/TargetPlatformManagerModule.cpp b/Engine/Source/Developer/TargetPlatform/Private/TargetPlatformManagerModule.cpp
--- a/Engine/Source/Developer/TargetPlatform/Private/TargetPlatformManagerModule.cpp
+++ b/Engine/Source/Developer/TargetPlatform/Private/TargetPlatformManagerModule.cpp
@@ -583,32 +583,6 @@ public:
 #endif
 	};
 
-	struct FTextureHintHelper
-	{
-		static ITextureFormat* GetFormatFromModule(ITextureFormatModule* Module)
-		{
-			return Module->GetTextureFormat();
-		}
-		static const TCHAR* GetAllModuleWildcard()
-		{
-			return TEXT("*TextureFormat*");
-		}
-		static const TCHAR* GetFormatDesc()
-		{
-			return TEXT("texture");
-		}
-#if WITH_ENGINE 
-		static void GetHintedModules(ITargetPlatform* Platform, TArray<FName>& Hints)
-		{
-			Platform->GetTextureFormatModuleHints(Hints);
-		}
-		static void GetRequiredFormats(ITargetPlatform* Platform, TArray<FName>& RequiredFormats)
-		{
-			Platform->GetAllTextureFormats(RequiredFormats);
-		}
-#endif
-	};
-
 	virtual const TArray<const IAudioFormat*>& GetAudioFormats() override
 	{
 		return GetFormatsWithHints<IAudioFormat, IAudioFormatModule, FAudioHintHelper>();
@@ -638,7 +612,7 @@ public:
 
 	virtual const TArray<const ITextureFormat*>& GetTextureFormats() override
 	{
-//		return GetFormatsWithHints<ITextureFormat, ITextureFormatModule, FTextureHintHelper>();
+		// note that this gets ALL ITextureFormat Modules, not just ones relevant to the current TargetPlatform
 		return TextureFormatManager->GetTextureFormats();
 	}
 
diff --git a/Engine/Source/Editor/UnrealEd/Private/CookOnTheFlyServer.cpp b/Engine/Source/Editor/UnrealEd/Private/CookOnTheFlyServer.cpp
--- a/Engine/Source/Editor/UnrealEd/Private/CookOnTheFlyServer.cpp
+++ b/Engine/Source/Editor/UnrealEd/Private/CookOnTheFlyServer.cpp
@@ -5092,6 +5092,13 @@ void GetAdditionalCurrentIniVersionStrings( const UCookOnTheFlyServer* CookOnThe
 	}
 
 	{
+
+		#if 1
+		// this is the only place that TargetPlatform::GetAllTextureFormats is used
+		// instead use ITextureFormatManagerModule::GetTextureFormats ?
+		//	then GetAllTextureFormats can be removed completely
+
+		// get all texture formats for this target platform, then find the modules that encode them
 		TArray<FName> AllTextureFormats;
 		TargetPlatform->GetAllTextureFormats(AllTextureFormats);
 		TArray<const ITextureFormat*> SupportedTextureFormats;
@@ -5100,13 +5107,18 @@ void GetAdditionalCurrentIniVersionStrings( const UCookOnTheFlyServer* CookOnThe
 			const ITextureFormat* TextureFormat = TPM->FindTextureFormat(TextureName);
 			if ( TextureFormat )
 			{
-				SupportedTextureFormats.Add(TextureFormat);
+				SupportedTextureFormats.AddUnique(TextureFormat);
 			}
 			else
 			{
 				UE_LOG(LogCook, Warning, TEXT("Unable to find texture format \"%s\" which is required by \"%s\""), *TextureName.ToString(), *TargetPlatform->PlatformName());
 			}
 		}
+		#else
+		//note: this gets All ITextureFormat modules in the Engine, not just ones relevant to this TargetPlatform
+		const TArray<const ITextureFormat*> & SupportedTextureFormats = TPM->GetTextureFormats();
+		#endif
+		
 		GetVersionFormatNumbersForIniVersionStrings(IniVersionMap, TEXT("TextureFormat"), SupportedTextureFormats);
 	}
 
