HLOD - Fixed custom primitive data not taken into account in some specific cases
* Ensure construction scripts are executed properly before building HLODs from commandlets.
* Include custom primitive data in the static mesh components HLOD hashes

diff --git a/Engine/Source/Editor/UnrealEd/Private/HierarchicalLOD.cpp b/Engine/Source/Editor/UnrealEd/Private/HierarchicalLOD.cpp
--- a/Engine/Source/Editor/UnrealEd/Private/HierarchicalLOD.cpp
+++ b/Engine/Source/Editor/UnrealEd/Private/HierarchicalLOD.cpp
@@ -48,6 +48,7 @@
 #include "IMeshReductionManagerModule.h"
 #include "Engine/HLODProxy.h"
 #include "Engine/LevelStreaming.h"
+#include "AssetCompilingManager.h"
 
 DEFINE_LOG_CATEGORY_STATIC(LogLODGenerator, Log, All);
 
@@ -756,6 +757,12 @@ void FHierarchicalLODBuilder::ClearPreviewBuild()
 
 void FHierarchicalLODBuilder::BuildMeshesForLODActors(bool bForceAll)
 {	
+	// Finalize asset compilation before we potentially uses them during the HLOD generation.
+	FAssetCompilingManager::Get().FinishAllCompilation();
+
+	// This may execute pending construction scripts.
+	FAssetCompilingManager::Get().ProcessAsyncTasks();
+
 	bool bVisibleLevelsWarning = false;
 
 	const TArray<ULevel*>& Levels = World->GetLevels();
diff --git a/Engine/Source/Runtime/Engine/Private/HLODProxy.cpp b/Engine/Source/Runtime/Engine/Private/HLODProxy.cpp
--- a/Engine/Source/Runtime/Engine/Private/HLODProxy.cpp
+++ b/Engine/Source/Runtime/Engine/Private/HLODProxy.cpp
@@ -457,6 +457,11 @@ uint32 UHLODProxy::GetCRC(UStaticMeshComponent* InComponent, uint32 InCRC, const
 	InComponent->GetLightMapResolution(Width, Height);
 	Ar << Width;
 	Ar << Height;
+
+	if (!InComponent->GetCustomPrimitiveData().Data.IsEmpty())
+	{
+		Ar << InComponent->GetCustomPrimitiveData();
+	}
 	
 	// incorporate vertex colors
 	for(FStaticMeshComponentLODInfo& LODInfo : InComponent->LODData)
