EditorDomain: Add DDC Resource Usage Stats.
#rnx

diff --git a/Engine/Source/Editor/UnrealEd/Private/EditorDomain/EditorDomain.cpp b/Engine/Source/Editor/UnrealEd/Private/EditorDomain/EditorDomain.cpp
--- a/Engine/Source/Editor/UnrealEd/Private/EditorDomain/EditorDomain.cpp
+++ b/Engine/Source/Editor/UnrealEd/Private/EditorDomain/EditorDomain.cpp
@@ -376,6 +376,8 @@ int64 FEditorDomain::FileSize(const FPackagePath& PackagePath, EPackageSegment P
 		RequestEditorDomainPackage(PackagePath, PackageSource->Digest.Hash, SkipFlags,
 			*Owner, MoveTemp(MetaDataGetComplete));
 	}
+	COOK_STAT(auto Timer = UE::EditorDomain::CookStats::Usage.TimeAsyncWait());
+	COOK_STAT(Timer.TrackCyclesOnly());
 	Owner->Wait();
 	return FileSize;
 }
diff --git a/Engine/Source/Editor/UnrealEd/Private/EditorDomain/EditorDomainArchive.cpp b/Engine/Source/Editor/UnrealEd/Private/EditorDomain/EditorDomainArchive.cpp
--- a/Engine/Source/Editor/UnrealEd/Private/EditorDomain/EditorDomainArchive.cpp
+++ b/Engine/Source/Editor/UnrealEd/Private/EditorDomain/EditorDomainArchive.cpp
@@ -49,6 +49,8 @@ void FEditorDomainPackageSegments::WaitForReady() const
 	{
 		return;
 	}
+	COOK_STAT(auto Timer = UE::EditorDomain::CookStats::Usage.TimeAsyncWait());
+	COOK_STAT(Timer.TrackCyclesOnly());
 	const_cast<UE::DerivedData::FRequestOwner&>(RequestOwner).Wait();
 	Source = AsyncSource;
 }
@@ -286,6 +288,7 @@ void FEditorDomainPackageSegments::OnRecordRequestComplete(UE::DerivedData::FCac
 		const FCbObject& MetaData = Response.Record.GetMeta();
 		bExistsInEditorDomain = true;
 		bStorageValid = MetaData["Valid"].AsBool(false);
+		COOK_STAT(UE::EditorDomain::CookStats::Usage.TimeSyncWork().AddHit(Response.Record.GetMeta().GetSize()));
 	}
 
 	if (bStorageValid)
@@ -435,6 +438,9 @@ void FEditorDomainPackageSegments::SendSegmentRequest(FSegment& Segment)
 					UE_LOG(LogEditorDomain, Error, TEXT("Package %s has corrupted cache data in segment %d."), *PackagePath.GetDebugName(), (int32)(&Segment - &Segments[0]));
 					Segment.Data.Reset();
 				}
+				COOK_STAT(auto Timer = UE::EditorDomain::CookStats::Usage.TimeSyncWork());
+				COOK_STAT(Timer.AddHit(SegmentSize));
+				COOK_STAT(Timer.TrackCyclesOnly()); // We're using TrackCyclesOnly to track the size without counting another hit. We actually don't care about the cycles
 			}
 			Segment.bAsyncComplete = true;
 		});
diff --git a/Engine/Source/Editor/UnrealEd/Private/EditorDomain/EditorDomainUtils.cpp b/Engine/Source/Editor/UnrealEd/Private/EditorDomain/EditorDomainUtils.cpp
--- a/Engine/Source/Editor/UnrealEd/Private/EditorDomain/EditorDomainUtils.cpp
+++ b/Engine/Source/Editor/UnrealEd/Private/EditorDomain/EditorDomainUtils.cpp
@@ -87,6 +87,22 @@ static ValueType MapFindRef(const TMap<KeyType, ValueType, SetAllocator, KeyFunc
 	return FoundValue ? *FoundValue : DefaultValue;
 }
 
+#if ENABLE_COOK_STATS
+namespace UE::EditorDomain::CookStats
+{
+
+FCookStats::FDDCResourceUsageStats Usage;
+void Register()
+{
+	FCookStatsManager::CookStatsCallbacks.AddLambda([](FCookStatsManager::AddStatFuncRef AddStat)
+		{
+			Usage.LogStats(AddStat, TEXT("EditorDomain.Usage"), TEXT("Package"));
+		});
+}
+
+}
+#endif
+
 namespace UE::EditorDomain
 {
 
@@ -1077,6 +1093,9 @@ void UtilsInitialize()
 	// where UtilsInitialize is called is too early; trying to call UStruct->GetSchemaHash at that
 	// time will break the UClass. Defer the construction of allowlist-based data until OnPostEngineInit
 	GUtilsPostInitDelegate = FCoreDelegates::OnPostEngineInit.AddLambda([]() { UtilsPostEngineInit(); });
+
+
+	COOK_STAT(UE::EditorDomain::CookStats::Register());
 }
 
 void UtilsPostEngineInit()
@@ -1619,6 +1638,8 @@ bool TrySavePackage(UPackage* Package)
 			*WriteToString<256>(PackageName));
 		return false;
 	}
+	COOK_STAT(auto Timer = UE::EditorDomain::CookStats::Usage.TimeSyncWork());
+
 	UE_LOG(LogEditorDomain, Verbose, TEXT("Saving to EditorDomain: %s."), *WriteToString<256>(PackageName));
 
 	SCOPED_CUSTOM_EDITORDOMAINTIMER(EditorDomain_TrySavePackage)
@@ -1846,6 +1867,11 @@ bool TrySavePackage(UPackage* Package)
 		{
 			RecordBuilder.AddValue(Attachment.ValueId, Attachment.Buffer);
 		}
+		COOK_STAT(Timer.AddMiss(FileSize + MetaData.GetSaveSize()));
+	}
+	else
+	{
+		COOK_STAT(Timer.AddMiss(MetaData.GetSaveSize()));
 	}
 	RecordBuilder.SetMeta(MetaData.Save().AsObject());
 	FRequestOwner Owner(EPriority::Normal);
diff --git a/Engine/Source/Editor/UnrealEd/Private/EditorDomain/EditorDomainUtils.h b/Engine/Source/Editor/UnrealEd/Private/EditorDomain/EditorDomainUtils.h
--- a/Engine/Source/Editor/UnrealEd/Private/EditorDomain/EditorDomainUtils.h
+++ b/Engine/Source/Editor/UnrealEd/Private/EditorDomain/EditorDomainUtils.h
@@ -13,6 +13,7 @@
 #include "Memory/MemoryFwd.h"
 #include "Misc/ConfigCacheIni.h"
 #include "Misc/EnumClassFlags.h"
+#include "ProfilingDebugging/CookStats.h"
 #include "Serialization/CustomVersion.h"
 #include "Templates/Function.h"
 #include "UObject/NameTypes.h"
@@ -112,4 +113,11 @@ FClassDigestMap& GetClassDigests();
 /** Initializes some global config-driven values used by the EditorDomain and TargetDomain. */
 void UtilsInitialize();
 
+#if ENABLE_COOK_STATS
+namespace CookStats
+{
+	extern FCookStats::FDDCResourceUsageStats Usage;
+}
+#endif
+
 }
\ No newline at end of file
