SPropertyMenuComponentPicker::CanPaste() no longer causes log spam if a string containing a space is in the clipboard.

diff --git a/Engine/Source/Editor/PropertyEditor/Private/UserInterface/PropertyEditor/SPropertyMenuComponentPicker.cpp b/Engine/Source/Editor/PropertyEditor/Private/UserInterface/PropertyEditor/SPropertyMenuComponentPicker.cpp
--- a/Engine/Source/Editor/PropertyEditor/Private/UserInterface/PropertyEditor/SPropertyMenuComponentPicker.cpp
+++ b/Engine/Source/Editor/PropertyEditor/Private/UserInterface/PropertyEditor/SPropertyMenuComponentPicker.cpp
@@ -216,7 +216,7 @@ bool SPropertyMenuComponentPicker::CanPaste()
 		bCanPaste = !Class.IsEmpty() && !PossibleObjectPath.IsEmpty();
 		if (bCanPaste)
 		{
-			bCanPaste = LoadClass<UActorComponent>(nullptr, *Class) != nullptr;
+			bCanPaste = LoadClass<UActorComponent>(nullptr, *Class, nullptr, LOAD_Quiet | LOAD_NoWarn) != nullptr;
 		}
 		if (bCanPaste)
 		{
diff --git a/Engine/Source/Editor/UnrealEd/Classes/Editor/PropertyEditorTestObject.h b/Engine/Source/Editor/UnrealEd/Classes/Editor/PropertyEditorTestObject.h
--- a/Engine/Source/Editor/UnrealEd/Classes/Editor/PropertyEditorTestObject.h
+++ b/Engine/Source/Editor/UnrealEd/Classes/Editor/PropertyEditorTestObject.h
@@ -472,6 +472,7 @@ class UPropertyEditorTestObject : public UObject
 	UPROPERTY(EditAnywhere, Category=Assets, meta=(AllowedClasses="PointLight, SpotLight", ExactClass=true))
 	FSoftObjectPath ExactPointOrSpotLightActorReference;
 
+	// NOTE: intentionally misplaced space in AllowedClasses
 	UPROPERTY(EditAnywhere, Category=Assets, meta=(AllowedClasses="Light ,StaticMeshActor", DisplayThumbnail))
 	FSoftObjectPath LightOrStaticMeshActorReference;
 
@@ -487,6 +488,9 @@ class UPropertyEditorTestObject : public UObject
 	UPROPERTY(EditAnywhere, Category=Assets)
 	FSoftObjectPath DisabledByCanEditChange;
 
+	UPROPERTY(EditAnywhere, Category=Assets)
+	FComponentReference ComponentReference;
+
 	UPROPERTY(EditAnywhere, Category=StructTests, meta=(InlineEditConditionToggle))
 	bool bEditCondition;
 
