Adds a different default screen percentage behavior for path tracing in editor

diff --git a/Engine/Source/Editor/UnrealEd/Classes/Editor/EditorPerformanceSettings.h b/Engine/Source/Editor/UnrealEd/Classes/Editor/EditorPerformanceSettings.h
--- a/Engine/Source/Editor/UnrealEd/Classes/Editor/EditorPerformanceSettings.h
+++ b/Engine/Source/Editor/UnrealEd/Classes/Editor/EditorPerformanceSettings.h
@@ -74,6 +74,10 @@ class UEditorPerformanceSettings : public UDeveloperSettings
 	UPROPERTY(EditAnywhere, config, Category=ViewportResolution, meta=(
 		DisplayName="Override project's default screen percentage mode for realtime editor viewports."))
 	EEditorUserScreenPercentageModeOverride RealtimeScreenPercentageMode;
+	
+	UPROPERTY(EditAnywhere, config, Category=ViewportResolution, meta=(
+		DisplayName="Override project's default screen percentage mode for path-traced editor viewports."))
+	EEditorUserScreenPercentageModeOverride PathTracerScreenPercentageMode;
 
 	UPROPERTY(EditAnywhere, config, Category=ViewportResolution, meta=(
 		DisplayName="Override project's default screen percentage mode for non-realtime editor viewports."))
diff --git a/Engine/Source/Editor/UnrealEd/Private/EditorPerformanceSettings.cpp b/Engine/Source/Editor/UnrealEd/Private/EditorPerformanceSettings.cpp
--- a/Engine/Source/Editor/UnrealEd/Private/EditorPerformanceSettings.cpp
+++ b/Engine/Source/Editor/UnrealEd/Private/EditorPerformanceSettings.cpp
@@ -24,6 +24,7 @@ UEditorPerformanceSettings::UEditorPerformanceSettings(const FObjectInitializer&
 	, bDisplayHighDPIViewports(true)
 	, bOverridePIEScreenPercentage(true)
 	, RealtimeScreenPercentageMode(EEditorUserScreenPercentageModeOverride::ProjectDefault)
+	, PathTracerScreenPercentageMode(EEditorUserScreenPercentageModeOverride::ProjectDefault)
 	, NonRealtimeScreenPercentageMode(EEditorUserScreenPercentageModeOverride::ProjectDefault)
 	, bOverrideManualScreenPercentage(false)
 	, ManualScreenPercentage(100.0f)
diff --git a/Engine/Source/Editor/UnrealEd/Private/EditorViewportClient.cpp b/Engine/Source/Editor/UnrealEd/Private/EditorViewportClient.cpp
--- a/Engine/Source/Editor/UnrealEd/Private/EditorViewportClient.cpp
+++ b/Engine/Source/Editor/UnrealEd/Private/EditorViewportClient.cpp
@@ -2742,7 +2742,7 @@ bool FEditorViewportClient::SupportsPreviewResolutionFraction() const
 float FEditorViewportClient::GetDefaultPrimaryResolutionFractionTarget() const
 {
 	FStaticResolutionFractionHeuristic StaticHeuristic;
-	StaticHeuristic.Settings.PullEditorRenderingSettings(bIsRealtime);
+	StaticHeuristic.Settings.PullEditorRenderingSettings(bIsRealtime, /* bIsPathTraced = */ EngineShowFlags.PathTracing);
 
 	if (SupportsLowDPIPreview() && IsLowDPIPreview()) // TODO: && ViewFamily.SupportsScreenPercentage())
 	{
diff --git a/Engine/Source/Editor/UnrealEd/Private/Settings/EditorProjectSettings.cpp b/Engine/Source/Editor/UnrealEd/Private/Settings/EditorProjectSettings.cpp
--- a/Engine/Source/Editor/UnrealEd/Private/Settings/EditorProjectSettings.cpp
+++ b/Engine/Source/Editor/UnrealEd/Private/Settings/EditorProjectSettings.cpp
@@ -14,6 +14,11 @@ TAutoConsoleVariable<int32> CVarEditorViewportDefaultScreenPercentageMode(
 	TEXT("Controls the default screen percentage mode for realtime editor viewports."),
 	ECVF_Default);
 
+TAutoConsoleVariable<int32> CVarEditorViewportDefaultPathTracerScreenPercentageMode(
+	TEXT("r.Editor.Viewport.ScreenPercentageMode.PathTracer"), 0,
+	TEXT("Controls the default screen percentage mode for path-traced viewports."),
+	ECVF_Default);
+
 TAutoConsoleVariable<int32> CVarEditorViewportNonRealtimeDefaultScreenPercentageMode(
 	TEXT("r.Editor.Viewport.ScreenPercentageMode.NonRealTime"), 2,
 	TEXT("Controls the default screen percentage mode for non-realtime editor viewports."),
@@ -178,6 +183,7 @@ void ULevelEditor2DSettings::PostEditChangeProperty(struct FPropertyChangedEvent
 UEditorPerformanceProjectSettings::UEditorPerformanceProjectSettings(const FObjectInitializer& ObjectInitializer)
 	: Super(ObjectInitializer)
 	, RealtimeScreenPercentageMode(EScreenPercentageMode::BasedOnDisplayResolution)
+	, PathTracerScreenPercentageMode(EScreenPercentageMode::Manual)
 	, NonRealtimeScreenPercentageMode(EScreenPercentageMode::BasedOnDPIScale)
 	, ManualScreenPercentage(100.0f)
 	, MinViewportRenderingResolution(720)
@@ -221,6 +227,7 @@ void UEditorPerformanceProjectSettings::ExportResolutionValuesToConsoleVariables
 		const UEditorPerformanceProjectSettings* EditorProjectSettings = GetDefault<UEditorPerformanceProjectSettings>();
 
 		CVarEditorViewportDefaultScreenPercentageMode->Set(int32(EditorProjectSettings->RealtimeScreenPercentageMode), ECVF_SetByProjectSetting);
+		CVarEditorViewportDefaultPathTracerScreenPercentageMode->Set(int32(EditorProjectSettings->PathTracerScreenPercentageMode), ECVF_SetByProjectSetting);
 		CVarEditorViewportNonRealtimeDefaultScreenPercentageMode->Set(int32(EditorProjectSettings->NonRealtimeScreenPercentageMode), ECVF_SetByProjectSetting);
 		CVarEditorViewportDefaultScreenPercentage->Set(EditorProjectSettings->ManualScreenPercentage, ECVF_SetByProjectSetting);
 		CVarEditorViewportDefaultMinRenderingResolution->Set(EditorProjectSettings->MinViewportRenderingResolution, ECVF_SetByProjectSetting);
@@ -232,6 +239,7 @@ void UEditorPerformanceProjectSettings::ExportResolutionValuesToConsoleVariables
 	{
 		const UEditorPerformanceSettings* EditorUserSettings = GetDefault<UEditorPerformanceSettings>();
 
+		// Real-time viewports
 		{
 			if (EditorUserSettings->RealtimeScreenPercentageMode == EEditorUserScreenPercentageModeOverride::Manual)
 			{
@@ -247,6 +255,23 @@ void UEditorPerformanceProjectSettings::ExportResolutionValuesToConsoleVariables
 			}
 		}
 
+		// Path traced viewports
+		{
+			if (EditorUserSettings->PathTracerScreenPercentageMode == EEditorUserScreenPercentageModeOverride::Manual)
+			{
+				CVarEditorViewportDefaultPathTracerScreenPercentageMode->Set(int32(EScreenPercentageMode::Manual), ECVF_SetByProjectSetting);
+			}
+			else if (EditorUserSettings->PathTracerScreenPercentageMode == EEditorUserScreenPercentageModeOverride::BasedOnDisplayResolution)
+			{
+				CVarEditorViewportDefaultPathTracerScreenPercentageMode->Set(int32(EScreenPercentageMode::BasedOnDisplayResolution), ECVF_SetByProjectSetting);
+			}
+			else if (EditorUserSettings->PathTracerScreenPercentageMode == EEditorUserScreenPercentageModeOverride::BasedOnDPIScale)
+			{
+				CVarEditorViewportDefaultPathTracerScreenPercentageMode->Set(int32(EScreenPercentageMode::BasedOnDPIScale), ECVF_SetByProjectSetting);
+			}
+		}
+
+		// Non-realtime viewports
 		{
 			if (EditorUserSettings->NonRealtimeScreenPercentageMode == EEditorUserScreenPercentageModeOverride::Manual)
 			{
@@ -296,24 +321,3 @@ void UEditorPerformanceProjectSettings::ExportResolutionValuesToConsoleVariables
 		}
 	}
 }
-
-void SetupEditorResolutionFractionHeuristicSettings(FStaticResolutionFractionHeuristic::FUserSettings* OutputSettings, bool bIsRealTime)
-{
-	FStaticResolutionFractionHeuristic::FUserSettings Settings;
-
-	// Setup up with the per project settings
-	{
-		if (bIsRealTime)
-		{
-			OutputSettings->Mode = EScreenPercentageMode(FMath::Clamp(CVarEditorViewportDefaultScreenPercentageMode->GetInt(), 0, 2));
-		}
-		else
-		{
-			OutputSettings->Mode = EScreenPercentageMode(FMath::Clamp(CVarEditorViewportNonRealtimeDefaultScreenPercentageMode->GetInt(), 0, 2));
-		}
-		OutputSettings->GlobalResolutionFraction = CVarEditorViewportDefaultScreenPercentage->GetFloat() / 100.0f;
-		OutputSettings->MinRenderingResolution = CVarEditorViewportDefaultMinRenderingResolution->GetInt();
-		OutputSettings->MaxRenderingResolution = CVarEditorViewportDefaultMaxRenderingResolution->GetInt();
-	}
-
-}
diff --git a/Engine/Source/Editor/UnrealEd/Public/Settings/EditorProjectSettings.h b/Engine/Source/Editor/UnrealEd/Public/Settings/EditorProjectSettings.h
--- a/Engine/Source/Editor/UnrealEd/Public/Settings/EditorProjectSettings.h
+++ b/Engine/Source/Editor/UnrealEd/Public/Settings/EditorProjectSettings.h
@@ -192,6 +192,10 @@ public:
 	UPROPERTY(EditAnywhere, config, Category=ViewportResolution, meta=(
 		DisplayName="Default screen percentage mode for realtime editor viewports."))
 	EScreenPercentageMode RealtimeScreenPercentageMode;
+	
+	UPROPERTY(EditAnywhere, config, Category=ViewportResolution, meta=(
+		DisplayName="Default screen percentage mode for path traced editor viewports."))
+	EScreenPercentageMode PathTracerScreenPercentageMode;
 
 	UPROPERTY(EditAnywhere, config, Category=ViewportResolution, meta=(
 		DisplayName="Default screen percentage mode for non-realtime editor viewports."))
diff --git a/Engine/Source/Runtime/Engine/Private/GameViewportClient.cpp b/Engine/Source/Runtime/Engine/Private/GameViewportClient.cpp
--- a/Engine/Source/Runtime/Engine/Private/GameViewportClient.cpp
+++ b/Engine/Source/Runtime/Engine/Private/GameViewportClient.cpp
@@ -1656,7 +1656,7 @@ void UGameViewportClient::Draw(FViewport* InViewport, FCanvas* SceneCanvas)
 #if WITH_EDITOR
 				if (FStaticResolutionFractionHeuristic::FUserSettings::EditorOverridePIESettings())
 				{
-					StaticHeuristic.Settings.PullEditorRenderingSettings(/* bIsRealTime = */ true);
+					StaticHeuristic.Settings.PullEditorRenderingSettings(/* bIsRealTime = */ true, /* bIsPathTraced = */ ViewFamily.EngineShowFlags.PathTracing);
 				}
 				else
 #endif
diff --git a/Engine/Source/Runtime/Engine/Private/LegacyScreenPercentageDriver.cpp b/Engine/Source/Runtime/Engine/Private/LegacyScreenPercentageDriver.cpp
--- a/Engine/Source/Runtime/Engine/Private/LegacyScreenPercentageDriver.cpp
+++ b/Engine/Source/Runtime/Engine/Private/LegacyScreenPercentageDriver.cpp
@@ -199,16 +199,21 @@ void FStaticResolutionFractionHeuristic::FUserSettings::PullRunTimeRenderingSett
 	AutoPixelCountMultiplier = CVarAutoPixelCountMultiplier.GetValueOnGameThread();
 }
 
-void FStaticResolutionFractionHeuristic::FUserSettings::PullEditorRenderingSettings(bool bIsRealTime)
+void FStaticResolutionFractionHeuristic::FUserSettings::PullEditorRenderingSettings(bool bIsRealTime, bool bIsPathTraced)
 #if WITH_EDITOR
 {
 	static auto CVarEditorViewportDefaultScreenPercentageRealTimeMode = IConsoleManager::Get().FindConsoleVariable(TEXT("r.Editor.Viewport.ScreenPercentageMode.RealTime"));
+	static auto CVarEditorViewportDefaultScreenPercentagePathTracerMode = IConsoleManager::Get().FindConsoleVariable(TEXT("r.Editor.Viewport.ScreenPercentageMode.PathTracer"));
 	static auto CVarEditorViewportDefaultScreenPercentageMode = IConsoleManager::Get().FindConsoleVariable(TEXT("r.Editor.Viewport.ScreenPercentageMode.NonRealTime"));
 	static auto CVarEditorViewportDefaultScreenPercentage = IConsoleManager::Get().FindConsoleVariable(TEXT("r.Editor.Viewport.ScreenPercentage"));
 	static auto CVarEditorViewportDefaultMinRenderingResolution = IConsoleManager::Get().FindConsoleVariable(TEXT("r.Editor.Viewport.MinRenderingResolution"));
 	static auto CVarEditorViewportDefaultMaxRenderingResolution = IConsoleManager::Get().FindConsoleVariable(TEXT("r.Editor.Viewport.MaxRenderingResolution"));
 
-	if (bIsRealTime)
+	if (bIsPathTraced)
+	{
+		Mode = EScreenPercentageMode(FMath::Clamp(CVarEditorViewportDefaultScreenPercentagePathTracerMode->GetInt(), 0, 2));
+	}
+	else if (bIsRealTime)
 	{
 		Mode = EScreenPercentageMode(FMath::Clamp(CVarEditorViewportDefaultScreenPercentageRealTimeMode->GetInt(), 0, 2));
 	}
diff --git a/Engine/Source/Runtime/Engine/Public/LegacyScreenPercentageDriver.h b/Engine/Source/Runtime/Engine/Public/LegacyScreenPercentageDriver.h
--- a/Engine/Source/Runtime/Engine/Public/LegacyScreenPercentageDriver.h
+++ b/Engine/Source/Runtime/Engine/Public/LegacyScreenPercentageDriver.h
@@ -94,7 +94,7 @@ struct ENGINE_API FStaticResolutionFractionHeuristic
 		void PullRunTimeRenderingSettings();
 
 		/** Pulls the user settings from the editor cvars. */
-		void PullEditorRenderingSettings(bool bIsRealTime);
+		void PullEditorRenderingSettings(bool bIsRealTime, bool bIsPathTraced);
 	};
 
 	FUserSettings Settings;
