SavePackage: Add IsCachedCookedPlatformDataLoaded calls; it is required by the contract of BeginCacheForCookedPlatformData. Add a Timeout result to return to the cookerto wait if any IsCachedCookedPlatformDataLoaded calls are incomplete.
#rnx

diff --git a/Engine/Source/Editor/UnrealEd/Private/CookOnTheFlyServer.cpp b/Engine/Source/Editor/UnrealEd/Private/CookOnTheFlyServer.cpp
--- a/Engine/Source/Editor/UnrealEd/Private/CookOnTheFlyServer.cpp
+++ b/Engine/Source/Editor/UnrealEd/Private/CookOnTheFlyServer.cpp
@@ -3158,6 +3158,8 @@ private: // Used only by UCookOnTheFlyServer, which has private access
 	FString Filename;
 	uint32 SaveFlags = 0;
 	bool bReferencedOnlyByEditorOnlyData = false;
+	bool bHasTimeOut = false;
+	bool bHasRetryErrorCode = false;
 	bool bHasFirstPlatformResults = false;
 
 	// General Package Data that is delay-loaded the first time we save a platform
@@ -3394,6 +3396,15 @@ void UCookOnTheFlyServer::PumpSaves(UE::Cook::FTickStackData& StackData, uint32
 
 		FSaveCookedPackageContext Context(*this, PackageData, PlatformsForPackage, StackData);
 		SaveCookedPackage(Context);
+		if (Context.bHasTimeOut)
+		{
+			// Timeouts can occur because of new objects created during the save, so we need to update our object cache,
+			// so we call ReleaseCookedPlatformData and ClearObjectCache to clear it and recache on next attempt.
+			ReleaseCookedPlatformData(PackageData, false/* bCompletedSave */);
+			PackageData.ClearObjectCache();
+			SaveQueue.Add(&PackageData);
+			continue;
+		}
 
 		ReleaseCookedPlatformData(PackageData, true /* bCompletedSave */);
 		PackageData.SendToState(EPackageState::Idle, ESendFlags::QueueAdd);
@@ -4435,18 +4446,32 @@ void FSaveCookedPackageContext::SetupPlatform(const ITargetPlatform* InTargetPla
 	Info.PackageName = Package->GetFName();
 	Info.LooseFilePath = PlatFilename;
 	PackageWriter->BeginPackage(Info);
+	// Set platform-specific save flags
+	FPackageData::FPlatformData& PlatformData = PackageData.FindOrAddPlatformData(TargetPlatform);
+	uint32 PlatformSaveFlagsMask = SAVE_AllowTimeout;
+	SaveFlags = (SaveFlags & ~PlatformSaveFlagsMask);
+	if (!PlatformData.bSaveTimedOut)
+	{
+		// If we timedout before, do not allow another timeout, otherwise do allow it
+		SaveFlags |= SAVE_AllowTimeout;
+	}
 
 	// Indicate Setup was successful
 	bPlatformSetupSuccessful = true;
 	SavePackageResult = ESavePackageResult::Success;
 }
 
+bool IsRetryErrorCode(ESavePackageResult Result)
+{
+	return (Result == ESavePackageResult::ReferencedOnlyByEditorOnlyData) |
+		(Result == ESavePackageResult::Timeout);
+}
+
 void FSaveCookedPackageContext::FinishPlatform()
 {
 	TRACE_CPUPROFILER_EVENT_SCOPE(FSaveCookedPackageContext::FinishPlatform);
 
 	bool bSuccessful = SavePackageResult.IsSuccessful();
-	bool bLocalReferencedOnlyByEditorOnlyData = SavePackageResult == ESavePackageResult::ReferencedOnlyByEditorOnlyData;
 
 	FAssetRegistryGenerator& Generator = *(COTFS.PlatformManager->GetPlatformData(TargetPlatform)->RegistryGenerator);
 	if (bPlatformSetupSuccessful)
@@ -4489,7 +4514,8 @@ void FSaveCookedPackageContext::FinishPlatform()
 
 	// In success or failure we want to mark the package as cooked, unless the failure was due to being referenced only
 	// by editor-only data; we may need to save it later when new content loads it through non editor-only references.
-	if (!bLocalReferencedOnlyByEditorOnlyData)
+	bool bIsRetryErrorCode = IsRetryErrorCode(SavePackageResult.Result);
+	if (!bIsRetryErrorCode)
 	{
 		PackageData.SetPlatformCooked(TargetPlatform, bSuccessful);
 	}
@@ -4510,10 +4536,11 @@ void FSaveCookedPackageContext::FinishPlatform()
 	}
 
 	// Accumulate results for SaveCookedPackage_Finish
+	bool bLocalReferencedOnlyByEditorOnlyData = SavePackageResult.Result == ESavePackageResult::ReferencedOnlyByEditorOnlyData;
 	if (bHasFirstPlatformResults && bLocalReferencedOnlyByEditorOnlyData != bReferencedOnlyByEditorOnlyData)
 	{
 		UE_LOG(LogCook, Error, TEXT("Package %s had different values for IsReferencedOnlyByEditorOnlyData from multiple platforms. ")
-			TEXT("Treading all platforms as IsReferencedOnlyByEditorOnlyData = true; this will cause the package to be ignored on the platforms that need it."),
+			TEXT("Treating all platforms as IsReferencedOnlyByEditorOnlyData = true; this will cause the package to be ignored on the platforms that need it."),
 			*Package->GetName());
 		bReferencedOnlyByEditorOnlyData = true;
 	}
@@ -4521,6 +4548,13 @@ void FSaveCookedPackageContext::FinishPlatform()
 	{
 		bReferencedOnlyByEditorOnlyData = bLocalReferencedOnlyByEditorOnlyData;
 	}
+	if (SavePackageResult.Result == ESavePackageResult::Timeout)
+	{
+		PackageData.FindOrAddPlatformData(TargetPlatform).bSaveTimedOut = true;
+		bHasTimeOut = true;
+	}
+
+	bHasRetryErrorCode |= bIsRetryErrorCode;
 	bHasFirstPlatformResults = true;
 }
 
@@ -4578,7 +4612,7 @@ void FSaveCookedPackageContext::FinishPackage()
 		}
 	}
 
-	if (!bReferencedOnlyByEditorOnlyData)
+	if (!bHasRetryErrorCode)
 	{
 		if ((COTFS.CurrentCookMode == ECookMode::CookOnTheFly) && !PackageData.GetIsUrgent())
 		{
@@ -7502,6 +7536,7 @@ void UCookOnTheFlyServer::ResetCook(TConstArrayView<TPair<const ITargetPlatform*
 				{
 					PlatformData->bCookAttempted = false;
 					PlatformData->bCookSucceeded = false;
+					PlatformData->bSaveTimedOut = false;
 				}
 			}
 		}
diff --git a/Engine/Source/Editor/UnrealEd/Private/Cooker/CookPackageData.cpp b/Engine/Source/Editor/UnrealEd/Private/Cooker/CookPackageData.cpp
--- a/Engine/Source/Editor/UnrealEd/Private/Cooker/CookPackageData.cpp
+++ b/Engine/Source/Editor/UnrealEd/Private/Cooker/CookPackageData.cpp
@@ -46,7 +46,7 @@ static FAutoConsoleVariableRef CVarPollAsyncPeriod(
 //////////////////////////////////////////////////////////////////////////
 // FPackageData
 FPackageData::FPlatformData::FPlatformData()
-	: bRequested(false), bCookAttempted(false), bCookSucceeded(false), bExplored(false)
+	: bRequested(false), bCookAttempted(false), bCookSucceeded(false), bExplored(false), bSaveTimedOut(false)
 {
 }
 
@@ -288,6 +288,8 @@ void FPackageData::SetPlatformCooked(const ITargetPlatform* TargetPlatform, bool
 			bModified = bModified | (Pair.Value.bCookAttempted == false);
 			Pair.Value.bCookAttempted = true;
 			Pair.Value.bCookSucceeded = bSucceeded;
+			// Clear the SaveTimedOut when get a cook result, in case we save again later and need to allow retry again
+			Pair.Value.bSaveTimedOut = false;
 		}
 		else
 		{
@@ -299,6 +301,7 @@ void FPackageData::SetPlatformCooked(const ITargetPlatform* TargetPlatform, bool
 		FPlatformData& Value = PlatformDatas.FindOrAdd(TargetPlatform);
 		Value.bCookAttempted = true;
 		Value.bCookSucceeded = bSucceeded;
+		Value.bSaveTimedOut = false;
 		bModified = true;
 	}
 	if (bModified && !bHasAnyOthers)
@@ -324,6 +327,7 @@ void FPackageData::ClearCookProgress()
 		Pair.Value.bCookAttempted = false;
 		Pair.Value.bCookSucceeded = false;
 		Pair.Value.bExplored = false;
+		Pair.Value.bSaveTimedOut = false;
 	}
 	if (bModifiedCookAttempted)
 	{
@@ -343,6 +347,7 @@ void FPackageData::ClearCookProgress(const ITargetPlatform* TargetPlatform)
 			Pair.Value.bCookAttempted = false;
 			Pair.Value.bCookSucceeded = false;
 			Pair.Value.bExplored = false;
+			Pair.Value.bSaveTimedOut = false;
 		}
 		else
 		{
diff --git a/Engine/Source/Editor/UnrealEd/Private/Cooker/CookPackageData.h b/Engine/Source/Editor/UnrealEd/Private/Cooker/CookPackageData.h
--- a/Engine/Source/Editor/UnrealEd/Private/Cooker/CookPackageData.h
+++ b/Engine/Source/Editor/UnrealEd/Private/Cooker/CookPackageData.h
@@ -127,6 +127,7 @@ public:
 		bool bCookAttempted : 1;
 		bool bCookSucceeded : 1;
 		bool bExplored : 1;
+		bool bSaveTimedOut : 1;
 	};
 
 	FPackageData(FPackageDatas& PackageDatas, const FName& InPackageName, const FName& InFileName);
diff --git a/Engine/Source/Runtime/CoreUObject/Private/UObject/SavePackage2.cpp b/Engine/Source/Runtime/CoreUObject/Private/UObject/SavePackage2.cpp
--- a/Engine/Source/Runtime/CoreUObject/Private/UObject/SavePackage2.cpp
+++ b/Engine/Source/Runtime/CoreUObject/Private/UObject/SavePackage2.cpp
@@ -2067,23 +2067,62 @@ ESavePackageResult FinalizeFile(FStructuredArchive::FRecord& StructuredArchiveRo
 	return ESavePackageResult::Success;
 }
 
-void BeginCachePlatformCookedData(FSaveContext& SaveContext)
+ESavePackageResult BeginCachePlatformCookedData(FSaveContext& SaveContext)
 {
 #if WITH_EDITOR
+	// TODO: Remove BeginCacheForCookedPlatformData from SavePackage; it is already called by the cooker.
 	// Cache platform cooked data
 	if (SaveContext.IsCooking() && !SaveContext.IsConcurrent())
 	{
+		const ITargetPlatform* TargetPlatform = SaveContext.GetTargetPlatform();
 		TArray<UObject*> ObjectsInPackage;
 		GetObjectsWithPackage(SaveContext.GetPackage(), ObjectsInPackage);
-		for (UObject* Object : ObjectsInPackage)
+		for (int32 Index = 0; Index < ObjectsInPackage.Num(); /** incremented in loop */)
 		{
-			if (!SaveContext.IsUnsaveable(Object))
+			UObject* Object = ObjectsInPackage[Index];
+			if (SaveContext.IsUnsaveable(Object))
+			{
+				ObjectsInPackage.RemoveAtSwap(Index);
+				continue;
+			}
+			Object->BeginCacheForCookedPlatformData(TargetPlatform);
+			if (Object->IsCachedCookedPlatformDataLoaded(TargetPlatform))
 			{
-				Object->BeginCacheForCookedPlatformData(SaveContext.GetTargetPlatform());
+				ObjectsInPackage.RemoveAtSwap(Index);
+				continue;
+			}
+			++Index;
+		}
+		if (ObjectsInPackage.Num())
+		{
+			if (SaveContext.GetSaveArgs().SaveFlags & SAVE_AllowTimeout)
+			{
+				return ESavePackageResult::Timeout;
+			}
+			const float MaxWaitSeconds = 30.f;
+			const double EndTimeSeconds = FPlatformTime::Seconds() + MaxWaitSeconds;
+			while (ObjectsInPackage.Num())
+			{
+				if (FPlatformTime::Seconds() > EndTimeSeconds)
+				{
+					UE_LOG(LogSavePackage, Error, TEXT("Save of %s failed: timed out waiting for IsCachedCookedPlatformDataLoaded on %s."),
+						SaveContext.GetFilename(), *ObjectsInPackage[0]->GetFullName());
+					return ESavePackageResult::Error;
+				}
+				for (int32 Index = 0; Index < ObjectsInPackage.Num(); /** incremented in loop */)
+				{
+					if (ObjectsInPackage[Index]->IsCachedCookedPlatformDataLoaded(TargetPlatform))
+					{
+						ObjectsInPackage.RemoveAtSwap(Index);
+						continue;
+					}
+					++Index;
+				}
 			}
 		}
 	}
 #endif
+	return ESavePackageResult::Success;
 }
 
 void ClearCachedPlatformCookedData(FSaveContext& SaveContext)
@@ -2530,7 +2569,11 @@ FSavePackageResultStruct UPackage::Save2(UPackage* InPackage, UObject* InAsset,
 
 	// Trigger platform cooked data caching before package harvesting 
 	// because it might modify some property and hence affect harvested property name for tagged property for example
-	BeginCachePlatformCookedData(SaveContext);
+	SaveContext.Result = BeginCachePlatformCookedData(SaveContext);
+	if (SaveContext.Result != ESavePackageResult::Success)
+	{
+		return SaveContext.Result;
+	}
 
 	SlowTask.EnterProgressFrame();
 	{
diff --git a/Engine/Source/Runtime/CoreUObject/Public/UObject/ObjectMacros.h b/Engine/Source/Runtime/CoreUObject/Public/UObject/ObjectMacros.h
--- a/Engine/Source/Runtime/CoreUObject/Public/UObject/ObjectMacros.h
+++ b/Engine/Source/Runtime/CoreUObject/Public/UObject/ObjectMacros.h
@@ -99,6 +99,7 @@ enum ESaveFlags
 	SAVE_Unversioned_Properties		= 0x00004000,	///< Properties are saved without property name information, and are saved/loaded in the order of the current binary.
 													///This is only reasonable to use with full cooked builds for distribution, or with a domain that selects the payload matching a given version externally.
 	SAVE_Optional					= 0x00008000,	///< Indicate that we to save optional exports. This flag is only valid while cooking. Optional exports are filtered if not specified during cooking. 
+	SAVE_AllowTimeout				= 0x00010000,   ///< If present, if SavePackage is blocked by an asynchronous operation, it will quickly fail with ESavePackageResult::Timeout, otherwise it will wait for a while and then fail with ESavePackageResult::Error.
 	SAVE_Unversioned = SAVE_Unversioned_Native | SAVE_Unversioned_Properties,
 };
 
diff --git a/Engine/Source/Runtime/CoreUObject/Public/UObject/Package.h b/Engine/Source/Runtime/CoreUObject/Public/UObject/Package.h
--- a/Engine/Source/Runtime/CoreUObject/Public/UObject/Package.h
+++ b/Engine/Source/Runtime/CoreUObject/Public/UObject/Package.h
@@ -60,6 +60,8 @@ enum class ESavePackageResult
 	ValidatorSuppress,
 	/** Internal save result used to identify a valid empty internal save realm to skip over. @see ESaveRealm */
 	EmptyRealm,
+	/** SavePackage is blocked by an asynchronous operation, so it quickly aborted. Can only be returned if SAVE_AllowTimeout is present in SaveFlags */
+	Timeout,
 };
 
 /**
