Enforce FBX animated import length to be frame aligned as per the import sampling rate

diff --git a/Engine/Source/Editor/UnrealEd/Private/SkeletalMeshEdit.cpp b/Engine/Source/Editor/UnrealEd/Private/SkeletalMeshEdit.cpp
--- a/Engine/Source/Editor/UnrealEd/Private/SkeletalMeshEdit.cpp
+++ b/Engine/Source/Editor/UnrealEd/Private/SkeletalMeshEdit.cpp
@@ -916,6 +916,15 @@ bool UnFbx::FFbxImporter::ValidateAnimStack(TArray<FbxNode*>& SortedLinks, TArra
 		return false;
 	}
 
+	const float SequenceLengthInSeconds = FGenericPlatformMath::Max<float>(AnimTimeSpan.GetDuration().GetSecondDouble(), MINIMUM_ANIMATION_LENGTH);
+	const FFrameRate TargetFrameRate(ResampleRate, 1);
+	const float SubFrame = TargetFrameRate.AsFrameTime(SequenceLengthInSeconds).GetSubFrame();
+	if (!FMath::IsNearlyZero(SubFrame))
+	{
+		AddTokenizedErrorMessage(FTokenizedMessage::Create(EMessageSeverity::Error, FText::Format(LOCTEXT("Error_InvalidImportLength", "Animation length {0} is not compatible with import frame-rate {1} (sub frame {2}), animation has to be frame-border aligned."), FText::AsNumber(SequenceLengthInSeconds), TargetFrameRate.ToPrettyText(), FText::AsNumber(SubFrame))), FFbxErrors::Animation_InvalidData);
+		return false;
+	}
+
 	const FBXImportOptions* ImportOption = GetImportOptions();
 	// only add morph time if not setrange. If Set Range there is no reason to override time
 	if ( bImportMorph && ImportOption->AnimationLengthImportType != FBXALIT_SetRange)
