Specifying the name of the style for the icon directly

diff --git a/Engine/Source/Editor/UnrealEd/Private/Tools/DefaultEdMode.cpp b/Engine/Source/Editor/UnrealEd/Private/Tools/DefaultEdMode.cpp
--- a/Engine/Source/Editor/UnrealEd/Private/Tools/DefaultEdMode.cpp
+++ b/Engine/Source/Editor/UnrealEd/Private/Tools/DefaultEdMode.cpp
@@ -3,7 +3,6 @@
 #include "Tools/DefaultEdMode.h"
 
 #include "EditorModes.h"
-#include "Styling/AppStyle.h"
 #include "Textures/SlateIcon.h"
 #include "LevelEditorViewport.h"
 #include "Elements/Framework/TypedElementList.h"
@@ -27,7 +26,7 @@ UEdModeDefault::UEdModeDefault()
 	Info = FEditorModeInfo(
 		FBuiltinEditorModes::EM_Default,
 		NSLOCTEXT("DefaultMode", "DisplayName", "Selection"),
-		FSlateIcon(FAppStyle::GetAppStyleSetName(), "LevelEditor.SelectMode", "LevelEditor.SelectMode.Small"),
+		FSlateIcon("EditorStyle", "LevelEditor.SelectMode", "LevelEditor.SelectMode.Small"),
 		true, 0);
 }
 
