Increase threshold when comparing calculated frame-time subframe - used to enforce importing animations at frame-boundaries (introduced in CL 20162458)

diff --git a/Engine/Source/Editor/UnrealEd/Private/SkeletalMeshEdit.cpp b/Engine/Source/Editor/UnrealEd/Private/SkeletalMeshEdit.cpp
--- a/Engine/Source/Editor/UnrealEd/Private/SkeletalMeshEdit.cpp
+++ b/Engine/Source/Editor/UnrealEd/Private/SkeletalMeshEdit.cpp
@@ -919,7 +919,7 @@ bool UnFbx::FFbxImporter::ValidateAnimStack(TArray<FbxNode*>& SortedLinks, TArra
 	const float SequenceLengthInSeconds = FGenericPlatformMath::Max<float>(AnimTimeSpan.GetDuration().GetSecondDouble(), MINIMUM_ANIMATION_LENGTH);
 	const FFrameRate TargetFrameRate(ResampleRate, 1);
 	const float SubFrame = TargetFrameRate.AsFrameTime(SequenceLengthInSeconds).GetSubFrame();
-	if (!FMath::IsNearlyZero(SubFrame))
+	if (!FMath::IsNearlyZero(SubFrame, KINDA_SMALL_NUMBER))
 	{
 		AddTokenizedErrorMessage(FTokenizedMessage::Create(EMessageSeverity::Error, FText::Format(LOCTEXT("Error_InvalidImportLength", "Animation length {0} is not compatible with import frame-rate {1} (sub frame {2}), animation has to be frame-border aligned."), FText::AsNumber(SequenceLengthInSeconds), TargetFrameRate.ToPrettyText(), FText::AsNumber(SubFrame))), FFbxErrors::Animation_InvalidData);
 		return false;
