Actor Factory: Prevent level sequence assets from being placed as level sequence actors if the class is not allowed.

diff --git a/Engine/Source/Editor/UnrealEd/Private/Factories/ActorFactory.cpp b/Engine/Source/Editor/UnrealEd/Private/Factories/ActorFactory.cpp
--- a/Engine/Source/Editor/UnrealEd/Private/Factories/ActorFactory.cpp
+++ b/Engine/Source/Editor/UnrealEd/Private/Factories/ActorFactory.cpp
@@ -136,6 +136,9 @@ ActorFactory.cpp:
 #include "LevelEditorViewport.h"
 #include "Editor.h"
 
+#include "ClassViewerFilter.h"
+#include "ClassViewerModule.h"
+
 DEFINE_LOG_CATEGORY(LogActorFactory);
 
 #define LOCTEXT_NAMESPACE "ActorFactory"
@@ -2062,6 +2065,19 @@ UActorFactoryLevelSequence::UActorFactoryLevelSequence(const FObjectInitializer&
 
 bool UActorFactoryLevelSequence::CanCreateActorFrom( const FAssetData& AssetData, FText& OutErrorMsg )
 {
+	FClassViewerModule& ClassViewerModule = FModuleManager::LoadModuleChecked<FClassViewerModule>("ClassViewer");
+	const TSharedPtr<IClassViewerFilter>& GlobalClassFilter = ClassViewerModule.GetGlobalClassViewerFilter();
+	TSharedRef<FClassViewerFilterFuncs> ClassFilterFuncs = ClassViewerModule.CreateFilterFuncs();
+	FClassViewerInitializationOptions ClassViewerOptions = {};
+
+	if (GlobalClassFilter.IsValid())
+	{
+		if (!GlobalClassFilter->IsClassAllowed(ClassViewerOptions, ALevelSequenceActor::StaticClass(), ClassFilterFuncs))
+		{
+			return false;
+		}
+	}
+
 	if ( UActorFactory::CanCreateActorFrom( AssetData, OutErrorMsg ) )
 	{
 		return true;
