World Partition: add assert when unable to lock the minimap texture when running the minimap generation world builder.

#rnx

diff --git a/Engine/Source/Editor/UnrealEd/Private/WorldPartition/WorldPartitionMiniMapBuilder.cpp b/Engine/Source/Editor/UnrealEd/Private/WorldPartition/WorldPartitionMiniMapBuilder.cpp
--- a/Engine/Source/Editor/UnrealEd/Private/WorldPartition/WorldPartitionMiniMapBuilder.cpp
+++ b/Engine/Source/Editor/UnrealEd/Private/WorldPartition/WorldPartitionMiniMapBuilder.cpp
@@ -113,7 +113,7 @@ bool UWorldPartitionMiniMapBuilder::PostRun(UWorld* World, FPackageSourceControl
 	}
 
 	FIntVector2 MinTileCoords(INT32_MAX);
-	FIntVector2 MaxTileCoords(-INT32_MAX);
+	FIntVector2 MaxTileCoords(INT32_MIN);
 	for (const FMinimapTile& Tile : MiniMapTiles)
 	{
 		MinTileCoords = FIntVector2(FMath::Min(MinTileCoords.X, Tile.Coordinates.X), FMath::Min(MinTileCoords.Y, Tile.Coordinates.Y));
@@ -133,6 +133,7 @@ bool UWorldPartitionMiniMapBuilder::PostRun(UWorld* World, FPackageSourceControl
 
 	    const uint32 DstDataStride = MiniMapTexture->Source.GetSizeX() * MiniMapTexture->Source.GetBytesPerPixel();
 	    uint8* DstDataBasePtr = MiniMapTexture->Source.LockMip(0);
+		check(DstDataBasePtr);
     
 	    for (const FMinimapTile& Tile : MiniMapTiles)
 	    {
@@ -229,4 +230,4 @@ bool UWorldPartitionMiniMapBuilder::PostRun(UWorld* World, FPackageSourceControl
 	}
 
 	return true;
-}
\ No newline at end of file
+}
