Improved precision when converting play-length to frames by using double when testing for subframe value - used to enforce importing animations at frame-boundaries (introduced in CL 20162458)

diff --git a/Engine/Source/Editor/UnrealEd/Private/SkeletalMeshEdit.cpp b/Engine/Source/Editor/UnrealEd/Private/SkeletalMeshEdit.cpp
--- a/Engine/Source/Editor/UnrealEd/Private/SkeletalMeshEdit.cpp
+++ b/Engine/Source/Editor/UnrealEd/Private/SkeletalMeshEdit.cpp
@@ -916,7 +916,7 @@ bool UnFbx::FFbxImporter::ValidateAnimStack(TArray<FbxNode*>& SortedLinks, TArra
 		return false;
 	}
 
-	const float SequenceLengthInSeconds = FGenericPlatformMath::Max<float>(AnimTimeSpan.GetDuration().GetSecondDouble(), MINIMUM_ANIMATION_LENGTH);
+	const double SequenceLengthInSeconds = FGenericPlatformMath::Max<double>(AnimTimeSpan.GetDuration().GetSecondDouble(), MINIMUM_ANIMATION_LENGTH);
 	const FFrameRate TargetFrameRate(ResampleRate, 1);
 	const float SubFrame = TargetFrameRate.AsFrameTime(SequenceLengthInSeconds).GetSubFrame();
 	if (!FMath::IsNearlyZero(SubFrame, KINDA_SMALL_NUMBER) && !FMath::IsNearlyEqual(SubFrame, 1.0f, KINDA_SMALL_NUMBER))
