EditorFactories - quiet silly static analysis false alarm

diff --git a/Engine/Source/Editor/UnrealEd/Private/Factories/EditorFactories.cpp b/Engine/Source/Editor/UnrealEd/Private/Factories/EditorFactories.cpp
--- a/Engine/Source/Editor/UnrealEd/Private/Factories/EditorFactories.cpp
+++ b/Engine/Source/Editor/UnrealEd/Private/Factories/EditorFactories.cpp
@@ -2846,7 +2846,8 @@ public:
 		{
 			check( sizeof(ColorDataType) == 8 );
 			uint16 RGBA[4] = { 0xFFFF,0xFFFF,0xFFFF, 0 };
-			memcpy(&WhiteWithZeroAlpha,RGBA,8);
+			check( sizeof(RGBA) == 8 );
+			memcpy(&WhiteWithZeroAlpha,RGBA, sizeof(ColorDataType));
 		}
 
 		// Left -> Right
