Fix crash in UWorldPartitionNavigationDataBuilder.
Use World passed in parameters instead of GetWorld() on the builder. There are no world in the builder outer chain.

diff --git a/Engine/Source/Editor/UnrealEd/Private/WorldPartition/WorldPartitionNavigationDataBuilder.cpp b/Engine/Source/Editor/UnrealEd/Private/WorldPartition/WorldPartitionNavigationDataBuilder.cpp
--- a/Engine/Source/Editor/UnrealEd/Private/WorldPartition/WorldPartitionNavigationDataBuilder.cpp
+++ b/Engine/Source/Editor/UnrealEd/Private/WorldPartition/WorldPartitionNavigationDataBuilder.cpp
@@ -28,7 +28,7 @@ UWorldPartitionNavigationDataBuilder::UWorldPartitionNavigationDataBuilder(const
 bool UWorldPartitionNavigationDataBuilder::PreRun(UWorld* World, FPackageSourceControlHelper& PackageHelper)
 {
 	// Set runtime data layer to be included in the base navmesh generation.
-	UDataLayerSubsystem* DataLayerSubsystem = UWorld::GetSubsystem<UDataLayerSubsystem>(GetWorld());
+	UDataLayerSubsystem* DataLayerSubsystem = UWorld::GetSubsystem<UDataLayerSubsystem>(World);
 	for (const TObjectPtr<UDataLayerAsset> DataLayer : World->GetWorldSettings()->BaseNavmeshDataLayers)
 	{
 		if (DataLayer != nullptr)
