Added the ability to prevent physics asset creation on import by default

diff --git a/Engine/Source/Editor/UnrealEd/Private/Fbx/FbxMainImport.cpp b/Engine/Source/Editor/UnrealEd/Private/Fbx/FbxMainImport.cpp
--- a/Engine/Source/Editor/UnrealEd/Private/Fbx/FbxMainImport.cpp
+++ b/Engine/Source/Editor/UnrealEd/Private/Fbx/FbxMainImport.cpp
@@ -47,6 +47,10 @@ DEFINE_LOG_CATEGORY(LogFbx);
 namespace UnFbx
 {
 
+static bool GDisableAutomaticPhysicsAssetCreation = false;
+static FAutoConsoleVariableRef DisableAutomaticPhysicsAssetCreationCVar(TEXT("FbxImport.DisableAutomaticPhysicsAssetCreation"),
+	GDisableAutomaticPhysicsAssetCreation, TEXT("Prevents physics assets from being created automatically by FBX import (False: disabled, True: enabled"));
+
 TSharedPtr<FFbxImporter> FFbxImporter::StaticInstance;
 
 TSharedPtr<FFbxImporter> FFbxImporter::StaticPreviewInstance;
@@ -125,6 +129,11 @@ FBXImportOptions* GetImportOptions( UnFbx::FFbxImporter* FbxImporter, UFbxImport
 			ImportUI->PhysicsAsset = NULL;
 		}
 
+		if (GDisableAutomaticPhysicsAssetCreation)
+		{
+			ImportUI->bCreatePhysicsAsset = false;
+		}
+
 		if(bForceImportType)
 		{
 			ImportUI->MeshTypeToImport = ImportType;
