Fix crash exiting monolithic editor after importing fbx

diff --git a/Engine/Source/Editor/UnrealEd/Private/Fbx/FbxMainImport.cpp b/Engine/Source/Editor/UnrealEd/Private/Fbx/FbxMainImport.cpp
--- a/Engine/Source/Editor/UnrealEd/Private/Fbx/FbxMainImport.cpp
+++ b/Engine/Source/Editor/UnrealEd/Private/Fbx/FbxMainImport.cpp
@@ -553,6 +553,14 @@ FFbxImporter::FFbxImporter()
 	ImportOptions->MaterialBasePath = NAME_None;
 	
 	CurPhase = NOTSTARTED;
+
+	//The FFbxImporter is a singleton is constructor is protected
+	//We must release the resource in the pre-exit delegate because in some cases the
+	//Instance is not valid anymore when the destructor get called (i.e. when we build the editor in monolithic)
+	FCoreDelegates::OnPreExit.AddLambda([]()
+	{
+		FFbxImporter::GetInstance()->CleanUp();
+	});
 }
 	
 //-------------------------------------------------------------------------
@@ -560,7 +568,7 @@ FFbxImporter::FFbxImporter()
 //-------------------------------------------------------------------------
 FFbxImporter::~FFbxImporter()
 {
-	CleanUp();
+	//The clean up should have been done in the pre-exit core delegate implement in the FFbxImporter constructor
 }
 
 //-------------------------------------------------------------------------
