ZenCOTF: Fix for missing global shader files in zen file manifest (fallout from 20015198)

#rnx

diff --git a/Engine/Source/Editor/UnrealEd/Private/CookOnTheFlyServer.cpp b/Engine/Source/Editor/UnrealEd/Private/CookOnTheFlyServer.cpp
--- a/Engine/Source/Editor/UnrealEd/Private/CookOnTheFlyServer.cpp
+++ b/Engine/Source/Editor/UnrealEd/Private/CookOnTheFlyServer.cpp
@@ -636,7 +636,6 @@ void UCookOnTheFlyServer::AddCookOnTheFlyPlatformFromGameThread(ITargetPlatform*
 
 	// Initialize systems that nothing in AddCookOnTheFlyPlatformFromGameThread references
 	// Functions in this section are not dependent upon each other and can be ordered arbitrarily or for async performance
-	BeginCookPackageWriters(BeginContext);
 	BeginCookTargetPlatforms(BeginContext.TargetPlatforms);
 	InitializeShadersForCookOnTheFly(BeginContext.TargetPlatforms);
 
@@ -646,6 +645,10 @@ void UCookOnTheFlyServer::AddCookOnTheFlyPlatformFromGameThread(ITargetPlatform*
 	// This will miss settings that are accessed during the cook
 	// TODO: A better way of handling ini settings
 	SaveCurrentIniSettings(TargetPlatform);
+
+	// Initialize package writers last,
+	// they may depend on completed artifacts (like global shader files) from systems initialized earlier
+	BeginCookPackageWriters(BeginContext);
 }
 
 void UCookOnTheFlyServer::OnTargetPlatformsInvalidated()
