Add skeletal mesh UStreamableRenderAsset data to fbx reimport save and restore data. This will prevent loosing those settings when we re-import a skeletalmesh.

#rnx

diff --git a/Engine/Source/Editor/UnrealEd/Private/ImportUtils/SkeletalMeshImportUtils.cpp b/Engine/Source/Editor/UnrealEd/Private/ImportUtils/SkeletalMeshImportUtils.cpp
--- a/Engine/Source/Editor/UnrealEd/Private/ImportUtils/SkeletalMeshImportUtils.cpp
+++ b/Engine/Source/Editor/UnrealEd/Private/ImportUtils/SkeletalMeshImportUtils.cpp
@@ -315,6 +315,7 @@ TSharedPtr<FExistingSkelMeshData> SkeletalMeshImportUtils::SaveExistingSkelMeshD
 	{
 		ExistingMeshDataPtr->ExistingUMetaDataTagValues = *MetaDataTagValues;
 	}
+	
 	ExistingMeshDataPtr->UseMaterialNameSlotWorkflow = InternalImportUtils::IsUsingMaterialSlotNameWorkflow(SourceSkeletalMesh->GetAssetImportData());
 	ExistingMeshDataPtr->MinLOD = SourceSkeletalMesh->GetMinLod();
 	ExistingMeshDataPtr->QualityLevelMinLOD = SourceSkeletalMesh->GetQualityLevelMinLod();
@@ -324,6 +325,8 @@ TSharedPtr<FExistingSkelMeshData> SkeletalMeshImportUtils::SaveExistingSkelMeshD
 	ExistingMeshDataPtr->MaxNumStreamedLODs = SourceSkeletalMesh->GetMaxNumStreamedLODs();
 	ExistingMeshDataPtr->MaxNumOptionalLODs = SourceSkeletalMesh->GetMaxNumOptionalLODs();
 
+	ExistingMeshDataPtr->ExistingStreamableRenderAssetData.Save(SourceSkeletalMesh);
+
 	const TArray<FSkeletalMaterial>& SourceMaterials = SourceSkeletalMesh->GetMaterials();
 	//Add the existing Material slot name data
 	for (int32 MaterialIndex = 0; MaterialIndex < SourceMaterials.Num(); ++MaterialIndex)
@@ -624,6 +627,8 @@ void SkeletalMeshImportUtils::RestoreExistingSkelMeshData(const TSharedPtr<const
 	SkeletalMesh->SetMaxNumStreamedLODs(MeshData->MaxNumStreamedLODs);
 	SkeletalMesh->SetMaxNumOptionalLODs(MeshData->MaxNumOptionalLODs);
 
+	MeshData->ExistingStreamableRenderAssetData.Restore(SkeletalMesh);
+
 	FSkeletalMeshModel* SkeletalMeshImportedModel = SkeletalMesh->GetImportedModel();
 
 	//Create a remap material Index use to find the matching section later
diff --git a/Engine/Source/Editor/UnrealEd/Public/ImportUtils/SkelImport.h b/Engine/Source/Editor/UnrealEd/Public/ImportUtils/SkelImport.h
--- a/Engine/Source/Editor/UnrealEd/Public/ImportUtils/SkelImport.h
+++ b/Engine/Source/Editor/UnrealEd/Public/ImportUtils/SkelImport.h
@@ -47,6 +47,29 @@ struct FSkeletalMeshLodImportDataBackup
 	ESkeletalMeshSkinningImportVersions MeshSkinningImportVersion = ESkeletalMeshSkinningImportVersions::Before_Versionning;
 };
 
+struct FStreamableRenderAssetData
+{
+	int32 NumCinematicMipLevels;
+	FPerQualityLevelInt NoRefStreamingLODBias;
+	bool NeverStream;
+	bool bGlobalForceMipLevelsToBeResident;
+
+	void Save(const USkeletalMesh* SkeletalMesh)
+	{
+		NumCinematicMipLevels = SkeletalMesh->NumCinematicMipLevels;
+		NoRefStreamingLODBias = SkeletalMesh->GetNoRefStreamingLODBias();
+		NeverStream = SkeletalMesh->NeverStream;
+		bGlobalForceMipLevelsToBeResident = SkeletalMesh->bGlobalForceMipLevelsToBeResident;
+	}
+	void Restore(USkeletalMesh* SkeletalMesh) const
+	{
+		SkeletalMesh->NumCinematicMipLevels = NumCinematicMipLevels;
+		SkeletalMesh->SetNoRefStreamingLODBias(NoRefStreamingLODBias);
+		SkeletalMesh->NeverStream = NeverStream;
+		SkeletalMesh->bGlobalForceMipLevelsToBeResident = bGlobalForceMipLevelsToBeResident;
+	}
+};
+
 struct FExistingSkelMeshData
 {
 	TArray<USkeletalMeshSocket*>			ExistingSockets;
@@ -109,6 +132,9 @@ struct FExistingSkelMeshData
 	bool bExistingSupportRayTracing;
 	int32 ExistingRayTracingMinLOD;
 	EClothLODBiasMode ExistingClothLODBiasMode;
+
+	//Streamable member
+	FStreamableRenderAssetData ExistingStreamableRenderAssetData;
 };
 
 /** 
