Cooker: Move shader file creation in CookByTheBook and CookOnTheFly to occurr immediately after BeginCookSandbox, since later steps in each mode can depend on the files being present.
Merge duplicated code in SaveGlobalShaderMapFiles and InitializeShadersForCookOnTheFly.

diff --git a/Engine/Source/Editor/UnrealEd/Classes/CookOnTheSide/CookOnTheFlyServer.h b/Engine/Source/Editor/UnrealEd/Classes/CookOnTheSide/CookOnTheFlyServer.h
--- a/Engine/Source/Editor/UnrealEd/Classes/CookOnTheSide/CookOnTheFlyServer.h
+++ b/Engine/Source/Editor/UnrealEd/Classes/CookOnTheSide/CookOnTheFlyServer.h
@@ -31,6 +31,7 @@ class IAssetRegistry;
 class ICookedPackageWriter;
 class IPlugin;
 class ITargetPlatform;
+enum class ODSCRecompileCommand;
 struct FBeginCookContext;
 struct FPropertyChangedEvent;
 
@@ -1084,7 +1085,7 @@ private:
 	*  
 	*  @param	Platforms		List of platforms to make global shader maps for
 	*/
-	void SaveGlobalShaderMapFiles(const TArrayView<const ITargetPlatform* const>& Platforms);
+	void SaveGlobalShaderMapFiles(const TArrayView<const ITargetPlatform* const>& Platforms, ODSCRecompileCommand RecompileCommand);
 
 
 	/** Create sandbox file in directory using current settings supplied */
diff --git a/Engine/Source/Editor/UnrealEd/Private/CookOnTheFlyServer.cpp b/Engine/Source/Editor/UnrealEd/Private/CookOnTheFlyServer.cpp
--- a/Engine/Source/Editor/UnrealEd/Private/CookOnTheFlyServer.cpp
+++ b/Engine/Source/Editor/UnrealEd/Private/CookOnTheFlyServer.cpp
@@ -598,15 +598,8 @@ void UCookOnTheFlyServer::CookOnTheFlyDeferredInitialize()
 
 void UCookOnTheFlyServer::InitializeShadersForCookOnTheFly(const TArrayView<ITargetPlatform* const>& NewTargetPlatforms)
 {
-	for (const ITargetPlatform* TargetPlatform : NewTargetPlatforms)
-	{
-		const FString& PlatformName = TargetPlatform->PlatformName();
-		UE_LOG(LogCook, Display, TEXT("Initializing shaders for cook-on-the-fly on platform '%s'"), *PlatformName);
-
-		TArray<uint8> GlobalShaderMap;
-		FShaderRecompileData RecompileArgs(PlatformName, SP_NumPlatforms, ODSCRecompileCommand::Global, nullptr, nullptr, &GlobalShaderMap);
-		RecompileShadersForRemote(RecompileArgs, GetSandboxDirectory(PlatformName));
-	}
+	UE_LOG(LogCook, Display, TEXT("Initializing shaders for cook-on-the-fly"));
+	SaveGlobalShaderMapFiles(NewTargetPlatforms, ODSCRecompileCommand::Global);
 }
 
 void UCookOnTheFlyServer::AddCookOnTheFlyPlatformFromGameThread(ITargetPlatform* TargetPlatform)
@@ -634,10 +627,14 @@ void UCookOnTheFlyServer::AddCookOnTheFlyPlatformFromGameThread(ITargetPlatform*
 	// Clear in-memory CookedPackages, or preserve them and cook iteratively in-process
 	BeginCookSandbox(BeginContext);
 
+	// Initialize systems that need to write files to the sandbox directory, for consumption later in AddCookOnTheFlyPlatformFromGameThread
+	// Functions in this section are not dependent upon each other and can be ordered arbitrarily or for async performance
+	InitializeShadersForCookOnTheFly(BeginContext.TargetPlatforms);
+
 	// Initialize systems that nothing in AddCookOnTheFlyPlatformFromGameThread references
 	// Functions in this section are not dependent upon each other and can be ordered arbitrarily or for async performance
+	BeginCookPackageWriters(BeginContext);
 	BeginCookTargetPlatforms(BeginContext.TargetPlatforms);
-	InitializeShadersForCookOnTheFly(BeginContext.TargetPlatforms);
 
 	// Execute save operations that are done in CookByTheBookFinished for CBTB, but that CookOnTheFly does at the beginning since it does not have a definite end
 	check(PlatformData->bIsSandboxInitialized); // This should have been set by BeginCookSandbox, and it is what we use to determine whether a platform has been initialized
@@ -645,10 +642,6 @@ void UCookOnTheFlyServer::AddCookOnTheFlyPlatformFromGameThread(ITargetPlatform*
 	// This will miss settings that are accessed during the cook
 	// TODO: A better way of handling ini settings
 	SaveCurrentIniSettings(TargetPlatform);
-
-	// Initialize package writers last,
-	// they may depend on completed artifacts (like global shader files) from systems initialized earlier
-	BeginCookPackageWriters(BeginContext);
 }
 
 void UCookOnTheFlyServer::OnTargetPlatformsInvalidated()
@@ -6618,25 +6611,19 @@ bool UCookOnTheFlyServer::IsCookByTheBookRunning() const
 }
 
 
-void UCookOnTheFlyServer::SaveGlobalShaderMapFiles(const TArrayView<const ITargetPlatform* const>& Platforms)
+void UCookOnTheFlyServer::SaveGlobalShaderMapFiles(const TArrayView<const ITargetPlatform* const>& Platforms, ODSCRecompileCommand RecompileCommand)
 {
-	// we don't support this behavior
-	check( !IsCookingDLC() );
-	for (int32 Index = 0; Index < Platforms.Num(); Index++)
+	check(!IsCookingDLC()); // GlobalShaderMapFiles are not supported when cooking DLC
+	check(IsInGameThread());
+	for (const ITargetPlatform* TargetPlatform : Platforms)
 	{
-		TArray<FString> Files;
-		TArray<uint8> GlobalShaderMap;
-
-		// make sure global shaders are up to date!
-		FShaderRecompileData RecompileData(Platforms[Index]->PlatformName(), SP_NumPlatforms, ODSCRecompileCommand::Changed, &Files, nullptr, &GlobalShaderMap);
-
-		check( IsInGameThread() );
-
-		FString OutputDir = GetSandboxDirectory(RecompileData.PlatformName);
-
-		UE_LOG(LogCook, Display, TEXT("Checking global shaders for platform %s"), *RecompileData.PlatformName);
+		const FString& PlatformName = TargetPlatform->PlatformName();
+		UE_LOG(LogCook, Display, TEXT("Compiling global%s shaders for platform '%s'"),
+			RecompileCommand == ODSCRecompileCommand::Changed ? TEXT(" changed") : TEXT(""), *PlatformName);
 
-		RecompileShadersForRemote(RecompileData, OutputDir);
+		TArray<uint8> GlobalShaderMap;
+		FShaderRecompileData RecompileData(PlatformName, SP_NumPlatforms, RecompileCommand, nullptr, nullptr, &GlobalShaderMap);
+		RecompileShadersForRemote(RecompileData, GetSandboxDirectory(PlatformName));
 	}
 }
 
@@ -6776,7 +6763,8 @@ void UCookOnTheFlyServer::BeginCookFinishShaderCodeLibrary(FBeginCookContext& Be
 	{
 		OpenGlobalShaderLibrary();
 
-		SaveGlobalShaderMapFiles(BeginContext.TargetPlatforms);
+		// make sure global shaders are up to date!
+		SaveGlobalShaderMapFiles(BeginContext.TargetPlatforms, ODSCRecompileCommand::Changed);
 
 		SaveAndCloseGlobalShaderLibrary();
 	}
@@ -8005,6 +7993,10 @@ void UCookOnTheFlyServer::StartCookByTheBook( const FCookByTheBookStartupOptions
 	// Clear in-memory CookedPackages, or preserve them and cook iteratively in-process
 	BeginCookSandbox(BeginContext);
 
+	// Initialize systems that need to write files to the sandbox directory, for consumption later in StartCookByTheBook
+	// Functions in this section are not dependent upon each other and can be ordered arbitrarily or for async performance
+	BeginCookFinishShaderCodeLibrary(BeginContext);
+
 	// Initialize systems that nothing in StartCookByTheBook references
 	// Functions in this section are not dependent upon each other and can be ordered arbitrarily or for async performance
 	BeginCookEditorSystems();
@@ -8016,8 +8008,6 @@ void UCookOnTheFlyServer::StartCookByTheBook( const FCookByTheBookStartupOptions
 	InitializePollables();
 	RecordDLCPackagesFromBaseGame(BeginContext);
 	RegisterCookByTheBookDelegates();
-
-	BeginCookFinishShaderCodeLibrary(BeginContext);
 }
 
 FBeginCookContext UCookOnTheFlyServer::SetCookByTheBookOptions(const FCookByTheBookStartupOptions& StartupOptions)
