Added "IsFastCook" to cooking analytics.

diff --git a/Engine/Source/Editor/UnrealEd/Classes/Commandlets/CookCommandlet.h b/Engine/Source/Editor/UnrealEd/Classes/Commandlets/CookCommandlet.h
--- a/Engine/Source/Editor/UnrealEd/Classes/Commandlets/CookCommandlet.h
+++ b/Engine/Source/Editor/UnrealEd/Classes/Commandlets/CookCommandlet.h
@@ -33,6 +33,8 @@ class UCookCommandlet
 	bool bIterativeCooking;
 	/** Prototype cook-on-the-fly server */
 	bool bCookOnTheFly; 
+	/** Is using fast cook */
+	bool bFastCook;
 	/** Cook everything */
 	bool bCookAll;
 	/** Skip saving any packages in Engine/Content/Editor* UNLESS TARGET HAS EDITORONLY DATA (in which case it will save those anyway) */
diff --git a/Engine/Source/Editor/UnrealEd/Private/Commandlets/CookCommandlet.cpp b/Engine/Source/Editor/UnrealEd/Private/Commandlets/CookCommandlet.cpp
--- a/Engine/Source/Editor/UnrealEd/Private/Commandlets/CookCommandlet.cpp
+++ b/Engine/Source/Editor/UnrealEd/Private/Commandlets/CookCommandlet.cpp
@@ -82,9 +82,9 @@ namespace DetailedCookStats
 	bool IsCookAll = false;
 	bool IsCookOnTheFly = false;
 	bool IsIterativeCook = false;
+	bool IsFastCook = false;
 	bool IsUnversioned = false;
 
-
 	FCookStatsManager::FAutoRegisterCallback RegisterCookStats([](FCookStatsManager::AddStatFuncRef AddStat)
 	{
 		const FString StatName(TEXT("Cook.Profile"));
@@ -114,6 +114,7 @@ namespace DetailedCookStats
 		ADD_COOK_STAT_FLT(" 0. 1. 13", IsIterativeCook);
 		ADD_COOK_STAT_FLT(" 0. 1. 14", IsUnversioned);
 		ADD_COOK_STAT_FLT(" 0. 1. 15", CookLabel);
+		ADD_COOK_STAT_FLT(" 0. 1. 16", IsFastCook);
 		
 		#undef ADD_COOK_STAT_FLT
 	});
@@ -657,10 +658,12 @@ int32 UCookCommandlet::Main(const FString& CmdLineParams)
 	ShowProgress = !Switches.Contains(TEXT("DIFFONLY"));
 	bNoShaderCooking = bCookOnTheFly; // Do not cook any shaders into the shader maps. Always true if we are running w/ cook on the fly
 	bIgnoreIniSettingsOutOfDate = Switches.Contains(TEXT("IgnoreIniSettingsOutOfDate"));
+	bFastCook = Switches.Contains(TEXT("FastCook"));
 
 	COOK_STAT(DetailedCookStats::IsCookAll = bCookAll);
 	COOK_STAT(DetailedCookStats::IsCookOnTheFly = bCookOnTheFly);
 	COOK_STAT(DetailedCookStats::IsIterativeCook = bIterativeCooking);
+	COOK_STAT(DetailedCookStats::IsFastCook = bFastCook);
 	COOK_STAT(DetailedCookStats::IsUnversioned = bUnversioned);
 
 	COOK_STAT(DetailedCookStats::CookProject = FApp::GetProjectName());
