Optimize cook PopReadyRequest to avoid an unnecessary search

#rnx

diff --git a/Engine/Source/Editor/UnrealEd/Private/Cooker/CookPackageData.cpp b/Engine/Source/Editor/UnrealEd/Private/Cooker/CookPackageData.cpp
--- a/Engine/Source/Editor/UnrealEd/Private/Cooker/CookPackageData.cpp
+++ b/Engine/Source/Editor/UnrealEd/Private/Cooker/CookPackageData.cpp
@@ -2278,14 +2278,16 @@ uint32 FRequestQueue::ReadyRequestsNum() const
 
 FPackageData* FRequestQueue::PopReadyRequest()
 {
-	for (FPackageData* PackageData : UrgentRequests)
+	for (auto Iterator = UrgentRequests.CreateIterator(); Iterator; ++Iterator)
 	{
-		UrgentRequests.Remove(PackageData);
+		FPackageData* PackageData = *Iterator;
+		Iterator.RemoveCurrent();
 		return PackageData;
 	}
-	for (FPackageData* PackageData : NormalRequests)
+	for (auto Iterator = NormalRequests.CreateIterator(); Iterator; ++Iterator)
 	{
-		NormalRequests.Remove(PackageData);
+		FPackageData* PackageData = *Iterator;
+		Iterator.RemoveCurrent();
 		return PackageData;
 	}
 	return nullptr;
