CookOnTheFly: Move SaveAssetRegistry to before BeginCookPackageWriters.
#rnx

diff --git a/Engine/Source/Editor/UnrealEd/Private/CookOnTheFlyServer.cpp b/Engine/Source/Editor/UnrealEd/Private/CookOnTheFlyServer.cpp
--- a/Engine/Source/Editor/UnrealEd/Private/CookOnTheFlyServer.cpp
+++ b/Engine/Source/Editor/UnrealEd/Private/CookOnTheFlyServer.cpp
@@ -630,15 +630,15 @@ void UCookOnTheFlyServer::AddCookOnTheFlyPlatformFromGameThread(ITargetPlatform*
 	// Initialize systems that need to write files to the sandbox directory, for consumption later in AddCookOnTheFlyPlatformFromGameThread
 	// Functions in this section are not dependent upon each other and can be ordered arbitrarily or for async performance
 	InitializeShadersForCookOnTheFly(BeginContext.TargetPlatforms);
+	// SaveAssetRegistry is done in CookByTheBookFinished for CBTB, but we need at the start of CookOnTheFly to send as startup information to connecting clients
+	PlatformData->RegistryGenerator->SaveAssetRegistry(GetSandboxAssetRegistryFilename(), true);
 
 	// Initialize systems that nothing in AddCookOnTheFlyPlatformFromGameThread references
 	// Functions in this section are not dependent upon each other and can be ordered arbitrarily or for async performance
 	BeginCookPackageWriters(BeginContext);
 	BeginCookTargetPlatforms(BeginContext.TargetPlatforms);
 
-	// Execute save operations that are done in CookByTheBookFinished for CBTB, but that CookOnTheFly does at the beginning since it does not have a definite end
-	check(PlatformData->bIsSandboxInitialized); // This should have been set by BeginCookSandbox, and it is what we use to determine whether a platform has been initialized
-	PlatformData->RegistryGenerator->SaveAssetRegistry(GetSandboxAssetRegistryFilename(), true);
+	// SaveCurrentIniSettings is done in CookByTheBookFinished for CBTB, but we don't have a definite end point in CookOnTheFly so we write it at the start
 	// This will miss settings that are accessed during the cook
 	// TODO: A better way of handling ini settings
 	SaveCurrentIniSettings(TargetPlatform);
