Cooker: Remove branch for CookOnTheFly from DecideNextCookAction. Implement the reduction of latency it was providing by changing the value of CookRetryBusyTime when in CookOnTheFly.
Fix precision issues with some times used in polling that were stored as floats instead of doubles.
#rnx

diff --git a/Engine/Source/Editor/UnrealEd/Classes/CookOnTheSide/CookOnTheFlyServer.h b/Engine/Source/Editor/UnrealEd/Classes/CookOnTheSide/CookOnTheFlyServer.h
--- a/Engine/Source/Editor/UnrealEd/Classes/CookOnTheSide/CookOnTheFlyServer.h
+++ b/Engine/Source/Editor/UnrealEd/Classes/CookOnTheSide/CookOnTheFlyServer.h
@@ -209,6 +209,9 @@ private:
 	uint32 PackagesPerGC;
 	/** Amount of time that is allowed to be idle before forcing a garbage collect. Set to 0 to never force GC due to idle time */
 	double IdleTimeToGC;
+	/** Amount of time to wait when save and load are busy waiting on async operations before trying them again. */
+	float CookProgressRetryBusyPeriodSeconds = 0.f;
+
 	// Memory Limits for when to Collect Garbage
 	uint64 MemoryMaxUsedVirtual;
 	uint64 MemoryMaxUsedPhysical;
@@ -1178,12 +1181,12 @@ private:
 
 
 	/** Timers for tracking how long we have been busy, to manage retries and warnings of deadlock */
-	float SaveBusyTimeLastRetry = 0.f;
-	float SaveBusyTimeStarted = 0.f;
-	float LoadBusyTimeLastRetry = 0.f;
-	float LoadBusyTimeStarted = 0.f;
+	double SaveBusyRetryTimeSeconds = MAX_flt;
+	double SaveBusyWarnTimeSeconds = MAX_flt;
+	double LoadBusyRetryTimeSeconds = MAX_flt;
+	double LoadBusyWarnTimeSeconds = MAX_flt;
 	/** Tracking for the ticking of tickable cook objects */
-	float LastCookableObjectTickTime = 0.f;
+	double LastCookableObjectTickTime = 0.;
 
 	// These helper structs are all TUniquePtr rather than inline members so that we can keep their headers private.  See class header comments for their purpose.
 	TUniquePtr<UE::Cook::FPackageTracker> PackageTracker;
diff --git a/Engine/Source/Editor/UnrealEd/Private/CookOnTheFlyServer.cpp b/Engine/Source/Editor/UnrealEd/Private/CookOnTheFlyServer.cpp
--- a/Engine/Source/Editor/UnrealEd/Private/CookOnTheFlyServer.cpp
+++ b/Engine/Source/Editor/UnrealEd/Private/CookOnTheFlyServer.cpp
@@ -1330,19 +1330,21 @@ void UCookOnTheFlyServer::SetSaveBusy(bool bInBusy)
 		bSaveBusy = bInBusy;
 		if (bSaveBusy)
 		{
-			SaveBusyTimeStarted = FPlatformTime::Seconds();
-			SaveBusyTimeLastRetry = SaveBusyTimeStarted;
+			const double CurrentTime = FPlatformTime::Seconds();
+			SaveBusyRetryTimeSeconds = CurrentTime + CookProgressRetryBusyPeriodSeconds;
+			SaveBusyWarnTimeSeconds = CurrentTime + CookProgressRetryBusyPeriodSeconds;
 		}
 		else
 		{
-			SaveBusyTimeStarted = 0;
-			SaveBusyTimeLastRetry = 0;
+			SaveBusyRetryTimeSeconds = MAX_flt;
+			SaveBusyWarnTimeSeconds = MAX_flt;
 		}
 	}
 	else if (bSaveBusy)
 	{
-		const float CurrentTime = FPlatformTime::Seconds();
-		if (CurrentTime - SaveBusyTimeStarted > GCookProgressWarnBusyTime)
+		const double CurrentTime = FPlatformTime::Seconds();
+		SaveBusyRetryTimeSeconds = CurrentTime + CookProgressRetryBusyPeriodSeconds;
+		if (CurrentTime >= SaveBusyWarnTimeSeconds)
 		{
 			// Issue a status update. For each UObject we're still waiting on, check whether the long duration is expected using type-specific checks
 			// Make the status update a warning if the long duration is not reported as expected.
@@ -1445,7 +1447,7 @@ void UCookOnTheFlyServer::SetSaveBusy(bool bInBusy)
 				UE_LOG(LogCook, Display, TEXT("    <None>"));
 			}
 
-			SaveBusyTimeStarted = CurrentTime;
+			SaveBusyWarnTimeSeconds = CurrentTime + GCookProgressWarnBusyTime;
 		}
 	}
 }
@@ -1459,19 +1461,21 @@ void UCookOnTheFlyServer::SetLoadBusy(bool bInLoadBusy)
 		bLoadBusy = bInLoadBusy;
 		if (bLoadBusy)
 		{
-			LoadBusyTimeStarted = FPlatformTime::Seconds();
-			LoadBusyTimeLastRetry = LoadBusyTimeStarted;
+			const double CurrentTime = FPlatformTime::Seconds();
+			LoadBusyRetryTimeSeconds = CurrentTime + CookProgressRetryBusyPeriodSeconds;
+			LoadBusyWarnTimeSeconds = CurrentTime + CookProgressRetryBusyPeriodSeconds;
 		}
 		else
 		{
-			LoadBusyTimeStarted = 0;
-			LoadBusyTimeLastRetry = 0;
+			LoadBusyRetryTimeSeconds = MAX_flt;
+			LoadBusyWarnTimeSeconds = MAX_flt;
 		}
 	}
 	else if (bLoadBusy)
 	{
-		const float CurrentTime = FPlatformTime::Seconds();
-		if (CurrentTime - LoadBusyTimeStarted > GCookProgressWarnBusyTime)
+		const double CurrentTime = FPlatformTime::Seconds();
+		LoadBusyRetryTimeSeconds = CurrentTime + CookProgressRetryBusyPeriodSeconds;
+		if (CurrentTime >= LoadBusyWarnTimeSeconds)
 		{
 			int DisplayCount = 0;
 			const int DisplayMax = 10;
@@ -1501,7 +1505,7 @@ void UCookOnTheFlyServer::SetLoadBusy(bool bInLoadBusy)
 			{
 				UE_LOG(LogCook, Display, TEXT("    <None>"));
 			}
-			LoadBusyTimeStarted = CurrentTime;
+			LoadBusyWarnTimeSeconds = CurrentTime + GCookProgressWarnBusyTime;
 		}
 	}
 }
@@ -1510,7 +1514,7 @@ void UCookOnTheFlyServer::UpdateDisplay(ECookTickFlags TickFlags, bool bForceDis
 {
 	using namespace UE::Cook;
 
-	const float CurrentTime = FPlatformTime::Seconds();
+	const double CurrentTime = FPlatformTime::Seconds();
 	const float DeltaProgressDisplayTime = CurrentTime - LastProgressDisplayTime;
 	const int32 CookedPackagesCount = PackageDatas->GetNumCooked();
 	const int32 CookPendingCount = ExternalRequests->GetNumRequests() + PackageDatas->GetMonitor().GetNumInProgress();
@@ -1764,8 +1768,8 @@ void UCookOnTheFlyServer::WaitForAsync(UE::Cook::FTickStackData& StackData)
 	float SleepDuration = WaitForAsyncSleepSeconds;
 	SleepDuration = FMath::Min(SleepDuration, StackData.Timer.GetEndTimeSeconds() - CurrentTime);
 	SleepDuration = FMath::Min(SleepDuration, PollNextTimeIdleSeconds - CurrentTime);
-	SleepDuration = FMath::Min(SleepDuration, SaveBusyTimeLastRetry + GCookProgressRetryBusyTime - CurrentTime);
-	SleepDuration = FMath::Min(SleepDuration, LoadBusyTimeLastRetry + GCookProgressRetryBusyTime - CurrentTime);
+	SleepDuration = FMath::Min(SleepDuration, SaveBusyRetryTimeSeconds - CurrentTime);
+	SleepDuration = FMath::Min(SleepDuration, LoadBusyRetryTimeSeconds - CurrentTime);
 	SleepDuration = FMath::Max(SleepDuration, 0);
 	FPlatformProcess::Sleep(SleepDuration);
 }
@@ -1853,43 +1857,13 @@ UCookOnTheFlyServer::ECookAction UCookOnTheFlyServer::DecideNextCookAction(UE::C
 		return ECookAction::Load;
 	}
 
-	if (IsCookOnTheFlyMode() && !IsRealtimeMode())
+	if (NumSaves > 0 && CurrentTime >= SaveBusyRetryTimeSeconds)
 	{
-		if (NumSaves > 0 && bLoadBusy)
-		{
-			return ECookAction::Save;
-		}
-		if (NumLoads > 0 && bSaveBusy)
-		{
-			return ECookAction::Load;
-		}
-		if (NumSaves > 0)
-		{
-			return ECookAction::Save;
-		}
-		if (NumLoads > 0)
-		{
-			return ECookAction::Load;
-		}
+		return ECookAction::Save;
 	}
-	else
+	if (NumLoads > 0 && CurrentTime >= LoadBusyRetryTimeSeconds)
 	{
-		if (bSaveBusy & (NumSaves > 0))
-		{
-			if (CurrentTime - SaveBusyTimeLastRetry > GCookProgressRetryBusyTime)
-			{
-				SaveBusyTimeLastRetry = CurrentTime;
-				return ECookAction::Save;
-			}
-		}
-		if (bLoadBusy & (NumLoads > 0))
-		{
-			if (CurrentTime - LoadBusyTimeLastRetry > GCookProgressRetryBusyTime)
-			{
-				LoadBusyTimeLastRetry = CurrentTime;
-				return ECookAction::Load;
-			}
-		}
+		return ECookAction::Load;
 	}
 
 	if (PackageDatas->GetMonitor().GetNumInProgress() > 0)
@@ -4702,6 +4676,16 @@ void UCookOnTheFlyServer::Initialize( ECookMode::Type DesiredCookMode, ECookInit
 	IdleTimeToGC = 20.0;
 	GConfig->GetDouble( TEXT("CookSettings"), TEXT("IdleTimeToGC"), IdleTimeToGC, GEditorIni );
 
+	if (IsCookOnTheFlyMode() && !IsRealtimeMode())
+	{
+		// Remove sleeps when waiting on async operations and otherwise idle; busy wait instead to minimize latency
+		CookProgressRetryBusyPeriodSeconds = 0.0f;
+	}
+	else
+	{
+		CookProgressRetryBusyPeriodSeconds = GCookProgressRetryBusyTime;
+	}
+
 	auto ReadMemorySetting = [](const TCHAR* SettingName, uint64& TargetVariable)
 	{
 		int32 ValueInMB = 0;
@@ -7078,7 +7062,7 @@ void UCookOnTheFlyServer::CookByTheBookFinished()
 
 	{
 		UE_SCOPED_COOKTIMER(TickCookableObjects);
-		float CurrentTime = FPlatformTime::Seconds();
+		const double CurrentTime = FPlatformTime::Seconds();
 		FTickableCookObject::TickObjects(CurrentTime - LastCookableObjectTickTime, true /* bTickComplete */);
 		LastCookableObjectTickTime = CurrentTime;
 	}
