Sort save packages with maps first

diff --git a/Engine/Source/Editor/UnrealEd/Private/FileHelpers.cpp b/Engine/Source/Editor/UnrealEd/Private/FileHelpers.cpp
--- a/Engine/Source/Editor/UnrealEd/Private/FileHelpers.cpp
+++ b/Engine/Source/Editor/UnrealEd/Private/FileHelpers.cpp
@@ -4239,6 +4239,32 @@ FEditorFileUtils::EPromptReturnCode FEditorFileUtils::PromptForCheckoutAndSave(
 	// and then go ahead and save the specified packages
 	if ( PackagesToSave.Num() > 0 && ReturnResponse == PR_Success )
 	{
+		// Sort packages to save
+		// Dialog sorts the packagelist, user in dialog can sort it by columns
+		// This brings the package list closer to the original before the dialog opened when saving dirty packages (Maps first, non-maps after)
+		// Also there is a a few situations where a dirty package can reference an unsaved map, the user will be prompted to give map a name
+		// If the dirty non-map package references the unsaved map we must have the dirty non-map package save first
+		TArray<UPackage*> SortedPackagesToSave;
+		SortedPackagesToSave.Reserve(PackagesToSave.Num());
+
+		for (UPackage* Package : PackagesToSave)
+		{
+			if (Package && Package->ContainsMap())
+			{
+				SortedPackagesToSave.Add(Package);
+			}
+		}
+
+		for (UPackage* Package : PackagesToSave)
+		{
+			if (Package && !Package->ContainsMap())
+			{
+				SortedPackagesToSave.Add(Package);
+			}
+		}
+
+		PackagesToSave = MoveTemp(SortedPackagesToSave);
+
 		TArray<UPackage*> FailedPackages;
 
 		TArray<UPackage*> PackagesCheckedOutOrMadeWritable;
