Flag Material Graph and Nodes as Optional so that they would also be included in the optional sidecar data alongside material expressions

diff --git a/Engine/Source/Editor/UnrealEd/Classes/MaterialGraph/MaterialGraph.h b/Engine/Source/Editor/UnrealEd/Classes/MaterialGraph/MaterialGraph.h
--- a/Engine/Source/Editor/UnrealEd/Classes/MaterialGraph/MaterialGraph.h
+++ b/Engine/Source/Editor/UnrealEd/Classes/MaterialGraph/MaterialGraph.h
@@ -109,7 +109,7 @@ private:
 
 };
 
-UCLASS()
+UCLASS(Optional)
 class UNREALED_API UMaterialGraph : public UEdGraph
 {
 	GENERATED_UCLASS_BODY()
diff --git a/Engine/Source/Editor/UnrealEd/Classes/MaterialGraph/MaterialGraphNode_Base.h b/Engine/Source/Editor/UnrealEd/Classes/MaterialGraph/MaterialGraphNode_Base.h
--- a/Engine/Source/Editor/UnrealEd/Classes/MaterialGraph/MaterialGraphNode_Base.h
+++ b/Engine/Source/Editor/UnrealEd/Classes/MaterialGraph/MaterialGraphNode_Base.h
@@ -10,7 +10,7 @@
 class UEdGraphPin;
 class UEdGraphSchema;
 
-UCLASS(MinimalAPI)
+UCLASS(MinimalAPI, Optional)
 class UMaterialGraphNode_Base : public UEdGraphNode
 {
 	GENERATED_UCLASS_BODY()
diff --git a/Engine/Source/Editor/UnrealEd/Classes/MaterialGraph/MaterialGraphNode_Comment.h b/Engine/Source/Editor/UnrealEd/Classes/MaterialGraph/MaterialGraphNode_Comment.h
--- a/Engine/Source/Editor/UnrealEd/Classes/MaterialGraph/MaterialGraphNode_Comment.h
+++ b/Engine/Source/Editor/UnrealEd/Classes/MaterialGraph/MaterialGraphNode_Comment.h
@@ -9,7 +9,7 @@
 
 #include "MaterialGraphNode_Comment.generated.h"
 
-UCLASS(MinimalAPI)
+UCLASS(MinimalAPI, Optional)
 class UMaterialGraphNode_Comment : public UEdGraphNode_Comment
 {
 	GENERATED_UCLASS_BODY()
