Cooker: Fix incorrect too-short waiting period being used for the cooker-blocked warning messages.
#rnx

diff --git a/Engine/Source/Editor/UnrealEd/Private/CookOnTheFlyServer.cpp b/Engine/Source/Editor/UnrealEd/Private/CookOnTheFlyServer.cpp
--- a/Engine/Source/Editor/UnrealEd/Private/CookOnTheFlyServer.cpp
+++ b/Engine/Source/Editor/UnrealEd/Private/CookOnTheFlyServer.cpp
@@ -1332,7 +1332,7 @@ void UCookOnTheFlyServer::SetSaveBusy(bool bInBusy)
 		{
 			const double CurrentTime = FPlatformTime::Seconds();
 			SaveBusyRetryTimeSeconds = CurrentTime + CookProgressRetryBusyPeriodSeconds;
-			SaveBusyWarnTimeSeconds = CurrentTime + CookProgressRetryBusyPeriodSeconds;
+			SaveBusyWarnTimeSeconds = CurrentTime + GCookProgressWarnBusyTime;
 		}
 		else
 		{
@@ -1463,7 +1463,7 @@ void UCookOnTheFlyServer::SetLoadBusy(bool bInLoadBusy)
 		{
 			const double CurrentTime = FPlatformTime::Seconds();
 			LoadBusyRetryTimeSeconds = CurrentTime + CookProgressRetryBusyPeriodSeconds;
-			LoadBusyWarnTimeSeconds = CurrentTime + CookProgressRetryBusyPeriodSeconds;
+			LoadBusyWarnTimeSeconds = CurrentTime + GCookProgressWarnBusyTime;
 		}
 		else
 		{
