Cooker NeverCookDirectories: Extend the nevercook of ExternalActors folder to include that relative path in all plugins in addition to /Game.
#rnx

diff --git a/Engine/Source/Editor/UnrealEd/Private/CookOnTheFlyServer.cpp b/Engine/Source/Editor/UnrealEd/Private/CookOnTheFlyServer.cpp
--- a/Engine/Source/Editor/UnrealEd/Private/CookOnTheFlyServer.cpp
+++ b/Engine/Source/Editor/UnrealEd/Private/CookOnTheFlyServer.cpp
@@ -8449,7 +8449,14 @@ TArray<FName> UCookOnTheFlyServer::GetNeverCookPackageFileNames(TArrayView<const
 	}
 
 	// For all modes, never cook External Actors; they are handled by the parent map
-	NeverCookDirectories.Add(FString::Printf(TEXT("/Game/%s"), ULevel::GetExternalActorsFolderName()));
+	FString ExternalActorsFolderName = ULevel::GetExternalActorsFolderName();
+	FString FullExternalActorsPath = FPaths::Combine(TEXT("/Game/"), ExternalActorsFolderName);
+	NeverCookDirectories.Add(MoveTemp(FullExternalActorsPath));
+	for (TSharedRef<IPlugin>& Plugin : IPluginManager::Get().GetEnabledPluginsWithContent())
+	{
+		FullExternalActorsPath = FPaths::Combine(Plugin->GetMountedAssetPath(), ExternalActorsFolderName);
+		NeverCookDirectories.Add(MoveTemp(FullExternalActorsPath));
+	}
 
 	TArray<FString> NeverCookPackagesPaths;
 	FPackageName::FindPackagesInDirectories(NeverCookPackagesPaths, NeverCookDirectories);
