UEditorLevelUtils::CreateNewStreamingLevelForWorld: Activate default mode when level is created (to avoid leaving no modes active)

diff --git a/Engine/Source/Editor/UnrealEd/Private/EditorLevelUtils.cpp b/Engine/Source/Editor/UnrealEd/Private/EditorLevelUtils.cpp
--- a/Engine/Source/Editor/UnrealEd/Private/EditorLevelUtils.cpp
+++ b/Engine/Source/Editor/UnrealEd/Private/EditorLevelUtils.cpp
@@ -713,6 +713,15 @@ ULevelStreaming* UEditorLevelUtils::CreateNewStreamingLevelForWorld(UWorld& InWo
 	}
 	using namespace UE::EditorLevelUtils::Private;
 
+	// Make sure we reenable the default mode on exit
+	ON_SCOPE_EXIT
+	{
+		if (!IsRunningCommandlet())
+		{
+			GLevelEditorModeTools().ActivateDefaultMode();
+		}
+	};
+
 	// This is the world we are adding the new level to
 	UWorld* WorldToAddLevelTo = &InWorld;
 
