BulkDataRegistry: Rename BulkDataRegistryEditorDomain to BulkDataRegistry to prevent confusion with EditorDomain package saving.
#rnx

diff --git a/Engine/Source/Editor/UnrealEd/Private/EditorDomain/BulkDataRegistryEditorDomain.cpp b/Engine/Source/Editor/UnrealEd/Private/EditorDomain/BulkDataRegistryEditorDomain.cpp
deleted file mode 100644
diff --git a/Engine/Source/Editor/UnrealEd/Private/EditorDomain/BulkDataRegistryEditorDomain.h b/Engine/Source/Editor/UnrealEd/Private/EditorDomain/BulkDataRegistryEditorDomain.h
deleted file mode 100644
diff --git a/Engine/Source/Editor/UnrealEd/Private/EditorDomain/BulkDataRegistryImpl.cpp b/Engine/Source/Editor/UnrealEd/Private/EditorDomain/BulkDataRegistryImpl.cpp
new file mode 100644
--- /dev/null
+++ b/Engine/Source/Editor/UnrealEd/Private/EditorDomain/BulkDataRegistryImpl.cpp
@@ -0,0 +1,936 @@
+// Copyright Epic Games, Inc. All Rights Reserved.
+
+#include "BulkDataRegistryImpl.h"
+
+#include "Serialization/BulkDataRegistry.h"
+#include "Compression/CompressedBuffer.h"
+#include "DerivedDataCache.h"
+#include "DerivedDataCacheRecord.h"
+#include "EditorBuildInputResolver.h"
+#include "EditorDomain/EditorDomainUtils.h"
+#include "HAL/PlatformTime.h"
+#include "Interfaces/IPluginManager.h"
+#include "IO/IoHash.h"
+#include "Memory/SharedBuffer.h"
+#include "Misc/CoreDelegates.h"
+#include "Misc/ScopeLock.h"
+#include "Misc/ScopeRWLock.h"
+#include "Serialization/MemoryReader.h"
+#include "Templates/RefCounting.h"
+#include "TickableEditorObject.h"
+#include "UObject/Package.h"
+#include "UObject/UObjectGlobals.h"
+#include "UObject/UObjectIterator.h"
+
+namespace UE::BulkDataRegistry::Private
+{
+
+namespace Constants
+{
+	constexpr uint64 TempLoadedPayloadsSizeBudget = 1024 * 1024 * 100;
+	constexpr double TempLoadedPayloadsDuration = 60.;
+}
+
+TConstArrayView<uint8> MakeArrayView(FSharedBuffer Buffer)
+{
+	return TConstArrayView<uint8>(reinterpret_cast<const uint8*>(Buffer.GetData()), Buffer.GetSize());
+}
+
+/** Add a hook to the BulkDataRegistry's startup delegate to use the EditorDomain as the BulkDataRegistry */
+class FEditorDomainRegisterAsBulkDataRegistry
+{
+public:
+	FEditorDomainRegisterAsBulkDataRegistry()
+	{
+		IBulkDataRegistry::GetSetBulkDataRegistryDelegate().BindStatic(SetBulkDataRegistry);
+	}
+
+	static IBulkDataRegistry* SetBulkDataRegistry()
+	{
+		return new FBulkDataRegistryImpl();
+	}
+} GRegisterAsBulkDataRegistry;
+
+
+FBulkDataRegistryImpl::FBulkDataRegistryImpl()
+{
+	SharedDataLock = new FTaskSharedDataLock();
+	// We piggyback on the BulkDataRegistry hook to set the pointer to tunnel in the pointer to the EditorBuildInputResolver as well
+	SetGlobalBuildInputResolver(&UE::DerivedData::FEditorBuildInputResolver::Get());
+	FCoreUObjectDelegates::OnEndLoadPackage.AddRaw(this, &FBulkDataRegistryImpl::OnEndLoadPackage);
+}
+
+FBulkDataRegistryImpl::~FBulkDataRegistryImpl()
+{
+	FCoreUObjectDelegates::OnEndLoadPackage.RemoveAll(this);
+	SetGlobalBuildInputResolver(nullptr);
+
+	TMap<FGuid, FUpdatingPayload> LocalUpdatingPayloads;
+	{
+		FWriteScopeLock SharedDataScopeLock(SharedDataLock->ActiveLock);
+		FWriteScopeLock RegistryScopeLock(RegistryLock);
+		FScopeLock PendingPackageScopeLock(&PendingPackageLock);
+
+		// Disable all activity that might come in from other threads
+		bActive = false;
+		SharedDataLock->bActive = false;
+
+		// Take custody of UpdatingPayloads
+		Swap(LocalUpdatingPayloads, UpdatingPayloads);
+	}
+
+	// Since the UpdatingPayloads AsyncTasks can no longer access their Requesters, we have to call those callbacks
+	for (TPair<FGuid, FUpdatingPayload>& Pair : LocalUpdatingPayloads)
+	{
+		for (TUniqueFunction<void(bool, const FCompressedBuffer&)>& Requester : Pair.Value.Requesters)
+		{
+			Requester(false, FCompressedBuffer());
+		}
+	}
+	LocalUpdatingPayloads.Empty();
+
+	// Clear PendingPackages
+	TMap<FName, TUniquePtr<FPendingPackage>> LocalPendingPackages;
+	{
+		FScopeLock PendingPackageScopeLock(&PendingPackageLock);
+		// Take custody of PendingPackages
+		Swap(LocalPendingPackages, PendingPackages);
+	}
+	for (TPair<FName, TUniquePtr<FPendingPackage>>& PendingPackagePair : LocalPendingPackages)
+	{
+		PendingPackagePair.Value->Cancel();
+	}
+	LocalPendingPackages.Empty();
+
+	// Clear PendingPackagePayloadIds. We have to take custody of PendingPayloadIds after calling
+	// Cancel from all FPendingPackage, as the FPendingPackages have callbacks that may write to PendingPayloadIds
+	TMap<FGuid, TRefCountPtr<FPendingPayloadId>> LocalPendingPayloadIds;
+	{
+		FWriteScopeLock RegistryScopeLock(RegistryLock);
+		// Take custody of PendingPayloadIds
+		Swap(LocalPendingPayloadIds, PendingPayloadIds);
+	}
+	for (TPair<FGuid, TRefCountPtr<FPendingPayloadId>>& PendingPayloadIdPair : LocalPendingPayloadIds)
+	{
+		PendingPayloadIdPair.Value->Cancel();
+	}
+}
+
+void FBulkDataRegistryImpl::OnEndLoadPackage(const FEndLoadPackageContext& Context)
+{
+	TArray<TUniquePtr<FPendingPackage>> PackagesToWrite;
+	{
+		FScopeLock PendingPackageScopeLock(&PendingPackageLock);
+
+		for (UPackage* LoadedPackage : Context.LoadedPackages)
+		{
+			FName PackageName = LoadedPackage->GetFName();
+			uint32 KeyHash = GetTypeHash(PackageName);
+			TUniquePtr<FPendingPackage>* PendingPackage = PendingPackages.FindByHash(KeyHash, PackageName);
+			if (PendingPackage)
+			{
+				bool bShouldRemove;
+				bool bShouldWriteCache;
+				check(PendingPackage->IsValid());
+				(*PendingPackage)->OnEndLoad(bShouldRemove, bShouldWriteCache);
+				// We do not hold a lock when calling WriteCache, so we require that the package no longer
+				// be accessible to other threads through PendingPackages if we need to write the cache,
+				// so bShouldRemove must be true if bShouldWriteCache is
+				check(!bShouldWriteCache || bShouldRemove);
+				if (bShouldRemove)
+				{
+					if (bShouldWriteCache)
+					{
+						PackagesToWrite.Add(MoveTemp(*PendingPackage));
+					}
+					PendingPackages.RemoveByHash(KeyHash, PackageName);
+				}
+			}
+		}
+	}
+
+	for (TUniquePtr<FPendingPackage>& Package : PackagesToWrite)
+	{
+		Package->WriteCache();
+	}
+}
+
+UE::BulkDataRegistry::ERegisterResult
+FBulkDataRegistryImpl::TryRegister(UPackage* Owner, const UE::Serialization::FEditorBulkData& BulkData)
+{
+	if (!BulkData.GetIdentifier().IsValid())
+	{
+		return UE::BulkDataRegistry::ERegisterResult::Success;
+	}
+
+	bool bAllowedToReadWritePayloadIdFromCache = false;
+	FName PackageName;
+	UE::Serialization::FEditorBulkData CopyBulk(BulkData.CopyTornOff());
+	UE::Serialization::FEditorBulkData PendingPackageBulk;
+	if (Owner)
+	{
+		PackageName = Owner->GetFName();
+		if (Owner->GetFileSize() // We only cache the BulkDataList for disk packages
+			&& !Owner->GetHasBeenEndLoaded() // We only cache BulkDatas that are loaded before the package finishes loading
+			&& CopyBulk.CanSaveForRegistry()
+			)
+		{
+			bAllowedToReadWritePayloadIdFromCache = true;
+			PendingPackageBulk = CopyBulk;
+		}
+	}
+
+	{
+		FWriteScopeLock RegistryScopeLock(RegistryLock);
+		check(bActive); // Registrations should not come in after we destruct
+		UE::BulkDataRegistry::Private::FRegisteredBulk& RegisteredBulk = Registry.FindOrAdd(BulkData.GetIdentifier());
+		if (RegisteredBulk.bRegistered)
+		{
+			bool bAllowSharing = false;
+			// If the original BulkData has left memory and a new bulkdata registers with the same package owner,
+			// assume that it is the same BulkData being reloaded and allow the sharing
+			bAllowSharing |= (!RegisteredBulk.bInMemory && PackageName == RegisteredBulk.PackageName);
+			if (!bAllowSharing)
+			{
+				return UE::BulkDataRegistry::ERegisterResult::AlreadyExists;
+			}
+
+			if (RegisteredBulk.PackageName.IsNone())
+			{
+				RegisteredBulk.PackageName = PackageName;
+			}
+			RegisteredBulk.bRegistered = true;
+			RegisteredBulk.bInMemory = true;
+			RegisteredBulk.bAllowedToWritePayloadIdToCache = false;
+			RegisteredBulk.bPayloadAvailable = true;
+
+			// For updated registrations, skip reading payloadid from cache, and call ResaveSizeTracker.Update rather than Register.
+			ResaveSizeTracker.UpdateRegistrationData(Owner, BulkData);
+			return UE::BulkDataRegistry::ERegisterResult::Success;
+		}
+		else
+		{
+			bool bCachedLocationMatches = false;
+			UE::Serialization::FEditorBulkData& TargetBulkData = RegisteredBulk.BulkData;
+
+			if (TargetBulkData.GetIdentifier().IsValid())
+			{
+				// This BulkData was added when loading the cached list of BulkData for the package, but has not been registered by an in-memory bulkdata
+				check(TargetBulkData.GetIdentifier() == BulkData.GetIdentifier());
+				bCachedLocationMatches = BulkData.LocationMatches(TargetBulkData);
+			}
+
+			if (!bCachedLocationMatches || !BulkData.HasPlaceholderPayloadId() || TargetBulkData.HasPlaceholderPayloadId())
+			{
+				// Copy the new BulkData over the value we got from the cache; the new BulkData is more authoritative
+				TargetBulkData = MoveTemp(CopyBulk);
+				RegisteredBulk.bAllowedToWritePayloadIdToCache = bAllowedToReadWritePayloadIdFromCache;
+			}
+			else
+			{
+				// Otherwise Keep the TargetBulkData, since it matches the location and has already calculated the PayloadId
+				RegisteredBulk.bAllowedToWritePayloadIdToCache = false;
+			}
+
+			RegisteredBulk.PackageName = PackageName;
+			RegisteredBulk.bRegistered = true;
+			RegisteredBulk.bInMemory = true;
+			RegisteredBulk.bPayloadAvailable = true;
+		}
+	}
+	if (bAllowedToReadWritePayloadIdFromCache)
+	{
+		AddPendingPackageBulkData(PackageName, MoveTemp(PendingPackageBulk));
+	}
+	ResaveSizeTracker.Register(Owner, BulkData);
+	return UE::BulkDataRegistry::ERegisterResult::Success;
+}
+
+void FBulkDataRegistryImpl::UpdateRegistrationData(UPackage* Owner, const UE::Serialization::FEditorBulkData& BulkData)
+{
+	if (!BulkData.GetIdentifier().IsValid())
+	{
+		UE_LOG(LogBulkDataRegistry, Warning, TEXT("UpdateRegistrationData called with invalid BulkData for Owner %s."),
+			Owner ? *Owner->GetName() : TEXT("<unknown>"));
+		return;
+	}
+
+	FName PackageName = Owner ? Owner->GetFName() : NAME_None;
+	{
+		FWriteScopeLock RegistryScopeLock(RegistryLock);
+		check(bActive); // Registrations should not come in after we destruct
+		UE::BulkDataRegistry::Private::FRegisteredBulk& RegisteredBulk = Registry.FindOrAdd(BulkData.GetIdentifier());
+		// Add the owner if we previously did not have an owner, otherwise keep the first owner
+		if (RegisteredBulk.PackageName.IsNone())
+		{
+			RegisteredBulk.PackageName = PackageName;
+		}
+		RegisteredBulk.BulkData = BulkData.CopyTornOff();
+		RegisteredBulk.bAllowedToWritePayloadIdToCache = false;
+		RegisteredBulk.bRegistered = true;
+		RegisteredBulk.bInMemory = true;
+		RegisteredBulk.bPayloadAvailable = true;
+	}
+	ResaveSizeTracker.UpdateRegistrationData(Owner, BulkData);
+}
+
+void FBulkDataRegistryImpl::Unregister(const UE::Serialization::FEditorBulkData& BulkData)
+{
+	const FGuid& Key = BulkData.GetIdentifier();
+	FWriteScopeLock RegistryScopeLock(RegistryLock);
+	check(bActive); // Deregistrations should not come in after we destruct
+	Registry.Remove(Key);
+}
+
+void FBulkDataRegistryImpl::OnExitMemory(const UE::Serialization::FEditorBulkData& BulkData)
+{
+	const FGuid& Key = BulkData.GetIdentifier();
+	FWriteScopeLock RegistryScopeLock(RegistryLock);
+	check(bActive); // Deregistrations should not come in after we destruct
+	FRegisteredBulk* Existing = Registry.Find(Key);
+	if (Existing)
+	{
+		if (Existing->BulkData.IsMemoryOnlyPayload())
+		{
+			Existing->BulkData.Reset();
+			Existing->BulkData.TearOff(); // Keep the TearOff flag after resetting
+			Existing->bPayloadAvailable = false;
+		}
+		Existing->bInMemory = false;
+	}
+}
+
+TFuture<UE::BulkDataRegistry::FMetaData> FBulkDataRegistryImpl::GetMeta(const FGuid& BulkDataId)
+{
+	bool bIsWriteLock = false;
+	FRWScopeLock RegistryScopeLock(RegistryLock, SLT_ReadOnly);
+	for (;;)
+	{
+		FRegisteredBulk* Existing = nullptr;
+		if (bActive)
+		{
+			Existing = Registry.Find(BulkDataId);
+		}
+		if (!Existing || !Existing->bPayloadAvailable)
+		{
+			TPromise<UE::BulkDataRegistry::FMetaData> Promise;
+			Promise.SetValue(UE::BulkDataRegistry::FMetaData{ FIoHash(), 0 });
+			return Promise.GetFuture();
+		}
+
+		UE::Serialization::FEditorBulkData& BulkData = Existing->BulkData;
+		if (!BulkData.HasPlaceholderPayloadId())
+		{
+			TPromise<UE::BulkDataRegistry::FMetaData> Promise;
+			Promise.SetValue(UE::BulkDataRegistry::FMetaData{ BulkData.GetPayloadId(), static_cast<uint64>(BulkData.GetPayloadSize()) });
+			return Promise.GetFuture();
+		}
+
+		if (!bIsWriteLock)
+		{
+			bIsWriteLock = true;
+			RegistryScopeLock.ReleaseReadOnlyLockAndAcquireWriteLock_USE_WITH_CAUTION();
+			continue;
+		}
+
+		// The payload in the registry is missing its RawHash; start a thread to calculate it and subscribe our caller to the results
+		FUpdatingPayload& UpdatingPayload = UpdatingPayloads.FindOrAdd(BulkDataId);
+		if (!UpdatingPayload.AsyncTask)
+		{
+			UpdatingPayload.AsyncTask = new FAutoDeleteAsyncTask<FUpdatePayloadWorker>(this, BulkData);
+			UpdatingPayload.AsyncTask->StartBackgroundTask();
+		}
+		TPromise<UE::BulkDataRegistry::FMetaData> Promise;
+		TFuture<UE::BulkDataRegistry::FMetaData> Future = Promise.GetFuture();
+		UpdatingPayload.Requesters.Add([Promise = MoveTemp(Promise)](bool bValid, const FCompressedBuffer& Buffer) mutable
+			{
+				Promise.SetValue(UE::BulkDataRegistry::FMetaData{ Buffer.GetRawHash(), Buffer.GetRawSize() });
+			});
+		return Future;
+	}
+}
+
+TFuture<UE::BulkDataRegistry::FData> FBulkDataRegistryImpl::GetData(const FGuid& BulkDataId)
+{
+	TOptional<UE::Serialization::FEditorBulkData> CopyBulk;
+	{
+		bool bIsWriteLock = false;
+		FRWScopeLock RegistryScopeLock(RegistryLock, SLT_ReadOnly);
+		for (;;)
+		{
+			FRegisteredBulk* Existing = nullptr;
+			if (bActive)
+			{
+				Existing = Registry.Find(BulkDataId);
+			}
+			if (!Existing || !Existing->bPayloadAvailable)
+			{
+				TPromise<UE::BulkDataRegistry::FData> Result;
+				Result.SetValue(UE::BulkDataRegistry::FData{ FCompressedBuffer() });
+				return Result.GetFuture();
+			}
+
+			if (!Existing->BulkData.HasPlaceholderPayloadId() && !Existing->bHasTempPayload)
+			{
+				// The contract of FEditorBulkData does not guarantee that GetCompressedPayload() is a quick operation (it may load the data
+				// synchronously), so copy the BulkData into a temporary and call it outside the lock
+				CopyBulk.Emplace(Existing->BulkData);
+				break;
+			}
+
+			if (!bIsWriteLock)
+			{
+				bIsWriteLock = true;
+				RegistryScopeLock.ReleaseReadOnlyLockAndAcquireWriteLock_USE_WITH_CAUTION();
+				continue;
+			}
+
+			if (!Existing->BulkData.HasPlaceholderPayloadId())
+			{
+				check(Existing->bHasTempPayload);
+				// We are the first GetData call after the BulkData previously loaded its Payload to calculate the RawHash.
+				// Sidenote, this means GetCompressedPayload will be fast.
+				// But we also have the responsibility to dump the data from memory since we have now consumed it.
+				// Make sure we copy the data pointer before dumping it from the Registry version!
+				CopyBulk.Emplace(Existing->BulkData);
+				Existing->BulkData.UnloadData();
+				Existing->bHasTempPayload = false;
+				break;
+			}
+
+			// The payload in the registry is missing its RawHash, and we calculate that on demand whenever the data is requested,
+			// which is now. Instead of only returning the data to our caller, we load the data and use it to update the RawHash
+			// in the registry and then return the data to our caller.
+			FUpdatingPayload& UpdatingPayload = UpdatingPayloads.FindOrAdd(BulkDataId);
+			if (!UpdatingPayload.AsyncTask)
+			{
+				UpdatingPayload.AsyncTask = new FAutoDeleteAsyncTask<FUpdatePayloadWorker>(this, Existing->BulkData);
+			}
+			TPromise<UE::BulkDataRegistry::FData> Promise;
+			TFuture<UE::BulkDataRegistry::FData> Future = Promise.GetFuture();
+			UpdatingPayload.Requesters.Add([Promise=MoveTemp(Promise)](bool bValid, const FCompressedBuffer& Buffer) mutable
+				{
+					Promise.SetValue(UE::BulkDataRegistry::FData{ Buffer });
+				});
+			return Future;
+		}
+	}
+
+	// We are calling a function that returns a TFuture on the stack-local CopyBulk, which would cause a read-after-free if the asynchronous TFuture could
+	// read from the BulkData. However, the contract of FEditorBulkData guarantees that the TFuture gets a copy of all data it needs and
+	// does not read from the BulkData after returning from GetCompressedPayload, so a read-after-free is not possible.
+	return CopyBulk->GetCompressedPayload().Next([](FCompressedBuffer Payload)
+		{
+			return UE::BulkDataRegistry::FData{ Payload };
+		});
+}
+
+bool FBulkDataRegistryImpl::TryGetBulkData(const FGuid& BulkDataId, UE::Serialization::FEditorBulkData* OutBulk,
+	FName* OutOwner)
+{
+	FReadScopeLock RegistryScopeLock(RegistryLock);
+	FRegisteredBulk* Existing = nullptr;
+	if (bActive)
+	{
+		Existing = Registry.Find(BulkDataId);
+	}
+	if (Existing)
+	{
+		if (OutBulk)
+		{
+			*OutBulk = Existing->BulkData;
+		}
+		if (OutOwner)
+		{
+			*OutOwner = Existing->PackageName;
+		}
+		return true;
+	}
+	return false;
+}
+
+uint64 FBulkDataRegistryImpl::GetBulkDataResaveSize(FName PackageName)
+{
+	return ResaveSizeTracker.GetBulkDataResaveSize(PackageName);
+}
+
+void FBulkDataRegistryImpl::TickCook(float DeltaTime, bool bTickComplete)
+{
+	bool bWaitForCooldown = !bTickComplete;
+	check(bActive); // Ticks should not come in after we destruct
+
+	{
+		FWriteScopeLock RegistryScopeLock(RegistryLock);
+		PruneTempLoadedPayloads();
+	}
+}
+
+void FBulkDataRegistryImpl::Tick(float DeltaTime)
+{
+	TickCook(DeltaTime, false /* bTickComplete */);
+}
+
+void FBulkDataRegistryImpl::AddPendingPackageBulkData(FName PackageName, UE::Serialization::FEditorBulkData&& BulkData)
+{
+	FScopeLock PendingPackageScopeLock(&PendingPackageLock);
+	check(bActive); // Registrations should not come in after we destruct, and AsyncTasks should check bActive before calling
+	TUniquePtr<FPendingPackage>& PendingPackage = PendingPackages.FindOrAdd(PackageName);
+	if (!PendingPackage.IsValid())
+	{
+		PendingPackage = MakeUnique<FPendingPackage>(PackageName, this);
+	}
+	if (!PendingPackage->IsLoadInProgress())
+	{
+		return;
+	}
+	PendingPackage->AddBulkData(MoveTemp(BulkData));
+}
+
+void FBulkDataRegistryImpl::AddTempLoadedPayload(const FGuid& RegistryKey, uint64 PayloadSize)
+{
+	// Called within RegistryLock WriteLock
+	TempLoadedPayloads.Add(FTempLoadedPayload{ RegistryKey, PayloadSize, FPlatformTime::Seconds() + Constants::TempLoadedPayloadsDuration });
+	TempLoadedPayloadsSize += PayloadSize;
+}
+
+void FBulkDataRegistryImpl::PruneTempLoadedPayloads()
+{
+	// Called within RegistryLock WriteLock
+	if (!TempLoadedPayloads.IsEmpty())
+	{
+		double CurrentTime = FPlatformTime::Seconds();
+		while (!TempLoadedPayloads.IsEmpty() &&
+			(TempLoadedPayloadsSize > Constants::TempLoadedPayloadsSizeBudget
+				|| TempLoadedPayloads[0].EndTime <= CurrentTime))
+		{
+			FTempLoadedPayload Payload = TempLoadedPayloads.PopFrontValue();
+			FRegisteredBulk* Existing = Registry.Find(Payload.Guid);
+			if (Existing)
+			{
+				// UnloadData only unloads the in-memory data, and only if the BulkData can be reloaded from disk
+				Existing->BulkData.UnloadData();
+				Existing->bHasTempPayload = false;
+			}
+			TempLoadedPayloadsSize -= Payload.PayloadSize;
+		}
+	}
+}
+
+void FBulkDataRegistryImpl::WritePayloadIdToCache(FName PackageName, const UE::Serialization::FEditorBulkData& BulkData) const
+{
+	check(!PackageName.IsNone());
+	TArray<uint8> Bytes;
+	FMemoryWriter Writer(Bytes);
+	const_cast<UE::Serialization::FEditorBulkData&>(BulkData).SerializeForRegistry(Writer);
+	UE::EditorDomain::PutBulkDataPayloadId(PackageName, BulkData.GetIdentifier(), MakeSharedBufferFromArray(MoveTemp(Bytes)));
+}
+
+void FBulkDataRegistryImpl::ReadPayloadIdsFromCache(FName PackageName, TArray<TRefCountPtr<FPendingPayloadId>>&& OldPendings,
+	TArray<TRefCountPtr<FPendingPayloadId>>&& NewPendings)
+{
+	// Cancel any old requests for the Guids in NewPendings; we are about to overwrite them
+	// This cancellation has to occur outside of any lock, since the task may be in progress and
+	// and to enter the lock and Cancel will wait on it
+	for (TRefCountPtr<FPendingPayloadId>& OldPending : OldPendings)
+	{
+		OldPending->Cancel();
+	}
+	OldPendings.Empty();
+	for (TRefCountPtr<FPendingPayloadId>& NewPending : NewPendings)
+	{
+		// Creation of the Request has to occur outside of any lock, because the request
+		// may execute immediately on this thread and need to enter the lock; our locks are non-reentrant
+		UE::DerivedData::FRequestBarrier Barrier(NewPending->GetRequestOwner());
+		UE::EditorDomain::GetBulkDataPayloadId(PackageName, NewPending->GetBulkDataId(), NewPending->GetRequestOwner(),
+			[this, PackageName, NewPending](FSharedBuffer Buffer)
+		{
+			if (Buffer.IsNull())
+			{
+				return;
+			}
+			FMemoryReaderView Reader(MakeArrayView(Buffer));
+			UE::Serialization::FEditorBulkData CachedBulkData;
+			CachedBulkData.SerializeForRegistry(Reader);
+			const FGuid& BulkDataId = NewPending->GetBulkDataId();
+			if (Reader.IsError() || CachedBulkData.GetIdentifier() != BulkDataId)
+			{
+				UE_LOG(LogBulkDataRegistry, Warning, TEXT("Corrupt cache data for BulkDataPayloadId %s."), WriteToString<192>(PackageName, TEXT("/"), BulkDataId).ToString());
+				return;
+			}
+
+			FWriteScopeLock RegistryScopeLock(RegistryLock);
+			if (!bActive)
+			{
+				return;
+			}
+			TRefCountPtr<FPendingPayloadId> ExistingPending;
+			if (!PendingPayloadIds.RemoveAndCopyValue(BulkDataId, ExistingPending))
+			{
+				return;
+			}
+			check(ExistingPending->GetBulkDataId() == BulkDataId);
+			if (ExistingPending != NewPending)
+			{
+				// We removed ExistingPending because we thought it was equal to NewPending, but it's not, so put it back
+				PendingPayloadIds.Add(BulkDataId, MoveTemp(ExistingPending));
+				return;
+			}
+
+			FRegisteredBulk* ExistingRegisteredBulk = Registry.Find(BulkDataId);
+			if (!ExistingRegisteredBulk)
+			{
+				return;
+			}
+
+			UE::Serialization::FEditorBulkData& ExistingBulkData = ExistingRegisteredBulk->BulkData;
+			check(ExistingBulkData.GetIdentifier() == BulkDataId);
+			if (ExistingBulkData.HasPlaceholderPayloadId() && CachedBulkData.LocationMatches(ExistingBulkData))
+			{
+				ExistingBulkData = CachedBulkData;
+			}
+		});
+	}
+
+	// Assign the Requests we just created to the data for each guid in the map,
+	// which has to be edited only within the lock
+	// If for any reason (race condition, shutting down) we can not assign a Request,
+	// we have to cancel the request before returning, to make sure the callback does not hold a pointer
+	// to this that could become dangling.
+	{
+		FWriteScopeLock RegistryScopeLock(RegistryLock);
+		if (!bActive)
+		{
+			for (TRefCountPtr<FPendingPayloadId>& NewPending : NewPendings)
+			{
+				OldPendings.Add(MoveTemp(NewPending));
+			}
+		}
+		else
+		{
+			for (TRefCountPtr<FPendingPayloadId>& NewPending : NewPendings)
+			{
+				TRefCountPtr<FPendingPayloadId>* ExistingPending = PendingPayloadIds.Find(NewPending->GetBulkDataId());
+				if (!ExistingPending || *ExistingPending != NewPending)
+				{
+					OldPendings.Add(MoveTemp(NewPending));
+				}
+			}
+		}
+	}
+	NewPendings.Empty();
+	for (TRefCountPtr<FPendingPayloadId>& OldPending : OldPendings)
+	{
+		OldPending->Cancel();
+	}
+	OldPendings.Empty();
+}
+
+FPendingPackage::FPendingPackage(FName InPackageName, FBulkDataRegistryImpl* InOwner)
+	: PackageName(InPackageName)
+	, BulkDataListCacheRequest(UE::DerivedData::EPriority::Low)
+	, Owner(InOwner)
+{
+	PendingOperations = Flag_EndLoad | Flag_BulkDataListResults;
+
+	UE::EditorDomain::GetBulkDataList(PackageName, BulkDataListCacheRequest,
+		[this](FSharedBuffer Buffer) { OnBulkDataListResults(Buffer); });
+}
+
+void FPendingPackage::Cancel()
+{
+	// Called from outside Owner->PendingPackagesLock, so OnBulkDataList can complete on other thread while we wait
+	// Called after removing this from Owner->PendingPackages under a previous cover of the lock
+	// If OnBulkDataList is running on other thread its attempt to remove from PendingPackages will be a noop
+	if (!BulkDataListCacheRequest.Poll())
+	{
+		// Optimization: prevent WriteCache from running at all if we reach here first
+		PendingOperations.fetch_or(Flag_Canceled);
+		BulkDataListCacheRequest.Cancel();
+	}
+}
+
+void FPendingPackage::OnEndLoad(bool& bOutShouldRemove, bool& bOutShouldWriteCache)
+{
+	// Called from within Owner->PendingPackageLock
+	bLoadInProgress = false;
+	if (PendingOperations.fetch_and(~Flag_EndLoad) == Flag_EndLoad)
+	{
+		bOutShouldWriteCache = true;
+		bOutShouldRemove = true;
+	}
+	else
+	{
+		bOutShouldWriteCache = false;
+		bOutShouldRemove = false;
+	}
+}
+
+void FPendingPackage::OnBulkDataListResults(FSharedBuffer Buffer)
+{
+	if (!Buffer.IsNull())
+	{
+		FMemoryReaderView Reader(MakeArrayView(Buffer));
+		Serialize(Reader, CachedBulkDatas);
+		if (Reader.IsError())
+		{
+			CachedBulkDatas.Empty();
+		}
+	}
+
+	ReadCache();
+
+	if (PendingOperations.fetch_and(~Flag_BulkDataListResults) == Flag_BulkDataListResults)
+	{
+		// We do not hold a lock when writing the cache, so we need to remove this from PendingPackages 
+		// before calling WriteCache to avoid other threads being able to access it.
+		TUniquePtr<FPendingPackage> ThisPointer;
+		{
+			FScopeLock PendingPackageScopeLock(&Owner->PendingPackageLock);
+			Owner->PendingPackages.RemoveAndCopyValue(PackageName, ThisPointer);
+			// PackageName may have already been removed from PendingPackages if Owner is shutting down. In that case,
+			// ThisPointer will be null, and Owner holds the UniquePtr to *this.
+			// It will release that UniquePtr only after Cancel returns, which will be after this function returns.
+		}
+		WriteCache();
+
+		// Deleting *this will destruct BulkDataListCacheRequest, which by
+		// default calls Cancel, but Cancel will block on this callback we are
+		// currently in. Direct the owner to keep requests alive to avoid that deadlock.
+		BulkDataListCacheRequest.KeepAlive();
+		ThisPointer.Reset();
+		// *this may have been deleted and can no longer be accessed
+	}
+}
+
+void FPendingPackage::ReadCache()
+{
+	if (CachedBulkDatas.Num() == 0)
+	{
+		return;
+	}
+
+	TArray<TRefCountPtr<FPendingPayloadId>> OldPendings;
+	TArray<TRefCountPtr<FPendingPayloadId>> NewPendings;
+
+	// Add each CachedBulkData to the Registry, updating RawHash if it is missing.
+	// For every BulkData in this package in the Registry after the CachedBulkData has been added, 
+	// if the RawHash is missing from the CachedBulkData as well, queue a read of its RawHash
+	// from the separate PlaceholderPayloadId BulkTablePayloadId cache bucket.
+	{
+		FWriteScopeLock RegistryScopeLock(Owner->RegistryLock);
+		if (!Owner->bActive)
+		{
+			return;
+		}
+		for (const UE::Serialization::FEditorBulkData& BulkData : CachedBulkDatas)
+		{
+			FGuid BulkDataId = BulkData.GetIdentifier();
+			FRegisteredBulk& TargetRegisteredBulk = Owner->Registry.FindOrAdd(BulkData.GetIdentifier());
+
+			bool bCachedLocationMatches = true;
+			UE::Serialization::FEditorBulkData& TargetBulkData = TargetRegisteredBulk.BulkData;
+			if (!TargetBulkData.GetIdentifier().IsValid())
+			{
+				TargetBulkData = BulkData;
+				TargetRegisteredBulk.PackageName = PackageName;
+				TargetRegisteredBulk.bPayloadAvailable = true;
+			}
+			else
+			{
+				check(TargetBulkData.GetIdentifier() == BulkDataId);
+				bCachedLocationMatches = BulkData.LocationMatches(TargetBulkData);
+				if (bCachedLocationMatches && !BulkData.HasPlaceholderPayloadId() && TargetBulkData.HasPlaceholderPayloadId())
+				{
+					TargetBulkData = BulkData;
+					TargetRegisteredBulk.PackageName = PackageName;
+					TargetRegisteredBulk.bPayloadAvailable = true;
+				}
+			}
+
+			if (bCachedLocationMatches && TargetBulkData.HasPlaceholderPayloadId())
+			{
+				NewPendings.Emplace(new FPendingPayloadId(BulkDataId));
+			}
+		}
+
+		for (TRefCountPtr<FPendingPayloadId>& NewPending : NewPendings)
+		{
+			TRefCountPtr<FPendingPayloadId>& ExistingPending = Owner->PendingPayloadIds.FindOrAdd(NewPending->GetBulkDataId());
+			if (ExistingPending.IsValid())
+			{
+				OldPendings.Add(MoveTemp(ExistingPending));
+			}
+			ExistingPending = NewPending;
+		}
+	}
+
+	Owner->ReadPayloadIdsFromCache(PackageName, MoveTemp(OldPendings), MoveTemp(NewPendings));
+}
+
+void FPendingPackage::WriteCache()
+{
+	// If the BulkDataList cache read found some existing results, then exit; cache results are deterministic so there
+	// is no need to write the list to the cache again.
+	if (CachedBulkDatas.Num() > 0)
+	{
+		return;
+	}
+
+	check(BulkDatas.Num() > 0); // We should have >= 1 bulkdatas, or we would not have created the FPendingPackage
+	// Remove any duplicates in the runtime BulkDatas; elements later in the list override earlier elements
+	{
+		TSet<FGuid> BulkDataGuids;
+		for (int32 Index = BulkDatas.Num() - 1; Index >= 0; --Index)
+		{
+			bool bAlreadyExists;
+			BulkDataGuids.Add(BulkDatas[Index].GetIdentifier(), &bAlreadyExists);
+			if (bAlreadyExists)
+			{
+				BulkDatas.RemoveAt(Index);
+			}
+		}
+	}
+
+	// Sort the list by guid, to avoid indeterminism in the list
+	BulkDatas.Sort([](const UE::Serialization::FEditorBulkData& A, const UE::Serialization::FEditorBulkData& B)
+		{
+			return A.GetIdentifier() < B.GetIdentifier();
+		});
+
+	TArray<uint8> Bytes;
+	FMemoryWriter Writer(Bytes);
+	Serialize(Writer, BulkDatas);
+	UE::EditorDomain::PutBulkDataList(PackageName, MakeSharedBufferFromArray(MoveTemp(Bytes)));
+}
+
+FUpdatePayloadWorker::FUpdatePayloadWorker(FBulkDataRegistryImpl* InBulkDataRegistry,
+	const UE::Serialization::FEditorBulkData& InSourceBulk)
+	: BulkData(InSourceBulk)
+	, BulkDataRegistry(InBulkDataRegistry)
+{
+	SharedDataLock = InBulkDataRegistry->SharedDataLock;
+}
+
+void FUpdatePayloadWorker::DoWork()
+{
+	FUpdatingPayload LocalUpdatingPayload;
+	FCompressedBuffer Buffer;
+	bool bValid = true;
+	for (;;)
+	{
+		BulkData.UpdatePayloadId();
+		Buffer = BulkData.GetCompressedPayload().Get();
+
+		{
+			FReadScopeLock SharedDataScopeLock(SharedDataLock->ActiveLock);
+			if (!SharedDataLock->bActive)
+			{
+				// The BulkDataRegistry has destructed. Our list of requesters is on the BulkDataRegistry, so there's nothing we can do except exit
+				return;
+			}
+			FWriteScopeLock RegistryScopeLock(BulkDataRegistry->RegistryLock);
+
+			if (!BulkDataRegistry->UpdatingPayloads.RemoveAndCopyValue(BulkData.GetIdentifier(), LocalUpdatingPayload))
+			{
+				// The updating payload might not exist in the case of the Registry shutting down; it will clear the UpdatingPayloads to cancel our action
+				// Return canceled (which we treat the same as failed) to our requesters
+				bValid = false;
+				break;
+			}
+			check(BulkDataRegistry->bActive); // Only set to false at the same time as SharedDataLock->bActive
+
+			FRegisteredBulk* RegisteredBulk = BulkDataRegistry->Registry.Find(BulkData.GetIdentifier());
+			if (!RegisteredBulk)
+			{
+				// Some agent has deregistered the BulkData before we finished calculating its payload
+				// return failure to our requesters
+				bValid = false;
+				break;
+			}
+
+			if (!BulkData.LocationMatches(RegisteredBulk->BulkData))
+			{
+				// Some caller has assigned a new BulkData. We need to abandon the BulkData we just loaded and give our callers the
+				// information about the new one
+				check(RegisteredBulk->BulkData.GetIdentifier() == BulkData.GetIdentifier()); // The identifier in the BulkData should match the key for that BulkData in Registry
+				BulkData = RegisteredBulk->BulkData;
+				// Add our LocalUpdatingPayload back to UpdatingPayloads; we removed it because we thought we were done.
+				BulkDataRegistry->UpdatingPayloads.Add(BulkData.GetIdentifier(), MoveTemp(LocalUpdatingPayload));
+				continue;
+			}
+
+			// Store the new PayloadId in the Registry's entry for the BulkData; new MetaData requests will no longer need to wait for it
+			RegisteredBulk->BulkData = BulkData;
+
+			// Mark that the next GetData call should remove the temporary payload
+			RegisteredBulk->bHasTempPayload = true;
+			BulkDataRegistry->AddTempLoadedPayload(BulkData.GetIdentifier(), BulkData.GetPayloadSize());
+			BulkDataRegistry->PruneTempLoadedPayloads();
+
+			if (RegisteredBulk->bAllowedToWritePayloadIdToCache)
+			{
+				BulkDataRegistry->WritePayloadIdToCache(RegisteredBulk->PackageName, BulkData);
+				RegisteredBulk->bAllowedToWritePayloadIdToCache = false;
+			}
+			break;
+		}
+	}
+
+	if (!bValid)
+	{
+		Buffer.Reset();
+	}
+	for (TUniqueFunction<void(bool, const FCompressedBuffer&)>& Requester : LocalUpdatingPayload.Requesters)
+	{
+		Requester(bValid, Buffer);
+	}
+}
+
+FPendingPayloadId::FPendingPayloadId(const FGuid& InBulkDataId)
+	: BulkDataId(InBulkDataId)
+	, Request(UE::DerivedData::EPriority::Low)
+{
+	// The last reference to this can be released by the completion callback, which will deadlock
+	// trying to cancel the request. KeepAlive skips cancellation in the destructor.
+	Request.KeepAlive();
+}
+
+void FPendingPayloadId::Cancel()
+{
+	Request.Cancel();
+}
+
+void Serialize(FArchive& Ar, TArray<UE::Serialization::FEditorBulkData>& Datas)
+{
+	int32 Num = Datas.Num();
+	Ar << Num;
+
+	const uint32 MinSize = 4;
+	if (Ar.IsLoading())
+	{
+		if (Ar.IsError() || Num * MinSize > Ar.TotalSize() - Ar.Tell())
+		{
+			Ar.SetError();
+			Datas.Empty();
+			return;
+		}
+		else
+		{
+			Datas.Empty(Num);
+		}
+		for (int32 n = 0; n < Num; ++n)
+		{
+			UE::Serialization::FEditorBulkData& BulkData = Datas.Emplace_GetRef();
+			BulkData.SerializeForRegistry(Ar);
+		}
+	}
+	else
+	{
+		for (UE::Serialization::FEditorBulkData& BulkData : Datas)
+		{
+			BulkData.SerializeForRegistry(Ar);
+		}
+	}
+}
+
+} // namespace UE::BulkDataRegistry::Private
diff --git a/Engine/Source/Editor/UnrealEd/Private/EditorDomain/BulkDataRegistryImpl.h b/Engine/Source/Editor/UnrealEd/Private/EditorDomain/BulkDataRegistryImpl.h
new file mode 100644
--- /dev/null
+++ b/Engine/Source/Editor/UnrealEd/Private/EditorDomain/BulkDataRegistryImpl.h
@@ -0,0 +1,216 @@
+// Copyright Epic Games, Inc. All Rights Reserved.
+
+#pragma once
+
+#include "Async/AsyncWork.h"
+#include "Async/Future.h"
+#include <atomic>
+#include "Containers/Array.h"
+#include "Containers/Map.h"
+#include "Containers/RingBuffer.h"
+#include "DerivedDataRequestOwner.h"
+#include "HAL/CriticalSection.h"
+#include "Serialization/BulkDataRegistry.h"
+#include "Serialization/EditorBulkData.h"
+#include "Templates/RefCounting.h"
+#include "TickableEditorObject.h"
+
+class FCompressedBuffer;
+class UPackage;
+struct FIoHash;
+
+namespace UE::BulkDataRegistry::Private
+{
+
+class FBulkDataRegistryImpl;
+
+/** Struct for storage of a BulkData in the registry, including the BulkData itself and data about cache status. */
+struct FRegisteredBulk
+{
+	FRegisteredBulk()
+		: bHasTempPayload(false), bAllowedToWritePayloadIdToCache(false), bRegistered(false), bInMemory(false), bPayloadAvailable(false)
+	{
+	}
+	UE::Serialization::FEditorBulkData BulkData;
+	FName PackageName;
+	bool bHasTempPayload : 1;
+	bool bAllowedToWritePayloadIdToCache : 1;
+	bool bRegistered : 1;
+	bool bInMemory : 1;
+	bool bPayloadAvailable : 1;
+};
+
+/** Serialize an array of BulkDatas into or out of bytes saved/load from the registry's persistent cache. */
+void Serialize(FArchive& Ar, TArray<UE::Serialization::FEditorBulkData>& InDatas);
+
+/** A collection of bulkdatas that should be sent to the cache for the given package. */
+class FPendingPackage
+{
+public:
+	FPendingPackage(FName PackageName, FBulkDataRegistryImpl* InOwner);
+	FPendingPackage(FPendingPackage&& Other) = delete;
+	FPendingPackage(const FPendingPackage& Other) = delete;
+
+	void Cancel();
+	void AddBulkData(UE::Serialization::FEditorBulkData&& BulkData)
+	{
+		BulkDatas.Add(MoveTemp(BulkData));
+	}
+	void OnEndLoad(bool& bOutShouldRemove, bool& bOutShouldWriteCache);
+	bool IsLoadInProgress() const
+	{
+		return bLoadInProgress;
+	}
+	void WriteCache();
+
+private:
+	void OnBulkDataListResults(FSharedBuffer Buffer);
+	void ReadCache();
+
+	enum EFlags : int32
+	{
+		Flag_EndLoad = 1 << 0,
+		Flag_BulkDataListResults = 1 << 1,
+		Flag_Canceled = 1 << 2,
+	};
+	FName PackageName;
+	TArray<UE::Serialization::FEditorBulkData> BulkDatas;
+	TArray<UE::Serialization::FEditorBulkData> CachedBulkDatas;
+	UE::DerivedData::FRequestOwner BulkDataListCacheRequest;
+	FBulkDataRegistryImpl* Owner;
+	/**
+	 * When PendingOperations reaches zero, we can remove the FPendingPackage.
+	 */
+	std::atomic<int32> PendingOperations;
+	/**
+	 * True until the LoadPackage for *this is complete. The FPendingPackage may last longer
+	 * than the initial load period while it waits for a CacheRead, but BulkDatas can only
+	 * be written to its list during its initial load, to avoid non-deterministic changes
+	 * to the list when operations occur on the UPackage in the editor.
+	 */
+	bool bLoadInProgress = true;
+};
+
+/** Data about a BulkData that has loaded its payload for TryGetMeta and should drop it after GetData or a timeout. */
+struct FTempLoadedPayload
+{
+	FGuid Guid;
+	uint64 PayloadSize;
+	double EndTime;
+};
+
+/** An Active flag and a lock around it for informing AutoDeleteAsyncTasks that their shared data is no longer available. */
+class FTaskSharedDataLock : public FThreadSafeRefCountedObject
+{
+public:
+	FRWLock ActiveLock;
+	bool bActive = true;
+};
+
+/** A worker that updates the PayloadId for a BulkData that is missing its RawHash. */
+class FUpdatePayloadWorker : public FNonAbandonableTask
+{
+public:
+	FUpdatePayloadWorker(FBulkDataRegistryImpl* InBulkDataRegistry,
+		const UE::Serialization::FEditorBulkData& InSourceBulk);
+
+	void DoWork();
+	TStatId GetStatId() const
+	{
+		RETURN_QUICK_DECLARE_CYCLE_STAT(FUpdatePayloadWorker, STATGROUP_ThreadPoolAsyncTasks);
+	}
+
+private:
+	UE::Serialization::FEditorBulkData BulkData;
+	TRefCountPtr<FTaskSharedDataLock> SharedDataLock;
+	FBulkDataRegistryImpl* BulkDataRegistry;
+};
+
+/** Data storage for the FUpdatePayloadWorker that is updated while in flight for additional requesters. */
+struct FUpdatingPayload
+{
+	/** Pointer to the task. Once set to non-null, is never modified. The task ensure this FUpdatingPayload is destroyed before the task destructs. */
+	FAutoDeleteAsyncTask<FUpdatePayloadWorker>* AsyncTask = nullptr;
+	TArray<TUniqueFunction<void(bool bValid, const FCompressedBuffer& Buffer)>> Requesters;
+};
+
+/** Data storage for a BulkData that is loading its PayloadId from the cache. */
+class FPendingPayloadId : public FThreadSafeRefCountedObject
+{
+public:
+	explicit FPendingPayloadId(const FGuid& InBulkDataId);
+
+	FPendingPayloadId(FPendingPayloadId&& Other) = delete;
+	FPendingPayloadId(const FPendingPayloadId& Other) = delete;
+
+	void Cancel();
+
+	UE::DerivedData::FRequestOwner& GetRequestOwner()
+	{
+		return Request;
+	}
+	const FGuid& GetBulkDataId()
+	{
+		return BulkDataId;
+	}
+
+private:
+	FGuid BulkDataId;
+	UE::DerivedData::FRequestOwner Request;
+};
+
+/** Implementation of a BulkDataRegistry that stores its persistent data in a DDC bucket. */
+class FBulkDataRegistryImpl : public IBulkDataRegistry, public FTickableEditorObject, public FTickableCookObject
+{
+public:
+	FBulkDataRegistryImpl();
+	virtual ~FBulkDataRegistryImpl();
+
+	// IBulkDataRegistry interface
+	virtual UE::BulkDataRegistry::ERegisterResult
+		TryRegister(UPackage* Owner, const UE::Serialization::FEditorBulkData& BulkData) override;
+	virtual void UpdateRegistrationData(UPackage* Owner, const UE::Serialization::FEditorBulkData& BulkData) override;
+	virtual void Unregister(const UE::Serialization::FEditorBulkData& BulkData) override;
+	virtual void OnExitMemory(const UE::Serialization::FEditorBulkData& BulkData) override;
+	virtual TFuture<UE::BulkDataRegistry::FMetaData> GetMeta(const FGuid& BulkDataId) override;
+	virtual TFuture<UE::BulkDataRegistry::FData> GetData(const FGuid& BulkDataId) override;
+	virtual bool TryGetBulkData(const FGuid& BulkDataId, UE::Serialization::FEditorBulkData* OutBulk = nullptr,
+		FName* OutOwner = nullptr) override;
+
+	virtual uint64 GetBulkDataResaveSize(FName PackageName) override;
+
+	// FTickableEditorObject/FTickableCookObject interface
+	virtual void Tick(float DeltaTime) override;
+	virtual void TickCook(float DeltaTime, bool bTickComplete) override;
+	virtual ETickableTickType GetTickableTickType() const override { return ETickableTickType::Always; }
+	virtual TStatId GetStatId() const override { return TStatId(); }
+
+private:
+	void AddPendingPackageBulkData(FName PackageName, UE::Serialization::FEditorBulkData&& BulkData);
+	void PollPendingPackages(bool bWaitForCooldown);
+	void AddTempLoadedPayload(const FGuid& RegistryKey, uint64 PayloadSize);
+	void PruneTempLoadedPayloads();
+	void OnEndLoadPackage(const FEndLoadPackageContext& Context);
+	void WritePayloadIdToCache(FName PackageName, const UE::Serialization::FEditorBulkData& BulkData) const;
+	void ReadPayloadIdsFromCache(FName PackageName, TArray<TRefCountPtr<FPendingPayloadId>>&& OldPendings,
+		TArray<TRefCountPtr<FPendingPayloadId>>&& NewPendings);
+
+	friend class FPendingPackage;
+	friend class FUpdatePayloadWorker;
+
+	// All locks can be held at the same time. They must always be entered in order: SharedDataLock, RegistryLock, PendingPackageLock
+	TRefCountPtr<FTaskSharedDataLock> SharedDataLock;
+	FRWLock RegistryLock;
+	FCriticalSection PendingPackageLock;
+
+	TMap<FGuid, FRegisteredBulk> Registry;
+	FResaveSizeTracker ResaveSizeTracker;
+	TMap<FGuid, FUpdatingPayload> UpdatingPayloads;
+	TMap<FName, TUniquePtr<FPendingPackage>> PendingPackages;
+	TMap<FGuid, TRefCountPtr<FPendingPayloadId>> PendingPayloadIds;
+	TRingBuffer<FTempLoadedPayload> TempLoadedPayloads;
+	uint64 TempLoadedPayloadsSize = 0;
+	bool bActive = true;
+};
+
+} // namespace UE::BulkDataRegistry::Private
diff --git a/Engine/Source/Editor/UnrealEd/Private/EditorDomain/EditorDomainUtils.cpp b/Engine/Source/Editor/UnrealEd/Private/EditorDomain/EditorDomainUtils.cpp
--- a/Engine/Source/Editor/UnrealEd/Private/EditorDomain/EditorDomainUtils.cpp
+++ b/Engine/Source/Editor/UnrealEd/Private/EditorDomain/EditorDomainUtils.cpp
@@ -171,8 +171,8 @@ const TCHAR* EditorDomainVersion = TEXT("D3E2B4E479C34305952A97AA63A8E19D");
 
 // Identifier of the CacheBuckets for EditorDomain tables
 const TCHAR* EditorDomainPackageBucketName = TEXT("EditorDomainPackage");
-const TCHAR* EditorDomainBulkDataListBucketName = TEXT("EditorDomainBulkDataList");
-const TCHAR* EditorDomainBulkDataPayloadIdBucketName = TEXT("EditorDomainBulkDataPayloadId");
+const TCHAR* BulkDataListBucketName = TEXT("BulkDataList");
+const TCHAR* BulkDataPayloadIdBucketName = TEXT("BulkDataPayloadId");
 
 static bool GetEditorDomainSaveUnversioned()
 {
@@ -1118,14 +1118,14 @@ UE::DerivedData::FCacheKey GetEditorDomainPackageKey(const FIoHash& EditorDomain
 
 static UE::DerivedData::FCacheKey GetBulkDataListKey(const FIoHash& EditorDomainHash)
 {
-	static UE::DerivedData::FCacheBucket EditorDomainBulkDataListBucket(EditorDomainBulkDataListBucketName);
-	return UE::DerivedData::FCacheKey{ EditorDomainBulkDataListBucket, EditorDomainHash };
+	static UE::DerivedData::FCacheBucket BulkDataListBucket(BulkDataListBucketName);
+	return UE::DerivedData::FCacheKey{ BulkDataListBucket, EditorDomainHash };
 }
 
 static UE::DerivedData::FCacheKey GetBulkDataPayloadIdKey(const FIoHash& PackageAndGuidHash)
 {
-	static UE::DerivedData::FCacheBucket EditorDomainBulkDataPayloadIdBucket(EditorDomainBulkDataPayloadIdBucketName);
-	return UE::DerivedData::FCacheKey{ EditorDomainBulkDataPayloadIdBucket, PackageAndGuidHash };
+	static UE::DerivedData::FCacheBucket BulkDataPayloadIdBucket(BulkDataPayloadIdBucketName);
+	return UE::DerivedData::FCacheKey{ BulkDataPayloadIdBucket, PackageAndGuidHash };
 }
 
 static UE::DerivedData::FValueId GetBulkDataValueId()
