[Bug Fix]: Add components list search doesn't compare non-beautified component names

diff --git a/Engine/Source/Editor/UnrealEd/Private/ComponentTypeRegistry.cpp b/Engine/Source/Editor/UnrealEd/Private/ComponentTypeRegistry.cpp
--- a/Engine/Source/Editor/UnrealEd/Private/ComponentTypeRegistry.cpp
+++ b/Engine/Source/Editor/UnrealEd/Private/ComponentTypeRegistry.cpp
@@ -450,7 +450,7 @@ void FComponentTypeRegistryData::ForceRefreshComponentList()
 			{
 				if( A->GetSortPriority() == 0 && B->GetSortPriority() == 0 )
 				{
-					bResult = FCString::Stricmp(*A->GetClassName(), *B->GetClassName()) < 0;
+					bResult = FCString::Stricmp(*A->GetClassDisplayName(), *B->GetClassDisplayName()) < 0;
 				}
 				else
 				{
diff --git a/Engine/Source/Editor/UnrealEd/Private/SComponentClassCombo.cpp b/Engine/Source/Editor/UnrealEd/Private/SComponentClassCombo.cpp
--- a/Engine/Source/Editor/UnrealEd/Private/SComponentClassCombo.cpp
+++ b/Engine/Source/Editor/UnrealEd/Private/SComponentClassCombo.cpp
@@ -27,11 +27,16 @@
 
 #define LOCTEXT_NAMESPACE "ComponentClassCombo"
 
-FString FComponentClassComboEntry::GetClassName() const
+FString FComponentClassComboEntry::GetClassDisplayName() const
 {
 	return ComponentClass != nullptr ? ComponentClass->GetDisplayNameText().ToString() : ComponentName;
 }
 
+FString FComponentClassComboEntry::GetClassName() const
+{
+	return ComponentClass != nullptr ? ComponentClass->GetName() : ComponentName;
+}
+
 void FComponentClassComboEntry::AddReferencedObjects(FReferenceCollector& Collector)
 {
 	UClass* RawClass = ComponentClass;
@@ -204,8 +209,13 @@ void SComponentClassCombo::GenerateFilteredComponentList()
 				if (bAllowEntry && bHasFilterText)
 				{
 					// Finally, disallow class entries that don't match the search box text.
-					FString FriendlyComponentName = GetSanitizedComponentName(CurrentEntry);
-					bAllowEntry = TextFilter->TestTextFilter(FBasicStringFilterExpressionContext(FriendlyComponentName));
+					const FString ComponentName = CurrentEntry->GetClassName();
+					bAllowEntry = TextFilter->TestTextFilter(FBasicStringFilterExpressionContext(ComponentName));
+					if(!bAllowEntry)
+					{
+						const FString FriendlyComponentName = GetSanitizedComponentName(CurrentEntry);
+						TextFilter->TestTextFilter(FBasicStringFilterExpressionContext(FriendlyComponentName));
+					}
 				}
 			}
 
@@ -542,7 +552,7 @@ FString SComponentClassCombo::GetSanitizedComponentName(FComponentClassComboEntr
 	}
 	else
 	{
-		DisplayName = Entry->GetClassName();
+		DisplayName = Entry->GetClassDisplayName();
 	}
 	return FName::NameToDisplayString(DisplayName, false);
 }
@@ -569,7 +579,7 @@ TSharedRef<SToolTip> SComponentClassCombo::GetComponentToolTip(FComponentClassCo
 
 	// Fallback for components that don't currently have a loaded class
 	return SNew(SToolTip)
-		.Text(FText::FromString(Entry->GetClassName()));
+		.Text(FText::FromString(Entry->GetClassDisplayName()));
 }
 
 bool SComponentClassCombo::IsComponentClassAllowed(FComponentClassComboEntryPtr Entry) const
diff --git a/Engine/Source/Editor/UnrealEd/Public/SComponentClassCombo.h b/Engine/Source/Editor/UnrealEd/Public/SComponentClassCombo.h
--- a/Engine/Source/Editor/UnrealEd/Public/SComponentClassCombo.h
+++ b/Engine/Source/Editor/UnrealEd/Public/SComponentClassCombo.h
@@ -146,6 +146,7 @@ public:
     	return CustomizationArgs.OnSubobjectCreated;
     }
 	
+	FString GetClassDisplayName() const;
 	FString GetClassName() const;
 	FString GetComponentPath() const { return ComponentPath.ToString(); }
 
