CookOnTheFlyServer: Combine duplicated variables from cook modes to reduce complexity, avoid the need to check for null, and prepare for the addition of a CookWorker mode.
#rnx

diff --git a/Engine/Source/Editor/UnrealEd/Classes/CookOnTheSide/CookOnTheFlyServer.h b/Engine/Source/Editor/UnrealEd/Classes/CookOnTheSide/CookOnTheFlyServer.h
--- a/Engine/Source/Editor/UnrealEd/Classes/CookOnTheSide/CookOnTheFlyServer.h
+++ b/Engine/Source/Editor/UnrealEd/Classes/CookOnTheSide/CookOnTheFlyServer.h
@@ -125,6 +125,8 @@
 	class FSaveCookedPackageContext;
 	class ICookOnTheFlyRequestManager;
 	struct FConstructPackageData;
+	struct FCookByTheBookOptions;
+	struct FCookOnTheFlyOptions;
 	struct FCookerTimer;
 	struct FCookSavePackageContext;
 	struct FGeneratorPackage;
@@ -149,23 +151,6 @@
 	UCookOnTheFlyServer(const FObjectInitializer& ObjectInitializer = FObjectInitializer::Get());
 	UCookOnTheFlyServer(FVTableHelper& Helper); // Declare the FVTableHelper constructor manually so that we can forward-declare-only TUniquePtrs in the header without getting compile error in generated cpp
 
-public:
-	struct FCookByTheBookOptions;
-
-	//////////////////////////////////////////////////////////////////////////
-	// Cook on the fly startup options
-	struct FCookOnTheFlyOptions
-	{
-		/** Wether the network file server or the I/O store connection server should bind to any port */
-		bool bBindAnyPort = false;
-		/** Whether to save the cooked output to the Zen storage server. */
-		bool bZenStore = false;
-		/** Whether the network file server should use a platform-specific communication protocol instead of TCP (used when bZenStore == false) */
-		bool bPlatformProtocol = false;
-		/** Target platforms */
-		TArray<ITargetPlatform*> TargetPlatforms;
-	};
-
 private:
 	using FPollFunction = TUniqueFunction<void(UE::Cook::FTickStackData&)>;
 	struct FPollable
@@ -190,13 +175,13 @@
 	/** Directory to output to instead of the default should be empty in the case of DLC cooking */ 
 	FString OutputDirectoryOverride;
 
-	FCookByTheBookOptions* CookByTheBookOptions = nullptr;
+	TUniquePtr<UE::Cook::FCookByTheBookOptions> CookByTheBookOptions;
 	TUniquePtr<UE::Cook::FPlatformManager> PlatformManager;
 
 	//////////////////////////////////////////////////////////////////////////
 	// Cook on the fly options
 
-	FCookOnTheFlyOptions CookOnTheFlyOptions;
+	TUniquePtr<UE::Cook::FCookOnTheFlyOptions> CookOnTheFlyOptions;
 	/** Cook on the fly server uses the NetworkFileServer */
 	TArray<class INetworkFileServer*> NetworkFileServers;
 	FOnFileModifiedDelegate FileModifiedDelegate;
@@ -322,6 +307,8 @@
 	/* Update polled fields used by CookOnTheFly's network request handlers */
 	void TickNetwork();
 
+	void TickMainCookLoop(UE::Cook::FTickStackData& StackData);
+
 	/** Execute operations that need to be done after each Scheduler task, such as checking for new external requests. */
 	void TickCookStatus(UE::Cook::FTickStackData& StackData);
 	void SetSaveBusy(bool bInSaveBusy);
@@ -338,7 +325,7 @@
 		Poll,			// Execute pollables which have exceeded their period
 		PollIdle,		// Execute pollables which have exceeded their idle period
 		WaitForAsync,	// Sleep for a time slice while we wait for async tasks to complete
-		YieldTick,		// Progress is blocked by an async result. Temporarily exit TickCookOnTheSide.
+		YieldTick,		// Progress is blocked by an async result. Temporarily exit TickMainCookLoop.
 		Cancel,			// Cancel the current CookByTheBook
 	};
 	/** Inspect all tasks the scheduler could do and return which one it should do. */
@@ -414,6 +401,37 @@
 		COSR_RequiresGC_PackageCount= 0x00000100,
 	};
 
+	struct FCookByTheBookStartupOptions
+	{
+	public:
+		TArray<ITargetPlatform*> TargetPlatforms;
+		TArray<FString> CookMaps;
+		TArray<FString> CookDirectories;
+		TArray<FString> NeverCookDirectories;
+		TArray<FString> CookCultures;
+		TArray<FString> IniMapSections;
+		TArray<FString> CookPackages; // list of packages we should cook, used to specify specific packages to cook
+		ECookByTheBookOptions CookOptions = ECookByTheBookOptions::None;
+		FString DLCName;
+		FString CreateReleaseVersion;
+		FString BasedOnReleaseVersion;
+		bool bGenerateStreamingInstallManifests = false;
+		bool bGenerateDependenciesForMaps = false;
+		bool bErrorOnEngineContentUse = false; // this is a flag for dlc, will cause the cooker to error if the dlc references engine content
+	};
 
+	struct FCookOnTheFlyStartupOptions
+	{
+		/** Wether the network file server or the I/O store connection server should bind to any port */
+		bool bBindAnyPort = false;
+		/** Whether to save the cooked output to the Zen storage server. */
+		bool bZenStore = false;
+		/** Whether the network file server should use a platform-specific communication protocol instead of TCP (used when bZenStore == false) */
+		bool bPlatformProtocol = false;
+		/** Target platforms */
+		TArray<ITargetPlatform*> TargetPlatforms;
+	};
+
+
 	virtual ~UCookOnTheFlyServer();
 
@@ -455,7 +473,7 @@
 	*
 	* @return true on success, false otherwise.
 	*/
-	bool StartCookOnTheFly(FCookOnTheFlyOptions InCookOnTheFlyOptions); 
+	bool StartCookOnTheFly(FCookOnTheFlyStartupOptions InCookOnTheFlyOptions); 
 
 	/**
 	* Broadcast our the fileserver presence on the network
@@ -467,25 +485,6 @@
 	*/
 	void ShutdownCookOnTheFly();
 
-	struct FCookByTheBookStartupOptions
-	{
-	public:
-		TArray<ITargetPlatform*> TargetPlatforms;
-		TArray<FString> CookMaps;
-		TArray<FString> CookDirectories;
-		TArray<FString> NeverCookDirectories;
-		TArray<FString> CookCultures; 
-		TArray<FString> IniMapSections;
-		TArray<FString> CookPackages; // list of packages we should cook, used to specify specific packages to cook
-		ECookByTheBookOptions CookOptions = ECookByTheBookOptions::None;
-		FString DLCName;
-		FString CreateReleaseVersion;
-		FString BasedOnReleaseVersion;
-		bool bGenerateStreamingInstallManifests = false;
-		bool bGenerateDependenciesForMaps = false;
-		bool bErrorOnEngineContentUse = false; // this is a flag for dlc, will cause the cooker to error if the dlc references engine content
-	};
-
 	/**
 	* Start a cook by the book session
 	* Cook on the fly can't run at the same time as cook by the book
@@ -516,14 +515,19 @@
 	void PostLoadPackageFixup(UPackage* Package, TArray<FName>* OutDiscoveredPackageNames,
 		TMap<FName, UE::Cook::FInstigator>* OutInstigators);
 
-	/**
-	* Handles cook package requests until there are no more requests, then returns
-	*
-	* @param  CookFlags output of the things that might have happened in the cook on the side
-	* 
-	* @return returns ECookOnTheSideResult
-	*/
-	uint32 TickCookOnTheSide( const float TimeSlice, uint32 &CookedPackagesCount, ECookTickFlags TickFlags = ECookTickFlags::None, int32 InMaxNumPackagesToSave = 50);
+	/** Tick CBTB until it finishes or needs to yield. Should only be called when in CookByTheBook Mode. */
+	uint32 TickCookByTheBook(const float TimeSlice, uint32& CookedPackageCount, ECookTickFlags TickFlags = ECookTickFlags::None, int32 InMaxNumPackagesToSave = 50);
+	/** Tick COTF until it finishes or needs to yield. Should only be called when in CookOnTheFly Mode. */
+	uint32 TickCookOnTheFly(const float TimeSlice, uint32& CookedPackageCount, ECookTickFlags TickFlags = ECookTickFlags::None, int32 InMaxNumPackagesToSave = 50);
 
 	/**
+	 * Execute CookByTheBook as a load of all packages followed by a save of all packages. This mode is sometimes faster than normal CBTB.
+	 * Should be called at the beginning of the cook after StartCookByTheBook, and only if IsFullLoadAndSave.
+	 */
+	uint32 CookFullLoadAndSave();
+
+	/** Return whether the configuration and commandline settings have selected FullLoadAndSave. */
+	bool IsFullLoadAndSave() const;
+
+	/**
 	* Clear all the previously cooked data all cook requests from now on will be considered recook requests
@@ -590,12 +594,7 @@
 
 	bool HasRecompileShaderRequests() const;
 
-	UE_DEPRECATED(4.26, "Use HasRemainingWork instead")
-	bool HasCookRequests() const
-	{
-		return HasRemainingWork();
-	}
-	/** Return whether TickCookOnTheSide needs to take any action for the current session. If not, the session is done. Used for external managers of the cooker to know when to tick it. */
+	/** Return whether the tick needs to take any action for the current session. If not, the session is done. Used for external managers of the cooker to know when to tick it. */
 	bool HasRemainingWork() const;
 	void WaitForRequests(int TimeoutMs);
 
@@ -724,7 +723,7 @@
 	//////////////////////////////////////////////////////////////////////////
 	// cook by the book specific functions
 	/** Construct the on-stack data for StartCook functions, based on the arguments to StartCookByTheBook */
-	FBeginCookContext SetCookByTheBookOptions(const FCookByTheBookStartupOptions& StartupOptions);
+	FBeginCookContext CreateBeginCookByTheBookContext(const FCookByTheBookStartupOptions& StartupOptions);
 	/** Set the PlatformManager's session platforms and finish filling out the BeginContext for StartCookByTheBook */
 	void SelectSessionPlatforms(FBeginCookContext& BeginContext);
 	/** CollectFilesToCook and add them as external requests. */
@@ -800,9 +799,7 @@
 
 	void RegisterShaderChunkDataGenerator();
 
-	/**
-	* Call back from the TickCookOnTheSide when a cook by the book finishes (when started form StartCookByTheBook)
-	*/
+	/** Called at the end of CookByTheBook to write aggregated data such as AssetRegistry and shaders. */
 	void CookByTheBookFinished();
 
 	/** Clears session-lifetime data from COTFS and CookPackageDatas. */
@@ -837,7 +834,7 @@
 	// cook on the fly specific functions
 
 	/** Construct the on-stack data for StartCook functions, based on the arguments to StartCookOnTheFly */
-	FBeginCookContext CreateBeginCookOnTheFlyContext();
+	FBeginCookContext CreateBeginCookOnTheFlyContext(const FCookOnTheFlyStartupOptions& Options);
 	/** Construct the on-stack data for StartCook functions, based on the arguments when adding a COTF platform */
 	FBeginCookContext CreateAddPlatformContext(ITargetPlatform* TargetPlatform);
 	void GetCookOnTheFlyUnsolicitedFiles(const ITargetPlatform* TargetPlatform, const FString& PlatformName, TArray<FString>& UnsolicitedFiles, const FString& Filename, bool bIsCookable);
@@ -1144,8 +1141,6 @@
 	/** Try calling the splitter's populate to create the package */
 	UPackage* TryPopulateGeneratedPackage(UE::Cook::FPopulatePackageContext& Context);
 
-	uint32 FullLoadAndSave(uint32& CookedPackageCount);
-
 	ICookedPackageWriter& FindOrCreatePackageWriter(const ITargetPlatform* TargetPlatform);
 	void FindOrCreateSaveContexts(TConstArrayView<const ITargetPlatform*> TargetPlatforms);
 	UE::Cook::FCookSavePackageContext& FindOrCreateSaveContext(const ITargetPlatform* TargetPlatform);
@@ -1178,7 +1173,12 @@
 	bool bFirstCookInThisProcess = true;
 	bool bHasDeferredInitializeCookOnTheFly = false;
 	bool bImportBehaviorCallbackInstalled = false;
+	/** Cancel has been queued and will be processed next tick */
+	bool bCancelSession = false;
+	/** True if a session is started for any mode. If the session is started, then we have at least one TargetPlatform specified. */
+	bool bSessionRunning = false;
+	/** Whether we're using ZenStore for storage of cook results. If false, we are using LooseCookedPackageWriter. */
+	bool bZenStore = false;
 
-
 	/** Timers for tracking how long we have been busy, to manage retries and warnings of deadlock */
 	double SaveBusyRetryTimeSeconds = MAX_flt;
diff --git a/Engine/Source/Editor/UnrealEd/Private/Commandlets/CookCommandlet.cpp b/Engine/Source/Editor/UnrealEd/Private/Commandlets/CookCommandlet.cpp
--- a/Engine/Source/Editor/UnrealEd/Private/Commandlets/CookCommandlet.cpp
+++ b/Engine/Source/Editor/UnrealEd/Private/Commandlets/CookCommandlet.cpp
@@ -439,7 +439,7 @@ bool UCookCommandlet::CookOnTheFly( FGuid InstanceId, int32 Timeout, bool bForce
 	CookFlags |= bIgnoreIniSettingsOutOfDate || CookerSettings->bIgnoreIniSettingsOutOfDateForIteration ? ECookInitializationFlags::IgnoreIniSettingsOutOfDate : ECookInitializationFlags::None;
 	CookOnTheFlyServer->Initialize( ECookMode::CookOnTheFly, CookFlags );
 
-	UCookOnTheFlyServer::FCookOnTheFlyOptions CookOnTheFlyStartupOptions;
+	UCookOnTheFlyServer::FCookOnTheFlyStartupOptions CookOnTheFlyStartupOptions;
 	CookOnTheFlyStartupOptions.bBindAnyPort = InstanceId.IsValid();
 	CookOnTheFlyStartupOptions.bZenStore = Switches.Contains(TEXT("ZenStore"));
 	CookOnTheFlyStartupOptions.bPlatformProtocol = Switches.Contains(TEXT("PlatformProtocol"));
@@ -479,7 +479,7 @@ bool UCookCommandlet::CookOnTheFly( FGuid InstanceId, int32 Timeout, bool bForce
 			, LastCookActionTime(FPlatformTime::Seconds())
 		{}
 
-		/** Intended to be called with stats from a UCookOnTheFlyServer::TickCookOnTheSide() call. Determines if we should be calling GC after TickCookOnTheSide(). */
+		/** Intended to be called with stats from the COTF tick loop. Determines if we should be calling GC after TickCookOnTheFly(). */
 		void Update(uint32 CookedCount, UCookOnTheFlyServer::ECookOnTheSideResult ResultFlags)
 		{
 			if (ResultFlags & (UCookOnTheFlyServer::COSR_CookedMap | UCookOnTheFlyServer::COSR_CookedPackage | UCookOnTheFlyServer::COSR_WaitingOnCache))
@@ -559,7 +559,7 @@ bool UCookCommandlet::CookOnTheFly( FGuid InstanceId, int32 Timeout, bool bForce
 	while (!IsEngineExitRequested())
 	{
 		uint32 CookedPkgCount = 0;
-		uint32 TickResults = CookOnTheFlyServer->TickCookOnTheSide(/*TimeSlice =*/10.f, CookedPkgCount, ShowProgress ? ECookTickFlags::None : ECookTickFlags::HideProgressDisplay);
+		uint32 TickResults = CookOnTheFlyServer->TickCookOnTheFly(/*TimeSlice =*/10.f, CookedPkgCount, ShowProgress ? ECookTickFlags::None : ECookTickFlags::HideProgressDisplay);
 
 		// Flush the asset registry before GC
 		FAssetRegistryModule::TickAssetRegistry(-1.0f);
@@ -992,14 +992,19 @@ bool UCookCommandlet::CookByTheBook( const TArray<ITargetPlatform*>& Platforms)
 			COOK_STAT(FScopedDurationTimer StartCookByTheBookTimer(DetailedCookStats::StartCookByTheBookTimeSec));
 			CookOnTheFlyServer->StartCookByTheBook(StartupOptions);
 		}
-		while (CookOnTheFlyServer->IsCookByTheBookRunning())
+		DECLARE_SCOPE_CYCLE_COUNTER(TEXT("CookByTheBook.MainLoop"), STAT_CookByTheBook_MainLoop, STATGROUP_LoadTime);
+		if (CookOnTheFlyServer->IsFullLoadAndSave())
 		{
-			DECLARE_SCOPE_CYCLE_COUNTER( TEXT( "CookByTheBook.MainLoop" ), STAT_CookByTheBook_MainLoop, STATGROUP_LoadTime );
+			CookOnTheFlyServer->CookFullLoadAndSave();
+		}
+		else
+		{
+			while (CookOnTheFlyServer->IsInSession())
 			{
 				uint32 TickResults = 0;
 				uint32 UnusedVariable = 0;
 
-				TickResults = CookOnTheFlyServer->TickCookOnTheSide( MAX_flt, UnusedVariable,
+				TickResults = CookOnTheFlyServer->TickCookByTheBook(MAX_flt, UnusedVariable,
 					ShowProgress ? ECookTickFlags::None : ECookTickFlags::HideProgressDisplay, MAX_int32);
 
 				if ((TickResults & UCookOnTheFlyServer::COSR_RequiresGC) != 0)
@@ -1039,7 +1044,7 @@ bool UCookCommandlet::CookByTheBook( const TArray<ITargetPlatform*>& Platforms)
 							// GShaderCompilingManager->FinishAllCompilation();
 						}
 					}
-					else 
+					else
 					{
 						// cooker loaded some object which needs to be cleaned up before the cooker can proceed so force gc
 						GCReason = TEXT("COSR_RequiresGC");
@@ -1054,7 +1059,7 @@ bool UCookCommandlet::CookByTheBook( const TArray<ITargetPlatform*>& Platforms)
 
 					int32 NumObjectsBeforeGC = GUObjectArray.GetObjectArrayNumMinusAvailable();
 					int32 NumObjectsAvailableBeforeGC = GUObjectArray.GetObjectArrayEstimatedAvailable();
-					UE_LOG(LogCookCommandlet, Display, TEXT("GarbageCollection...%s (%s)"), (bPartialGC? TEXT(" partial gc") : TEXT("")), *GCReason);
+					UE_LOG(LogCookCommandlet, Display, TEXT("GarbageCollection...%s (%s)"), (bPartialGC ? TEXT(" partial gc") : TEXT("")), *GCReason);
 
 #if ENABLE_LOW_LEVEL_MEM_TRACKER
 					FLowLevelMemTracker::Get().UpdateStatsPerFrame();
diff --git a/Engine/Source/Editor/UnrealEd/Private/CookOnTheFlyServer.cpp b/Engine/Source/Editor/UnrealEd/Private/CookOnTheFlyServer.cpp
--- a/Engine/Source/Editor/UnrealEd/Private/CookOnTheFlyServer.cpp
+++ b/Engine/Source/Editor/UnrealEd/Private/CookOnTheFlyServer.cpp
@@ -461,7 +461,6 @@ private:
 UCookOnTheFlyServer::UCookOnTheFlyServer(const FObjectInitializer& ObjectInitializer)
 	: Super(ObjectInitializer),
 	CurrentCookMode(ECookMode::CookOnTheFly),
-	CookByTheBookOptions(nullptr),
 	CookFlags(ECookInitializationFlags::None),
 	bIsSavingPackage(false),
 	AssetRegistry(nullptr)
@@ -479,9 +478,6 @@ UCookOnTheFlyServer::~UCookOnTheFlyServer()
 	FCoreUObjectDelegates::GetPreGarbageCollectDelegate().RemoveAll(this);
 	FCoreUObjectDelegates::GetPostGarbageCollect().RemoveAll(this);
 	GetTargetPlatformManager()->GetOnTargetPlatformsInvalidatedDelegate().RemoveAll(this);
-
-	delete CookByTheBookOptions;
-	CookByTheBookOptions = nullptr;
 }
 
 // This tick only happens in the editor.  The cook commandlet directly calls tick on the side.
@@ -492,7 +488,7 @@ void UCookOnTheFlyServer::Tick(float DeltaTime)
 
 	check(IsCookingInEditor());
 
-	if (IsCookByTheBookMode() && !IsCookByTheBookRunning() && !GIsSlowTask && IsCookFlagSet(ECookInitializationFlags::BuildDDCInBackground))
+	if (IsCookByTheBookMode() && !IsInSession() && !GIsSlowTask && IsCookFlagSet(ECookInitializationFlags::BuildDDCInBackground))
 	{
 		// if we are in the editor then precache some stuff ;)
 		TArray<const ITargetPlatform*> CacheTargetPlatforms;
@@ -509,8 +505,18 @@ void UCookOnTheFlyServer::Tick(float DeltaTime)
 	}
 
 	uint32 CookedPackagesCount = 0;
-	const static float CookOnTheSideTimeSlice = 0.1f; // seconds
-	TickCookOnTheSide( CookOnTheSideTimeSlice, CookedPackagesCount);
+	const static float TickTimeSliceSeconds = 0.1f;
+	TickCancels();
+	if (IsCookOnTheFlyMode())
+	{
+		TickCookOnTheFly(TickTimeSliceSeconds, CookedPackagesCount);
+	}
+	else
+	{
+		check(IsCookByTheBookMode());
+		TickCookByTheBook(TickTimeSliceSeconds, CookedPackagesCount);
+	}
+
 	TickRequestManager();
 	TickRecompileShaderRequests();
 }
@@ -525,7 +531,7 @@ TStatId UCookOnTheFlyServer::GetStatId() const
 	RETURN_QUICK_DECLARE_CYCLE_STAT(UCookServer, STATGROUP_Tickables);
 }
 
-bool UCookOnTheFlyServer::StartCookOnTheFly(FCookOnTheFlyOptions InCookOnTheFlyOptions)
+bool UCookOnTheFlyServer::StartCookOnTheFly(FCookOnTheFlyStartupOptions InCookOnTheFlyOptions)
 {
 	LLM_SCOPE_BYTAG(Cooker);
 #if WITH_COTF
@@ -536,8 +542,7 @@ bool UCookOnTheFlyServer::StartCookOnTheFly(FCookOnTheFlyOptions InCookOnTheFlyO
 	NetworkRequestEvent = FPlatformProcess::GetSynchEventFromPool();
 #endif
 
-	CookOnTheFlyOptions = MoveTemp(InCookOnTheFlyOptions);
-	FBeginCookContext BeginContext = CreateBeginCookOnTheFlyContext();
+	FBeginCookContext BeginContext = CreateBeginCookOnTheFlyContext(InCookOnTheFlyOptions);
 	CreateSandboxFile(BeginContext);
 
 	CookOnTheFlyServerInterface = MakeUnique<UCookOnTheFlyServer::FCookOnTheFlyServerInterface>(*this);
@@ -556,15 +561,15 @@ bool UCookOnTheFlyServer::StartCookOnTheFly(FCookOnTheFlyOptions InCookOnTheFlyO
 
 	GRedirectCollector.OnStartupPackageLoadComplete();
 
-	for (ITargetPlatform* TargetPlatform : CookOnTheFlyOptions.TargetPlatforms)
+	for (ITargetPlatform* TargetPlatform : InCookOnTheFlyOptions.TargetPlatforms)
 	{
 		AddCookOnTheFlyPlatformFromGameThread(TargetPlatform);
 	}
 
 	UE_LOG(LogCook, Display, TEXT("Starting '%s' cook-on-the-fly server"),
-		CookOnTheFlyOptions.bZenStore ? TEXT("Zen") : TEXT("Network File"));
+		IsUsingZenStore() ? TEXT("Zen") : TEXT("Network File"));
 
-	if (CookOnTheFlyOptions.bZenStore)
+	if (IsUsingZenStore())
 	{
 		UE::Cook::FIoStoreCookOnTheFlyServerOptions ServerOptions;
 		ServerOptions.Port = -1; // Use default
@@ -573,8 +578,8 @@ bool UCookOnTheFlyServer::StartCookOnTheFly(FCookOnTheFlyOptions InCookOnTheFlyO
 	else
 	{
 		FNetworkFileServerOptions ServerOptions;
-		ServerOptions.Protocol = CookOnTheFlyOptions.bPlatformProtocol ? NFSP_Platform : NFSP_Tcp;
-		ServerOptions.Port = CookOnTheFlyOptions.bBindAnyPort ? 0 : -1;
+		ServerOptions.Protocol = CookOnTheFlyOptions->bPlatformProtocol ? NFSP_Platform : NFSP_Tcp;
+		ServerOptions.Port = CookOnTheFlyOptions->bBindAnyPort ? 0 : -1;
 		CookOnTheFlyRequestManager = UE::Cook::MakeNetworkFileCookOnTheFlyRequestManager(*CookOnTheFlyServerInterface, MoveTemp(ServerOptions));
 	}
 	BeginCookEditorSystems();
@@ -1005,7 +1010,7 @@ bool UCookOnTheFlyServer::IsUsingShaderCodeLibrary() const
 
 bool UCookOnTheFlyServer::IsUsingZenStore() const
 {
-	return (IsCookByTheBookMode() && CookByTheBookOptions->bZenStore) || (IsCookOnTheFlyMode() && CookOnTheFlyOptions.bZenStore);
+	return bZenStore;
 }
 
 bool UCookOnTheFlyServer::IsCookOnTheFlyMode() const
@@ -1020,35 +1025,23 @@ bool UCookOnTheFlyServer::IsUsingLegacyCookOnTheFlyScheduling() const
 
 bool UCookOnTheFlyServer::IsCreatingReleaseVersion()
 {
-	if (CookByTheBookOptions)
-	{
-		return !CookByTheBookOptions->CreateReleaseVersion.IsEmpty();
-	}
-
-	return false;
+	return !CookByTheBookOptions->CreateReleaseVersion.IsEmpty();
 }
 
 bool UCookOnTheFlyServer::IsCookingDLC() const
 {
-	// can only cook DLC in cook by the book
 	// we are cooking DLC when the DLC name is setup
-	
-	if (CookByTheBookOptions)
-	{
-		return !CookByTheBookOptions->DlcName.IsEmpty();
-	}
-
-	return false;
+	return !CookByTheBookOptions->DlcName.IsEmpty();
 }
 
 bool UCookOnTheFlyServer::IsCookingAgainstFixedBase() const
 {
-	return IsCookingDLC() && CookByTheBookOptions && CookByTheBookOptions->bCookAgainstFixedBase;
+	return IsCookingDLC() && CookByTheBookOptions->bCookAgainstFixedBase;
 }
 
 bool UCookOnTheFlyServer::ShouldPopulateFullAssetRegistry() const
 {
-	return !IsCookingDLC() || (CookByTheBookOptions && CookByTheBookOptions->bDlcLoadMainAssetRegistry);
+	return !IsCookingDLC() || CookByTheBookOptions->bDlcLoadMainAssetRegistry;
 }
 
 FString UCookOnTheFlyServer::GetBaseDirectoryForDLC() const
@@ -1167,131 +1160,120 @@ bool UCookOnTheFlyServer::RequestPackage(const FName& StandardPackageFName, cons
 	return RequestPackage(StandardPackageFName, PlatformManager->GetSessionPlatforms(), bForceFrontOfQueue);
 }
 
-uint32 UCookOnTheFlyServer::TickCookOnTheSide(const float TimeSlice, uint32 &CookedPackageCount, ECookTickFlags TickFlags, int32 InMaxNumPackagesToSave)
+uint32 UCookOnTheFlyServer::TickCookByTheBook(const float TimeSlice, uint32& CookedPackageCount, ECookTickFlags TickFlags, int32 InMaxNumPackagesToSave)
 {
+	check(IsCookByTheBookMode());
+
 	LLM_SCOPE_BYTAG(Cooker);
-	TickCancels();
-	TickNetwork();
-	if (!IsInSession())
+	COOK_STAT(FScopedDurationTimer TickTimer(DetailedCookStats::TickCookOnTheSideTimeSec));
+	UE::Cook::FTickStackData StackData(TimeSlice, IsRealtimeMode(), TickFlags, InMaxNumPackagesToSave);
+
+	TickMainCookLoop(StackData);
+
+	CookByTheBookOptions->CookTime += StackData.Timer.GetTimeTillNow();
+	// Make sure no UE_SCOPED_HIERARCHICAL_COOKTIMERs are around CookByTheBookFinishes or CancelCookByTheBook, as those functions delete memory for them
+	if (StackData.bCookCancelled)
 	{
-		return COSR_None;
+		CancelCookByTheBook();
 	}
-
-	if (IsCookByTheBookMode() && CookByTheBookOptions->bRunning && CookByTheBookOptions->bFullLoadAndSave)
+	else if (IsInSession() && StackData.bCookComplete)
 	{
-		uint32 Result = FullLoadAndSave(CookedPackageCount);
-
+		UpdateDisplay(TickFlags, true /* bForceDisplay */);
 		CookByTheBookFinished();
-
-		return Result;
 	}
+	CookedPackageCount += StackData.CookedPackageCount;
+	return StackData.ResultFlags;
+}
+
+uint32 UCookOnTheFlyServer::TickCookOnTheFly(const float TimeSlice, uint32& CookedPackageCount, ECookTickFlags TickFlags, int32 InMaxNumPackagesToSave)
+{
+	check(IsCookOnTheFlyMode());
 
+	LLM_SCOPE_BYTAG(Cooker);
 	COOK_STAT(FScopedDurationTimer TickTimer(DetailedCookStats::TickCookOnTheSideTimeSec));
+	UE::Cook::FTickStackData StackData(TimeSlice, IsRealtimeMode(), TickFlags, InMaxNumPackagesToSave);
+
+	TickNetwork();
+	TickMainCookLoop(StackData);
 
+	CookedPackageCount += StackData.CookedPackageCount;
+	return StackData.ResultFlags;
+}
+
+void UCookOnTheFlyServer::TickMainCookLoop(UE::Cook::FTickStackData& StackData)
+{
+	UE_SCOPED_HIERARCHICAL_COOKTIMER(TickMainCookLoop);
+	if (!IsInSession())
 	{
-		if (AssetRegistry == nullptr || AssetRegistry->IsLoadingAssets())
-		{
-			// early out
-			return COSR_None;
-		}
+		return;
 	}
 
-	UE::Cook::FTickStackData StackData(TimeSlice, IsRealtimeMode(), TickFlags, InMaxNumPackagesToSave);
-	bool bCookComplete = false;
-	bool bCookCancelled = false;
-
+	bool bContinueTick = true;
+	while (bContinueTick && (!IsEngineExitRequested() || CurrentCookMode == ECookMode::CookByTheBook))
 	{
-		UE_SCOPED_HIERARCHICAL_COOKTIMER(TickCookOnTheSide); // Make sure no UE_SCOPED_HIERARCHICAL_COOKTIMERs are around CookByTheBookFinishes or CancelCookByTheBook, as those functions delete memory for them
+		TickCookStatus(StackData);
 
-		bool bContinueTick = true;
-		while (bContinueTick && (!IsEngineExitRequested() || CurrentCookMode == ECookMode::CookByTheBook))
+		ECookAction CookAction = DecideNextCookAction(StackData);
+		int32 NumPushed;
+		bool bBusy;
+		switch (CookAction)
 		{
-			TickCookStatus(StackData);
-
-			ECookAction CookAction = DecideNextCookAction(StackData);
-			int32 NumPushed;
-			bool bBusy;
-			switch (CookAction)
+		case ECookAction::Request:
+			PumpRequests(StackData, NumPushed);
+			if (NumPushed > 0)
 			{
-			case ECookAction::Request:
-				PumpRequests(StackData, NumPushed);
-				if (NumPushed > 0)
-				{
-					SetLoadBusy(false);
-				}
-				break;
-			case ECookAction::Load:
-				PumpLoads(StackData, 0, NumPushed, bBusy);
-				SetLoadBusy(bBusy && NumPushed == 0); // Mark as busy if pump was blocked and we did not make any progress
-				if (NumPushed > 0)
-				{
-					SetSaveBusy(false);
-				}
-				break;
-			case ECookAction::LoadLimited:
-				PumpLoads(StackData, DesiredLoadQueueLength, NumPushed, bBusy);
-				SetLoadBusy(bBusy && NumPushed == 0); // Mark as busy if pump was blocked and we did not make any progress
-				if (NumPushed > 0)
-				{
-					SetSaveBusy(false);
-				}
-				break;
-			case ECookAction::Save:
-				PumpSaves(StackData, 0, NumPushed, bBusy);
-				SetSaveBusy(bBusy && NumPushed == 0); // Mark as busy if pump was blocked and we did not make any progress
-				break;
-			case ECookAction::SaveLimited:
-				PumpSaves(StackData, DesiredSaveQueueLength, NumPushed, bBusy);
-				SetSaveBusy(bBusy && NumPushed == 0); // Mark as busy if pump was blocked and we did not make any progress
-				break;
-			case ECookAction::Poll:
-				PumpPollables(StackData, false /* bIsIdle */);
-				break;
-			case ECookAction::PollIdle:
-				PumpPollables(StackData, true /* bIsIdle */);
-				break;
-			case ECookAction::WaitForAsync:
-				WaitForAsync(StackData);
-				break;
-			case ECookAction::YieldTick:
-				bContinueTick = false;
-				break;
-			case ECookAction::Done:
-				bContinueTick = false;
-				bCookComplete = true;
-				break;
-			case ECookAction::Cancel:
-				bCookCancelled = true;
-				bContinueTick = false;
-				break;
-			default:
-				check(false);
-				break;
+				SetLoadBusy(false);
+			}
+			break;
+		case ECookAction::Load:
+			PumpLoads(StackData, 0, NumPushed, bBusy);
+			SetLoadBusy(bBusy && NumPushed == 0); // Mark as busy if pump was blocked and we did not make any progress
+			if (NumPushed > 0)
+			{
+				SetSaveBusy(false);
 			}
+			break;
+		case ECookAction::LoadLimited:
+			PumpLoads(StackData, DesiredLoadQueueLength, NumPushed, bBusy);
+			SetLoadBusy(bBusy && NumPushed == 0); // Mark as busy if pump was blocked and we did not make any progress
+			if (NumPushed > 0)
+			{
+				SetSaveBusy(false);
+			}
+			break;
+		case ECookAction::Save:
+			PumpSaves(StackData, 0, NumPushed, bBusy);
+			SetSaveBusy(bBusy && NumPushed == 0); // Mark as busy if pump was blocked and we did not make any progress
+			break;
+		case ECookAction::SaveLimited:
+			PumpSaves(StackData, DesiredSaveQueueLength, NumPushed, bBusy);
+			SetSaveBusy(bBusy && NumPushed == 0); // Mark as busy if pump was blocked and we did not make any progress
+			break;
+		case ECookAction::Poll:
+			PumpPollables(StackData, false /* bIsIdle */);
+			break;
+		case ECookAction::PollIdle:
+			PumpPollables(StackData, true /* bIsIdle */);
+			break;
+		case ECookAction::WaitForAsync:
+			WaitForAsync(StackData);
+			break;
+		case ECookAction::YieldTick:
+			bContinueTick = false;
+			break;
+		case ECookAction::Done:
+			bContinueTick = false;
+			StackData.bCookComplete = true;
+			break;
+		case ECookAction::Cancel:
+			StackData.bCookCancelled = true;
+			bContinueTick = false;
+			break;
+		default:
+			check(false);
+			break;
 		}
 	}
-
-	if (bCookCancelled)
-	{
-		CancelCookByTheBook();
-	}
-
-	if (CookByTheBookOptions)
-	{
-		CookByTheBookOptions->CookTime += StackData.Timer.GetTimeTillNow();
-	}
-
-	if (IsCookByTheBookRunning() && bCookComplete)
-	{
-		check(IsCookByTheBookMode());
-
-		// if we are out of stuff and we are in cook by the book from the editor mode then we finish up
-		UpdateDisplay(TickFlags, true /* bForceDisplay */);
-		CookByTheBookFinished();
-	}
-
-	CookedPackageCount += StackData.CookedPackageCount;
-	return StackData.ResultFlags;
 }
 
 void UCookOnTheFlyServer::TickCookStatus(UE::Cook::FTickStackData& StackData)
@@ -1776,7 +1758,7 @@ void UCookOnTheFlyServer::WaitForAsync(UE::Cook::FTickStackData& StackData)
 
 UCookOnTheFlyServer::ECookAction UCookOnTheFlyServer::DecideNextCookAction(UE::Cook::FTickStackData& StackData)
 {
-	if (IsCookByTheBookMode() && CookByTheBookOptions->bCancel)
+	if (bCancelSession)
 	{
 		return ECookAction::Cancel;
 	}
@@ -2300,7 +2282,7 @@ void UCookOnTheFlyServer::QueueDiscoveredPackageData(UE::Cook::FPackageData& Pac
 	}
 	
 	if (!PackageData.IsInProgress() &&
-		(PackageData.IsGenerated() || !IsCookByTheBookMode() || !CookByTheBookOptions->bSkipHardReferences))
+		(PackageData.IsGenerated() || !CookByTheBookOptions->bSkipHardReferences))
 	{
 		PackageData.SetRequestData(TargetPlatforms, /*bIsUrgent*/ false, UE::Cook::FCompletionCallback(),
 			MoveTemp(Instigator));
@@ -2352,22 +2334,20 @@ void UCookOnTheFlyServer::TickNetwork()
 {
 	// Only CookOnTheFly handles network requests
 	// It is not safe to call PruneUnreferencedSessionPlatforms in CookByTheBook because StartCookByTheBook does not AddRef its session platforms
-	if (IsCookOnTheFlyMode())
+	check(IsCookOnTheFlyMode())
+	if (IsInSession() && !bCookOnTheFlyExternalRequests)
 	{
-		if (IsInSession() && !bCookOnTheFlyExternalRequests)
-		{
-			PlatformManager->PruneUnreferencedSessionPlatforms(*this);
-		}
-		else
+		PlatformManager->PruneUnreferencedSessionPlatforms(*this);
+	}
+	else
+	{
+		// Process callbacks in case there is a callback pending that needs to create a session
+		TArray<UE::Cook::FSchedulerCallback> Callbacks;
+		if (ExternalRequests->DequeueCallbacks(Callbacks))
 		{
-			// Process callbacks in case there is a callback pending that needs to create a session
-			TArray<UE::Cook::FSchedulerCallback> Callbacks;
-			if (ExternalRequests->DequeueCallbacks(Callbacks))
+			for (UE::Cook::FSchedulerCallback& Callback : Callbacks)
 			{
-				for (UE::Cook::FSchedulerCallback& Callback : Callbacks)
-				{
-					Callback();
-				}
+				Callback();
 			}
 		}
 	}
@@ -3085,7 +3065,7 @@ void UCookOnTheFlyServer::ProcessUnsolicitedPackages(TArray<FName>* OutDiscovere
 		{
 			UPackage* Package = PackageWithInstigator.Key;
 			FInstigator& Instigator = PackageWithInstigator.Value;
-			if (!IsCookByTheBookMode() || !CookByTheBookOptions->bSkipSoftReferences)
+			if (!CookByTheBookOptions->bSkipSoftReferences)
 			{
 				PostLoadPackageFixup(Package, OutDiscoveredPackageNames, OutInstigators);
 			}
@@ -3579,12 +3559,6 @@ void UCookOnTheFlyServer::TickPrecacheObjectsForPlatforms(const float TimeSlice,
 
 bool UCookOnTheFlyServer::HasExceededMaxMemory() const
 {
-	if (IsCookByTheBookMode() && CookByTheBookOptions->bFullLoadAndSave)
-	{
-		// FullLoadAndSave does the entire cook in one tick, so there is no need to GC after
-		return false;
-	}
-
 #if UE_GC_TRACK_OBJ_AVAILABLE
 	if (GUObjectArray.GetObjectArrayEstimatedAvailable() < MinFreeUObjectIndicesBeforeGC)
 	{
@@ -3755,7 +3729,7 @@ void UCookOnTheFlyServer::MarkPackageDirtyForCooker(UPackage* Package, bool bAll
 		return;
 	}
 
-	if (IsInSession() && !bAllowInSession && !(IsCookByTheBookMode() && CookByTheBookOptions->bFullLoadAndSave))
+	if (IsInSession() && !bAllowInSession)
 	{
 		ExternalRequests->AddCallback([this, PackageName]() { MarkPackageDirtyForCookerFromSchedulerThread(PackageName); });
 	}
@@ -3787,7 +3761,7 @@ void UCookOnTheFlyServer::MarkPackageDirtyForCookerFromSchedulerThread(const FNa
 		{
 			PackageData->SendToState(UE::Cook::EPackageState::Request, UE::Cook::ESendFlags::QueueAddAndRemove);
 		}
-		else if (IsCookByTheBookRunning() && bHadCookedPlatforms)
+		else if (IsCookByTheBookMode() && IsInSession() && bHadCookedPlatforms)
 		{
 			PackageData->UpdateRequestData(PlatformManager->GetSessionPlatforms(), false /* bIsUrgent */, UE::Cook::FCompletionCallback(),
 				UE::Cook::FInstigator(UE::Cook::EInstigator::Unspecified, GInstigatorMarkPackageDirty));
@@ -3809,7 +3783,7 @@ void UCookOnTheFlyServer::MarkPackageDirtyForCookerFromSchedulerThread(const FNa
 
 bool UCookOnTheFlyServer::IsInSession() const
 {
-	return IsCookByTheBookRunning() || (IsCookOnTheFlyMode() && PlatformManager->HasSelectedSessionPlatforms());
+	return bSessionRunning;
 }
 
 void UCookOnTheFlyServer::ShutdownCookOnTheFly()
@@ -4489,7 +4463,7 @@ void FSaveCookedPackageContext::FinishPlatform()
 		// TODO: Reenable BuildDefinitionList once FCbPackage support for empty FCbObjects is in
 		//Info.Attachments.Add({ "BuildDefinitionList", BuildDefinitionList });
 		Info.WriteOptions = IPackageWriter::EWriteOptions::Write;
-		if (COTFS.CookByTheBookOptions != nullptr)
+		if (COTFS.IsCookByTheBookMode())
 		{
 			Info.WriteOptions |= IPackageWriter::EWriteOptions::ComputeHash;
 		}
@@ -4498,7 +4472,7 @@ void FSaveCookedPackageContext::FinishPlatform()
 	}
 
 	// Update asset registry
-	if (COTFS.CookByTheBookOptions)
+	if (COTFS.IsCookByTheBookMode())
 	{
 		Generator.UpdateAssetRegistryPackageData(*Package, SavePackageResult, MoveTemp(*ArchiveCookContext.GetCookTagList()));
 	}
@@ -4552,7 +4526,7 @@ void FSaveCookedPackageContext::FinishPackage()
 {
 	// Add soft references discovered from the package
 	FName PackageFName = Package->GetFName();
-	if (!COTFS.IsCookByTheBookMode() || !COTFS.CookByTheBookOptions->bSkipSoftReferences)
+	if (!COTFS.CookByTheBookOptions->bSkipSoftReferences)
 	{
 		// Also request any localized variants of this package
 		if (COTFS.IsCookByTheBookMode() && !FPackageName::IsLocalizedPackage(Package->GetName()))
@@ -4635,9 +4609,11 @@ void UCookOnTheFlyServer::Initialize( ECookMode::Type DesiredCookMode, ECookInit
 	ExternalRequests = MakeUnique<UE::Cook::FExternalRequests>();
 	PackageTracker = MakeUnique<UE::Cook::FPackageTracker>(*PackageDatas.Get());
 	DiffModeHelper = MakeUnique<FDiffModeCookServerUtils>();
+	BuildDefinitions = MakeUnique<UE::Cook::FBuildDefinitions>();
+	CookByTheBookOptions = MakeUnique<UE::Cook::FCookByTheBookOptions>();
+	CookOnTheFlyOptions = MakeUnique<UE::Cook::FCookOnTheFlyOptions>();
 	AssetRegistry = IAssetRegistry::Get();
 
-	BuildDefinitions.Reset(new UE::Cook::FBuildDefinitions());
 	UE::Cook::InitializeTls();
 	UE::Cook::FPlatformManager::InitializeTls();
 
@@ -4709,31 +4685,6 @@ void UCookOnTheFlyServer::Initialize( ECookMode::Type DesiredCookMode, ECookInit
 	ReadMemorySetting(TEXT("MemoryMinFreeVirtual"), MemoryMinFreeVirtual);
 	ReadMemorySetting(TEXT("MemoryMinFreePhysical"), MemoryMinFreePhysical);
 
-	uint64 MaxMemoryAllowance;
-	if (ReadMemorySetting(TEXT("MaxMemoryAllowance"), MaxMemoryAllowance))
-	{ 
-		UE_LOG(LogCook, Warning, TEXT("CookSettings.MaxMemoryAllowance is deprecated. Use CookSettings.MemoryMaxUsedPhysical instead."));
-		MemoryMaxUsedPhysical = MaxMemoryAllowance;
-	}
-	uint64 MinMemoryBeforeGC;
-	if (ReadMemorySetting(TEXT("MinMemoryBeforeGC"), MinMemoryBeforeGC))
-	{
-		UE_LOG(LogCook, Warning, TEXT("CookSettings.MinMemoryBeforeGC is deprecated. Use CookSettings.MemoryMaxUsedVirtual instead."));
-		MemoryMaxUsedVirtual = MinMemoryBeforeGC;
-	}
-	uint64 MinFreeMemory;
-	if (ReadMemorySetting(TEXT("MinFreeMemory"), MinFreeMemory))
-	{
-		UE_LOG(LogCook, Warning, TEXT("CookSettings.MinFreeMemory is deprecated. Use CookSettings.MemoryMinFreePhysical instead."));
-		MemoryMinFreePhysical = MinFreeMemory;
-	}
-	uint64 MinReservedMemory;
-	if (ReadMemorySetting(TEXT("MinReservedMemory"), MinReservedMemory))
-	{
-		UE_LOG(LogCook, Warning, TEXT("CookSettings.MinReservedMemory is deprecated. Use CookSettings.MemoryMinFreePhysical instead."));
-		MemoryMinFreePhysical = MinReservedMemory;
-	}
-	
 	MaxPreloadAllocated = 16;
 	DesiredSaveQueueLength = 8;
 	DesiredLoadQueueLength = 8;
@@ -4787,7 +4738,6 @@ void UCookOnTheFlyServer::Initialize( ECookMode::Type DesiredCookMode, ECookInit
 
 	if (IsCookByTheBookMode())
 	{
-		CookByTheBookOptions = new FCookByTheBookOptions();
 		for (TObjectIterator<UPackage> It; It; ++It)
 		{
 			if ((*It) != GetTransientPackage())
@@ -4799,8 +4749,6 @@ void UCookOnTheFlyServer::Initialize( ECookMode::Type DesiredCookMode, ECookInit
 		bHybridIterativeDebug = FParse::Param(FCommandLine::Get(), TEXT("hybriditerativedebug"));
 	}
 	
-	UE_LOG(LogCook, Display, TEXT("Mobile HDR setting %d"), IsMobileHDR());
-
 	// See if there are any plugins that need to be remapped for the sandbox
 	const FProjectDescriptor* Project = IProjectManager::Get().GetCurrentProject();
 	if (Project != nullptr)
@@ -5556,13 +5504,14 @@ TMap<FName, FString> UCookOnTheFlyServer::CalculateCookSettingStrings() const
 	const FName NAME_CookMode(TEXT("CookMode"));
 
 	CookSettingStrings.Add(FName(TEXT("Version")), TEXT("C7C76F79"));
-	if (CookByTheBookOptions)
+	if (IsCookByTheBookMode())
 	{
 		CookSettingStrings.Add(NAME_CookMode, TEXT("CookByTheBook"));
 		CookSettingStrings.Add(FName(TEXT("DLCName")), CookByTheBookOptions->DlcName);
 	}
 	else
 	{
+		check(IsCookOnTheFlyMode());
 		CookSettingStrings.Add(NAME_CookMode, TEXT("CookOnTheFly"));
 	}
 	return CookSettingStrings;
@@ -6028,7 +5977,7 @@ void UCookOnTheFlyServer::BlockOnAssetRegistry()
 	UE_LOG(LogCook, Display, TEXT("Waiting for Asset Registry"));
 	if (ShouldPopulateFullAssetRegistry())
 	{
-		// Trigger or waitfor completion the primary AssetRegistry scan.
+		// Trigger or wait for completion the primary AssetRegistry scan.
 		// Additionally scan any cook-specific paths from ini
 		TArray<FString> ScanPaths;
 		GConfig->GetArray(TEXT("AssetRegistry"), TEXT("PathsToScanForCook"), ScanPaths, GEngineIni);
@@ -6078,7 +6027,7 @@ void UCookOnTheFlyServer::RefreshPlatformAssetRegistries(const TArrayView<const
 		// if we are cooking DLC, we will just spend a lot of time removing the shipped packages from the AR,
 		// so we don't bother copying them over. can easily save 10 seconds on a large project
 		bool bInitalizeFromExisting = !IsCookingDLC();
-		RegistryGenerator->Initialize(CookByTheBookOptions ? CookByTheBookOptions->StartupPackages : TArray<FName>(), bInitalizeFromExisting);
+		RegistryGenerator->Initialize(CookByTheBookOptions->StartupPackages, bInitalizeFromExisting);
 	}
 }
 
@@ -6611,7 +6560,7 @@ void UCookOnTheFlyServer::GetGameDefaultObjects(const TArray<ITargetPlatform*>&
 
 bool UCookOnTheFlyServer::IsCookByTheBookRunning() const
 {
-	return CookByTheBookOptions && CookByTheBookOptions->bRunning;
+	return IsCookByTheBookMode() && IsInSession();
 }
 
 
@@ -6761,7 +6710,7 @@ void UCookOnTheFlyServer::BeginCookStartShaderCodeLibrary(FBeginCookContext& Beg
 
 void UCookOnTheFlyServer::BeginCookFinishShaderCodeLibrary(FBeginCookContext& BeginContext)
 {
-	check(BeginContext.StartupOptions && CookByTheBookOptions); // CookByTheBook only for now
+	check(BeginContext.StartupOptions && IsCookByTheBookMode()); // CookByTheBook only for now
 	// don't resave the global shader map files in dlc
 	if (!IsCookingDLC() && !EnumHasAnyFlags(BeginContext.StartupOptions->CookOptions, ECookByTheBookOptions::ForceDisableSaveGlobalShaders))
 	{
@@ -7053,9 +7002,9 @@ void UCookOnTheFlyServer::CookByTheBookFinished()
 {
 	using namespace UE::Cook;
 
-	check( IsInGameThread() );
-	check( IsCookByTheBookMode() );
-	check( CookByTheBookOptions->bRunning == true );
+	check(IsInGameThread());
+	check(IsCookByTheBookMode());
+	check(IsInSession());
 	check(PackageDatas->GetRequestQueue().IsEmpty());
 	check(PackageDatas->GetLoadPrepareQueue().IsEmpty());
 	check(PackageDatas->GetLoadReadyQueue().IsEmpty());
@@ -7294,7 +7243,6 @@ void UCookOnTheFlyServer::CookByTheBookFinished()
 
 void UCookOnTheFlyServer::ShutdownCookSession()
 {
-	PlatformManager->ClearSessionPlatforms();
 	for (UE::Cook::FPackageData* PackageData : *PackageDatas)
 	{
 		PackageData->DestroyGeneratorPackage();
@@ -7302,12 +7250,10 @@ void UCookOnTheFlyServer::ShutdownCookSession()
 
 	ConditionalUninstallImportBehaviorCallback();
 
+	bCancelSession = false;
 	if (IsCookByTheBookMode())
 	{
-		check(CookByTheBookOptions);
-		CookByTheBookOptions->BasedOnReleaseCookedPackages.Empty();
-		CookByTheBookOptions->bRunning = false;
-		CookByTheBookOptions->bFullLoadAndSave = false;
+		CookByTheBookOptions->ClearSessionData();
 
 		UnregisterCookByTheBookDelegates();
 
@@ -7316,6 +7262,7 @@ void UCookOnTheFlyServer::ShutdownCookSession()
 		OutputHierarchyTimers();
 		ClearHierarchyTimers();
 	}
+	PlatformManager->ClearSessionPlatforms(*this);
 }
 
 void UCookOnTheFlyServer::PrintFinishStats()
@@ -7387,22 +7334,19 @@ void UCookOnTheFlyServer::WriteMapDependencyGraph(const ITargetPlatform* TargetP
 
 void UCookOnTheFlyServer::QueueCancelCookByTheBook()
 {
-	if ( IsCookByTheBookMode() )
+	if (IsCookByTheBookMode() && IsInSession())
 	{
-		check( CookByTheBookOptions != NULL );
-		CookByTheBookOptions->bCancel = true;
+		bCancelSession = true;
 	}
 }
 
 void UCookOnTheFlyServer::CancelCookByTheBook()
 {
-	if ( IsCookByTheBookMode() && CookByTheBookOptions->bRunning )
+	check(IsCookByTheBookMode());
+	check(IsInGameThread());
+	if (IsInSession())
 	{
-		check(CookByTheBookOptions);
-		check( IsInGameThread() );
-
 		CancelAllQueues();
-
 		ShutdownCookSession();
 	} 
 }
@@ -7411,7 +7355,6 @@ void UCookOnTheFlyServer::StopAndClearCookedData()
 {
 	if ( IsCookByTheBookMode() )
 	{
-		check( CookByTheBookOptions != NULL );
 		CancelCookByTheBook();
 	}
 	else
@@ -7979,13 +7922,12 @@ void UCookOnTheFlyServer::StartCookByTheBook( const FCookByTheBookStartupOptions
 
 	// Initialize systems and settings that the rest of StartCookByTheBook depends on
 	// Functions in this section are ordered and can depend on the functions before them
-	FBeginCookContext BeginContext = SetCookByTheBookOptions(CookByTheBookStartupOptions);
+	FBeginCookContext BeginContext = CreateBeginCookByTheBookContext(CookByTheBookStartupOptions);
 	BlockOnAssetRegistry();
 	CreateSandboxFile(BeginContext);
 	SetBeginCookConfigSettings(BeginContext);
 	SelectSessionPlatforms(BeginContext);
 	SetBeginCookIterativeFlags(BeginContext);
-	CookByTheBookOptions->bRunning = true;
 
 	// Initialize systems referenced by later stages or that need to start early for async performance
 	// Functions in this section must not need to read/write the SandboxDirectory or MemoryCookedPackages
@@ -8014,13 +7956,13 @@ void UCookOnTheFlyServer::StartCookByTheBook( const FCookByTheBookStartupOptions
 	RegisterCookByTheBookDelegates();
 }
 
-FBeginCookContext UCookOnTheFlyServer::SetCookByTheBookOptions(const FCookByTheBookStartupOptions& StartupOptions)
+FBeginCookContext UCookOnTheFlyServer::CreateBeginCookByTheBookContext(const FCookByTheBookStartupOptions& StartupOptions)
 {
 	FBeginCookContext BeginContext;
 
 	BeginContext.StartupOptions = &StartupOptions;
 	const ECookByTheBookOptions& CookOptions = StartupOptions.CookOptions;
-	CookByTheBookOptions->bCancel = false;
+	bZenStore = !!(CookOptions & ECookByTheBookOptions::ZenStore);
 	CookByTheBookOptions->CookTime = 0.0f;
 	CookByTheBookOptions->CookStartTime = FPlatformTime::Seconds();
 	CookByTheBookOptions->bGenerateStreamingInstallManifests = StartupOptions.bGenerateStreamingInstallManifests;
@@ -8028,8 +7970,7 @@ FBeginCookContext UCookOnTheFlyServer::SetCookByTheBookOptions(const FCookByTheB
 	CookByTheBookOptions->CreateReleaseVersion = StartupOptions.CreateReleaseVersion;
 	CookByTheBookOptions->bSkipHardReferences = !!(CookOptions & ECookByTheBookOptions::SkipHardReferences);
 	CookByTheBookOptions->bSkipSoftReferences = !!(CookOptions & ECookByTheBookOptions::SkipSoftReferences);
-	CookByTheBookOptions->bFullLoadAndSave = !!(CookOptions & ECookByTheBookOptions::FullLoadAndSave);
-	CookByTheBookOptions->bZenStore = !!(CookOptions & ECookByTheBookOptions::ZenStore);
+	CookByTheBookOptions->bFullLoadAndSave = !!(CookOptions & ECookByTheBookOptions::FullLoadAndSave) && !IsCookingInEditor();
 	CookByTheBookOptions->bCookAgainstFixedBase = !!(CookOptions & ECookByTheBookOptions::CookAgainstFixedBase);
 	CookByTheBookOptions->bDlcLoadMainAssetRegistry = !!(CookOptions & ECookByTheBookOptions::DlcLoadMainAssetRegistry);
 	CookByTheBookOptions->bErrorOnEngineContentUse = StartupOptions.bErrorOnEngineContentUse;
@@ -8054,8 +7995,11 @@ FBeginCookContext UCookOnTheFlyServer::SetCookByTheBookOptions(const FCookByTheB
 	return BeginContext;
 }
 
-FBeginCookContext UCookOnTheFlyServer::CreateBeginCookOnTheFlyContext()
+FBeginCookContext UCookOnTheFlyServer::CreateBeginCookOnTheFlyContext(const FCookOnTheFlyStartupOptions& Options)
 {
+	bZenStore = Options.bZenStore;
+	CookOnTheFlyOptions->bBindAnyPort = Options.bBindAnyPort;
+	CookOnTheFlyOptions->bPlatformProtocol = Options.bPlatformProtocol;
 	return FBeginCookContext();
 }
 
@@ -8348,7 +8292,7 @@ void UCookOnTheFlyServer::BeginCookPackageWriters(FBeginCookContext& BeginContex
 
 void UCookOnTheFlyServer::SelectSessionPlatforms(FBeginCookContext& BeginContext)
 {
-	PlatformManager->SelectSessionPlatforms(BeginContext.TargetPlatforms);
+	PlatformManager->SelectSessionPlatforms(*this, BeginContext.TargetPlatforms);
 	if (PackageTracker->HasBeenConsumed())
 	{
 		bPackageFilterDirty = true;
@@ -8520,17 +8464,17 @@ void UCookOnTheFlyServer::MaybeMarkPackageAsAlreadyLoaded(UPackage *Package)
 	{
 		bShouldMarkAsAlreadyProcessed = PackageData->HasAllCookedPlatforms(PlatformManager->GetSessionPlatforms(), true /* bIncludeFailed */);
 
-		FString Platforms;
-		for (const TPair<const ITargetPlatform*, UE::Cook::FPackageData::FPlatformData>& Pair : PackageData->GetPlatformDatas())
+		if (IsCookFlagSet(ECookInitializationFlags::LogDebugInfo))
 		{
-			if (Pair.Value.bCookAttempted)
+			FString Platforms;
+			for (const TPair<const ITargetPlatform*, UE::Cook::FPackageData::FPlatformData>& Pair : PackageData->GetPlatformDatas())
 			{
-				Platforms += TEXT(" ");
-				Platforms += Pair.Key->PlatformName();
+				if (Pair.Value.bCookAttempted)
+				{
+					Platforms += TEXT(" ");
+					Platforms += Pair.Key->PlatformName();
+				}
 			}
-		}
-		if (IsCookFlagSet(ECookInitializationFlags::LogDebugInfo))
-		{
 			if (!bShouldMarkAsAlreadyProcessed)
 			{
 				UE_LOG(LogCook, Display, TEXT("Reloading package %s slowly because it wasn't cooked for all platforms %s."), *StandardName.ToString(), *Platforms);
@@ -8726,7 +8670,7 @@ bool UCookOnTheFlyServer::GetAllPackageFilenamesFromAssetRegistry(const FString&
 					}
 					else
 					{
-						const bool bContainsMap = AssetDatas[AssetIndex]->PackageFlags & PKG_ContainsMap;
+						const bool bContainsMap = !!(AssetDatas[AssetIndex]->PackageFlags & PKG_ContainsMap);
 						PackageFileName = FPackageDatas::LookupFileNameOnDisk(PackageName,
 							false /* bRequireExists */, bContainsMap);
 						if (!PackageFileName.IsNone())
@@ -8747,11 +8691,16 @@ bool UCookOnTheFlyServer::GetAllPackageFilenamesFromAssetRegistry(const FString&
 	return false;
 }
 
-uint32 UCookOnTheFlyServer::FullLoadAndSave(uint32& CookedPackageCount)
+bool UCookOnTheFlyServer::IsFullLoadAndSave() const
+{
+	return CookByTheBookOptions->bFullLoadAndSave;
+}
+
+uint32 UCookOnTheFlyServer::CookFullLoadAndSave()
 {
 	UE_SCOPED_HIERARCHICAL_COOKTIMER(FullLoadAndSave);
-	check(CurrentCookMode == ECookMode::CookByTheBook);
-	check(CookByTheBookOptions);
+	check(IsCookByTheBookMode());
+	check(!IsCookingInEditor());
 	check(IsInGameThread());
 
 	uint32 Result = 0;
@@ -8980,7 +8929,7 @@ uint32 UCookOnTheFlyServer::FullLoadAndSave(uint32& CookedPackageCount)
 					}
 				}
 
-				if (!IsCookByTheBookMode() || !CookByTheBookOptions->bSkipSoftReferences)
+				if (!CookByTheBookOptions->bSkipSoftReferences)
 				{
 					UE_SCOPED_HIERARCHICAL_COOKTIMER(ResolveStringReferences);
 					TSet<FName> StringAssetPackages;
@@ -9294,7 +9243,6 @@ uint32 UCookOnTheFlyServer::FullLoadAndSave(uint32& CookedPackageCount)
 			GIsSavingPackage = false;
 		}
 
-		CookedPackageCount += ParallelSavedPackages;
 		if (ParallelSavedPackages > 0)
 		{
 			Result |= COSR_CookedPackage;
@@ -9323,6 +9271,7 @@ uint32 UCookOnTheFlyServer::FullLoadAndSave(uint32& CookedPackageCount)
 		}
 	}
 
+	CookByTheBookFinished();
 	return Result;
 }
 
@@ -9434,10 +9383,6 @@ void UCookOnTheFlyServer::PropertyImportBehaviorCallback(const FObjectImport& Im
 
 void UCookOnTheFlyServer::GenerateLocalizationReferences(TConstArrayView<FString> CookCultures)
 {
-	if (!CookByTheBookOptions)
-	{
-		return;
-	}
 	CookByTheBookOptions->SourceToLocalizedPackageVariants.Reset();
 	CookByTheBookOptions->AllCulturesToCook.Reset();
 
@@ -9490,11 +9435,6 @@ void UCookOnTheFlyServer::GenerateLocalizationReferences(TConstArrayView<FString
 
 void UCookOnTheFlyServer::RegisterLocalizationChunkDataGenerator()
 {
-	if (!CookByTheBookOptions)
-	{
-		return;
-	}
-
 	// Get the list of localization targets to chunk, and remove any targets that we've been asked not to stage
 	const UProjectPackagingSettings* const PackagingSettings = GetDefault<UProjectPackagingSettings>();
 	TArray<FString> LocalizationTargetsToChunk = PackagingSettings->LocalizationTargetsToChunk;
diff --git a/Engine/Source/Editor/UnrealEd/Private/Cooker/CookPlatformManager.cpp b/Engine/Source/Editor/UnrealEd/Private/Cooker/CookPlatformManager.cpp
--- a/Engine/Source/Editor/UnrealEd/Private/Cooker/CookPlatformManager.cpp
+++ b/Engine/Source/Editor/UnrealEd/Private/Cooker/CookPlatformManager.cpp
@@ -101,12 +101,6 @@ namespace Cook
 		return SessionPlatforms;
 	}
 
-	bool FPlatformManager::HasSelectedSessionPlatforms() const
-	{
-		FRWScopeLockConditional ScopeLock(SessionLock, FRWScopeLockType::SLT_ReadOnly, !IsSchedulerThread());
-		return bHasSelectedSessionPlatforms;
-	}
-
 	bool FPlatformManager::HasSessionPlatform(FPlatformId TargetPlatform) const
 	{
 		const bool bIsSchedulerThread = IsSchedulerThread();
@@ -115,7 +109,7 @@ namespace Cook
 		return SessionPlatforms.Contains(TargetPlatform);
 	}
 
-	void FPlatformManager::SelectSessionPlatforms(const TArrayView<FPlatformId const>& TargetPlatforms)
+	void FPlatformManager::SelectSessionPlatforms(UCookOnTheFlyServer& COTFS, const TArrayView<FPlatformId const>& TargetPlatforms)
 	{
 		checkf(IsSchedulerThread(), TEXT("Writing to SessionPlatforms is only allowed from the scheduler thread."));
 		for (FPlatformId TargetPlatform : TargetPlatforms)
@@ -126,18 +120,20 @@ namespace Cook
 		FRWScopeLock ScopeLock(SessionLock, FRWScopeLockType::SLT_Write);
 		SessionPlatforms.Empty(TargetPlatforms.Num());
 		SessionPlatforms.Append(TargetPlatforms.GetData(), TargetPlatforms.Num());
+		COTFS.bSessionRunning = true;
 		bHasSelectedSessionPlatforms = true;
 	}
 
-	void FPlatformManager::ClearSessionPlatforms()
+	void FPlatformManager::ClearSessionPlatforms(UCookOnTheFlyServer& COTFS)
 	{
 		checkf(IsSchedulerThread(), TEXT("Writing to SessionPlatforms is only allowed from the scheduler thread."));
 		FRWScopeLock ScopeLock(SessionLock, FRWScopeLockType::SLT_Write);
 		SessionPlatforms.Empty();
+		COTFS.bSessionRunning = false;
 		bHasSelectedSessionPlatforms = false;
 	}
 
-	void FPlatformManager::AddSessionPlatform(FPlatformId TargetPlatform)
+	void FPlatformManager::AddSessionPlatform(UCookOnTheFlyServer& COTFS, FPlatformId TargetPlatform)
 	{
 		checkf(IsSchedulerThread(), TEXT("Writing to SessionPlatforms is only allowed from the scheduler thread."));
 		CreatePlatformData(TargetPlatform);
@@ -148,6 +144,7 @@ namespace Cook
 			return;
 		}
 		SessionPlatforms.Add(TargetPlatform);
+		COTFS.bSessionRunning = true;
 		bHasSelectedSessionPlatforms = true;
 	}
 
@@ -264,6 +261,7 @@ namespace Cook
 					SessionPlatforms.Remove(TargetPlatform);
 					if (SessionPlatforms.Num() == 0)
 					{
+						CookOnTheFlyServer.bSessionRunning = false;
 						bHasSelectedSessionPlatforms = false;
 					}
 				}
@@ -286,7 +284,7 @@ namespace Cook
 					const ITargetPlatform* TargetPlatform = GetTargetPlatformManager()->FindTargetPlatform(PlatformName.ToString());
 					if (TargetPlatform)
 					{
-						AddSessionPlatform(TargetPlatform);
+						AddSessionPlatform(*LocalCookOnTheFlyServer, TargetPlatform);
 						LocalCookOnTheFlyServer->bPackageFilterDirty = true;
 					}
 				});
diff --git a/Engine/Source/Editor/UnrealEd/Private/Cooker/CookPlatformManager.h b/Engine/Source/Editor/UnrealEd/Private/Cooker/CookPlatformManager.h
--- a/Engine/Source/Editor/UnrealEd/Private/Cooker/CookPlatformManager.h
+++ b/Engine/Source/Editor/UnrealEd/Private/Cooker/CookPlatformManager.h
@@ -75,20 +75,17 @@ namespace Cook
 		/** Returns the set of TargetPlatforms that is active for the CurrentCookByTheBook session or CookOnTheFly request. */
 		const TArray<const ITargetPlatform*>& GetSessionPlatforms() const;
 
-		/** Return whether the platforms have been selected for the current session (may be empty if current session is a null session). */
-		bool HasSelectedSessionPlatforms() const;
-
 		/** Return whether the given platform is already in the list of platforms for the current session. */
 		bool HasSessionPlatform(FPlatformId TargetPlatform) const;
 
 		/** Specify the set of TargetPlatforms to use for the currently-initializing CookByTheBook session or CookOnTheFly request. */
-		void SelectSessionPlatforms(const TArrayView<FPlatformId const>& TargetPlatforms);
+		void SelectSessionPlatforms(UCookOnTheFlyServer& COTFS, const TArrayView<FPlatformId const>& TargetPlatforms);
 
 		/** Mark that the list of TargetPlatforms for the session has no longer been set; it will be an error to try to read them until SelectSessionPlatforms is called. */
-		void ClearSessionPlatforms();
+		void ClearSessionPlatforms(UCookOnTheFlyServer& COTFS);
 
 		/** Add @param TargetPlatform to the session platforms if not already present. */
-		void AddSessionPlatform(FPlatformId TargetPlatform);
+		void AddSessionPlatform(UCookOnTheFlyServer& COTFS, FPlatformId TargetPlatform);
 
 		/**
 		 * Get The PlatformData for the given Platform.
@@ -177,7 +174,7 @@ namespace Cook
 		
 		bool bArePlatformsPrepopulated = false;
 
-		/** It is invalid to attempt to cook if session platforms have not been selected. */
+		/** Duplicate of COTFS.bSessionRunning. Used for assertions on other threads. It is invalid to attempt to cook if session platforms have not been selected. */
 		bool bHasSelectedSessionPlatforms = false;
 	};
 
diff --git a/Engine/Source/Editor/UnrealEd/Private/Cooker/CookRequestCluster.cpp b/Engine/Source/Editor/UnrealEd/Private/Cooker/CookRequestCluster.cpp
--- a/Engine/Source/Editor/UnrealEd/Private/Cooker/CookRequestCluster.cpp
+++ b/Engine/Source/Editor/UnrealEd/Private/Cooker/CookRequestCluster.cpp
@@ -129,16 +129,13 @@ void FRequestCluster::AddClusters(UCookOnTheFlyServer& COTFS, FPackageDataSet& U
 	};
 
 	bool bErrorOnEngineContentUse = false;
-	UCookOnTheFlyServer::FCookByTheBookOptions* Options = COTFS.CookByTheBookOptions;
+	UE::Cook::FCookByTheBookOptions& Options = *COTFS.CookByTheBookOptions;
 	FString DLCPath;
-	if (Options)
+	bErrorOnEngineContentUse = Options.bErrorOnEngineContentUse;
+	if (bErrorOnEngineContentUse)
 	{
-		bErrorOnEngineContentUse = Options->bErrorOnEngineContentUse;
-		if (bErrorOnEngineContentUse)
-		{
-			DLCPath = FPaths::Combine(*COTFS.GetBaseDirectoryForDLC(), TEXT("Content"));
-			FPaths::MakeStandardFilename(DLCPath);
-		}
+		DLCPath = FPaths::Combine(*COTFS.GetBaseDirectoryForDLC(), TEXT("Content"));
+		FPaths::MakeStandardFilename(DLCPath);
 	}
 
 	UE::Cook::FRequestCluster* MRUCluster = nullptr;
@@ -189,12 +186,12 @@ FRequestCluster::FFileNameRequest::FFileNameRequest(FFilePlatformRequest&& FileR
 
 void FRequestCluster::Initialize(UCookOnTheFlyServer& COTFS)
 {
-	UCookOnTheFlyServer::FCookByTheBookOptions* Options = COTFS.CookByTheBookOptions;
-	if (Options)
+	if (!COTFS.IsCookOnTheFlyMode())
 	{
-		bAllowHardDependencies = !Options->bSkipHardReferences;
-		bAllowSoftDependencies = !Options->bSkipSoftReferences;
-		bErrorOnEngineContentUse = Options->bErrorOnEngineContentUse;
+		UE::Cook::FCookByTheBookOptions& Options = *COTFS.CookByTheBookOptions;
+		bAllowHardDependencies = !Options.bSkipHardReferences;
+		bAllowSoftDependencies = !Options.bSkipSoftReferences;
+		bErrorOnEngineContentUse = Options.bErrorOnEngineContentUse;
 	}
 	else
 	{
@@ -284,7 +281,7 @@ void FRequestCluster::FetchPackageNames(const FCookerTimer& CookerTimer, bool& b
 	}
 	for (;NextRequest < InRequestsNum; ++NextRequest)
 	{
-		if (NextRequest % TimerCheckPeriod == 0 && CookerTimer.IsTimeUp())
+		if ((NextRequest+1) % TimerCheckPeriod == 0 && CookerTimer.IsTimeUp())
 		{
 			break;
 		}
diff --git a/Engine/Source/Editor/UnrealEd/Private/Cooker/CookTypes.h b/Engine/Source/Editor/UnrealEd/Private/Cooker/CookTypes.h
--- a/Engine/Source/Editor/UnrealEd/Private/Cooker/CookTypes.h
+++ b/Engine/Source/Editor/UnrealEd/Private/Cooker/CookTypes.h
@@ -157,18 +157,21 @@ namespace UE::Cook
 		static float ForeverTimeSlice;
 	};
 
-	/** Fields from TickCookOnTheSide that helper functions will need access to. */
+	/** Temporary-lifetime data about the current tick of the cooker. */
 	struct FTickStackData
 	{
-		/** The return value from TickCookOnTheSide; it is a bitmask of flags in enum ECookOnTheSideResult */
+		/** A bitmask of flags of type enum ECookOnTheSideResult that were set during the tick. */
 		uint32 ResultFlags = 0;
-		/** The number of packages that have been cooked in the current call to TickCookOnTheSide. It is used for e.g. decisions about when to collect garbage. */
+		/** The number of packages that have been cooked in the current tick. It is used for e.g. decisions about when to collect garbage. */
 		uint32 CookedPackageCount = 0;
-		/** The CookerTimer for the current TickCookOnTheSide request. Used by slow reentrant operations that need to check whether they have timed out. */
+		/** The CookerTimer for the current tick. Used by slow reentrant operations that need to check whether they have timed out. */
 		FCookerTimer Timer;
-		/** CookFlags describing details of the caller's desired behavior of the current TickCookOnTheSide request. */
+		/** CookFlags describing details of the caller's desired behavior for the current tick. */
 		ECookTickFlags TickFlags;
 
+		bool bCookComplete = false;
+		bool bCookCancelled = false;
+
 		explicit FTickStackData(const float& TimeSlice, const bool bIsRealtimeMode, ECookTickFlags InTickFlags, int32 InMaxNumPackagesToSave)
 			:Timer(TimeSlice, bIsRealtimeMode, InMaxNumPackagesToSave), TickFlags(InTickFlags)
 		{
@@ -235,21 +238,16 @@ inline bool RouteIsCachedCookedPlatformDataLoaded(UObject* Obj, const ITargetPla
 //////////////////////////////////////////////////////////////////////////
 // Cook by the book options
 
-struct UCookOnTheFlyServer::FCookByTheBookOptions
+namespace UE::Cook
 {
-public:
-	/** Should we generate streaming install manifests (only valid option in cook by the book) */
-	bool							bGenerateStreamingInstallManifests = false;
 
-	/** Should we generate a seperate manifest for map dependencies */
-	bool							bGenerateDependenciesForMaps = false;
-
-	/** Is cook by the book currently running */
-	bool							bRunning = false;
-
-	/** Cancel has been queued will be processed next tick */
-	bool							bCancel = false;
+struct FCookByTheBookOptions
+{
+public:
+	// Process-lifetime variables
+	TArray<FName>					StartupPackages;
 
+	// Session-lifetime variables
 	/** DlcName setup if we are cooking dlc will be used as the directory to save cooked files to */
 	FString							DlcName;
 
@@ -257,28 +255,53 @@ public:
 	FString							CreateReleaseVersion;
 
 	/** If we are based on a release version of the game this is the set of packages which were cooked in that release. Map from platform name to list of uncooked package filenames */
-	TMap<FName, TArray<FName>>			BasedOnReleaseCookedPackages;
+	TMap<FName, TArray<FName>>		BasedOnReleaseCookedPackages;
+
+	/** Mapping from source packages to their localized variants (based on the culture list in FCookByTheBookStartupOptions) */
+	TMap<FName, TArray<FName>>		SourceToLocalizedPackageVariants;
+	/** List of all the cultures (e.g. "en") that need to be cooked */
+	TArray<FString>					AllCulturesToCook;
 
 	/** Timing information about cook by the book */
 	double							CookTime = 0.0;
 	double							CookStartTime = 0.0;
 
+	/** Should we generate streaming install manifests (only valid option in cook by the book) */
+	bool							bGenerateStreamingInstallManifests = false;
+
+	/** Should we generate a seperate manifest for map dependencies */
+	bool							bGenerateDependenciesForMaps = false;
+
 	/** error when detecting engine content being used in this cook */
 	bool							bErrorOnEngineContentUse = false;
 	bool							bSkipHardReferences = false;
 	bool							bSkipSoftReferences = false;
 	bool							bFullLoadAndSave = false;
-	bool							bZenStore = false;
 	bool							bCookAgainstFixedBase = false;
 	bool							bDlcLoadMainAssetRegistry = false;
-	TArray<FName>					StartupPackages;
 
-	/** Mapping from source packages to their localized variants (based on the culture list in FCookByTheBookStartupOptions) */
-	TMap<FName, TArray<FName>>		SourceToLocalizedPackageVariants;
-	/** List of all the cultures (e.g. "en") that need to be cooked */
-	TArray<FString>					AllCulturesToCook;
+	void ClearSessionData()
+	{
+		FCookByTheBookOptions EmptyOptions;
+		// Preserve Process-lifetime variables
+		EmptyOptions.StartupPackages = MoveTemp(StartupPackages);
+		*this = MoveTemp(EmptyOptions);
+	}
 };
 
+//////////////////////////////////////////////////////////////////////////
+// Cook on the fly startup options
+struct FCookOnTheFlyOptions
+{
+	/** Wether the network file server or the I/O store connection server should bind to any port */
+	bool bBindAnyPort = false;
+	/** Whether the network file server should use a platform-specific communication protocol instead of TCP (used when bZenStore == false) */
+	bool bPlatformProtocol = false;
+};
+
+}
+
+
 /** Helper struct for FBeginCookContext; holds the context data for each platform being cooked */
 struct FBeginCookContextPlatform
 {
diff --git a/Engine/Source/Editor/UnrealEd/Private/UnrealEdEngine.cpp b/Engine/Source/Editor/UnrealEd/Private/UnrealEdEngine.cpp
--- a/Engine/Source/Editor/UnrealEd/Private/UnrealEdEngine.cpp
+++ b/Engine/Source/Editor/UnrealEd/Private/UnrealEdEngine.cpp
@@ -234,7 +234,7 @@ void UUnrealEdEngine::Init(IEngineLoop* InEngineLoop)
 			CookServer = NewObject<UCookOnTheFlyServer>();
 			CookServer->Initialize(ECookMode::CookOnTheFlyFromTheEditor, BaseCookingFlags);
 
-			UCookOnTheFlyServer::FCookOnTheFlyOptions CookOnTheFlyStartupOptions;
+			UCookOnTheFlyServer::FCookOnTheFlyStartupOptions CookOnTheFlyStartupOptions;
 			CookOnTheFlyStartupOptions.bBindAnyPort = false;
 			CookOnTheFlyStartupOptions.bZenStore = GetDefault<UProjectPackagingSettings>()->bUseZenStore;
 			CookServer->StartCookOnTheFly(CookOnTheFlyStartupOptions);
