Unified Zen and legacy COTF network protocols.
*Reduces code duplication
*Enables the same ODSC flow to be used for both COTF variants
*The client will now autodetect if it should run in Zen mode or not

diff --git a/Engine/Source/Developer/CookOnTheFlyNetServer/CookOnTheFlyNetServer.Build.cs b/Engine/Source/Developer/CookOnTheFlyNetServer/CookOnTheFlyNetServer.Build.cs
new file mode 100644
--- /dev/null
+++ b/Engine/Source/Developer/CookOnTheFlyNetServer/CookOnTheFlyNetServer.Build.cs
@@ -0,0 +1,24 @@
+// Copyright Epic Games, Inc. All Rights Reserved.
+
+using UnrealBuildTool;
+
+public class CookOnTheFlyNetServer : ModuleRules
+{
+    public CookOnTheFlyNetServer(ReadOnlyTargetRules Target) : base(Target)
+    {
+        PrivateDependencyModuleNames.AddRange(
+                new string[]
+                {
+                    "TargetPlatform"
+                });
+
+        PublicDependencyModuleNames.AddRange(
+            new string[] {
+                "Core",
+                "Sockets",
+                "Networking",
+                "CookOnTheFly"
+            }
+        );
+    }
+}
diff --git a/Engine/Source/Developer/CookOnTheFlyNetServer/Internal/CookOnTheFlyNetServer.h b/Engine/Source/Developer/CookOnTheFlyNetServer/Internal/CookOnTheFlyNetServer.h
new file mode 100644
--- /dev/null
+++ b/Engine/Source/Developer/CookOnTheFlyNetServer/Internal/CookOnTheFlyNetServer.h
@@ -0,0 +1,80 @@
+// Copyright Epic Games, Inc. All Rights Reserved.
+
+#pragma once
+
+#include "CoreMinimal.h"
+#include "Modules/ModuleInterface.h"
+#include "Async/Future.h"
+#include "Misc/Timespan.h"
+#include "Serialization/MemoryReader.h"
+#include "Serialization/MemoryWriter.h"
+#include "CookOnTheFly.h"
+
+class FInternetAddr;
+
+DECLARE_LOG_CATEGORY_EXTERN(LogCookOnTheFlyNetworkServer, Log, All);
+
+namespace UE { namespace Cook
+{
+
+enum class ECookOnTheFlyNetworkServerProtocol : uint32
+{
+	Tcp,
+	Platform
+};
+
+struct FCookOnTheFlyNetworkServerOptions
+{
+	/* Server protocol*/
+	ECookOnTheFlyNetworkServerProtocol Protocol = ECookOnTheFlyNetworkServerProtocol::Tcp;
+
+	/* The port number to bind to (-1 = default port, 0 = any available port) */
+	int32 Port = INDEX_NONE;
+
+	/* Active target platform(s) */
+	TArray<ITargetPlatform*> TargetPlatforms;
+
+	/* Zen project name */
+	FString ZenProjectName;
+};
+
+class ICookOnTheFlyClientConnection
+{
+public:
+	virtual ~ICookOnTheFlyClientConnection() {}
+
+	virtual FName GetPlatformName() const = 0;
+	virtual const ITargetPlatform* GetTargetPlatform() const = 0;
+	virtual bool GetIsSingleThreaded() const = 0;
+	virtual bool SendMessage(const FCookOnTheFlyMessage& Message) = 0;
+};
+
+class ICookOnTheFlyNetworkServer
+{
+public:
+	virtual ~ICookOnTheFlyNetworkServer() {}
+
+	DECLARE_EVENT_OneParam(ICookOnTheFlyNetworkServer, FClientConnectionEvent, ICookOnTheFlyClientConnection&);
+	virtual FClientConnectionEvent& OnClientConnected() = 0;
+	virtual FClientConnectionEvent& OnClientDisconnected() = 0;
+
+	DECLARE_DELEGATE_RetVal_TwoParams(bool, FHandleRequestDelegate, ICookOnTheFlyClientConnection&, const FCookOnTheFlyRequest&);
+	virtual FHandleRequestDelegate& OnRequest(ECookOnTheFlyMessage MessageType) = 0;
+
+	virtual bool Start() = 0;
+	virtual bool IsReadyToAcceptConnections(void) const = 0;
+	virtual bool GetAddressList(TArray<TSharedPtr<FInternetAddr>>& OutAddresses) const = 0;
+	virtual FString GetSupportedProtocol() const = 0;
+	virtual int32 NumConnections() const = 0;
+};
+
+class ICookOnTheFlyNetworkServerModule
+	: public IModuleInterface
+{
+public:
+	virtual ~ICookOnTheFlyNetworkServerModule() { }
+
+	virtual TSharedPtr<ICookOnTheFlyNetworkServer> CreateServer(const FCookOnTheFlyNetworkServerOptions& Options) = 0;
+};
+
+}} // namespace UE::Cook
diff --git a/Engine/Source/Developer/CookOnTheFlyNetServer/Private/CookOnTheFlyNetServer.cpp b/Engine/Source/Developer/CookOnTheFlyNetServer/Private/CookOnTheFlyNetServer.cpp
new file mode 100644
--- /dev/null
+++ b/Engine/Source/Developer/CookOnTheFlyNetServer/Private/CookOnTheFlyNetServer.cpp
@@ -0,0 +1,31 @@
+// Copyright Epic Games, Inc. All Rights Reserved.
+
+#include "CookOnTheFlyNetServer.h"
+#include "Modules/ModuleManager.h"
+#include "TCPServer.h"
+#include "PlatformProtocolServer.h"
+
+DEFINE_LOG_CATEGORY(LogCookOnTheFlyNetworkServer);
+
+class FCookOnTheFlyNetworkServerModule final
+	: public UE::Cook::ICookOnTheFlyNetworkServerModule
+{
+public:
+	virtual TSharedPtr<UE::Cook::ICookOnTheFlyNetworkServer> CreateServer(const UE::Cook::FCookOnTheFlyNetworkServerOptions& Options) override
+	{
+		using namespace UE::Cook;
+
+		switch (Options.Protocol)
+		{
+		case ECookOnTheFlyNetworkServerProtocol::Tcp:
+			return MakeShared<FCookOnTheFlyServerTCP>(Options.Port, Options.TargetPlatforms, Options.ZenProjectName);
+		case ECookOnTheFlyNetworkServerProtocol::Platform:
+			return MakeShared<FCookOnTheFlyServerPlatformProtocol>(Options.TargetPlatforms, Options.ZenProjectName);
+		default:
+			UE_LOG(LogCookOnTheFlyNetworkServer, Fatal, TEXT("Unsupported protocol: %d"), Options.Protocol);
+			return nullptr;
+		}
+	}
+};
+
+IMPLEMENT_MODULE(FCookOnTheFlyNetworkServerModule, CookOnTheFlyNetServer);
diff --git a/Engine/Source/Developer/CookOnTheFlyNetServer/Private/CookOnTheFlyNetServerBase.cpp b/Engine/Source/Developer/CookOnTheFlyNetServer/Private/CookOnTheFlyNetServerBase.cpp
new file mode 100644
--- /dev/null
+++ b/Engine/Source/Developer/CookOnTheFlyNetServer/Private/CookOnTheFlyNetServerBase.cpp
@@ -0,0 +1,175 @@
+// Copyright Epic Games, Inc. All Rights Reserved.
+
+#include "CookOnTheFlyNetServerBase.h"
+#include "Serialization/BufferArchive.h"
+#include "Interfaces/ITargetPlatform.h"
+#include "Serialization/ArrayReader.h"
+#include "HAL/RunnableThread.h"
+
+FCookOnTheFlyClientConnectionBase::FCookOnTheFlyClientConnectionBase(FCookOnTheFlyNetworkServerBase& InOwner)
+	: Owner(InOwner)
+{
+}
+
+FCookOnTheFlyClientConnectionBase::~FCookOnTheFlyClientConnectionBase()
+{
+	if (WorkerThread)
+	{
+		delete WorkerThread;
+	}
+	if (bClientConnectedBroadcasted && Owner.ClientDisconnectedEvent.IsBound())
+	{
+		Owner.ClientDisconnectedEvent.Broadcast(*this);
+	}
+}
+
+bool FCookOnTheFlyClientConnectionBase::Initialize()
+{
+	FArrayReader HandshakeRequestPayload;
+	if (!ReceivePayload(HandshakeRequestPayload))
+	{
+		UE_LOG(LogCookOnTheFlyNetworkServer, Warning, TEXT("Failed to receive handshake payload"));
+		return false;
+	}
+
+	TArray<FString> TargetPlatformNames;
+	HandshakeRequestPayload << TargetPlatformNames;
+	HandshakeRequestPayload << bIsSingleThreaded;
+
+	bool bIsOk = true;
+	FString PlatformNameString;
+	if (!TargetPlatformNames.IsEmpty())
+	{
+		// figure out the best matching target platform for the set of valid ones
+		for (int32 TPIndex = 0; TPIndex < TargetPlatformNames.Num() && !TargetPlatform; TPIndex++)
+		{
+			UE_LOG(LogCookOnTheFlyNetworkServer, Verbose, TEXT("    Possible Target Platform from client: %s"), *TargetPlatformNames[TPIndex]);
+
+			// look for a matching target platform
+			for (int32 ActiveTPIndex = 0; ActiveTPIndex < Owner.ActiveTargetPlatforms.Num(); ActiveTPIndex++)
+			{
+				UE_LOG(LogCookOnTheFlyNetworkServer, Verbose, TEXT("   Checking against: %s"), *Owner.ActiveTargetPlatforms[ActiveTPIndex]->PlatformName());
+				if (Owner.ActiveTargetPlatforms[ActiveTPIndex]->PlatformName() == TargetPlatformNames[TPIndex])
+				{
+					TargetPlatform = Owner.ActiveTargetPlatforms[ActiveTPIndex];
+					PlatformNameString = TargetPlatform->PlatformName();
+					PlatformName = FName(PlatformNameString);
+					break;
+				}
+			}
+		}
+
+		// if we didn't find one, reject client and also print some warnings
+		if (!TargetPlatform)
+		{
+			// reject client we can't cook/compile shaders for you!
+			UE_LOG(LogCookOnTheFlyNetworkServer, Warning, TEXT("Unable to find target platform for client, terminating client connection!"));
+
+			for (int32 TPIndex = 0; TPIndex < TargetPlatformNames.Num(); TPIndex++)
+			{
+				UE_LOG(LogCookOnTheFlyNetworkServer, Warning, TEXT("    Target platforms from client: %s"), *TargetPlatformNames[TPIndex]);
+			}
+			for (int32 ActiveTPIndex = 0; ActiveTPIndex < Owner.ActiveTargetPlatforms.Num(); ActiveTPIndex++)
+			{
+				UE_LOG(LogCookOnTheFlyNetworkServer, Warning, TEXT("    Active target platforms on server: %s"), *Owner.ActiveTargetPlatforms[ActiveTPIndex]->PlatformName());
+			}
+			bIsOk = false;
+		}
+	}
+
+	if (!bIsOk)
+	{
+		return false;
+	}
+
+	if (Owner.ClientConnectedEvent.IsBound())
+	{
+		Owner.ClientConnectedEvent.Broadcast(*this);
+		bClientConnectedBroadcasted = true;
+	}
+
+	FBufferArchive HandshakeResponsePayload;
+	HandshakeResponsePayload << bIsOk;
+	HandshakeResponsePayload << PlatformNameString;
+	HandshakeResponsePayload << Owner.ZenProjectName;
+
+	if (!SendPayload(HandshakeResponsePayload))
+	{
+		return false;
+	}
+
+#if UE_BUILD_DEBUG
+	// this thread needs more space in debug builds as it tries to log messages and such
+	const static uint32 ThreadSize = 2 * 1024 * 1024;
+#else
+	const static uint32 ThreadSize = 1 * 1024 * 1024;
+#endif
+	WorkerThread = FRunnableThread::Create(this, TEXT("FCookOnTheFlyClientConnection"), ThreadSize, TPri_AboveNormal);
+
+	return true;
+}
+
+uint32 FCookOnTheFlyClientConnectionBase::Run()
+{
+	while (!bStopRequested)
+	{
+		// read a header and payload pair
+		FArrayReader Payload;
+		if (!ReceivePayload(Payload))
+		{
+			UE_LOG(LogCookOnTheFlyNetworkServer, Warning, TEXT("ReceivePayload failed"));
+			break;
+		}
+		ProcessPayload(Payload);
+	}
+	return 0;
+}
+
+
+bool FCookOnTheFlyClientConnectionBase::ProcessPayload(FArchive& Payload)
+{
+	using namespace UE::Cook;
+
+	FCookOnTheFlyRequest Request;
+	Payload << Request;
+
+	return Owner.ProcessRequest(*this, Request);
+}
+
+bool FCookOnTheFlyClientConnectionBase::SendMessage(const UE::Cook::FCookOnTheFlyMessage& Message)
+{
+	FBufferArchive Payload;
+	Payload.Reserve(Message.TotalSize());
+	Payload << const_cast<UE::Cook::FCookOnTheFlyMessage&>(Message);
+	return SendPayload(Payload);
+}
+
+FCookOnTheFlyNetworkServerBase::FCookOnTheFlyNetworkServerBase(const TArray<ITargetPlatform*>& InActiveTargetPlatforms, const FString& InZenProjectName)
+	: ActiveTargetPlatforms(InActiveTargetPlatforms)
+	, ZenProjectName(InZenProjectName)
+{
+
+}
+
+bool FCookOnTheFlyNetworkServerBase::ProcessRequest(FCookOnTheFlyClientConnectionBase& Connection, const UE::Cook::FCookOnTheFlyRequest& Request)
+{
+	using namespace UE::Cook;
+
+	ECookOnTheFlyMessage RequestType = Request.GetMessageType();
+	FHandleRequestDelegate* Handler = Handlers.Find(RequestType);
+	if (!Handler)
+	{
+		return false;
+	}
+	if (!Handler->IsBound())
+	{
+		return false;
+	}
+	return Handler->Execute(Connection, Request);
+}
+
+UE::Cook::ICookOnTheFlyNetworkServer::FHandleRequestDelegate& FCookOnTheFlyNetworkServerBase::OnRequest(UE::Cook::ECookOnTheFlyMessage MessageType)
+{
+	FHandleRequestDelegate& Handler = Handlers.FindOrAdd(MessageType);
+	return Handler;
+}
diff --git a/Engine/Source/Developer/CookOnTheFlyNetServer/Private/CookOnTheFlyNetServerBase.h b/Engine/Source/Developer/CookOnTheFlyNetServer/Private/CookOnTheFlyNetServerBase.h
new file mode 100644
--- /dev/null
+++ b/Engine/Source/Developer/CookOnTheFlyNetServer/Private/CookOnTheFlyNetServerBase.h
@@ -0,0 +1,114 @@
+// Copyright Epic Games, Inc. All Rights Reserved.
+
+#pragma once
+
+#include "CookOnTheFlyNetServer.h"
+#include "HAL/Runnable.h"
+
+class FCookOnTheFlyNetworkServerBase;
+
+class FCookOnTheFlyClientConnectionBase
+	: public UE::Cook::ICookOnTheFlyClientConnection
+	, private FRunnable
+{
+public:
+	FCookOnTheFlyClientConnectionBase(FCookOnTheFlyNetworkServerBase& InOwner);
+	virtual ~FCookOnTheFlyClientConnectionBase();
+
+	FName GetPlatformName() const override
+	{
+		return PlatformName;
+	}
+
+	const ITargetPlatform* GetTargetPlatform() const
+	{
+		return TargetPlatform;
+	}
+
+	bool GetIsSingleThreaded() const
+	{
+		return bIsSingleThreaded;
+	}
+
+	bool Initialize();
+
+	bool SendMessage(const UE::Cook::FCookOnTheFlyMessage& Message) override;
+
+	bool IsRunning()
+	{
+		return bRunning;
+	}
+
+	void Disconnect()
+	{
+		Stop();
+	}
+
+private:
+	virtual void OnInit() {}
+
+	virtual bool Init() override
+	{
+		OnInit();
+		return true;
+	}
+
+	virtual uint32 Run() override;
+
+	virtual void Stop() override
+	{
+		bStopRequested = true;
+	}
+
+	virtual void OnExit() {}
+
+	virtual void Exit() override
+	{
+		OnExit();
+		bRunning = false;
+	}
+
+	bool ProcessPayload(FArchive& Payload);
+	virtual bool ReceivePayload(FArrayReader& Payload) = 0;
+	virtual bool SendPayload(const TArray<uint8>& Out) = 0;
+
+	FCookOnTheFlyNetworkServerBase& Owner;
+	FName PlatformName;
+	ITargetPlatform* TargetPlatform = nullptr;
+	bool bIsSingleThreaded = false;
+	TAtomic<bool> bStopRequested{ false };
+	TAtomic<bool> bRunning{ true };
+	FRunnableThread* WorkerThread = nullptr;
+	bool bClientConnectedBroadcasted = false;
+};
+
+class FCookOnTheFlyNetworkServerBase
+	: public UE::Cook::ICookOnTheFlyNetworkServer
+{
+	friend class FCookOnTheFlyClientConnectionBase;
+
+public:
+	FCookOnTheFlyNetworkServerBase(const TArray<ITargetPlatform*>& InActiveTargetPlatforms, const FString& InZenProjectName);
+
+	bool ProcessRequest(FCookOnTheFlyClientConnectionBase& Connection, const UE::Cook::FCookOnTheFlyRequest& Request);
+
+	DECLARE_DERIVED_EVENT(FCookOnTheFlyNetworkServerBase, UE::Cook::ICookOnTheFlyNetworkServer::FClientConnectionEvent, FClientConnectionEvent);
+	FClientConnectionEvent& OnClientConnected() override
+	{
+		return ClientConnectedEvent;
+	}
+
+	FClientConnectionEvent& OnClientDisconnected() override
+	{
+		return ClientDisconnectedEvent;
+	}
+
+	FHandleRequestDelegate& OnRequest(UE::Cook::ECookOnTheFlyMessage MessageType) override;
+
+private:
+	FClientConnectionEvent ClientConnectedEvent;
+	FClientConnectionEvent ClientDisconnectedEvent;
+	TMap<UE::Cook::ECookOnTheFlyMessage, FHandleRequestDelegate> Handlers;
+	TArray<ITargetPlatform*> ActiveTargetPlatforms;
+	FString ZenProjectName;
+};
diff --git a/Engine/Source/Developer/CookOnTheFlyNetServer/Private/PlatformProtocolServer.cpp b/Engine/Source/Developer/CookOnTheFlyNetServer/Private/PlatformProtocolServer.cpp
new file mode 100644
--- /dev/null
+++ b/Engine/Source/Developer/CookOnTheFlyNetServer/Private/PlatformProtocolServer.cpp
@@ -0,0 +1,301 @@
+// Copyright Epic Games, Inc. All Rights Reserved.
+
+#include "PlatformProtocolServer.h"
+#include "HAL/RunnableThread.h"
+#include "Misc/OutputDeviceRedirector.h"
+#include "NetworkMessage.h"
+#include "Interfaces/ITargetPlatform.h"
+#include "Interfaces/ITargetDevice.h"
+#include "Interfaces/ITargetDeviceSocket.h"
+#include "GenericPlatform/GenericPlatformHostCommunication.h"
+
+
+
+namespace
+{
+	class FSimpleAbstractSocket_PlatformProtocol : public FSimpleAbstractSocket
+	{
+	public:
+
+		FSimpleAbstractSocket_PlatformProtocol(ITargetDeviceSocketPtr InSocket)
+			: Socket(InSocket)
+		{
+			check(Socket != nullptr);
+		}
+
+		virtual bool Receive(uint8* Results, int32 Size) const override
+		{
+			return Socket->Receive(Results, Size);
+		}
+
+		virtual bool Send(const uint8* Buffer, int32 Size) const override
+		{
+			return Socket->Send(Buffer, Size);
+		}
+
+		virtual uint32 GetMagic() const override
+		{
+			return 0x9E2B83C7;
+		}
+
+	private:
+
+		ITargetDeviceSocketPtr Socket;
+	};
+}
+
+
+class FCookOnTheFlyServerPlatformProtocol::FConnectionThreaded
+	: public FCookOnTheFlyClientConnectionBase
+{
+public:
+
+	FConnectionThreaded(FCookOnTheFlyServerPlatformProtocol& InOwner, ITargetDevicePtr InDevice, ITargetDeviceSocketPtr InSocket)
+		: FCookOnTheFlyClientConnectionBase(InOwner)
+		, Device(InDevice)
+		, Socket(InSocket)
+	{
+	}
+	
+	virtual ~FConnectionThreaded()
+	{
+		Device->CloseConnection(Socket);
+		Socket = nullptr;
+	}
+
+
+	virtual void OnInit() override
+	{
+#if PLATFORM_WINDOWS
+		FWindowsPlatformMisc::CoInitialize(ECOMModel::Multithreaded);
+#endif
+	}
+
+	virtual void OnExit() override
+	{
+#if PLATFORM_WINDOWS
+		FWindowsPlatformMisc::CoUninitialize();
+#endif
+	}
+
+	virtual bool ReceivePayload(FArrayReader& Payload) override
+	{
+		return FNFSMessageHeader::ReceivePayload(Payload, FSimpleAbstractSocket_PlatformProtocol(Socket));
+	}
+
+	virtual bool SendPayload(const TArray<uint8>& Out) override
+	{
+		return FNFSMessageHeader::WrapAndSendPayload(Out, FSimpleAbstractSocket_PlatformProtocol(Socket));
+	}
+
+	FString GetName() const
+	{
+		return FString::Printf(TEXT("%s (%s)"), *Device->GetName(), *Device->GetTargetPlatform().PlatformName());
+	}
+
+	ITargetDevicePtr GetDevice() const
+	{
+		return Device;
+	}
+
+private:
+	ITargetDevicePtr	   Device;
+	ITargetDeviceSocketPtr Socket;
+};
+
+
+FCookOnTheFlyServerPlatformProtocol::FCookOnTheFlyServerPlatformProtocol(const TArray<ITargetPlatform*>& InTargetPlatforms, const FString& InZenProjectName)
+	: FCookOnTheFlyNetworkServerBase(InTargetPlatforms, InZenProjectName)
+	, TargetPlatforms(InTargetPlatforms)
+{
+	Running = false;
+	StopRequested = false;
+}
+
+bool FCookOnTheFlyServerPlatformProtocol::Start()
+{
+	UE_LOG(LogCookOnTheFlyNetworkServer, Display, TEXT("Unreal Network File Server (custom protocol) starting up..."));
+
+	// Check the list of platforms once on start (any missing platforms will be ignored later on to avoid spamming the log).
+	for (ITargetPlatform* TargetPlatform : TargetPlatforms)
+	{
+		if (!TargetPlatform->SupportsFeature(ETargetPlatformFeatures::DirectDataExchange))
+		{
+			UE_LOG(LogCookOnTheFlyNetworkServer, Error, TEXT("Platform '%s' does not support direct communication with targets (it will be ignored)."), *TargetPlatform->PlatformName());
+		}
+	}
+
+	// Create a thread that will be updating the list of connected target devices.
+	Thread = FRunnableThread::Create(this, TEXT("FNetworkFileServerCustomProtocol"), 8 * 1024, TPri_AboveNormal);
+
+	return true;
+}
+
+
+FCookOnTheFlyServerPlatformProtocol::~FCookOnTheFlyServerPlatformProtocol()
+{
+	// Kill the running thread.
+	if (Thread != nullptr)
+	{
+		Thread->Kill(true);
+
+		delete Thread;
+		Thread = nullptr;
+	}
+}
+
+
+uint32 FCookOnTheFlyServerPlatformProtocol::Run()
+{
+#if PLATFORM_WINDOWS
+	FWindowsPlatformMisc::CoInitialize(ECOMModel::Multithreaded);
+#endif
+
+	Running = true;
+
+	// Go until requested to be done.
+	while (!StopRequested)
+	{
+		UpdateConnections();
+
+		FPlatformProcess::Sleep(1.0f);
+	}
+
+#if PLATFORM_WINDOWS
+	FWindowsPlatformMisc::CoUninitialize();
+#endif
+
+	return 0;
+}
+
+
+void FCookOnTheFlyServerPlatformProtocol::Stop()
+{
+	StopRequested = true;
+}
+
+
+void FCookOnTheFlyServerPlatformProtocol::Exit()
+{
+	// Close all connections.
+	for (auto* Connection : Connections)
+	{
+		delete Connection;
+	}
+
+	Connections.Empty();
+}
+
+
+FString FCookOnTheFlyServerPlatformProtocol::GetSupportedProtocol() const
+{
+	return FString("custom");
+}
+
+
+bool FCookOnTheFlyServerPlatformProtocol::GetAddressList(TArray<TSharedPtr<FInternetAddr> >& OutAddresses) const
+{
+	return 0;
+}
+
+
+bool FCookOnTheFlyServerPlatformProtocol::IsReadyToAcceptConnections() const
+{
+	return Running;
+}
+
+
+int32 FCookOnTheFlyServerPlatformProtocol::NumConnections() const
+{
+	return Connections.Num();
+}
+
+void FCookOnTheFlyServerPlatformProtocol::UpdateConnections()
+{
+	RemoveClosedConnections();
+
+	AddConnectionsForNewDevices();
+}
+
+
+void FCookOnTheFlyServerPlatformProtocol::RemoveClosedConnections()
+{
+	for (int32 ConnectionIndex = 0; ConnectionIndex < Connections.Num(); ++ConnectionIndex)
+	{
+		FConnectionThreaded* Connection = Connections[ConnectionIndex];
+
+		if (!Connection->IsRunning())
+		{
+			UE_LOG(LogCookOnTheFlyNetworkServer, Display, TEXT("Client %s disconnected."), *Connection->GetName());
+			Connections.RemoveAtSwap(ConnectionIndex);
+			delete Connection;
+		}
+	}
+}
+
+
+void FCookOnTheFlyServerPlatformProtocol::AddConnectionsForNewDevices()
+{
+	for (ITargetPlatform* TargetPlatform : TargetPlatforms)
+	{
+		if (TargetPlatform->SupportsFeature(ETargetPlatformFeatures::DirectDataExchange))
+		{
+			AddConnectionsForNewDevices(TargetPlatform);
+		}
+	}
+}
+
+
+void FCookOnTheFlyServerPlatformProtocol::AddConnectionsForNewDevices(ITargetPlatform* TargetPlatform)
+{
+	TArray<ITargetDevicePtr> TargetDevices;
+
+	TargetPlatform->GetAllDevices(TargetDevices);
+
+	for (ITargetDevicePtr Device : TargetDevices)
+	{
+		if (Device->IsConnected())
+		{
+			FConnectionThreaded** ExistingConnection =
+				Connections.FindByPredicate(
+					[&Device](const FConnectionThreaded* Connection)
+					{
+						return Connection->GetDevice() == Device;
+					}
+			);
+
+			// Checking IsProtocolAvailable first would make more sense, but internally it queries
+			// COM interfaces, which throws exceptions if the protocol is already in use. While we catch
+			// and process these exceptions, Visual Studio intercepts then as well and outputs messages
+			// spamming the log, which hinders the debugging experience.
+			if (!ExistingConnection)
+			{
+				if (Device->IsProtocolAvailable(EHostProtocol::CookOnTheFly))
+				{
+					ITargetDeviceSocketPtr Socket = Device->OpenConnection(EHostProtocol::CookOnTheFly);
+
+					if (Socket)
+					{
+						if (Socket->Connected())
+						{
+							FConnectionThreaded* Connection = new FConnectionThreaded(*this, Device, Socket);
+							if (Connection->Initialize())
+							{
+								Connections.Add(Connection);
+								UE_LOG(LogCookOnTheFlyNetworkServer, Display, TEXT("Client %s connected."), *Connection->GetName());
+							}
+							else
+							{
+								delete Connection;
+							}
+						}
+						else
+						{
+							Device->CloseConnection(Socket);
+						}
+					}
+				}
+			}
+		}
+	}
+}
diff --git a/Engine/Source/Developer/CookOnTheFlyNetServer/Private/PlatformProtocolServer.h b/Engine/Source/Developer/CookOnTheFlyNetServer/Private/PlatformProtocolServer.h
new file mode 100644
--- /dev/null
+++ b/Engine/Source/Developer/CookOnTheFlyNetServer/Private/PlatformProtocolServer.h
@@ -0,0 +1,84 @@
+// Copyright Epic Games, Inc. All Rights Reserved.
+
+#pragma once
+
+#include "CoreMinimal.h"
+#include "HAL/Runnable.h"
+#include "CookOnTheFlyNetServerBase.h"
+
+class ITargetPlatform;
+
+
+/**
+ * This class wraps the server thread and network connection.
+ *
+ * It uses ITargetDeviceSocket to exchanging data with the connected targets.
+ * This interface is an abstraction for direct pc-target communication as provided
+ * by the platforms. IPlatformHostCommunication/IPlatformHostSocket are the
+ * corresponding interfaces used on the game side.
+ *
+ * This implementation is based on FNetworkFileServer (the TCP-based server).
+ */
+class FCookOnTheFlyServerPlatformProtocol
+	: public FRunnable
+	, public FCookOnTheFlyNetworkServerBase
+{
+public:
+
+	/**
+	 * Creates and initializes a new instance.
+	 *
+	 * @param InFileServerOptions Network file server options
+	 */
+	FCookOnTheFlyServerPlatformProtocol(const TArray<ITargetPlatform*>& TargetPlatforms, const FString& InZenProjectName);
+
+	/**
+	 * Destructor.
+	 */
+	~FCookOnTheFlyServerPlatformProtocol();
+
+public:
+
+	//~ Begin FRunnable Interface
+
+	virtual uint32 Run()  override;
+	virtual void   Stop() override;
+	virtual void   Exit() override;
+
+	//~ End FRunnable Interface
+
+public:
+
+	//~ Begin ICookOnTheFlyNetworkServer Interface
+
+	virtual bool    IsReadyToAcceptConnections() const override;
+	virtual bool    GetAddressList(TArray<TSharedPtr<FInternetAddr> >& OutAddresses) const override;
+	virtual FString GetSupportedProtocol() const override;
+	virtual int32   NumConnections() const override;
+	virtual bool    Start() override;
+	//~ End ICookOnTheFlyNetworkServer Interface
+
+private:
+
+	void UpdateConnections();
+	void AddConnectionsForNewDevices();
+	void AddConnectionsForNewDevices(ITargetPlatform* TargetPlatform);
+	void RemoveClosedConnections();
+
+	class FConnectionThreaded;
+
+	// File server options
+	TArray<ITargetPlatform*> TargetPlatforms;
+
+	// Holds the server thread object.
+	FRunnableThread* Thread;
+
+	// Holds the list of all client connections.
+	TArray<FConnectionThreaded*> Connections;
+
+	// Holds a flag indicating whether the thread should stop executing
+	std::atomic<bool> StopRequested;
+
+	// Is the Listener thread up and running. 
+	std::atomic<bool> Running;
+};
diff --git a/Engine/Source/Developer/CookOnTheFlyNetServer/Private/TCPServer.cpp b/Engine/Source/Developer/CookOnTheFlyNetServer/Private/TCPServer.cpp
new file mode 100644
--- /dev/null
+++ b/Engine/Source/Developer/CookOnTheFlyNetServer/Private/TCPServer.cpp
@@ -0,0 +1,276 @@
+// Copyright Epic Games, Inc. All Rights Reserved.
+
+#include "TCPServer.h"
+#include "HAL/RunnableThread.h"
+#include "Misc/OutputDeviceRedirector.h"
+#include "IPAddress.h"
+#include "Sockets.h"
+#include "SocketSubsystem.h"
+#include "NetworkMessage.h"
+
+class FCookOnTheFlyClientConnectionTCP
+	: public FCookOnTheFlyClientConnectionBase
+{
+public:
+
+	FCookOnTheFlyClientConnectionTCP(FCookOnTheFlyServerTCP& InOwner, FSocket* InSocket)
+		: FCookOnTheFlyClientConnectionBase(InOwner)
+		, Socket(InSocket)
+	{
+	}
+
+	~FCookOnTheFlyClientConnectionTCP()
+	{
+		ISocketSubsystem::Get()->DestroySocket(Socket);
+	}
+
+	virtual bool ReceivePayload(FArrayReader& Payload) override
+	{
+		return FNFSMessageHeader::ReceivePayload(Payload, FSimpleAbstractSocket_FSocket(Socket));
+	}
+
+	virtual bool SendPayload(const TArray<uint8>& Payload) override
+	{
+		return FNFSMessageHeader::WrapAndSendPayload(Payload, FSimpleAbstractSocket_FSocket(Socket));
+	}
+
+	virtual void OnExit() override
+	{
+		Socket->Close();
+	}
+
+	void GetAddress(FInternetAddr& Addr)
+	{
+		Socket->GetAddress(Addr);
+	}
+
+	void GetPeerAddress(FInternetAddr& Addr)
+	{
+		Socket->GetPeerAddress(Addr);
+	}
+
+private:
+	FSocket* Socket;
+};
+
+
+
+/* FCookOnTheFlyServerTCP constructors
+ *****************************************************************************/
+
+FCookOnTheFlyServerTCP::FCookOnTheFlyServerTCP(int32 InPort, const TArray<ITargetPlatform*>& InTargetPlatforms, const FString& InZenProjectName)
+	: FCookOnTheFlyNetworkServerBase(InTargetPlatforms, InZenProjectName)
+	, Port(InPort)
+{
+	if (Port < 0)
+	{
+		Port = DEFAULT_TCP_FILE_SERVING_PORT;
+	}
+
+	Running.Set(false);
+	StopRequested.Set(false);
+}
+
+FCookOnTheFlyServerTCP::~FCookOnTheFlyServerTCP()
+{
+	// Kill the running thread.
+	if (Thread != NULL)
+	{
+		Thread->Kill(true);
+
+		delete Thread;
+		Thread = NULL;
+	}
+
+	// We are done with the socket.
+	Socket->Close();
+	ISocketSubsystem::Get()->DestroySocket(Socket);
+	Socket = NULL;
+}
+
+bool FCookOnTheFlyServerTCP::Start()
+{
+	// make sure sockets are going
+	ISocketSubsystem* SocketSubsystem = ISocketSubsystem::Get();
+
+	if (!SocketSubsystem)
+	{
+		UE_LOG(LogCookOnTheFlyNetworkServer, Error, TEXT("Could not get socket subsystem."));
+	}
+	else
+	{
+		// listen on any IP address
+		ListenAddr = SocketSubsystem->GetLocalBindAddr(*GLog);
+		ListenAddr->SetPort(Port);
+
+		// create a server TCP socket
+		Socket = SocketSubsystem->CreateSocket(NAME_Stream, TEXT("FCookOnTheFlyServerTCP-listen"), ListenAddr->GetProtocolType());
+
+		if (!Socket)
+		{
+			UE_LOG(LogCookOnTheFlyNetworkServer, Error, TEXT("Could not create listen socket."));
+		}
+		else
+		{
+			Socket->SetReuseAddr();
+			Socket->SetNoDelay();
+
+			// bind to the address
+			if (!Socket->Bind(*ListenAddr))
+			{
+				UE_LOG(LogCookOnTheFlyNetworkServer, Warning, TEXT("Failed to bind listen socket %s in FCookOnTheFlyServerTCP"), *ListenAddr->ToString(true));
+			}
+			// listen for connections
+			else if (!Socket->Listen(16))
+			{
+				UE_LOG(LogCookOnTheFlyNetworkServer, Warning, TEXT("Failed to listen on socket %s in FCookOnTheFlyServerTCP"), *ListenAddr->ToString(true));
+			}
+			else
+			{
+				// set the port on the listen address to be the same as the port on the socket
+				int32 port = Socket->GetPortNo();
+				check((Port == 0 && port != 0) || port == Port);
+				ListenAddr->SetPort(port);
+
+				// now create a thread to accept connections
+				Thread = FRunnableThread::Create(this, TEXT("FCookOnTheFlyServerTCP"), 8 * 1024, TPri_AboveNormal);
+				return true;
+			}
+		}
+	}
+	return false;
+}
+
+/* FRunnable overrides
+ *****************************************************************************/
+
+uint32 FCookOnTheFlyServerTCP::Run()
+{
+	Running.Set(true);
+	// go until requested to be done
+	while (!StopRequested.GetValue())
+	{
+		bool bReadReady = false;
+
+		// clean up closed connections
+		for (int32 ConnectionIndex = 0; ConnectionIndex < Connections.Num(); ++ConnectionIndex)
+		{
+			FCookOnTheFlyClientConnectionTCP* Connection = Connections[ConnectionIndex];
+
+			if (!Connection->IsRunning())
+			{
+				UE_LOG(LogCookOnTheFlyNetworkServer, Display, TEXT("Client disconnected."));
+				Connections.RemoveAtSwap(ConnectionIndex);
+				delete Connection;
+			}
+		}
+
+		// check for incoming connections
+		if (Socket->WaitForPendingConnection(bReadReady, FTimespan::FromSeconds(0.25f)))
+		{
+			if (bReadReady)
+			{
+				FSocket* ClientSocket = Socket->Accept(TEXT("Remote Console Connection"));
+
+				if (ClientSocket != NULL)
+				{
+					TSharedPtr<FInternetAddr> Addr = ISocketSubsystem::Get(PLATFORM_SOCKETSUBSYSTEM)->CreateInternetAddr();
+					ClientSocket->GetAddress(*Addr);
+					TSharedPtr<FInternetAddr> PeerAddr = ISocketSubsystem::Get(PLATFORM_SOCKETSUBSYSTEM)->CreateInternetAddr();
+					ClientSocket->GetPeerAddress(*PeerAddr);
+
+					for (auto PreviousConnection : Connections)
+					{
+						TSharedPtr<FInternetAddr> PreviousAddr = ISocketSubsystem::Get(PLATFORM_SOCKETSUBSYSTEM)->CreateInternetAddr();;
+						PreviousConnection->GetAddress(*PreviousAddr);
+						TSharedPtr<FInternetAddr> PreviousPeerAddr = ISocketSubsystem::Get(PLATFORM_SOCKETSUBSYSTEM)->CreateInternetAddr();;
+						PreviousConnection->GetPeerAddress(*PreviousPeerAddr);
+						if ((*Addr == *PreviousAddr) &&
+							(*PeerAddr == *PreviousPeerAddr))
+						{
+							// kill the connection 
+							PreviousConnection->Disconnect();
+							UE_LOG(LogCookOnTheFlyNetworkServer, Warning, TEXT("Killing client connection because new client connected from same address."));
+						}
+					}
+
+					FCookOnTheFlyClientConnectionTCP* Connection = new FCookOnTheFlyClientConnectionTCP(*this, ClientSocket);
+					if (Connection->Initialize())
+					{
+						Connections.Add(Connection);
+						UE_LOG(LogCookOnTheFlyNetworkServer, Display, TEXT("Client connected."));
+					}
+					else
+					{
+						delete Connection;
+					}
+				}
+			}
+		}
+		else
+		{
+			FPlatformProcess::Sleep(0.25f);
+		}
+	}
+
+	return 0;
+}
+
+
+void FCookOnTheFlyServerTCP::Exit()
+{
+	// close all connections
+	for (int32 ConnectionIndex = 0; ConnectionIndex < Connections.Num(); ConnectionIndex++)
+	{
+
+		delete Connections[ConnectionIndex];
+	}
+
+	Connections.Empty();
+}
+
+
+/* ICookOnTheFlyNetworkServer overrides
+ *****************************************************************************/
+
+FString FCookOnTheFlyServerTCP::GetSupportedProtocol() const
+{
+	return FString("tcp");
+}
+
+
+bool FCookOnTheFlyServerTCP::GetAddressList(TArray<TSharedPtr<FInternetAddr> >& OutAddresses) const
+{
+	if (ListenAddr.IsValid())
+	{
+		FString ListenAddressString = ListenAddr->ToString(true);
+
+		if (ListenAddressString.StartsWith(TEXT("0.0.0.0")))
+		{
+			if (ISocketSubsystem::Get()->GetLocalAdapterAddresses(OutAddresses))
+			{
+				for (int32 AddressIndex = 0; AddressIndex < OutAddresses.Num(); ++AddressIndex)
+				{
+					OutAddresses[AddressIndex]->SetPort(ListenAddr->GetPort());
+				}
+			}
+		}
+		else
+		{
+			OutAddresses.Add(ListenAddr);
+		}
+	}
+
+	return (OutAddresses.Num() > 0);
+}
+
+
+bool FCookOnTheFlyServerTCP::IsReadyToAcceptConnections(void) const
+{
+	return (Running.GetValue() != 0);
+}
+
+int32 FCookOnTheFlyServerTCP::NumConnections(void) const
+{
+	return Connections.Num();
+}
diff --git a/Engine/Source/Developer/CookOnTheFlyNetServer/Private/TCPServer.h b/Engine/Source/Developer/CookOnTheFlyNetServer/Private/TCPServer.h
new file mode 100644
--- /dev/null
+++ b/Engine/Source/Developer/CookOnTheFlyNetServer/Private/TCPServer.h
@@ -0,0 +1,80 @@
+// Copyright Epic Games, Inc. All Rights Reserved.
+
+#pragma once
+
+#include "CoreMinimal.h"
+#include "HAL/ThreadSafeCounter.h"
+#include "HAL/Runnable.h"
+#include "CookOnTheFlyNetServerBase.h"
+
+class FSocket;
+
+
+/**
+ * This class wraps the server thread and network connection
+ */
+class FCookOnTheFlyServerTCP
+	: public FRunnable
+	, public FCookOnTheFlyNetworkServerBase
+{
+public:
+
+	/**
+	 * Creates and initializes a new instance.
+	 *
+	 * @param InFileServerOptions Network file server options
+	 */
+	FCookOnTheFlyServerTCP(int32 InPort, const TArray<ITargetPlatform*>& InTargetPlatforms, const FString& InZenProjectName);
+
+	/**
+	 * Destructor.
+	 */
+	~FCookOnTheFlyServerTCP();
+
+public:
+
+	// FRunnable Interface
+
+	virtual bool Init() override
+	{
+		return true;
+	}
+
+	virtual uint32 Run() override;
+
+	virtual void Stop() override
+	{
+		StopRequested.Set(true);
+	}
+
+	virtual void Exit() override;
+
+	// ICookOnTheFlyNetworkServer interface
+
+	virtual bool IsReadyToAcceptConnections(void) const override;
+	virtual bool GetAddressList(TArray<TSharedPtr<FInternetAddr>>& OutAddresses) const override;
+	virtual FString GetSupportedProtocol() const override;
+	virtual int32 NumConnections() const override;
+	virtual bool Start() override;
+
+private:
+	int32 Port;
+
+	// Holds the server (listening) socket.
+	FSocket* Socket;
+
+	// Holds the server thread object.
+	FRunnableThread* Thread;
+
+	// Holds the list of all client connections.
+	TArray<class FCookOnTheFlyClientConnectionTCP*> Connections;
+
+	// Holds a flag indicating whether the thread should stop executing
+	FThreadSafeCounter StopRequested;
+
+	// Is the Listner thread up and running. 
+	FThreadSafeCounter Running;
+
+	// Holds the address that the server is bound to.
+	TSharedPtr<FInternetAddr> ListenAddr;
+};
diff --git a/Engine/Source/Editor/UnrealEd/Classes/CookOnTheSide/CookOnTheFlyServer.h b/Engine/Source/Editor/UnrealEd/Classes/CookOnTheSide/CookOnTheFlyServer.h
--- a/Engine/Source/Editor/UnrealEd/Classes/CookOnTheSide/CookOnTheFlyServer.h
+++ b/Engine/Source/Editor/UnrealEd/Classes/CookOnTheSide/CookOnTheFlyServer.h
@@ -124,6 +124,7 @@ namespace UE::Cook
 	class FRequestCluster;
 	class FSaveCookedPackageContext;
 	class ICookOnTheFlyRequestManager;
+	class ICookOnTheFlyNetworkServer;
 	struct FConstructPackageData;
 	struct FCookByTheBookOptions;
 	struct FCookOnTheFlyOptions;
@@ -186,6 +187,7 @@ private:
 	TArray<class INetworkFileServer*> NetworkFileServers;
 	FOnFileModifiedDelegate FileModifiedDelegate;
 	TUniquePtr<UE::Cook::ICookOnTheFlyRequestManager> CookOnTheFlyRequestManager;
+	TSharedPtr<UE::Cook::ICookOnTheFlyNetworkServer> CookOnTheFlyNetworkServer;
 
 	//////////////////////////////////////////////////////////////////////////
 	// General cook options
diff --git a/Engine/Source/Editor/UnrealEd/Private/CookOnTheFlyServer.cpp b/Engine/Source/Editor/UnrealEd/Private/CookOnTheFlyServer.cpp
--- a/Engine/Source/Editor/UnrealEd/Private/CookOnTheFlyServer.cpp
+++ b/Engine/Source/Editor/UnrealEd/Private/CookOnTheFlyServer.cpp
@@ -126,6 +126,7 @@
 #include "UObject/UObjectArray.h"
 #include "UObject/UObjectIterator.h"
 #include "ZenStoreWriter.h"
+#include "CookOnTheFlyNetServer.h"
 
 #define LOCTEXT_NAMESPACE "Cooker"
 
@@ -289,7 +290,7 @@ public:
 		return Cooker.SandboxFile->GetSandboxDirectory();
 	}
 
-	virtual const ITargetPlatform* AddPlatform(const FName& PlatformName) override
+	const ITargetPlatform* AddPlatform(const FName& PlatformName)
 	{
 		UE::Cook::FPlatformManager::FReadScopeLock PlatformScopeLock(Cooker.PlatformManager->ReadLockPlatforms());
 		const ITargetPlatform* TargetPlatform = AddPlatformInternal(PlatformName);
@@ -371,16 +372,6 @@ public:
 			});
 	}
 
-	virtual bool EnqueueRecompileShaderRequest(UE::Cook::FRecompileShaderRequest RecompileShaderRequest) override
-	{
-		UE_LOG(LogCook, Verbose, TEXT("Enqueing recompile shader(s) request, Platform='%s', ShaderPlatform='%d'"),
-			*RecompileShaderRequest.RecompileArguments.PlatformName, RecompileShaderRequest.RecompileArguments.ShaderPlatform);
-
-		Cooker.PackageTracker->RecompileRequests.Enqueue(MoveTemp(RecompileShaderRequest));
-
-		return true;
-	}
-
 	virtual ICookedPackageWriter& GetPackageWriter(const ITargetPlatform* TargetPlatform) override
 	{
 		return GetPackageWriterInternal(TargetPlatform);
@@ -533,6 +524,8 @@ TStatId UCookOnTheFlyServer::GetStatId() const
 
 bool UCookOnTheFlyServer::StartCookOnTheFly(FCookOnTheFlyStartupOptions InCookOnTheFlyOptions)
 {
+	using namespace UE::Cook;
+
 	LLM_SCOPE_BYTAG(Cooker);
 #if WITH_COTF
 	check(IsCookOnTheFlyMode());
@@ -569,18 +562,93 @@ bool UCookOnTheFlyServer::StartCookOnTheFly(FCookOnTheFlyStartupOptions InCookOn
 	UE_LOG(LogCook, Display, TEXT("Starting '%s' cook-on-the-fly server"),
 		IsUsingZenStore() ? TEXT("Zen") : TEXT("Network File"));
 
+	FCookOnTheFlyNetworkServerOptions NetworkServerOptions;
+	NetworkServerOptions.Protocol = CookOnTheFlyOptions->bPlatformProtocol ? ECookOnTheFlyNetworkServerProtocol::Platform : ECookOnTheFlyNetworkServerProtocol::Tcp;
+	NetworkServerOptions.Port = CookOnTheFlyOptions->bBindAnyPort ? 0 : -1;
+	if (!InCookOnTheFlyOptions.TargetPlatforms.IsEmpty())
+	{
+		NetworkServerOptions.TargetPlatforms = InCookOnTheFlyOptions.TargetPlatforms;
+	}
+	else
+	{
+		NetworkServerOptions.TargetPlatforms = TPM.GetTargetPlatforms();
+	}
+	if (IsUsingZenStore())
+	{
+		NetworkServerOptions.ZenProjectName = FApp::GetZenStoreProjectId();
+	}
+
+	ICookOnTheFlyNetworkServerModule& CookOnTheFlyNetworkServerModule = FModuleManager::LoadModuleChecked<ICookOnTheFlyNetworkServerModule>(TEXT("CookOnTheFlyNetServer"));
+	CookOnTheFlyNetworkServer = CookOnTheFlyNetworkServerModule.CreateServer(NetworkServerOptions);
+
+	CookOnTheFlyNetworkServer->OnClientConnected().AddLambda([this](ICookOnTheFlyClientConnection& Connection)
+		{
+			if (Connection.GetTargetPlatform())
+			{
+				CookOnTheFlyServerInterface->AddPlatform(Connection.GetPlatformName());
+			}
+		});
+
+	CookOnTheFlyNetworkServer->OnRequest(ECookOnTheFlyMessage::RecompileShaders).BindLambda([this](ICookOnTheFlyClientConnection& Connection, const FCookOnTheFlyRequest& Request)
+		{
+			FCookOnTheFlyResponse Response(Request);
+
+			if (!Connection.GetTargetPlatform())
+			{
+				UE_LOG(LogCook, Warning, TEXT("RecompileShadersRequest from editor client"));
+				Response.SetStatus(UE::Cook::ECookOnTheFlyMessageStatus::Error);
+			}
+			else
+			{
+				TArray<FString> RecompileModifiedFiles;
+				TArray<uint8> MeshMaterialMaps;
+				TArray<uint8> GlobalShaderMap;
+
+				FShaderRecompileData RecompileData(Connection.GetTargetPlatform()->PlatformName(), &RecompileModifiedFiles, &MeshMaterialMaps, &GlobalShaderMap);
+				{
+					TUniquePtr<FArchive> Ar = Request.ReadBody();
+					*Ar << RecompileData;
+				}
+
+				FEventRef RecompileCompletedEvent;
+				UE::Cook::FRecompileShaderCompletedCallback RecompileCompleted = [this, &RecompileCompletedEvent]()
+				{
+					RecompileCompletedEvent->Trigger();
+				};
+
+				PackageTracker->RecompileRequests.Enqueue({ RecompileData, MoveTemp(RecompileCompleted) });
+
+				RecompileCompletedEvent->Wait();
+
+				{
+					TUniquePtr<FArchive> Ar = Response.WriteBody();
+					*Ar << MeshMaterialMaps;
+					*Ar << GlobalShaderMap;
+				}
+			}
+			return Connection.SendMessage(Response);
+		});
+
 	if (IsUsingZenStore())
 	{
-		UE::Cook::FIoStoreCookOnTheFlyServerOptions ServerOptions;
-		ServerOptions.Port = -1; // Use default
-		CookOnTheFlyRequestManager = UE::Cook::MakeIoStoreCookOnTheFlyRequestManager(*CookOnTheFlyServerInterface, AssetRegistry, MoveTemp(ServerOptions));
+		CookOnTheFlyRequestManager = MakeIoStoreCookOnTheFlyRequestManager(*CookOnTheFlyServerInterface, AssetRegistry, CookOnTheFlyNetworkServer.ToSharedRef());
+	}
+	else
+	{
+		CookOnTheFlyRequestManager = MakeNetworkFileCookOnTheFlyRequestManager(*CookOnTheFlyServerInterface, CookOnTheFlyNetworkServer.ToSharedRef());
+	}
+
+	if (CookOnTheFlyNetworkServer->Start())
+	{
+		TArray<TSharedPtr<FInternetAddr>> ListenAddresses;
+		if (CookOnTheFlyNetworkServer->GetAddressList(ListenAddresses))
+		{
+			UE_LOG(LogCook, Display, TEXT("Unreal Network File Server is ready for client connections on %s!"), *ListenAddresses[0]->ToString(true));
+		}
 	}
 	else
 	{
-		FNetworkFileServerOptions ServerOptions;
-		ServerOptions.Protocol = CookOnTheFlyOptions->bPlatformProtocol ? NFSP_Platform : NFSP_Tcp;
-		ServerOptions.Port = CookOnTheFlyOptions->bBindAnyPort ? 0 : -1;
-		CookOnTheFlyRequestManager = UE::Cook::MakeNetworkFileCookOnTheFlyRequestManager(*CookOnTheFlyServerInterface, MoveTemp(ServerOptions));
+		UE_LOG(LogCook, Fatal, TEXT("Failed starting Unreal Network file server!"));
 	}
 	BeginCookEditorSystems();
 
diff --git a/Engine/Source/Editor/UnrealEd/Private/Cooker/CookOnTheFlyServerInterface.h b/Engine/Source/Editor/UnrealEd/Private/Cooker/CookOnTheFlyServerInterface.h
--- a/Engine/Source/Editor/UnrealEd/Private/Cooker/CookOnTheFlyServerInterface.h
+++ b/Engine/Source/Editor/UnrealEd/Private/Cooker/CookOnTheFlyServerInterface.h
@@ -52,13 +52,6 @@ public:
 	/** Returns the cooker sandbox directory path. */
 	virtual FString GetSandboxDirectory() const = 0;
 	
-	/**
-	 * Add platform to the cook-on-the-fly session.
-	 *
-	 * @param PlatformName The platform name.
-	 */
-	virtual const ITargetPlatform* AddPlatform(const FName& PlatformName) = 0;
-
 	/**
 	 * Remove platform from the cook-on-the-fly session.
 	 *
@@ -89,13 +82,6 @@ public:
 
 	virtual void MarkPackageDirty(const FName& PackageName) = 0;
 
-	/**
-	 * Enqueue a new shader compile request.
-	 *
-	 * @param RecompileShaderRequest Recompile shader request parameters.
-	 */
-	virtual bool EnqueueRecompileShaderRequest(FRecompileShaderRequest RecompileShaderRequest) = 0;
-
 	/**
 	 * Returns the package store writer for the specified platform.
 	 */
diff --git a/Engine/Source/Editor/UnrealEd/Private/Cooker/ExternalCookOnTheFlyServer.cpp b/Engine/Source/Editor/UnrealEd/Private/Cooker/ExternalCookOnTheFlyServer.cpp
--- a/Engine/Source/Editor/UnrealEd/Private/Cooker/ExternalCookOnTheFlyServer.cpp
+++ b/Engine/Source/Editor/UnrealEd/Private/Cooker/ExternalCookOnTheFlyServer.cpp
@@ -37,10 +37,10 @@ void FExternalCookOnTheFlyServer::HandleRegisterServiceMessage(const FZenCookOnT
 	}
 	UE_LOG(LogCookOnTheFly, Verbose, TEXT("Accepting service from %s"), *Context->GetSender().ToString());
 	UE::Cook::FCookOnTheFlyHostOptions HostOptions;
-	HostOptions.Hosts.Add(TEXT("127.0.0.1"));
-	HostOptions.Port = Message.Port;
-	ServerConnection = CookOnTheFlyModule.ConnectToServer(HostOptions);
-	if (ServerConnection)
+	HostOptions.Hosts.Add(FString::Printf(TEXT("127.0.0.1:%d"), Message.Port));
+	
+	CookOnTheFlyServerConnection = CookOnTheFlyModule.ConnectToServer(HostOptions);
+	if (CookOnTheFlyServerConnection)
 	{
 		UE_LOG(LogCookOnTheFly, Display, TEXT("Connected to server"));
 		AssetRegistry.OnAssetUpdatedOnDisk().AddRaw(this, &FExternalCookOnTheFlyServer::AssetUpdatedOnDisk);
@@ -67,24 +67,24 @@ void FExternalCookOnTheFlyServer::Tick(float DeltaSeconds)
 
 			UE::Cook::FCookOnTheFlyRequest Request(UE::Cook::ECookOnTheFlyMessage::RecookPackages);
 			Request.SetBodyTo(UE::ZenCookOnTheFly::Messaging::FRecookPackagesRequest{ MoveTemp(PackageIdsToRecook) });
-			UE::Cook::FCookOnTheFlyResponse Response = ServerConnection->SendRequest(Request).Get();
+			UE::Cook::FCookOnTheFlyResponse Response = CookOnTheFlyServerConnection->SendRequest(Request).Get();
 			if (!Response.IsOk())
 			{
 				UE_LOG(LogCookOnTheFly, Warning, TEXT("Failed to send RecookPackages request"));
 			}
 		}
 	}
-	if (ServerConnection && !ServerConnection->IsConnected())
+	if (CookOnTheFlyServerConnection && !CookOnTheFlyServerConnection->IsConnected())
 	{
 		AssetRegistry.OnAssetUpdatedOnDisk().RemoveAll(this);
-		ServerConnection.Reset();
+		CookOnTheFlyServerConnection.Reset();
 		UE_LOG(LogCookOnTheFly, Display, TEXT("Disconnected from server"));
 	}
 }
 
 bool FExternalCookOnTheFlyServer::IsTickable() const
 {
-	return ServerConnection.IsValid();
+	return CookOnTheFlyServerConnection.IsValid();
 }
 
 void FExternalCookOnTheFlyServer::AssetUpdatedOnDisk(const FAssetData& AssetData)
diff --git a/Engine/Source/Editor/UnrealEd/Private/Cooker/ExternalCookOnTheFlyServer.h b/Engine/Source/Editor/UnrealEd/Private/Cooker/ExternalCookOnTheFlyServer.h
--- a/Engine/Source/Editor/UnrealEd/Private/Cooker/ExternalCookOnTheFlyServer.h
+++ b/Engine/Source/Editor/UnrealEd/Private/Cooker/ExternalCookOnTheFlyServer.h
@@ -31,7 +31,7 @@ private:
 	IAssetRegistry& AssetRegistry;
 	TSharedPtr<FMessageEndpoint, ESPMode::ThreadSafe> MessageEndpoint;
 	const FString ServiceId;
-	TUniquePtr<UE::Cook::ICookOnTheFlyServerConnection> ServerConnection;
+	TUniquePtr<UE::Cook::ICookOnTheFlyServerConnection> CookOnTheFlyServerConnection;
 	FCriticalSection AllPackagesToRecookCritical;
 	TSet<FName> AllPackagesToRecook;
 };
diff --git a/Engine/Source/Editor/UnrealEd/Private/Cooker/IoStoreCookOnTheFlyRequestManager.cpp b/Engine/Source/Editor/UnrealEd/Private/Cooker/IoStoreCookOnTheFlyRequestManager.cpp
--- a/Engine/Source/Editor/UnrealEd/Private/Cooker/IoStoreCookOnTheFlyRequestManager.cpp
+++ b/Engine/Source/Editor/UnrealEd/Private/Cooker/IoStoreCookOnTheFlyRequestManager.cpp
@@ -6,7 +6,6 @@
 #include "CookOnTheFlyMessages.h"
 #include "Modules/ModuleManager.h"
 #include "PackageStoreWriter.h"
-#include "ShaderCompiler.h"
 #include "HAL/PlatformTime.h"
 #include "Misc/PackageName.h"
 #include "UObject/SavePackage.h"
@@ -23,403 +22,26 @@
 #include "MessageEndpoint.h"
 #include "MessageEndpointBuilder.h"
 #include "Cooker/ExternalCookOnTheFlyServer.h"
+#include "CookOnTheFlyNetServer.h"
 
-class FIoStoreCookOnTheFlyNetworkServer
+class FIoStoreCookOnTheFlyRequestManager final
+	: public UE::Cook::ICookOnTheFlyRequestManager
 {
-	static constexpr uint32 ServerSenderId = ~uint32(0);
-
 public:
-	/**
-	 * Connection status
-	 */
-	enum class EConnectionStatus
-	{
-		Disconnected,
-		Connected
-	};
-
-	using FRequestHandler = TFunction<bool(const FName&, bool, const UE::Cook::FCookOnTheFlyRequest&, UE::Cook::FCookOnTheFlyResponse&)>;
-
-	using FClientConnectionHandler = TFunction<bool(const FName&, bool, EConnectionStatus)>;
-
-	using FFillRequest = TFunction<void(FArchive&)>;
-
-	using FProcessResponse = TFunction<bool(FArchive&)>;
-
-	/**
-	 * Cook-on-the-fly connection server options.
-	 */
-	struct FServerOptions
-	{
-		/** The port to listen for new connections. */
-		int32 Port = UE::Cook::DefaultCookOnTheFlyServingPort;
-
-		/** Callback invoked when a client has connected or disconnected. */
-		FClientConnectionHandler HandleClientConnection;
-
-		/** Callback invoked when the server receives a new request. */
-		FRequestHandler HandleRequest;
-	};
-
-	FIoStoreCookOnTheFlyNetworkServer(FServerOptions InOptions)
-		: Options(MoveTemp(InOptions))
+	FIoStoreCookOnTheFlyRequestManager(UE::Cook::ICookOnTheFlyServer& InCookOnTheFlyServer, const IAssetRegistry* AssetRegistry, TSharedRef<UE::Cook::ICookOnTheFlyNetworkServer> InConnectionServer)
+		: CookOnTheFlyServer(InCookOnTheFlyServer)
+		, ConnectionServer(InConnectionServer)
 		, ServiceId(FExternalCookOnTheFlyServer::GenerateServiceId())
-	{
-		MessageEndpoint = FMessageEndpoint::Builder("FCookOnTheFly");
-	}
-
-	~FIoStoreCookOnTheFlyNetworkServer()
-	{
-		StopServer();
-	}
-
-	bool StartServer()
-	{
-		const int32 Port = Options.Port > 0 ? Options.Port : UE::Cook::DefaultCookOnTheFlyServingPort;
-		UE_LOG(LogCookOnTheFly, Log, TEXT("Starting COTF server on port '%d'"), Port);
-
-		check(!bIsRunning);
-		bIsRunning = false;
-		bStopRequested = false;
-
-		ISocketSubsystem& SocketSubsystem = *ISocketSubsystem::Get();
-
-		ListenAddr = SocketSubsystem.GetLocalBindAddr(*GLog);
-		ListenAddr->SetPort(Port);
-
-		// create a server TCP socket
-		Socket = SocketSubsystem.CreateSocket(NAME_Stream, TEXT("COTF-Server"), ListenAddr->GetProtocolType());
-
-		if (!Socket)
-		{
-			UE_LOG(LogCookOnTheFly, Error, TEXT("Could not create listen socket"));
-			return false;
-		}
-
-		Socket->SetReuseAddr();
-		Socket->SetNoDelay();
-
-		if (!Socket->Bind(*ListenAddr))
-		{
-			UE_LOG(LogCookOnTheFly, Error, TEXT("Failed to bind socket to address '%s'"), *ListenAddr->ToString(true));
-			return false;
-		}
-
-		if (!Socket->Listen(16))
-		{
-			UE_LOG(LogCookOnTheFly, Warning, TEXT("Failed to listen on address '%s'"), *ListenAddr->ToString(true));
-			return false;
-		}
-
-		ListenAddr->SetPort(Socket->GetPortNo());
-
-		ServerThread = AsyncThread([this] { return ServerThreadEntry(); }, 8 * 1024, TPri_AboveNormal);
-
-		UE_LOG(LogCookOnTheFly, Display, TEXT("COTF server is ready for client(s) on '%s'!"), *ListenAddr->ToString(true));
-
-		if (MessageEndpoint.IsValid())
-		{
-			FZenCookOnTheFlyRegisterServiceMessage* RegisterServiceMessage = FMessageEndpoint::MakeMessage<FZenCookOnTheFlyRegisterServiceMessage>();
-			RegisterServiceMessage->ServiceId = ServiceId;
-			RegisterServiceMessage->Port = ListenAddr->GetPort();
-			MessageEndpoint->Publish(RegisterServiceMessage);
-		}
-
-		return true;
-	}
-
-	void StopServer()
-	{
-		if (bIsRunning && !bStopRequested)
-		{
-			bStopRequested = true;
-
-			ISocketSubsystem& SocketSubsystem = *ISocketSubsystem::Get();
-
-			while (bIsRunning)
-			{
-				FPlatformProcess::Sleep(0.25f);
-			}
-
-			for (TUniquePtr<FClient>& Client : Clients)
-			{
-				Client->bStopRequested = true;
-				Client->Socket->Close();
-				Client->Thread.Wait();
-				SocketSubsystem.DestroySocket(Client->Socket);
-			}
-
-			Clients.Empty();
-		}
-	}
-
-	bool BroadcastMessage(const UE::Cook::FCookOnTheFlyMessage& Message, const FName& PlatformName = NAME_None)
 	{
 		using namespace UE::Cook;
 
-		FCookOnTheFlyMessageHeader Header = Message.GetHeader();
-
-		Header.MessageType = Header.MessageType | ECookOnTheFlyMessage::Message;
-		Header.MessageStatus = ECookOnTheFlyMessageStatus::Ok;
-		Header.SenderId = ServerSenderId;
-		Header.CorrelationId = NextCorrelationId++;
-		Header.Timestamp = FDateTime::UtcNow().GetTicks();
-
-		FBufferArchive MessagePayload;
-		MessagePayload.Reserve(Message.TotalSize());
-
-		MessagePayload << Header;
-		MessagePayload << const_cast<TArray<uint8>&>(Message.GetBody());
-
-		UE_LOG(LogCookOnTheFly, Verbose, TEXT("Sending: %s, Size='%lld'"), *Header.ToString(), Message.TotalSize());
-
-		TArray<FClient*, TInlineAllocator<4>> ClientsToBroadcast;
-		{
-			FScopeLock _(&ClintsCriticalSection);
-
-			for (TUniquePtr<FClient>& Client : Clients)
-			{
-				if (PlatformName.IsNone() || Client->PlatformName == PlatformName)
-				{
-					ClientsToBroadcast.Add(Client.Get());
-				}
-			}
-		}
-
-		bool bBroadcasted = true;
-		for (FClient* Client : ClientsToBroadcast)
-		{
-			if (!FNFSMessageHeader::WrapAndSendPayload(MessagePayload, FSimpleAbstractSocket_FSocket(Client->Socket)))
-			{
-				UE_LOG(LogCookOnTheFly, Warning, TEXT("Failed to send message '%s' to client '%s' (Id='%d', Platform='%s')"),
-					LexToString(Message.GetHeader().MessageType), *Client->PeerAddr->ToString(true), Client->ClientId, *Client->PlatformName.ToString());
-
-				Client->bIsRunning = false;
-				bBroadcasted = false;
-			}
-
-			Client->LastActivityTime = FPlatformTime::Seconds();
-		}
-
-		return bBroadcasted;
-	}
-
-private:
-	struct FClient
-	{
-		FSocket* Socket = nullptr;
-		TSharedPtr<FInternetAddr> Addr;
-		TSharedPtr<FInternetAddr> PeerAddr;
-		TFuture<void> Thread;
-		TAtomic<bool> bIsRunning{ false };
-		TAtomic<bool> bStopRequested{ false };
-		TAtomic<bool> bIsProcessingRequest{ false };
-		TAtomic<double> LastActivityTime{ 0 };
-		uint32 ClientId = 0;
-		FName PlatformName;
-		bool bIsSingleThreaded = false;
-	};
-
-	void ServerThreadEntry()
-	{
-		using namespace UE::Cook;
-
-		ISocketSubsystem& SocketSubsystem = *ISocketSubsystem::Get();
-		bIsRunning = true;
-
-		while (!bStopRequested)
-		{
-			bool bIsReady = false;
-			if (Socket->WaitForPendingConnection(bIsReady, FTimespan::FromSeconds(0.25f)))
-			{
-				if (bIsReady)
-				{
-					if (FSocket* ClientSocket = Socket->Accept(TEXT("COTF-Client")))
-					{
-						TUniquePtr<FClient> Client = MakeUnique<FClient>();
-						Client->Socket = ClientSocket;
-						Client->Addr = SocketSubsystem.CreateInternetAddr();
-						Client->PeerAddr = SocketSubsystem.CreateInternetAddr();
-						ClientSocket->GetAddress(*Client->Addr);
-						ClientSocket->GetPeerAddress(*Client->PeerAddr);
-						Client->ClientId = NextClientId++;
-						Client->bIsRunning = true;
-						Client->LastActivityTime = FPlatformTime::Seconds();
-						Client->Thread = AsyncThread([this, ConnectedClient = Client.Get()]()
-						{
-							ClientThreadEntry(ConnectedClient);
-						});
-
-						UE_LOG(LogCookOnTheFly, Display, TEXT("New client connected from address '%s' (ID='%d')"), *Client->PeerAddr->ToString(true), Client->ClientId);
-						{
-							FScopeLock _(&ClintsCriticalSection);
-							Clients.Add(MoveTemp(Client));
-						}
-					}
-				}
-			}
-			else
-			{
-				FPlatformProcess::Sleep(0.25f);
-			}
-
-			{
-				FScopeLock _(&ClintsCriticalSection);
-				for (auto It = Clients.CreateIterator(); It; ++It)
-				{
-					TUniquePtr<FClient>& Client = *It;
-					if (!Client->bIsRunning)
-					{
-						UE_LOG(LogCookOnTheFly, Display, TEXT("Closing connection to client on address '%s' (Id='%d', Platform='%s')"),
-							*Client->PeerAddr->ToString(true), Client->ClientId, *Client->PlatformName.ToString());
-
-						Options.HandleClientConnection(Client->PlatformName, Client->bIsSingleThreaded, EConnectionStatus::Disconnected);
-
-						Client->Socket->Close();
-						Client->Thread.Wait();
-						SocketSubsystem.DestroySocket(Client->Socket);
-						It.RemoveCurrent();
-					}
-				}
-			}
-		}
-
-		bIsRunning = false;
-	}
-
-	void ClientThreadEntry(FClient* Client)
-	{
-		while (!bStopRequested && !Client->bStopRequested)
-		{
-			Client->LastActivityTime = FPlatformTime::Seconds();
-			if (!ProcesseRequest(*Client))
-			{
-				break;
-			}
-		}
-
-		Client->bIsRunning = false;
-	}
-
-	bool ProcesseRequest(FClient& Client)
-	{
-		using namespace UE::Cook;
-
-		Client.bIsProcessingRequest = false;
-
-		FArrayReader RequestPayload;
-		if (!FNFSMessageHeader::ReceivePayload(RequestPayload, FSimpleAbstractSocket_FSocket(Client.Socket)))
-		{
-			UE_LOG(LogCookOnTheFly, Warning, TEXT("Unable to receive request from client"));
-			return false;
-		}
-
-		Client.bIsProcessingRequest = true;
-
-		FCookOnTheFlyRequest Request;
-		RequestPayload << Request;
-
-		UE_LOG(LogCookOnTheFly, Verbose, TEXT("Received: %s, Size='%lld'"), *Request.GetHeader().ToString(), Request.TotalSize());
-
-		EnumRemoveFlags(Request.GetHeader().MessageType, ECookOnTheFlyMessage::TypeFlags);
-
-		FCookOnTheFlyResponse Response;
-		bool bRequestOk = false;
-		bool bIsResponse = false;
-
-		switch (Request.GetHeader().MessageType)
-		{
-		case ECookOnTheFlyMessage::Handshake:
-		{
-			ProcessHandshake(Client, Request, Response);
-			bRequestOk = Options.HandleClientConnection(Client.PlatformName, Client.bIsSingleThreaded, EConnectionStatus::Connected);
-			break;
-		}
-		default:
-		{
-			bRequestOk = Options.HandleRequest(Client.PlatformName, Client.bIsSingleThreaded, Request, Response);
-			break;
-		}
-		}
-
-		if (bRequestOk && !bIsResponse)
-		{
-			FCookOnTheFlyMessageHeader& ResponseHeader = Response.GetHeader();
-
-			ResponseHeader.MessageType = Request.GetHeader().MessageType | ECookOnTheFlyMessage::Response;
-			ResponseHeader.SenderId = ServerSenderId;
-			ResponseHeader.CorrelationId = Request.GetHeader().CorrelationId;
-			ResponseHeader.Timestamp = FDateTime::UtcNow().GetTicks();
-
-			Response.SetHeader(ResponseHeader);
-
-			FBufferArchive ResponsePayload;
-			ResponsePayload.Reserve(Response.TotalSize());
-
-			ResponsePayload << Response;
-			bRequestOk = FNFSMessageHeader::WrapAndSendPayload(ResponsePayload, FSimpleAbstractSocket_FSocket(Client.Socket));
-		}
-
-		return bRequestOk;
-	}
-
-	void ProcessHandshake(FClient& Client, UE::Cook::FCookOnTheFlyRequest& HandshakeRequest, UE::Cook::FCookOnTheFlyResponse& Response)
-	{
-		FString PlatformName;
-		FString ProjectName;
-
-		{
-			TUniquePtr<FArchive> Ar = HandshakeRequest.ReadBody();
-			*Ar << PlatformName;
-			*Ar << ProjectName;
-			*Ar << Client.bIsSingleThreaded;
-		}
-
-		if (PlatformName.Len())
-		{
-			Client.PlatformName = FName(*PlatformName);
-		}
-		Response.SetBodyTo(Client.ClientId);
-		Response.SetStatus(UE::Cook::ECookOnTheFlyMessageStatus::Ok);
-	}
-
-	FClient* GetClientById(uint32 ClientId)
-	{
-		FScopeLock _(&ClintsCriticalSection);
+		ConnectionServer->OnClientConnected().AddRaw(this, &FIoStoreCookOnTheFlyRequestManager::OnClientConnected);
+		ConnectionServer->OnClientDisconnected().AddRaw(this, &FIoStoreCookOnTheFlyRequestManager::OnClientDisconnected);
+		ConnectionServer->OnRequest(ECookOnTheFlyMessage::GetCookedPackages).BindRaw(this, &FIoStoreCookOnTheFlyRequestManager::HandleClientRequest);
+		ConnectionServer->OnRequest(ECookOnTheFlyMessage::CookPackage).BindRaw(this, &FIoStoreCookOnTheFlyRequestManager::HandleClientRequest);
+		ConnectionServer->OnRequest(ECookOnTheFlyMessage::RecookPackages).BindRaw(this, &FIoStoreCookOnTheFlyRequestManager::HandleClientRequest);
 
-		for (TUniquePtr<FClient>& Client : Clients)
-		{
-			if (Client->ClientId == ClientId)
-			{
-				return Client.Get();
-			}
-		}
-
-		return nullptr;
-	}
-
-	FIoStoreCookOnTheFlyNetworkServer::FServerOptions Options;
-	TFuture<void> ServerThread;
-	TSharedPtr<FInternetAddr> ListenAddr;
-	FSocket* Socket = nullptr;
-	FCriticalSection ClintsCriticalSection;
-	TArray<TUniquePtr<FClient>> Clients;
-	TAtomic<bool> bIsRunning{ false };
-	TAtomic<bool> bStopRequested{ false };
-	uint32 NextClientId = 1;
-	TAtomic<uint32> NextCorrelationId{ 1 };
-	TSharedPtr<FMessageEndpoint, ESPMode::ThreadSafe> MessageEndpoint;
-	const FString ServiceId;
-};
-
-class FIoStoreCookOnTheFlyRequestManager final
-	: public UE::Cook::ICookOnTheFlyRequestManager
-{
-public:
-	FIoStoreCookOnTheFlyRequestManager(UE::Cook::ICookOnTheFlyServer& InCookOnTheFlyServer, const IAssetRegistry* AssetRegistry, UE::Cook::FIoStoreCookOnTheFlyServerOptions InOptions)
-		: CookOnTheFlyServer(InCookOnTheFlyServer)
-		, Options(MoveTemp(InOptions))
-	{
+		MessageEndpoint = FMessageEndpoint::Builder("FCookOnTheFly");
 		TArray<FAssetData> AllAssets;
 		AssetRegistry->GetAllAssets(AllAssets, true);
 		for (const FAssetData& AssetData : AllAssets)
@@ -640,32 +262,19 @@ private:
 	{
 		using namespace UE::Cook;
 
-		ICookOnTheFlyModule& CookOnTheFlyModule = FModuleManager::LoadModuleChecked<ICookOnTheFlyModule>(TEXT("CookOnTheFly"));
-
-		const int32 Port = Options.Port > 0 ? Options.Port : UE::Cook::DefaultCookOnTheFlyServingPort;
-
-		ConnectionServer = MakeUnique<FIoStoreCookOnTheFlyNetworkServer>(FIoStoreCookOnTheFlyNetworkServer::FServerOptions
+		TArray<TSharedPtr<FInternetAddr>> ListenAddresses;
+		if (MessageEndpoint.IsValid() && ConnectionServer->GetAddressList(ListenAddresses))
 		{
-			Port,
-			[this](const FName& PlatformName, bool bIsSingleThreaded, FIoStoreCookOnTheFlyNetworkServer::EConnectionStatus ConnectionStatus)
-			{
-				return HandleClientConnection(PlatformName, bIsSingleThreaded, ConnectionStatus);
-			},
-			[this](const FName& PlatformName, bool bIsSingleThreaded, const FCookOnTheFlyRequest& Request, FCookOnTheFlyResponse& Response)
-			{ 
-				return HandleClientRequest(PlatformName, bIsSingleThreaded, Request, Response);
-			}
-		});
-
-		const bool bServerStarted = ConnectionServer->StartServer();
-
-		return bServerStarted;
+			FZenCookOnTheFlyRegisterServiceMessage* RegisterServiceMessage = FMessageEndpoint::MakeMessage<FZenCookOnTheFlyRegisterServiceMessage>();
+			RegisterServiceMessage->ServiceId = ServiceId;
+			RegisterServiceMessage->Port = ListenAddresses[0]->GetPort();
+			MessageEndpoint->Publish(RegisterServiceMessage);
+		}
+		return true;
 	}
 
 	virtual void Shutdown() override
 	{
-		ConnectionServer->StopServer();
-		ConnectionServer.Reset();
 	}
 
 	virtual void OnPackageGenerated(const FName& PackageName)
@@ -747,95 +356,102 @@ private:
 	}
 
 private:
-	bool HandleClientConnection(const FName& PlatformName, bool bIsSingleThreaded, FIoStoreCookOnTheFlyNetworkServer::EConnectionStatus ConnectionStatus)
+	void OnClientConnected(UE::Cook::ICookOnTheFlyClientConnection& Connection)
 	{
-		if (PlatformName.IsNone())
-		{
-			return true;
-		}
-
-		if (ConnectionStatus == FIoStoreCookOnTheFlyNetworkServer::EConnectionStatus::Connected)
+		const ITargetPlatform* TargetPlatform = Connection.GetTargetPlatform();
+		if (TargetPlatform)
 		{
-			const ITargetPlatform* TargetPlatform = CookOnTheFlyServer.AddPlatform(PlatformName);
-			if (!TargetPlatform)
-			{
-				return false;
-			}
-
 			IPackageStoreWriter* PackageWriter = CookOnTheFlyServer.GetPackageWriter(TargetPlatform).AsPackageStoreWriter();
 			check(PackageWriter); // This class should not be used except when COTFS is using an IPackageStoreWriter
 
+			FName PlatformName = Connection.GetPlatformName();
 			FScopeLock _(&ContextsCriticalSection);
-
 			if (!PlatformContexts.Contains(PlatformName))
 			{
 				TUniquePtr<FPlatformContext>& Context = PlatformContexts.Add(PlatformName, MakeUnique<FPlatformContext>(PlatformName, PackageWriter));
-					
+
 				PackageWriter->GetEntries([&Context](TArrayView<const FPackageStoreEntryResource> Entries,
 					TArrayView<const IPackageStoreWriter::FOplogCookInfo> CookInfos)
-				{
-					Context->AddExistingPackages(Entries, CookInfos);
-				});
+					{
+						Context->AddExistingPackages(Entries, CookInfos);
+					});
 
 				PackageWriter->OnEntryCreated().AddRaw(this, &FIoStoreCookOnTheFlyRequestManager::OnPackageStoreEntryCreated);
 				PackageWriter->OnCommit().AddRaw(this, &FIoStoreCookOnTheFlyRequestManager::OnPackageCooked);
 				PackageWriter->OnMarkUpToDate().AddRaw(this, &FIoStoreCookOnTheFlyRequestManager::OnPackagesMarkedUpToDate);
 			}
-
-			if (bIsSingleThreaded)
+			FPlatformContext& Context = GetContext(PlatformName);
+			FScopeLock __(&Context.GetLock());
+			if (Connection.GetIsSingleThreaded())
 			{
-				FPlatformContext& Context = GetContext(PlatformName);
-				FScopeLock __(&Context.GetLock());
 				Context.AddSingleThreadedClient();
 			}
+		}
 
-			return true;
+		UE_LOG(LogCookOnTheFly, Display, TEXT("New client connected"));
+		{
+			FScopeLock _(&ClientsCriticalSection);
+			Clients.Add(&Connection);
 		}
-		else
+	}
+
+	void OnClientDisconnected(UE::Cook::ICookOnTheFlyClientConnection& Connection)
+	{
 		{
-			if (bIsSingleThreaded)
+			FScopeLock _(&ClientsCriticalSection);
+			Clients.Remove(&Connection);
+		}
+		FName PlatformName = Connection.GetPlatformName();
+		if (!PlatformName.IsNone())
+		{
+			if (Connection.GetIsSingleThreaded())
 			{
 				FPlatformContext& Context = GetContext(PlatformName);
 				FScopeLock __(&Context.GetLock());
 				Context.RemoveSingleThreadedClient();
 			}
-			CookOnTheFlyServer.RemovePlatform(PlatformName);
-			return true;
 		}
 	}
 
-	bool HandleClientRequest(const FName& PlatformName, bool bIsSingleThreaded, const UE::Cook::FCookOnTheFlyRequest& Request, UE::Cook::FCookOnTheFlyResponse& Response)
+	bool HandleClientRequest(UE::Cook::ICookOnTheFlyClientConnection& Connection, const UE::Cook::FCookOnTheFlyRequest& Request)
 	{
-		bool bRequestOk = false;
+		using namespace UE::Cook;
 
 		const double StartTime = FPlatformTime::Seconds();
 
-		UE_LOG(LogCookOnTheFly, Verbose, TEXT("New request, Type='%s', Client='%s'"), LexToString(Request.GetHeader().MessageType), *PlatformName.ToString());
+		UE_LOG(LogCookOnTheFly, Verbose, TEXT("Received: %s, Size='%lld'"), *Request.GetHeader().ToString(), Request.TotalSize());
 
-		switch (Request.GetHeader().MessageType)
-		{
+		FCookOnTheFlyResponse Response(Request);
+		bool bRequestOk = false;
+
+		UE_LOG(LogCookOnTheFly, Verbose, TEXT("New request, Type='%s', Client='%s'"), LexToString(Request.GetHeader().MessageType), *Connection.GetPlatformName().ToString());
+
+		switch (Request.GetMessageType())
+		{ 
 			case UE::Cook::ECookOnTheFlyMessage::CookPackage:
-				bRequestOk = HandleCookPackageRequest(PlatformName, bIsSingleThreaded, Request, Response);
+				bRequestOk = HandleCookPackageRequest(Connection.GetPlatformName(), Connection.GetIsSingleThreaded(), Request, Response);
 				break;
 			case UE::Cook::ECookOnTheFlyMessage::GetCookedPackages:
-				bRequestOk = HandleGetCookedPackagesRequest(PlatformName, Request, Response);
-				break;
-			case UE::Cook::ECookOnTheFlyMessage::RecompileShaders:
-				bRequestOk = HandleRecompileShadersRequest(PlatformName, Request, Response);
+				bRequestOk = HandleGetCookedPackagesRequest(Connection.GetPlatformName(), Request, Response);
 				break;
 			case UE::Cook::ECookOnTheFlyMessage::RecookPackages:
-				bRequestOk = HandleRecookPackagesRequest(PlatformName, Request, Response);
+				bRequestOk = HandleRecookPackagesRequest(Connection.GetPlatformName(), Request, Response);
 				break;
 			default:
-				UE_LOG(LogCookOnTheFly, Fatal, TEXT("Unknown request, Type='%s', Client='%s'"), LexToString(Request.GetHeader().MessageType), *PlatformName.ToString());
-				break;
+				return false;
 		}
 
+		if (!bRequestOk)
+		{
+			Response.SetStatus(UE::Cook::ECookOnTheFlyMessageStatus::Error);
+		}
+		bRequestOk = Connection.SendMessage(Response);
+		
 		const double Duration = FPlatformTime::Seconds() - StartTime;
 
 		UE_LOG(LogCookOnTheFly, Verbose, TEXT("Request handled, Type='%s', Client='%s', Status='%s', Duration='%.6lfs'"),
-			LexToString(Request.GetHeader().MessageType),
-			*PlatformName.ToString(),
+			LexToString(Request.GetMessageType()),
+			*Connection.GetPlatformName().ToString(),
 			bRequestOk ? TEXT("Ok") : TEXT("Failed"),
 			Duration);
 
@@ -853,15 +469,17 @@ private:
 			Response.SetStatus(ECookOnTheFlyMessageStatus::Error);
 			return true;
 		}
-		
-		FCompletedPackages CompletedPackages;
+		else
 		{
+			FScopeLock _(&ContextsCriticalSection);
 			FPlatformContext& Context = GetContext(PlatformName);
-			FScopeLock _(&Context.GetLock());
+			FScopeLock __(&Context.GetLock());
+			FCompletedPackages CompletedPackages;
 			Context.GetCompletedPackages(CompletedPackages);
+
+			Response.SetBodyTo(MoveTemp(CompletedPackages));
 		}
 
-		Response.SetBodyTo( MoveTemp(CompletedPackages) );
 		Response.SetStatus(ECookOnTheFlyMessageStatus::Ok);
 
 		return true;
@@ -943,6 +561,40 @@ private:
 		return true;
 	}
 
+	bool BroadcastMessage(const UE::Cook::FCookOnTheFlyMessage& Message, const FName& PlatformName = NAME_None)
+	{
+		using namespace UE::Cook;
+
+		UE_LOG(LogCookOnTheFly, Verbose, TEXT("Sending: %s, Size='%lld'"), *Message.GetHeader().ToString(), Message.TotalSize());
+
+		TArray<ICookOnTheFlyClientConnection*, TInlineAllocator<4>> ClientsToBroadcast;
+		{
+			FScopeLock _(&ClientsCriticalSection);
+
+			for (ICookOnTheFlyClientConnection* Client : Clients)
+			{
+				if (PlatformName.IsNone() || Client->GetPlatformName() == PlatformName)
+				{
+					ClientsToBroadcast.Add(Client);
+				}
+			}
+		}
+
+		bool bBroadcasted = true;
+		for (ICookOnTheFlyClientConnection* Client : ClientsToBroadcast)
+		{
+			if (!Client->SendMessage(Message))
+			{
+				UE_LOG(LogCookOnTheFly, Warning, TEXT("Failed to send message '%s' to client (Platform='%s')"),
+					LexToString(Message.GetHeader().MessageType), *Client->GetPlatformName().ToString());
+
+				bBroadcasted = false;
+			}
+		}
+
+		return bBroadcasted;
+	}
+
 	void OnPackageStoreEntryCreated(const IPackageStoreWriter::FEntryCreatedEventArgs& EventArgs)
 	{
 		FPlatformContext& Context = GetContext(EventArgs.PlatformName);
@@ -986,7 +638,7 @@ private:
 				*Ar << ChunkIds;
 			}
 
-			ConnectionServer->BroadcastMessage(Message, EventArgs.PlatformName);
+			BroadcastMessage(Message, EventArgs.PlatformName);
 		}
 
 		FCompletedPackages NewCompletedPackages;
@@ -1043,63 +695,10 @@ private:
 
 			FCookOnTheFlyMessage Message(ECookOnTheFlyMessage::PackagesCooked);
 			Message.SetBodyTo<FCompletedPackages>(MoveTemp(NewCompletedPackages));
-			ConnectionServer->BroadcastMessage(Message, PlatformName);
-		}
-	}
-
-	bool HandleRecompileShadersRequest(const FName& PlatformName, const UE::Cook::FCookOnTheFlyRequest& Request, UE::Cook::FCookOnTheFlyResponse& Response)
-	{
-		TRACE_CPUPROFILER_EVENT_SCOPE(CookOnTheFly::HandleRecompileShadersRequest);
-
-		if (PlatformName.IsNone())
-		{
-			UE_LOG(LogCookOnTheFly, Warning, TEXT("RecompileShadersRequest from editor client"));
-			Response.SetStatus(UE::Cook::ECookOnTheFlyMessageStatus::Error);
-			return true;
-		}
-
-		TArray<FString> RecompileModifiedFiles;
-		TArray<uint8> MeshMaterialMaps;
-		TArray<uint8> GlobalShaderMap;
-
-		FShaderRecompileData RecompileData(PlatformName.ToString(), &RecompileModifiedFiles, &MeshMaterialMaps, &GlobalShaderMap);
-		{
-			TUniquePtr<FArchive> Ar = Request.ReadBody();
-			*Ar << RecompileData;
+			BroadcastMessage(Message, PlatformName);
 		}
-
-		RecompileShaders(RecompileData);
-
-		{
-			TUniquePtr<FArchive> Ar = Response.WriteBody();
-			*Ar << MeshMaterialMaps;
-			*Ar << GlobalShaderMap;
-		}
-
-		Response.SetStatus(UE::Cook::ECookOnTheFlyMessageStatus::Ok);
-
-		return true;
 	}
-
-	void RecompileShaders(FShaderRecompileData& RecompileData)
-	{
-		FEvent* RecompileCompletedEvent = FPlatformProcess::GetSynchEventFromPool();
-		UE::Cook::FRecompileShaderCompletedCallback RecompileCompleted = [this, RecompileCompletedEvent]()
-		{
-			RecompileCompletedEvent->Trigger();
-		};
-
-		const bool bEnqueued = CookOnTheFlyServer.EnqueueRecompileShaderRequest(UE::Cook::FRecompileShaderRequest
-		{
-			RecompileData, 
-			MoveTemp(RecompileCompleted)
-		});
-		check(bEnqueued);
-
-		RecompileCompletedEvent->Wait();
-		FPlatformProcess::ReturnSynchEventToPool(RecompileCompletedEvent);
-	}
-
+	
 	FPlatformContext& GetContext(const FName& PlatformName)
 	{
 		FScopeLock _(&ContextsCriticalSection);
@@ -1123,8 +722,11 @@ private:
 	}
 
 	UE::Cook::ICookOnTheFlyServer& CookOnTheFlyServer;
-	UE::Cook::FIoStoreCookOnTheFlyServerOptions Options;
-	TUniquePtr<FIoStoreCookOnTheFlyNetworkServer> ConnectionServer;
+	TSharedRef<UE::Cook::ICookOnTheFlyNetworkServer> ConnectionServer;
+	FCriticalSection ClientsCriticalSection;
+	TArray<UE::Cook::ICookOnTheFlyClientConnection*> Clients;
+	TSharedPtr<FMessageEndpoint, ESPMode::ThreadSafe> MessageEndpoint;
+	const FString ServiceId;
 	FCriticalSection ContextsCriticalSection;
 	TMap<FName, TUniquePtr<FPlatformContext>> PlatformContexts;
 	FCriticalSection AllKnownPackagesCriticalSection;
@@ -1136,9 +738,9 @@ private:
 namespace UE { namespace Cook
 {
 
-TUniquePtr<ICookOnTheFlyRequestManager> MakeIoStoreCookOnTheFlyRequestManager(ICookOnTheFlyServer& CookOnTheFlyServer, const IAssetRegistry* AssetRegistry, FIoStoreCookOnTheFlyServerOptions Options)
+TUniquePtr<ICookOnTheFlyRequestManager> MakeIoStoreCookOnTheFlyRequestManager(ICookOnTheFlyServer& CookOnTheFlyServer, const IAssetRegistry* AssetRegistry, TSharedRef<ICookOnTheFlyNetworkServer> ConnectionServer)
 {
-	return MakeUnique<FIoStoreCookOnTheFlyRequestManager>(CookOnTheFlyServer, AssetRegistry, Options);
+	return MakeUnique<FIoStoreCookOnTheFlyRequestManager>(CookOnTheFlyServer, AssetRegistry, ConnectionServer);
 }
 
 }}
diff --git a/Engine/Source/Editor/UnrealEd/Private/Cooker/IoStoreCookOnTheFlyRequestManager.h b/Engine/Source/Editor/UnrealEd/Private/Cooker/IoStoreCookOnTheFlyRequestManager.h
--- a/Engine/Source/Editor/UnrealEd/Private/Cooker/IoStoreCookOnTheFlyRequestManager.h
+++ b/Engine/Source/Editor/UnrealEd/Private/Cooker/IoStoreCookOnTheFlyRequestManager.h
@@ -11,12 +11,8 @@ namespace UE { namespace Cook
 
 class ICookOnTheFlyServer;
 class ICookOnTheFlyRequestManager;
+class ICookOnTheFlyNetworkServer;
 
-struct FIoStoreCookOnTheFlyServerOptions
-{
-	int32 Port = -1; // -1 indicates the default COTF serving port
-};
-
-TUniquePtr<ICookOnTheFlyRequestManager> MakeIoStoreCookOnTheFlyRequestManager(ICookOnTheFlyServer& CookOnTheFlyServer, const IAssetRegistry* AssetRegistry, FIoStoreCookOnTheFlyServerOptions Options);
+TUniquePtr<ICookOnTheFlyRequestManager> MakeIoStoreCookOnTheFlyRequestManager(ICookOnTheFlyServer& CookOnTheFlyServer, const IAssetRegistry* AssetRegistry, TSharedRef<ICookOnTheFlyNetworkServer> ConnectionServer);
 
 }} // namespace UE::Cook
diff --git a/Engine/Source/Editor/UnrealEd/Private/Cooker/NetworkFileCookOnTheFlyRequestManager.cpp b/Engine/Source/Editor/UnrealEd/Private/Cooker/NetworkFileCookOnTheFlyRequestManager.cpp
--- a/Engine/Source/Editor/UnrealEd/Private/Cooker/NetworkFileCookOnTheFlyRequestManager.cpp
+++ b/Engine/Source/Editor/UnrealEd/Private/Cooker/NetworkFileCookOnTheFlyRequestManager.cpp
@@ -10,18 +10,16 @@ class FNetworkFileCookOnTheFlyRequestManager final
 	: public UE::Cook::ICookOnTheFlyRequestManager
 {
 public:
-	FNetworkFileCookOnTheFlyRequestManager(UE::Cook::ICookOnTheFlyServer& InCookOnTheFlyServer, FNetworkFileServerOptions FileServerOptions)
+	FNetworkFileCookOnTheFlyRequestManager(UE::Cook::ICookOnTheFlyServer& InCookOnTheFlyServer, TSharedRef<UE::Cook::ICookOnTheFlyNetworkServer> NetworkServer)
 		: CookOnTheFlyServer(InCookOnTheFlyServer)
 	{
-		FileServerOptions.Delegates.NewConnectionDelegate			= FNewConnectionDelegate::CreateRaw(this, &FNetworkFileCookOnTheFlyRequestManager::OnNewConnection);
-		FileServerOptions.Delegates.FileRequestDelegate				= FFileRequestDelegate::CreateRaw(this, &FNetworkFileCookOnTheFlyRequestManager::OnFileRequest);
-		FileServerOptions.Delegates.RecompileShadersDelegate		= FRecompileShadersDelegate::CreateRaw(this, &FNetworkFileCookOnTheFlyRequestManager::OnRecompileShaders);
-		FileServerOptions.Delegates.SandboxPathOverrideDelegate		= FSandboxPathDelegate::CreateRaw(this, &FNetworkFileCookOnTheFlyRequestManager::OnGetSandboxPath);
-		FileServerOptions.Delegates.OnFileModifiedCallback			= &FileModifiedDelegate;
-		FileServerOptions.bRestrictPackageAssetsToSandbox			= true; // prevents sending uncooked packages if the package fails to cook
+		FNetworkFileDelegateContainer Delegates;
+		Delegates.FileRequestDelegate				= FFileRequestDelegate::CreateRaw(this, &FNetworkFileCookOnTheFlyRequestManager::OnFileRequest);
+		Delegates.SandboxPathOverrideDelegate		= FSandboxPathDelegate::CreateRaw(this, &FNetworkFileCookOnTheFlyRequestManager::OnGetSandboxPath);
+		Delegates.OnFileModifiedCallback			= &FileModifiedDelegate;
 
 		NetworkFileServer.Reset(FModuleManager::LoadModuleChecked<INetworkFileSystemModule>("NetworkFileSystem")
-			.CreateNetworkFileServer(MoveTemp(FileServerOptions), /* bLoadTargetPlatforms */ true));
+			.CreateNetworkFileServer(NetworkServer, Delegates));
 	}
 
 	virtual ~FNetworkFileCookOnTheFlyRequestManager()
@@ -57,11 +55,6 @@ public:
 	}
 
 private:
-	bool OnNewConnection(const FString& VersionInfo, const FString& PlatformName)
-	{
-		return CookOnTheFlyServer.AddPlatform(FName(*PlatformName)) != nullptr;
-	}
-
 	void OnFileRequest(FString& Filename, const FString& PlatformNameStr, TArray<FString>& UnsolicitedFiles)
 	{
 		using namespace UE::Cook;
@@ -95,25 +88,6 @@ private:
 		}
 	}
 
-	void OnRecompileShaders(const FShaderRecompileData& RecompileData)
-	{
-		FEvent* RecompileCompletedEvent = FPlatformProcess::GetSynchEventFromPool();
-		UE::Cook::FRecompileShaderCompletedCallback RecompileCompleted = [this, RecompileCompletedEvent]()
-		{
-			RecompileCompletedEvent->Trigger();
-		};
-
-		const bool bEnqueued = CookOnTheFlyServer.EnqueueRecompileShaderRequest(UE::Cook::FRecompileShaderRequest
-		{
-			RecompileData,
-			MoveTemp(RecompileCompleted)
-		});
-		check(bEnqueued);
-
-		RecompileCompletedEvent->Wait();
-		FPlatformProcess::ReturnSynchEventToPool(RecompileCompletedEvent);
-	}
-
 	FString OnGetSandboxPath()
 	{
 		return CookOnTheFlyServer.GetSandboxDirectory();
@@ -127,9 +101,9 @@ private:
 namespace UE { namespace Cook
 {
 
-TUniquePtr<ICookOnTheFlyRequestManager> MakeNetworkFileCookOnTheFlyRequestManager(ICookOnTheFlyServer& CookOnTheFlyServer, const FNetworkFileServerOptions& FileServerOptions)
+TUniquePtr<ICookOnTheFlyRequestManager> MakeNetworkFileCookOnTheFlyRequestManager(ICookOnTheFlyServer& CookOnTheFlyServer, TSharedRef<ICookOnTheFlyNetworkServer> NetworkServer)
 {
-	return MakeUnique<FNetworkFileCookOnTheFlyRequestManager>(CookOnTheFlyServer, FileServerOptions);
+	return MakeUnique<FNetworkFileCookOnTheFlyRequestManager>(CookOnTheFlyServer, NetworkServer);
 }
 
 }} // namespace UE::Cook
diff --git a/Engine/Source/Editor/UnrealEd/Private/Cooker/NetworkFileCookOnTheFlyRequestManager.h b/Engine/Source/Editor/UnrealEd/Private/Cooker/NetworkFileCookOnTheFlyRequestManager.h
--- a/Engine/Source/Editor/UnrealEd/Private/Cooker/NetworkFileCookOnTheFlyRequestManager.h
+++ b/Engine/Source/Editor/UnrealEd/Private/Cooker/NetworkFileCookOnTheFlyRequestManager.h
@@ -11,7 +11,8 @@ namespace UE { namespace Cook
 
 class ICookOnTheFlyServer;
 class ICookOnTheFlyRequestManager;
+class ICookOnTheFlyNetworkServer;
 
-TUniquePtr<ICookOnTheFlyRequestManager> MakeNetworkFileCookOnTheFlyRequestManager(ICookOnTheFlyServer& CookOnTheFlyServer, const FNetworkFileServerOptions& FileServerOptions);
+TUniquePtr<ICookOnTheFlyRequestManager> MakeNetworkFileCookOnTheFlyRequestManager(ICookOnTheFlyServer& CookOnTheFlyServer, TSharedRef<ICookOnTheFlyNetworkServer> NetworkServer);
 
 }} // namespace UE::Cook
diff --git a/Engine/Source/Editor/UnrealEd/UnrealEd.Build.cs b/Engine/Source/Editor/UnrealEd/UnrealEd.Build.cs
--- a/Engine/Source/Editor/UnrealEd/UnrealEd.Build.cs
+++ b/Engine/Source/Editor/UnrealEd/UnrealEd.Build.cs
@@ -212,6 +212,7 @@ public class UnrealEd : ModuleRules
 				"AnimationBlueprintLibrary",
 				"MaterialBaking",
 				"CookOnTheFly",
+				"CookOnTheFlyNetServer",
 				"Zen",
 			}
 		);
diff --git a/Engine/Source/Runtime/CookOnTheFly/CookOnTheFly.Build.cs b/Engine/Source/Runtime/CookOnTheFly/CookOnTheFly.Build.cs
--- a/Engine/Source/Runtime/CookOnTheFly/CookOnTheFly.Build.cs
+++ b/Engine/Source/Runtime/CookOnTheFly/CookOnTheFly.Build.cs
@@ -9,7 +9,8 @@ public class CookOnTheFly : ModuleRules
         PublicDependencyModuleNames.AddRange(
             new string[] {
                 "Core",
-                "Sockets"
+                "Sockets",
+                "Networking"
             }
         );
     }
diff --git a/Engine/Source/Runtime/CookOnTheFly/Internal/CookOnTheFly.h b/Engine/Source/Runtime/CookOnTheFly/Internal/CookOnTheFly.h
--- a/Engine/Source/Runtime/CookOnTheFly/Internal/CookOnTheFly.h
+++ b/Engine/Source/Runtime/CookOnTheFly/Internal/CookOnTheFly.h
@@ -8,18 +8,16 @@
 #include "Misc/Timespan.h"
 #include "Serialization/MemoryReader.h"
 #include "Serialization/MemoryWriter.h"
+#include "Delegates/Delegate.h"
+
+class FArrayReader;
+class FInternetAddr;
 
 COOKONTHEFLY_API DECLARE_LOG_CATEGORY_EXTERN(LogCookOnTheFly, Log, All);
 
 namespace UE { namespace Cook
 {
 
-enum
-{ 
-	/* The default port used by the the cook-on-the-fly server. */
-	DefaultCookOnTheFlyServingPort = 42899
-};
-
 /**
  * Flags and message types to be used with the cook-on-the-fly server.
  *
@@ -30,16 +28,12 @@ enum class ECookOnTheFlyMessage : uint32
 	/* Represents no message. */
 	None				= 0x00,
 	
-	/* A one way message. */
-	Message				= 0x01,
 	/* A request message. */
 	Request				= 0x02,
 	/* A response message. */
 	Response			= 0x04,
 	TypeFlags			= 0x0F,
 
-	/* The handshake request message. */
-	Handshake			= 0x10,
 	/* Request to cook a package. */
 	CookPackage			= 0x20,
 	/* Get all currenlty cooked packages. */
@@ -52,6 +46,8 @@ enum class ECookOnTheFlyMessage : uint32
 	FilesAdded			= 0x60,
 	/* Request to recook packages. */
 	RecookPackages		= 0x70,
+	/* Legacy message for NetworkPlatformFile */
+	NetworkPlatformFile	= 0x80,
 };
 ENUM_CLASS_FLAGS(ECookOnTheFlyMessage);
 
@@ -66,8 +62,6 @@ inline const TCHAR* LexToString(ECookOnTheFlyMessage Message)
 	{
 		case ECookOnTheFlyMessage::None:
 			return TEXT("None");
-		case ECookOnTheFlyMessage::Handshake:
-			return TEXT("Handshake");
 		case ECookOnTheFlyMessage::CookPackage:
 			return TEXT("CookPackage");
 		case ECookOnTheFlyMessage::GetCookedPackages:
@@ -80,6 +74,8 @@ inline const TCHAR* LexToString(ECookOnTheFlyMessage Message)
 			return TEXT("FilesAdded");
 		case ECookOnTheFlyMessage::RecookPackages:
 			return TEXT("RecookPackages");
+		case ECookOnTheFlyMessage::NetworkPlatformFile:
+			return TEXT("NetworkPlatformFile");
 		default:
 			return TEXT("Unknown");
 	};
@@ -124,9 +120,7 @@ struct FCookOnTheFlyMessageHeader
 	/** Type of message */
 	ECookOnTheFlyMessage MessageType = ECookOnTheFlyMessage::None;
 	/** The message status. */
-	ECookOnTheFlyMessageStatus MessageStatus = ECookOnTheFlyMessageStatus::None;
-	/** Sender id */
-	uint32 SenderId = 0;
+	ECookOnTheFlyMessageStatus MessageStatus = ECookOnTheFlyMessageStatus::Ok;
 	/** Correlation id, used to match response with request. */
 	uint32 CorrelationId = 0;
 	/** When the message was sent. */
@@ -152,6 +146,13 @@ public:
 		Header.MessageType = MessageType;
 	}
 
+	ECookOnTheFlyMessage GetMessageType() const
+	{
+		ECookOnTheFlyMessage MessageType = Header.MessageType;
+		EnumRemoveFlags(MessageType, ECookOnTheFlyMessage::TypeFlags);
+		return MessageType;
+	}
+
 	/** Returns the message header. */
 	FCookOnTheFlyMessageHeader& GetHeader()
 	{
@@ -242,9 +243,47 @@ protected:
 	TArray<uint8> Body;
 };
 
-using FCookOnTheFlyRequest = FCookOnTheFlyMessage;
+/**
+ * Cook-on-the-fly request.
+ */
+class FCookOnTheFlyRequest
+	: public FCookOnTheFlyMessage
+{
+public:
+	/** Creates a new instance of a cook-on-the-fly request. */
+	COOKONTHEFLY_API FCookOnTheFlyRequest() = default;
 
-using FCookOnTheFlyResponse = FCookOnTheFlyMessage;
+	/** Creates a new instance of a cook-on-the-fly request with the specified request type. */
+	explicit FCookOnTheFlyRequest(ECookOnTheFlyMessage MessageType)
+		: FCookOnTheFlyMessage(MessageType | ECookOnTheFlyMessage::Request)
+	{
+	}
+};
+
+/**
+ * Cook-on-the-fly response.
+ */
+class FCookOnTheFlyResponse
+	: public FCookOnTheFlyMessage
+{
+public:
+	/** Creates a new instance of a cook-on-the-fly response. */
+	COOKONTHEFLY_API FCookOnTheFlyResponse() = default;
+
+	/** Creates a new instance of a cook-on-the-fly response with the specified response type. */
+	explicit FCookOnTheFlyResponse(ECookOnTheFlyMessage MessageType)
+		: FCookOnTheFlyMessage(MessageType | ECookOnTheFlyMessage::Response)
+	{
+	}
+
+	/** Creates a new instance of a cook-on-the-fly response for the specified request. */
+	explicit FCookOnTheFlyResponse(const FCookOnTheFlyRequest& Request)
+		: FCookOnTheFlyMessage(Request.GetMessageType())
+	{
+		Header.MessageType |= ECookOnTheFlyMessage::Response;
+		Header.CorrelationId = Request.GetHeader().CorrelationId;
+	}
+};
 
 /**
  * Cook-on-the-fly host address.
@@ -253,39 +292,34 @@ struct FCookOnTheFlyHostOptions
 {
 	/** Host address. */
 	TArray<FString> Hosts;
-	/** Host port. */
-	int32 Port = DefaultCookOnTheFlyServingPort;
 	/** How long to wait for the server to start. */
 	FTimespan ServerStartupWaitTime;
 };
 
-/**
- * A connection used to communicate with the cook-on-the-fly server.
- */
 class ICookOnTheFlyServerConnection
 {
 public:
 	virtual ~ICookOnTheFlyServerConnection() { }
 
+	virtual const FString& GetHost() const = 0;
+
+	virtual const FString& GetZenProjectName() const = 0;
+
+	virtual const FString& GetPlatformName() const = 0;
+
 	/**
 	 * Returns whether connected to the cook-on-the-fly server.
 	 */
 	virtual bool IsConnected() const = 0;
-	
-	/**
-	 * Disconnect from the server.
-	 */
-	virtual void Disconnect() = 0;
+
+	virtual bool IsSingleThreaded() const = 0;
 
 	/**
 	 * Sends a request to the server.
 	 *
 	 * @param Request The request message to send.
-	 * @param FillRequest Callback to populate the request message..
-	 * @param ProcessResponse Callback to process the response message.
 	 */
-	//virtual bool SendRequest(ECookOnTheFlyMessage Request, FFillRequest&& FillRequest, FProcessResponse&& ProcessResponse) = 0;
-	virtual TFuture<FCookOnTheFlyResponse> SendRequest(const FCookOnTheFlyRequest& Request) = 0;
+	virtual TFuture<FCookOnTheFlyResponse> SendRequest(FCookOnTheFlyRequest& Request) = 0;
 
 	/**
 	 * Event triggered when a new message has been sent from the server.
@@ -303,6 +337,8 @@ class ICookOnTheFlyModule
 public:
 	virtual ~ICookOnTheFlyModule() { }
 
+	virtual TSharedPtr<ICookOnTheFlyServerConnection> GetDefaultServerConnection() = 0;
+
 	/**
 	 * Connect to the cook-on-the-fly server.
 	 *
diff --git a/Engine/Source/Runtime/CookOnTheFly/Private/CookOnTheFly.cpp b/Engine/Source/Runtime/CookOnTheFly/Private/CookOnTheFly.cpp
--- a/Engine/Source/Runtime/CookOnTheFly/Private/CookOnTheFly.cpp
+++ b/Engine/Source/Runtime/CookOnTheFly/Private/CookOnTheFly.cpp
@@ -11,10 +11,9 @@ namespace UE { namespace Cook
 
 FString FCookOnTheFlyMessageHeader::ToString() const
 {
-	return FString::Printf(TEXT("Message='%s', Status='%s', Sender='%u', CorrelationId='%u'"),
+	return FString::Printf(TEXT("Message='%s', Status='%s', CorrelationId='%u'"),
 		LexToString(MessageType),
 		LexToString(MessageStatus),
-		SenderId,
 		CorrelationId);
 }
 
@@ -25,7 +24,6 @@ FArchive& operator<<(FArchive& Ar, FCookOnTheFlyMessageHeader& Header)
 	
 	Ar << MessageType;
 	Ar << MessageStatus;
-	Ar << Header.SenderId;
 	Ar << Header.CorrelationId;
 	Ar << Header.Timestamp;
 
diff --git a/Engine/Source/Runtime/CookOnTheFly/Private/CookOnTheFlyModule.cpp b/Engine/Source/Runtime/CookOnTheFly/Private/CookOnTheFlyModule.cpp
--- a/Engine/Source/Runtime/CookOnTheFly/Private/CookOnTheFlyModule.cpp
+++ b/Engine/Source/Runtime/CookOnTheFly/Private/CookOnTheFlyModule.cpp
@@ -1,17 +1,15 @@
 // Copyright Epic Games, Inc. All Rights Reserved.
 
-#include "CookOnTheFly.h"
+#include "CookOnTheFlyServerConnection.h"
+#include "TCPTransport.h"
+#include "PlatformTransport.h"
 #include "Modules/ModuleManager.h"
+#include "GenericPlatform/GenericPlatformHostCommunication.h"
+#include "HAL/PlatformTime.h"
+#include "Misc/CommandLine.h"
 
 DEFINE_LOG_CATEGORY(LogCookOnTheFly);
 
-namespace UE { namespace Cook
-{
-
-TUniquePtr<ICookOnTheFlyServerConnection> MakeServerConnection(const UE::Cook::FCookOnTheFlyHostOptions& HostOptions);
-
-}}
-
 class FCookOnTheFlyModule final
 	: public UE::Cook::ICookOnTheFlyModule
 {
@@ -31,9 +29,105 @@ public:
 	// ICookOnTheFlyModule interface
 	//
 
+	virtual TSharedPtr<UE::Cook::ICookOnTheFlyServerConnection> GetDefaultServerConnection() override
+	{
+		static TSharedPtr<UE::Cook::ICookOnTheFlyServerConnection> DefaultServerConnection = MakeShareable(CreateServerConnection(GetDefaultHostOptions()));
+		return DefaultServerConnection;
+	}
+
 	virtual TUniquePtr<UE::Cook::ICookOnTheFlyServerConnection> ConnectToServer(const UE::Cook::FCookOnTheFlyHostOptions& HostOptions) override
 	{
-		return UE::Cook::MakeServerConnection(HostOptions);
+		return TUniquePtr<UE::Cook::ICookOnTheFlyServerConnection>(CreateServerConnection(HostOptions));
+	}
+
+private:
+	TUniquePtr<ICookOnTheFlyServerTransport> CreateTransportForHostAddress(const FString& HostIp)
+	{
+		if (HostIp.StartsWith(TEXT("tcp://")))
+		{
+			return MakeUnique<FTCPTransport>();
+		}
+
+		if (HostIp.StartsWith(TEXT("platform://")))
+		{
+			if (FPlatformMisc::GetPlatformHostCommunication().Available())
+			{
+				return MakeUnique<FPlatformTransport>(EHostProtocol::CookOnTheFly, "CookOnTheFly");
+			}
+			else
+			{
+				UE_LOG(LogCookOnTheFly, Warning, TEXT("Platform transport (platform://) not supported for this platform."));
+
+				return nullptr;
+			}
+		}
+
+		// no transport specified assuming tcp
+		return MakeUnique<FTCPTransport>();
+	}
+
+	UE::Cook::ICookOnTheFlyServerConnection* CreateServerConnection(const UE::Cook::FCookOnTheFlyHostOptions& HostOptions)
+	{
+		if (HostOptions.Hosts.IsEmpty())
+		{
+			return nullptr;
+		}
+
+		const double ServerWaitEndTime = FPlatformTime::Seconds() + HostOptions.ServerStartupWaitTime.GetTotalSeconds();
+
+		for (;;)
+		{
+			for (const FString& Host : HostOptions.Hosts)
+			{
+				// Try to initialize with each of the IP addresses found in the command line until we 
+				// get a working one.
+
+				// find the correct transport for this ip address 
+				TUniquePtr<ICookOnTheFlyServerTransport> Transport = CreateTransportForHostAddress(Host);
+				if (Transport)
+				{
+					UE_LOG(LogCookOnTheFly, Verbose, TEXT("Created transport for %s."), *Host);
+					if (Transport->Initialize(*Host))
+					{
+						UE_LOG(LogCookOnTheFly, Display, TEXT("Created transport for %s."), *Host);
+						return MakeCookOnTheFlyServerConnection(MoveTemp(Transport), Host);
+					}
+					UE_LOG(LogCookOnTheFly, Verbose, TEXT("Failed to initialize %s."), *Host);
+				}
+			}
+
+			if (FPlatformTime::Seconds() > ServerWaitEndTime)
+			{
+				break;
+			}
+
+			FPlatformProcess::Sleep(1.0f);
+		};
+
+		UE_LOG(LogCookOnTheFly, Error, TEXT("Failed to connect to COTF server"));
+		return nullptr;
+	}
+
+	UE::Cook::FCookOnTheFlyHostOptions GetDefaultHostOptions()
+	{
+		TArray<FString> HostIpList;
+		FTimespan ServerStartupWaitTime;
+		FString HostIpString;
+		if (!FParse::Value(FCommandLine::Get(), TEXT("-FileHostIP="), HostIpString))
+		{
+			if (FParse::Param(FCommandLine::Get(), TEXT("CookOnTheFly")))
+			{
+				FParse::Value(FCommandLine::Get(), TEXT("-ZenStoreHost="), HostIpString);
+			}
+		}
+		if (HostIpString.ParseIntoArray(HostIpList, TEXT("+"), true) > 0)
+		{
+			double ServerWaitTimeInSeconds = 0.0;
+			FParse::Value(FCommandLine::Get(), TEXT("-CookOnTheFlyServerWaitTime="), ServerWaitTimeInSeconds);
+			ServerStartupWaitTime = FTimespan::FromSeconds(ServerWaitTimeInSeconds);
+		}
+
+		return { HostIpList, ServerStartupWaitTime };
 	}
 };
 
diff --git a/Engine/Source/Runtime/CookOnTheFly/Private/CookOnTheFlyServerConnection.cpp b/Engine/Source/Runtime/CookOnTheFly/Private/CookOnTheFlyServerConnection.cpp
--- a/Engine/Source/Runtime/CookOnTheFly/Private/CookOnTheFlyServerConnection.cpp
+++ b/Engine/Source/Runtime/CookOnTheFly/Private/CookOnTheFlyServerConnection.cpp
@@ -1,5 +1,6 @@
 // Copyright Epic Games, Inc. All Rights Reserved.
 
+#include "CookOnTheFlyServerConnection.h"
 #include "CookOnTheFly.h"
 #include "HAL/PlatformMisc.h"
 #include "Async/Async.h"
@@ -11,109 +12,57 @@
 #include "Misc/App.h"
 #include "Misc/DateTime.h"
 #include "Serialization/MemoryReader.h"
+#include "Misc/CoreDelegates.h"
 
 DEFINE_LOG_CATEGORY_STATIC(LogCotfServerConnection, Log, All);
 
-static TArray<TSharedPtr<FInternetAddr>> GetAddressFromString(ISocketSubsystem& SocketSubsystem, TArrayView<const FString> HostAddresses, const int32 Port)
-{
-	TArray<TSharedPtr<FInternetAddr>> InterntAddresses;
-
-	for (const FString& HostAddr : HostAddresses)
-	{
-		TSharedPtr<FInternetAddr> Addr = SocketSubsystem.GetAddressFromString(HostAddr);
-
-		if (!Addr.IsValid() || !Addr->IsValid())
-		{
-			FAddressInfoResult GAIRequest = SocketSubsystem.GetAddressInfo(*HostAddr, nullptr, EAddressInfoFlags::Default, NAME_None);
-			if (GAIRequest.ReturnCode == SE_NO_ERROR && GAIRequest.Results.Num() > 0)
-			{
-				Addr = GAIRequest.Results[0].Address;
-			}
-		}
-
-		if (Addr.IsValid() && Addr->IsValid())
-		{
-			Addr->SetPort(Port);
-			InterntAddresses.Emplace(MoveTemp(Addr));
-		}
-	}
-
-	return InterntAddresses;
-}
-
 class FCookOnTheFlyServerConnection final
 	: public UE::Cook::ICookOnTheFlyServerConnection
 {
 public:
-	FCookOnTheFlyServerConnection()
-		: bIsSingleThreaded(!FGenericPlatformProcess::SupportsMultithreading())
+	FCookOnTheFlyServerConnection(TUniquePtr<ICookOnTheFlyServerTransport> InTransport, const FString& InHost)
+		: Transport(MoveTemp(InTransport))
+		, Host(InHost)
+		, bIsSingleThreaded(!FGenericPlatformProcess::SupportsMultithreading())
 	{
-	}
-
-	~FCookOnTheFlyServerConnection()
-	{
-		Disconnect();
-	}
-
-	bool Connect(const UE::Cook::FCookOnTheFlyHostOptions& HostOptions)
-	{
-		TRACE_CPUPROFILER_EVENT_SCOPE(FCookOnTheFlyServerConnection::Connect);
-		check(HostOptions.Hosts.Num());
-
-		const int32 Port = HostOptions.Port > 0 ? HostOptions.Port : UE::Cook::DefaultCookOnTheFlyServingPort;
+		using namespace UE::Cook;
 
-		ISocketSubsystem& SocketSubsystem = *ISocketSubsystem::Get();
-		TArray<TSharedPtr<FInternetAddr>> HostAddresses = GetAddressFromString(SocketSubsystem, HostOptions.Hosts, Port);
+		check(Transport);
 
-		if (!HostAddresses.Num())
+		TArray<FString> TargetPlatformNames;
+		if (FPlatformProperties::RequiresCookedData())
 		{
-			UE_LOG(LogCotfServerConnection, Error, TEXT("No valid COTF server address found"));
-			return false;
+			FPlatformMisc::GetValidTargetPlatforms(TargetPlatformNames);
 		}
+		
+		FBufferArchive HandshakeRequestPayload;
+		HandshakeRequestPayload << TargetPlatformNames;
+		bool bIsSingleThreaded_ = bIsSingleThreaded;
+		HandshakeRequestPayload << bIsSingleThreaded_;
 
-		bool bConnected = false;
-		const double ServerWaitEndTime = FPlatformTime::Seconds() + HostOptions.ServerStartupWaitTime.GetTotalSeconds();
-
-		for (;;)
+		if (!Transport->SendPayload(HandshakeRequestPayload))
 		{
-			for (const TSharedPtr<FInternetAddr>& Addr : HostAddresses)
-			{
-				UE_LOG(LogCotfServerConnection, Display, TEXT("Connecting to COTF server at '%s'..."), *Addr->ToString(true));
-			
-				Socket.Reset(SocketSubsystem.CreateSocket(NAME_Stream, TEXT("COTF-ServerConnection"), Addr->GetProtocolType()));
-				if (!Socket.IsValid())
-				{
-					continue;
-				}
-				Socket->SetNoDelay();
-				if (!Socket->Connect(*Addr))
-				{
-					continue;
-				}
-				ServerAddr = Addr;
-				bConnected = true;
-				break;
-			}
-
-			if (bConnected || FPlatformTime::Seconds() > ServerWaitEndTime)
-			{
-				break;
-			}
-
-			FPlatformProcess::Sleep(1.0f);
-		};
+			UE_LOG(LogCotfServerConnection, Fatal, TEXT("Failed to send handshake request to server"));
+			return;
+		}
 
-		if (!bConnected)
+		FArrayReader ResponsePayload;
+		if (!Transport->ReceivePayload(ResponsePayload))
 		{
-			UE_LOG(LogCotfServerConnection, Error, TEXT("Failed to connect to COTF server"));
-			return false;
+			UE_LOG(LogCotfServerConnection, Fatal, TEXT("Failed to receive handshake response from server"));
+			return;
 		}
-
-		if (!SendHandshakeMessage())
+		bool bIsOk;
+		ResponsePayload << bIsOk;
+		if (!bIsOk)
 		{
-			UE_LOG(LogCotfServerConnection, Fatal, TEXT("Failed to handshake with COTF server at '%s'"), *ServerAddr->ToString(true));
-			return false;
+			UE_LOG(LogCotfServerConnection, Fatal, TEXT("Couldn't handshake with server"));
+			return;
 		}
+		ResponsePayload << PlatformName;
+		ResponsePayload << ZenProjectName;
+
+		FCoreDelegates::OnEnginePreExit.AddRaw(this, &FCookOnTheFlyServerConnection::OnEnginePreExit);
 
 		Thread = MakeUnique<FThread>(
 			TEXT("CotfServerConnection"),
@@ -126,79 +75,89 @@ public:
 				SingleThreadedTickFunction();
 			},
 			8 * 1024);
-		
-		UE_LOG(LogCotfServerConnection, Log, TEXT("Connected to COTF server at '%s'"), *ServerAddr->ToString(true));
+	}
 
-		return true;
+	~FCookOnTheFlyServerConnection()
+	{
+		FCoreDelegates::OnEnginePreExit.RemoveAll(this);
+		Disconnect();
+	}
+
+	virtual const FString& GetHost() const override
+	{
+		return Host;
+	}
+
+	virtual const FString& GetZenProjectName() const override
+	{
+		return ZenProjectName;
+	}
+
+	virtual const FString& GetPlatformName() const override
+	{
+		return PlatformName;
+	}
+
+	virtual bool IsSingleThreaded() const override
+	{
+		return bIsSingleThreaded;
 	}
 
 	virtual bool IsConnected() const override
 	{
-		return ClientId > 0;
+		return Transport.IsValid();
 	}
 
-	virtual TFuture<UE::Cook::FCookOnTheFlyResponse> SendRequest(const UE::Cook::FCookOnTheFlyRequest& Request) override
+	virtual TFuture<UE::Cook::FCookOnTheFlyResponse> SendRequest(UE::Cook::FCookOnTheFlyRequest& Request) override
 	{
 		using namespace UE::Cook;
 
 		const uint32 CorrelationId					= NextCorrelationId++;
-		FCookOnTheFlyMessageHeader RequestHeader	= Request.GetHeader();
+		FCookOnTheFlyMessageHeader& RequestHeader	= Request.GetHeader();
 
-		RequestHeader.MessageType	= RequestHeader.MessageType | ECookOnTheFlyMessage::Request;
 		RequestHeader.MessageStatus = ECookOnTheFlyMessageStatus::Ok;
-		RequestHeader.SenderId		= ClientId;
 		RequestHeader.CorrelationId = CorrelationId;
 		RequestHeader.Timestamp		= FDateTime::UtcNow().GetTicks();
 
 		FBufferArchive RequestPayload;
 		RequestPayload.Reserve(Request.TotalSize());
 
-		RequestPayload << RequestHeader;
-		RequestPayload << const_cast<TArray<uint8>&>(Request.GetBody());
+		RequestPayload << Request;
 
 		FPendingRequest* PendingRequest = Alloc(CorrelationId);
 		PendingRequest->RequestHeader = RequestHeader;
 		
 		TFuture<FCookOnTheFlyResponse> FutureResponse = PendingRequest->ResponsePromise.GetFuture();
 
-		UE_LOG(LogCotfServerConnection, Verbose, TEXT("Sending: %s, Size='%lld'"), *RequestHeader.ToString(), Request.TotalSize());
-
-		if (SendMessage(RequestPayload))
+		bool bOk = false;
+		if (Transport)
 		{
-			if (bIsSingleThreaded)
+			UE_LOG(LogCotfServerConnection, Verbose, TEXT("Sending: %s, Size='%lld'"), *RequestHeader.ToString(), Request.TotalSize());
+
+			FScopeLock _(&TransportCritical);
+			if (Transport->SendPayload(RequestPayload))
 			{
-				ProcessMessagesWhile([this, CorrelationId]()
-					{
-						return PendingRequests.Contains(CorrelationId);
-					});
+				if (bIsSingleThreaded)
+				{
+					ProcessMessagesWhile([this, CorrelationId]()
+						{
+							return PendingRequests.Contains(CorrelationId);
+						});
+				}
+				bOk = true;
 			}
-			return FutureResponse;
 		}
-		else
+		if (!bOk)
 		{
 			UE_LOG(LogCotfServerConnection, Warning, TEXT("Failed to send: %s, Size='%lld'"), *RequestHeader.ToString(), Request.TotalSize());
 
-			Free(PendingRequest);
-
-			FCookOnTheFlyResponse ErrorResponse;
+			FCookOnTheFlyResponse ErrorResponse(Request);
 			ErrorResponse.SetStatus(ECookOnTheFlyMessageStatus::Error);
+			PendingRequest->ResponsePromise.SetValue(ErrorResponse);
 
-			TPromise<FCookOnTheFlyResponse> ErrorResponsePromise;
-			ErrorResponsePromise.SetValue(ErrorResponse);
-			return ErrorResponsePromise.GetFuture();
-		}
-	}
-
-	virtual void Disconnect() override
-	{
-		if (!bStopRequested.Exchange(true))
-		{
-			Socket->Close();
-			Thread->Join();
-			Thread.Reset();
-			Socket.Reset();
-			ClientId  = 0;
+			Free(PendingRequest);
 		}
+		return FutureResponse;
 	}
 
 	DECLARE_DERIVED_EVENT(FCookOnTheFlyServerConnection, UE::Cook::ICookOnTheFlyServerConnection::FMessageEvent, FMessageEvent);
@@ -207,75 +166,18 @@ public:
 		return MessageEvent;
 	}
 
-	bool SendMessage(const TArray<uint8>& Message)
-	{
-		if (!FNFSMessageHeader::WrapAndSendPayload(Message, FSimpleAbstractSocket_FSocket(Socket.Get())))
-		{
-			UE_LOG(LogCotfServerConnection, Fatal, TEXT("Failed sending payload to COTF server"));
-			return false;
-		}
-
-		return true;
-	}
-
-	bool ReceiveMessage(FArrayReader& Message)
-	{
-		if (!FNFSMessageHeader::ReceivePayload(Message, FSimpleAbstractSocket_FSocket(Socket.Get())))
-		{
-			UE_LOG(LogCotfServerConnection, Warning, TEXT("Failed reveiving payload from COTF server"));
-			return false;
-		}
-
-		return true;
-	}
-
-	bool SendHandshakeMessage()
+private:
+	void Disconnect()
 	{
-		using namespace UE::Cook;
-
-		FCookOnTheFlyMessage HandshakeRequest(ECookOnTheFlyMessage::Handshake | ECookOnTheFlyMessage::Request);
+		if (!bStopRequested.Exchange(true))
 		{
-			FString PlatformName;
-			if (FPlatformProperties::RequiresCookedData())
+			if (Transport)
 			{
-				TArray<FString> TargetPlatformNames;
-				FPlatformMisc::GetValidTargetPlatforms(TargetPlatformNames);
-				check(TargetPlatformNames.Num() > 0);
-				PlatformName = MoveTemp(TargetPlatformNames[0]);
+				Transport->Disconnect();
 			}
-			FString ProjectName(FApp::GetProjectName());
-
-			TUniquePtr<FArchive> Ar = HandshakeRequest.WriteBody();
-			*Ar << PlatformName;
-			*Ar << ProjectName;
-			bool bLocalIsSingleThreded = bIsSingleThreaded;
-			*Ar << bLocalIsSingleThreded;
-		}
-
-		FBufferArchive HandshakeRequestPayload;
-		HandshakeRequestPayload << HandshakeRequest;
-
-		if (!SendMessage(HandshakeRequestPayload))
-		{
-			return false;
-		}
-
-		FArrayReader HandshakeResponsePayload;
-		if (!ReceiveMessage(HandshakeResponsePayload))
-		{
-			return false;
-		}
-
-		FCookOnTheFlyMessage HandshakeResponse;
-		HandshakeResponsePayload << HandshakeResponse;
-		{
-			TUniquePtr<FArchive> Ar = HandshakeResponse.ReadBody();
-			*Ar << ClientId;
+			Thread->Join();
+			Thread.Reset();
 		}
-
-		UE_CLOG(ClientId > 0, LogCotfServerConnection, Display, TEXT("Connected to server with ID='%d'"), ClientId);
-
-		return ClientId > 0;
 	}
 
 	bool ProcessMessagesWhile(TFunctionRef<bool()> Condition)
@@ -285,10 +187,9 @@ public:
 		while (Condition())
 		{
 			FArrayReader MessagePayload;
-			if (!ReceiveMessage(MessagePayload))
+			if (!Transport || !Transport->ReceivePayload(MessagePayload))
 			{
-				UE_LOG(LogCotfServerConnection, Warning, TEXT("Failed to receive message from '%s'"), *ServerAddr->ToString(true));
-				ClientId = 0;
+				UE_LOG(LogCotfServerConnection, Warning, TEXT("Failed to receive message from server"));
 				return false;
 			}
 
@@ -302,7 +203,6 @@ public:
 
 			const bool bIsResponse = EnumHasAnyFlags(MessageHeader.MessageType, ECookOnTheFlyMessage::Response);
 			const bool bIsRequest = EnumHasAnyFlags(MessageHeader.MessageType, ECookOnTheFlyMessage::Request);
-			EnumRemoveFlags(MessageHeader.MessageType, ECookOnTheFlyMessage::TypeFlags);
 
 			if (bIsRequest)
 			{
@@ -345,16 +245,15 @@ public:
 				return !bStopRequested.Load();
 			});
 
-		UE_LOG(LogCotfServerConnection, Display, TEXT("Terminating connection to server '%s'"), *ServerAddr->ToString(true));
-		ClientId = 0;
+		UE_LOG(LogCotfServerConnection, Display, TEXT("Terminating connection to server"));
+		Transport.Reset();
 	}
 
 	void SingleThreadedTickFunction()
 	{
 		ProcessMessagesWhile([this]()
 			{
-				uint32 PendingDataSize;
-				return Socket->HasPendingData(PendingDataSize);
+				return Transport && Transport->HasPendingPayload();
 			});
 	}
 
@@ -389,30 +288,27 @@ public:
 		return nullptr;
 	}
 
+	void OnEnginePreExit()
+	{
+		Disconnect();
+	}
+
+	FCriticalSection TransportCritical;
+	TUniquePtr<ICookOnTheFlyServerTransport> Transport;
+	FString Host;
 	FMessageEvent MessageEvent;
-	TSharedPtr<FInternetAddr> ServerAddr;
-	TUniquePtr<FSocket> Socket;
-	uint32 ClientId = 0;
 	TUniquePtr<FThread> Thread;
-	TAtomic<bool> bStopRequested { false };
+	TAtomic<bool> bStopRequested{ false };
 	const bool bIsSingleThreaded;
+	FString PlatformName;
+	FString ZenProjectName;
 
 	FCriticalSection RequestsCriticalSection;
 	TMap<uint32, TUniquePtr<FPendingRequest>> PendingRequests;
 	TAtomic<uint32> NextCorrelationId { 1 };
 };
 
-namespace UE { namespace Cook
-{
-
-TUniquePtr<ICookOnTheFlyServerConnection> MakeServerConnection(const UE::Cook::FCookOnTheFlyHostOptions& HostOptions)
+UE::Cook::ICookOnTheFlyServerConnection* MakeCookOnTheFlyServerConnection(TUniquePtr<ICookOnTheFlyServerTransport> InTransport, const FString& InHost)
 {
-	TUniquePtr<FCookOnTheFlyServerConnection> Connection = MakeUnique<FCookOnTheFlyServerConnection>();
-	if (!Connection->Connect(HostOptions))
-	{
-		Connection.Release();
-	}
-	return Connection;
-}
-
-}} // namespace UE::Cook
+	return new FCookOnTheFlyServerConnection(MoveTemp(InTransport), InHost);
+}
\ No newline at end of file
diff --git a/Engine/Source/Runtime/CookOnTheFly/Private/CookOnTheFlyServerConnection.h b/Engine/Source/Runtime/CookOnTheFly/Private/CookOnTheFlyServerConnection.h
new file mode 100644
--- /dev/null
+++ b/Engine/Source/Runtime/CookOnTheFly/Private/CookOnTheFlyServerConnection.h
@@ -0,0 +1,23 @@
+// Copyright Epic Games, Inc. All Rights Reserved.
+
+#pragma once
+
+#include "CookOnTheFly.h"
+
+class FArrayReader;
+
+class ICookOnTheFlyServerTransport
+{
+public:
+	ICookOnTheFlyServerTransport() {};
+	virtual ~ICookOnTheFlyServerTransport() {};
+
+	virtual bool Initialize(const TCHAR*) = 0;
+	virtual void Disconnect() = 0;
+	virtual bool SendPayload(const TArray<uint8>& Payload) = 0;
+	virtual bool HasPendingPayload() = 0;
+	virtual bool ReceivePayload(FArrayReader& Payload) = 0;
+
+};
+
+UE::Cook::ICookOnTheFlyServerConnection* MakeCookOnTheFlyServerConnection(TUniquePtr<ICookOnTheFlyServerTransport> InTransport, const FString& InHost);
diff --git a/Engine/Source/Runtime/CookOnTheFly/Private/PlatformTransport.cpp b/Engine/Source/Runtime/CookOnTheFly/Private/PlatformTransport.cpp
new file mode 100644
--- /dev/null
+++ b/Engine/Source/Runtime/CookOnTheFly/Private/PlatformTransport.cpp
@@ -0,0 +1,140 @@
+// Copyright Epic Games, Inc. All Rights Reserved.
+
+#include "PlatformTransport.h"
+#include "NetworkMessage.h"
+#include "HAL/PlatformProcess.h"
+#include "GenericPlatform/GenericPlatformHostSocket.h"
+#include "GenericPlatform/GenericPlatformHostCommunication.h"
+#include "CookOnTheFly.h"
+
+namespace
+{
+	/**
+	 * Socket abstraction required by FNFSMessageHeader that uses IPlatformHostSocket.
+	 */
+	class FSimpleAbstractSocket_PlatformProtocol : public FSimpleAbstractSocket
+	{
+	public:
+
+		FSimpleAbstractSocket_PlatformProtocol(IPlatformHostSocketPtr InHostSocket)
+			: HostSocket(InHostSocket)
+		{
+			check(HostSocket != nullptr);
+		}
+
+		virtual bool Receive(uint8* Results, int32 Size) const override
+		{
+			return (HostSocket->Receive(Results, Size) == IPlatformHostSocket::EResultNet::Ok);
+		}
+
+		virtual bool Send(const uint8* Buffer, int32 Size) const override
+		{
+			return (HostSocket->Send(Buffer, Size) == IPlatformHostSocket::EResultNet::Ok);
+		}
+
+		virtual uint32 GetMagic() const override
+		{
+			return 0x9E2B83C7;
+		}
+
+	private:
+		IPlatformHostSocketPtr HostSocket;
+	};
+}
+
+
+FPlatformTransport::FPlatformTransport(int32 InProtocolIndex, const FString& InProtocolName)
+	: ProtocolIndex(InProtocolIndex)
+	, ProtocolName(InProtocolName)
+	, HostSocket(nullptr)
+{
+}
+
+
+bool FPlatformTransport::Initialize(const TCHAR* InHostIp)
+{
+	check(HostSocket == nullptr);
+
+	IPlatformHostCommunication& HostCommunication = FPlatformMisc::GetPlatformHostCommunication();
+
+	if (!HostCommunication.Available())
+	{
+		return false;
+	}
+
+	HostSocket = HostCommunication.OpenConnection(ProtocolIndex, ProtocolName);
+
+	if (!HostSocket)
+	{
+		return false;
+	}
+
+	UE_LOG(LogCookOnTheFly, Display, TEXT("Waiting for the server to accept the connection (custom protocol)..."));
+
+	// We need to wait because this transport is a bit different than typical sockets.
+	// In this case, it's the client (game) that enables the communication by opening the socket.
+	// The server (pc) can connect only if it detects this enabled communication protocol.
+	return WaitUntilConnected();
+}
+
+
+bool FPlatformTransport::SendPayload(const TArray<uint8>& Payload)
+{
+	return FNFSMessageHeader::WrapAndSendPayload(Payload, FSimpleAbstractSocket_PlatformProtocol(HostSocket));
+}
+
+
+bool FPlatformTransport::ReceivePayload(FArrayReader& Payload)
+{
+	return FNFSMessageHeader::ReceivePayload(Payload, FSimpleAbstractSocket_PlatformProtocol(HostSocket));
+}
+
+
+bool FPlatformTransport::WaitUntilConnected()
+{
+	while (true)
+	{
+		IPlatformHostSocket::EConnectionState ConnectionState = HostSocket->GetState();
+
+		// We're interested in the Connected state, but break the loop in other non-default states
+		// in case any error has appeared.
+		if (HostSocket->GetState() != IPlatformHostSocket::EConnectionState::Created)
+		{
+			break;
+		}
+
+		FPlatformProcess::Sleep(0.5f);
+	}
+
+	return HostSocket->GetState() == IPlatformHostSocket::EConnectionState::Connected;
+}
+
+bool FPlatformTransport::HasPendingPayload()
+{
+	checkNoEntry();
+	return false;
+}
+
+void FPlatformTransport::Disconnect()
+{
+	IPlatformHostCommunication& HostCommunication = FPlatformMisc::GetPlatformHostCommunication();
+	if (HostCommunication.Available())
+	{
+		HostCommunication.CloseConnection(HostSocket);
+	}
+}
+
+FPlatformTransport::~FPlatformTransport()
+{
+	if (HostSocket)
+	{
+		IPlatformHostCommunication& HostCommunication = FPlatformMisc::GetPlatformHostCommunication();
+
+		if (HostCommunication.Available())
+		{
+			HostCommunication.CloseConnection(HostSocket);
+
+			HostSocket = nullptr;
+		}
+	}
+}
diff --git a/Engine/Source/Runtime/CookOnTheFly/Private/PlatformTransport.h b/Engine/Source/Runtime/CookOnTheFly/Private/PlatformTransport.h
new file mode 100644
--- /dev/null
+++ b/Engine/Source/Runtime/CookOnTheFly/Private/PlatformTransport.h
@@ -0,0 +1,40 @@
+// Copyright Epic Games, Inc. All Rights Reserved.
+
+#pragma once
+
+#include "CookOnTheFlyServerConnection.h"
+#include "GenericPlatform/GenericPlatformHostSocket.h"
+
+
+/**
+ * Implementation of ITransport using IPlatformHostCommunication/IPlatformHostSocket interfaces (custom target and host pc communication).
+ */
+class FPlatformTransport
+	: public ICookOnTheFlyServerTransport
+{
+public:
+
+	FPlatformTransport(int32 InProtocolIndex, const FString& InProtocolName);
+	~FPlatformTransport();
+
+	//~ Begin ICookOnTheFlyServerTransport interface
+	virtual bool Initialize(const TCHAR* HostIp) override;
+	virtual bool SendPayload(const TArray<uint8>& Payload) override;
+	virtual bool ReceivePayload(FArrayReader& Payload) override;
+	virtual bool HasPendingPayload() override;
+	virtual void Disconnect() override;
+	//~ End ICookOnTheFlyServerTransport interface
+
+private:
+
+	/**
+	 * Wait until HostSocket is in a non-default state (preferably Connected).
+	 * @return True is the socket is Connected, false otherwise (an error or if the host pc immediately disconnects).
+	 */
+	bool WaitUntilConnected();
+
+	int32				   ProtocolIndex;
+	FString				   ProtocolName;
+	IPlatformHostSocketPtr HostSocket;
+
+};
diff --git a/Engine/Source/Runtime/CookOnTheFly/Private/TCPTransport.cpp b/Engine/Source/Runtime/CookOnTheFly/Private/TCPTransport.cpp
new file mode 100644
--- /dev/null
+++ b/Engine/Source/Runtime/CookOnTheFly/Private/TCPTransport.cpp
@@ -0,0 +1,114 @@
+// Copyright Epic Games, Inc. All Rights Reserved.
+
+#include "TCPTransport.h"
+#include "NetworkMessage.h"
+#include "SocketSubsystem.h"
+#include "IPAddress.h"
+#include "Sockets.h"
+#include "MultichannelTcpSocket.h"
+#include "Interfaces/IPv4/IPv4Endpoint.h"
+#include "CookOnTheFly.h"
+
+FTCPTransport::FTCPTransport()
+	:FileSocket(NULL)
+	,MCSocket(NULL)
+{
+}
+
+bool FTCPTransport::Initialize(const TCHAR* InHostIp)
+{
+	ISocketSubsystem* SSS = ISocketSubsystem::Get();
+
+	FString HostIp = InHostIp;
+
+
+	// the protocol isn't required for tcp it's assumed by default
+	HostIp.RemoveFromStart(TEXT("tcp://"));
+
+	// convert the string to a ip addr structure
+	// DEFAULT_TCP_FILE_SERVING_PORT is overridden 
+	TSharedRef<FInternetAddr> Addr = SSS->CreateInternetAddr();
+	bool bIsValid;
+
+	FIPv4Endpoint FileServerEndpoint;
+	if (FIPv4Endpoint::Parse(InHostIp, FileServerEndpoint) || FIPv4Endpoint::FromHostAndPort(InHostIp, FileServerEndpoint))
+	{
+		Addr->SetIp(FileServerEndpoint.Address.Value);
+		Addr->SetPort(FileServerEndpoint.Port == 0 ? DEFAULT_TCP_FILE_SERVING_PORT : FileServerEndpoint.Port);
+		bIsValid = true;
+	}
+	else
+	{
+		Addr->SetIp(*HostIp, bIsValid);
+		Addr->SetPort(DEFAULT_TCP_FILE_SERVING_PORT);
+	}
+
+	if (bIsValid)
+	{
+		// create the socket
+		FileSocket = SSS->CreateSocket(NAME_Stream, TEXT("COTF tcp"), Addr->GetProtocolType());
+
+		// try to connect to the server
+		if (FileSocket == nullptr || FileSocket->Connect(*Addr) == false)
+		{
+			// on failure, shut it all down
+			SSS->DestroySocket(FileSocket);
+			FileSocket = NULL;
+			UE_LOG(LogCookOnTheFly, Error, TEXT("Failed to connect to COTF server at %s."), *Addr->ToString(true));
+		}
+	}
+
+#if USE_MCSOCKET_FOR_NFS
+	MCSocket = new FMultichannelTcpSocket(FileSocket, 64 * 1024 * 1024);
+#endif
+
+	return ( FileSocket || MCSocket );
+}
+
+
+bool FTCPTransport::SendPayload(const TArray<uint8>& Payload)
+{
+#if USE_MCSOCKET_FOR_NFS
+	return FNFSMessageHeader::WrapAndSendPayload(Payload, FSimpleAbstractSocket_FMultichannelTCPSocket(MCSocket, NFS_Channels::Main));
+#else 	
+	return FNFSMessageHeader::WrapAndSendPayload(Payload, FSimpleAbstractSocket_FSocket(FileSocket));
+#endif 
+}
+
+
+bool FTCPTransport::ReceivePayload(FArrayReader& Payload)
+{
+#if USE_MCSOCKET_FOR_NFS
+	return FNFSMessageHeader::ReceivePayload(Payload, FSimpleAbstractSocket_FMultichannelTCPSocket(MCSocket, NFS_Channels::Main));
+#else
+	return FNFSMessageHeader::ReceivePayload(Payload, FSimpleAbstractSocket_FSocket(FileSocket));
+#endif
+}
+
+bool FTCPTransport::HasPendingPayload()
+{
+#if USE_MCSOCKET_FOR_NFS
+	return MCSocket->DataAvailable(NFS_Channels::Main) > 0;
+#else
+	uint32 PendingDataSize;
+	return FileSocket->HasPendingData(PendingDataSize);
+#endif
+}
+
+void FTCPTransport::Disconnect()
+{
+#if USE_MCSOCKET_FOR_NFS
+	checkNoEntry();
+#else
+	FileSocket->Close();
+#endif
+}
+
+FTCPTransport::~FTCPTransport()
+{
+	// on failure, shut it all down
+	delete MCSocket;
+	MCSocket = NULL;
+	ISocketSubsystem::Get()->DestroySocket(FileSocket);
+	FileSocket = NULL;
+}
diff --git a/Engine/Source/Runtime/CookOnTheFly/Private/TCPTransport.h b/Engine/Source/Runtime/CookOnTheFly/Private/TCPTransport.h
new file mode 100644
--- /dev/null
+++ b/Engine/Source/Runtime/CookOnTheFly/Private/TCPTransport.h
@@ -0,0 +1,28 @@
+// Copyright Epic Games, Inc. All Rights Reserved.
+
+#pragma once
+
+#include "CookOnTheFlyServerConnection.h"
+
+class FTCPTransport
+	: public ICookOnTheFlyServerTransport
+{
+
+public:
+
+	FTCPTransport();
+	~FTCPTransport();
+
+	// ITransport Interface. 
+	virtual bool Initialize(const TCHAR* HostIp) override;
+	virtual bool SendPayload(const TArray<uint8>& Payload) override;
+	virtual bool ReceivePayload(FArrayReader& Payload) override;
+	virtual bool HasPendingPayload() override;
+	virtual void Disconnect() override;
+
+private: 
+
+	class FSocket*		FileSocket;
+	class FMultichannelTcpSocket* MCSocket;
+
+};
diff --git a/Engine/Source/Runtime/NetworkFile/NetworkFile.Build.cs b/Engine/Source/Runtime/NetworkFile/NetworkFile.Build.cs
--- a/Engine/Source/Runtime/NetworkFile/NetworkFile.Build.cs
+++ b/Engine/Source/Runtime/NetworkFile/NetworkFile.Build.cs
@@ -13,7 +13,8 @@ public class NetworkFile : ModuleRules
 				"Core",
 				"CoreUObject",
 				"Networking",
-				"Sockets"
+				"Sockets",
+				"CookOnTheFly"
 			});
 
 		PublicIncludePaths.Add("Runtime/CoreUObject/Public/UObject");
diff --git a/Engine/Source/Runtime/NetworkFile/Private/ITransport.h b/Engine/Source/Runtime/NetworkFile/Private/ITransport.h
deleted file mode 100644
diff --git a/Engine/Source/Runtime/NetworkFile/Private/NetworkPlatformFile.cpp b/Engine/Source/Runtime/NetworkFile/Private/NetworkPlatformFile.cpp
--- a/Engine/Source/Runtime/NetworkFile/Private/NetworkPlatformFile.cpp
+++ b/Engine/Source/Runtime/NetworkFile/Private/NetworkPlatformFile.cpp
@@ -17,18 +17,13 @@
 #include "Misc/PackageName.h"
 #include "Misc/PathViews.h"
 #include "SocketSubsystem.h"
-#include "GenericPlatform/GenericPlatformHostCommunication.h"
-
-#if ENABLE_HTTP_FOR_NETWORK_FILE
-#include "HTTPTransport.h"
-#endif
-#include "TCPTransport.h"
-#include "PlatformTransport.h"
 
 #include "HAL/IPlatformFileModule.h"
 #include "Templates/UniquePtr.h"
 
 #include "UObject/Object.h"
+#include "CookOnTheFly.h"
+#include "CookOnTheFlyMessages.h"
 
 DEFINE_LOG_CATEGORY(LogNetworkPlatformFile);
 
@@ -46,7 +41,6 @@ FNetworkPlatformFile::FNetworkPlatformFile()
 	, HeartbeatFrequency(5.0f)
 	, FinishedAsyncNetworkReadUnsolicitedFiles(NULL)
 	, FinishedAsyncWriteUnsolicitedFiles(NULL)
-	, Transport(NULL)
 {
 	TotalWriteTime = 0.0; // total non async time spent writing to disk
 	TotalNetworkSyncTime = 0.0; // total non async time spent syncing to network
@@ -62,77 +56,31 @@ FNetworkPlatformFile::FNetworkPlatformFile()
 bool FNetworkPlatformFile::ShouldBeUsed(IPlatformFile* Inner, const TCHAR* CmdLine) const
 {
 	FString HostIp;
-	return FParse::Value(CmdLine, TEXT("-FileHostIP="), HostIp);
-}
-
-ITransport *CreateTransportForHostAddress(const FString &HostIp )
-{
-	if ( HostIp.StartsWith(TEXT("tcp://")))
+	if (!FParse::Value(CmdLine, TEXT("-FileHostIP="), HostIp))
 	{
-		return new FTCPTransport();
-	}
-
-	if ( HostIp.StartsWith(TEXT("http://")))
-	{
-#if ENABLE_HTTP_FOR_NETWORK_FILE
-		return new FHTTPTransport();
-#endif
-	}
-
-	if ( HostIp.StartsWith(TEXT("platform://")))
-	{
-		if (FPlatformMisc::GetPlatformHostCommunication().Available())
-		{
-			return new FPlatformTransport(EHostProtocol::CookOnTheFly, "CookOnTheFly");
-		}
-		else
-		{
-			UE_LOG(LogNetworkPlatformFile, Warning, TEXT("Platform transport (platform://) not supported for this platform."));
-
-			return nullptr;
-		}
+		return false;
 	}
-
-	// no transport specified assuming tcp
-	return new FTCPTransport();
+	UE::Cook::ICookOnTheFlyModule& CookOnTheFlyModule = FModuleManager::LoadModuleChecked<UE::Cook::ICookOnTheFlyModule>(TEXT("CookOnTheFly"));
+	TSharedPtr<UE::Cook::ICookOnTheFlyServerConnection> DefaultConnection = CookOnTheFlyModule.GetDefaultServerConnection();
+	return !DefaultConnection.IsValid() || DefaultConnection->GetZenProjectName().IsEmpty();
 }
 
 bool FNetworkPlatformFile::Initialize(IPlatformFile* Inner, const TCHAR* CmdLine)
 {
-	bool bResult = false;
-	FString HostIpString;	
-	if (FParse::Value(CmdLine, TEXT("-FileHostIP="), HostIpString))
+	UE::Cook::ICookOnTheFlyModule& CookOnTheFlyModule = FModuleManager::LoadModuleChecked<UE::Cook::ICookOnTheFlyModule>(TEXT("CookOnTheFly"));
+	Connection = CookOnTheFlyModule.GetDefaultServerConnection();
+	if (!Connection.IsValid())
 	{
-		TArray<FString> HostIpList;
-		if (HostIpString.ParseIntoArray(HostIpList, TEXT("+"), true) > 0)
-		{
-			for (int32 HostIpIndex = 0; !bResult && HostIpIndex < HostIpList.Num(); ++HostIpIndex)
-			{
-				// Try to initialize with each of the IP addresses found in the command line until we 
-				// get a working one.
-
-				// find the correct transport for this ip address 
-				Transport = CreateTransportForHostAddress( HostIpList[HostIpIndex] );
-
-				UE_LOG(LogNetworkPlatformFile, Warning, TEXT("Created transport for %s."), *HostIpList[HostIpIndex]);
-
-				if ( Transport )
-				{
-					bResult = Transport->Initialize( *HostIpList[HostIpIndex] ) && InitializeInternal(Inner, *HostIpList[HostIpIndex]);		
-					if (bResult)
-						break;
-
-					UE_LOG(LogNetworkPlatformFile, Warning, TEXT("Failed to initialize %s."), *HostIpList[HostIpIndex]);
-
-					// try a different host might be a different protocol
-					delete Transport;
-				}
-				Transport = NULL;
-			}
-		}
+		UE_LOG(LogNetworkPlatformFile, Warning, TEXT("No COTF server connection."));
+		return false;
 	}
-
-	return bResult;
+	if (!InitializeInternal(Inner, *Connection->GetHost()))
+	{
+		UE_LOG(LogNetworkPlatformFile, Warning, TEXT("Failed to initialize %s."), *Connection->GetHost());
+		return false;
+	}
+	Connection->OnMessage().AddRaw(this, &FNetworkPlatformFile::OnCookOnTheFlyMessage);
+	return true;
 }
 
 bool FNetworkPlatformFile::InitializeInternal(IPlatformFile* Inner, const TCHAR* HostIP)
@@ -184,6 +132,9 @@ bool FNetworkPlatformFile::InitializeInternal(IPlatformFile* Inner, const TCHAR*
 
 bool FNetworkPlatformFile::SendPayloadAndReceiveResponse(TArray<uint8>& In, TArray<uint8>& Out)
 {
+	using namespace UE::Cook;
+	using namespace UE::ZenCookOnTheFly::Messaging;
+
 	{
 		FScopeLock ScopeLock(&SynchronizationObject);
 		if ( FinishedAsyncNetworkReadUnsolicitedFiles )
@@ -192,13 +143,53 @@ bool FNetworkPlatformFile::SendPayloadAndReceiveResponse(TArray<uint8>& In, TArr
 			FinishedAsyncNetworkReadUnsolicitedFiles = NULL;
 		}
 	}
-	
-	return Transport->SendPayloadAndReceiveResponse( In, Out );
+
+	FCookOnTheFlyRequest Request(ECookOnTheFlyMessage::NetworkPlatformFile);
+	TUniquePtr<FArchive> RequestPayload = Request.WriteBody();
+	RequestPayload->Serialize(In.GetData(), In.Num());
+	FCookOnTheFlyResponse Response = Connection->SendRequest(Request).Get();
+
+	if (!Response.IsOk())
+	{
+		UE_LOG(LogCookOnTheFly, Warning, TEXT("Failed to send 'NetworkPlatformFile' request"));
+		return false;
+	}
+
+	TUniquePtr<FArchive> ResponsePayload = Response.ReadBody();
+	Out.SetNum(ResponsePayload->TotalSize());
+	ResponsePayload->Serialize(Out.GetData(), ResponsePayload->TotalSize());
+	return true;
+}
+
+void FNetworkPlatformFile::OnCookOnTheFlyMessage(const UE::Cook::FCookOnTheFlyMessage& Message)
+{
+	using namespace UE::Cook;
+	using namespace UE::ZenCookOnTheFly::Messaging;
+
+	if (Message.GetHeader().MessageType == ECookOnTheFlyMessage::NetworkPlatformFile)
+	{
+		TUniquePtr<FArchive> PayloadReader = Message.ReadBody();
+		TArray<uint8> Payload;
+		Payload.SetNum(PayloadReader->TotalSize());
+		PayloadReader->Serialize(Payload.GetData(), PayloadReader->TotalSize());
+		PendingPayloads.Enqueue(MoveTemp(Payload));
+		NewPayloadEvent->Trigger();
+	}
 }
 
 bool FNetworkPlatformFile::ReceiveResponse(TArray<uint8> &Out )
 {
-	return Transport->ReceiveResponse( Out );
+	for (;;)
+	{
+		NewPayloadEvent->Wait();
+		TOptional<TArray<uint8>> Payload = PendingPayloads.Dequeue();
+		if (!Payload.IsSet())
+		{
+			Out = MoveTemp(Payload.GetValue());
+			break;
+		}
+	}
+	return true;
 }
 
 
@@ -216,7 +207,7 @@ void FNetworkPlatformFile::InitializeAfterSetActive()
 		FArrayReader Response;
 		if (!SendPayloadAndReceiveResponse(Payload, Response))
 		{
-			delete Transport; 
+			Connection.Reset();
 			return; 
 		}
 		else
@@ -416,7 +407,7 @@ FNetworkPlatformFile::~FNetworkPlatformFile()
 			FinishedAsyncWriteUnsolicitedFiles = NULL;
 		}
 		
-		delete Transport; // close our sockets.
+		Connection.Reset(); // close our sockets.
 	}
 }
 
@@ -892,46 +883,30 @@ bool FNetworkPlatformFile::SendWriteMessage(const uint8* Source, int64 BytesToWr
 
 bool FNetworkPlatformFile::SendMessageToServer(const TCHAR* Message, IPlatformFile::IFileServerMessageHandler* Handler)
 {
-	// handle the recompile shaders message
-	// @todo: Maybe we should just send the string message to the server, but then we'd have to 
-	// handle the return from the server in a generic way
+#if WITH_COTF
+	if (!Connection->IsConnected())
+	{
+		return false;
+	}
 	if (FCString::Stricmp(Message, TEXT("RecompileShaders")) == 0)
 	{
-		FNetworkFileArchive Payload(NFS_Messages::RecompileShaders);
-
-		// let the handler fill out the object
-		Handler->FillPayload(Payload);
-
-		FArrayReader Response;
+		UE::Cook::FCookOnTheFlyRequest Request(UE::Cook::ECookOnTheFlyMessage::RecompileShaders);
 		{
-			FScopeLock ScopeLock(&SynchronizationObject);
-			if (!SendPayloadAndReceiveResponse(Payload, Response))
-			{
-				return false;
-			}
+			TUniquePtr<FArchive> Ar = Request.WriteBody();
+			Handler->FillPayload(*Ar);
 		}
 
-		// locally delete any files that were modified on the server, so that any read will recache the file
-		// this has to be done in this class, not in the Handler (which can't access these members)
-		TArray<FString> ModifiedFiles;
-		Response << ModifiedFiles;
-
-		if( InnerPlatformFile != NULL )
+		UE::Cook::FCookOnTheFlyResponse Response = Connection->SendRequest(Request).Get();
+		if (Response.IsOk())
 		{
-			for (int32 Index = 0; Index < ModifiedFiles.Num(); Index++)
-			{
-				InnerPlatformFile->DeleteFile(*ModifiedFiles[Index]);
-				CachedLocalFiles.Remove(ModifiedFiles[Index]);
-				ServerFiles.AddFileOrDirectory(ModifiedFiles[Index], FDateTime::UtcNow());
-			}
+			TUniquePtr<FArchive> Ar = Response.ReadBody();
+			Handler->ProcessResponse(*Ar);
 		}
 
-
-		// let the handler process the response directly
-		Handler->ProcessResponse(Response);
+		return Response.IsOk();
 	}
-
-	return true;
+#endif
+	return false;
 }
 
 
diff --git a/Engine/Source/Runtime/NetworkFile/Private/PlatformTransport.cpp b/Engine/Source/Runtime/NetworkFile/Private/PlatformTransport.cpp
deleted file mode 100644
diff --git a/Engine/Source/Runtime/NetworkFile/Private/PlatformTransport.h b/Engine/Source/Runtime/NetworkFile/Private/PlatformTransport.h
deleted file mode 100644
diff --git a/Engine/Source/Runtime/NetworkFile/Private/TCPTransport.cpp b/Engine/Source/Runtime/NetworkFile/Private/TCPTransport.cpp
deleted file mode 100644
diff --git a/Engine/Source/Runtime/NetworkFile/Private/TCPTransport.h b/Engine/Source/Runtime/NetworkFile/Private/TCPTransport.h
deleted file mode 100644
diff --git a/Engine/Source/Runtime/NetworkFile/Public/NetworkPlatformFile.h b/Engine/Source/Runtime/NetworkFile/Public/NetworkPlatformFile.h
--- a/Engine/Source/Runtime/NetworkFile/Public/NetworkPlatformFile.h
+++ b/Engine/Source/Runtime/NetworkFile/Public/NetworkPlatformFile.h
@@ -7,9 +7,17 @@
 #include "NetworkMessage.h"
 #include "ServerTOC.h"
 #include "Misc/CoreMisc.h" // included for FSelfRegisteringExec
+#include "Containers/SpscQueue.h"
+#include "HAL/Event.h"
 
 class FScopedEvent;
 
+namespace UE { namespace Cook
+{
+	class ICookOnTheFlyServerConnection;
+	class FCookOnTheFlyMessage;
+}}
+
 DECLARE_LOG_CATEGORY_EXTERN(LogNetworkPlatformFile, Log, All);
 
 
@@ -222,6 +230,8 @@ private:
 	 */
 	void EnsureFileIsLocal(const FString& Filename);
 
+	void OnCookOnTheFlyMessage(const UE::Cook::FCookOnTheFlyMessage& Message);
+
 protected:
 	/**
 	* Does normal path standardization, and also any extra modifications to make string comparisons against
@@ -289,8 +299,9 @@ private:
 	FScopedEvent *FinishedAsyncNetworkReadUnsolicitedFiles;
 	FScopedEvent *FinishedAsyncWriteUnsolicitedFiles;
 
-    // Our network Transport. 
-	class ITransport* Transport; 
+	TSharedPtr<UE::Cook::ICookOnTheFlyServerConnection> Connection;
+	TSpscQueue<TArray<uint8>> PendingPayloads;
+	FEventRef NewPayloadEvent;
 
 	static FString MP4Extension;
 	static FString BulkFileExtension;
diff --git a/Engine/Source/Runtime/NetworkFileSystem/NetworkFileSystem.Build.cs b/Engine/Source/Runtime/NetworkFileSystem/NetworkFileSystem.Build.cs
--- a/Engine/Source/Runtime/NetworkFileSystem/NetworkFileSystem.Build.cs
+++ b/Engine/Source/Runtime/NetworkFileSystem/NetworkFileSystem.Build.cs
@@ -20,6 +20,8 @@ namespace UnrealBuildTool.Rules
 					"SandboxFile",
 					"TargetPlatform",
 					"DesktopPlatform",
+					"CookOnTheFly",
+					"CookOnTheFlyNetServer"
 				});
 
 			PublicIncludePaths.AddRange(
diff --git a/Engine/Source/Runtime/NetworkFileSystem/Private/NetworkFileServer.cpp b/Engine/Source/Runtime/NetworkFileSystem/Private/NetworkFileServer.cpp
--- a/Engine/Source/Runtime/NetworkFileSystem/Private/NetworkFileServer.cpp
+++ b/Engine/Source/Runtime/NetworkFileSystem/Private/NetworkFileServer.cpp
@@ -1,7 +1,6 @@
 // Copyright Epic Games, Inc. All Rights Reserved.
 
 #include "NetworkFileServer.h"
-#include "HAL/RunnableThread.h"
 #include "Misc/OutputDeviceRedirector.h"
 #include "IPAddress.h"
 #include "Sockets.h"
@@ -9,268 +8,106 @@
 #include "NetworkMessage.h"
 #include "NetworkFileSystemLog.h"
 #include "NetworkFileServerConnection.h"
+#include "CookOnTheFlyNetServer.h"
+#include "CookOnTheFly.h"
+#include "CookOnTheFlyMessages.h"
 
-class FNetworkFileServerClientConnectionThreaded
-	:  public FNetworkFileServerClientConnection
-	  ,protected FRunnable 
+class FCookOnTheFlyNetworkFileServerConnection
+	: public FNetworkFileServerClientConnection
 {
 public:
-
-	FNetworkFileServerClientConnectionThreaded(FSocket* InSocket, const FNetworkFileServerOptions& Options)
-		: FNetworkFileServerClientConnection(Options)
-		, Socket(InSocket)
-	{
-		Running.Set(true);
-		StopRequested.Reset();
-	
-#if UE_BUILD_DEBUG
-		// this thread needs more space in debug builds as it tries to log messages and such
-		const static uint32 NetworkFileServerThreadSize = 2 * 1024 * 1024; 
-#else
-		const static uint32 NetworkFileServerThreadSize = 1 * 1024 * 1024; 
-#endif
-		WorkerThread = FRunnableThread::Create(this, TEXT("FNetworkFileServerClientConnection"), NetworkFileServerThreadSize, TPri_AboveNormal);
-	}
-
-
-	virtual bool Init() override
+	FCookOnTheFlyNetworkFileServerConnection(UE::Cook::ICookOnTheFlyClientConnection& InConnection, const FNetworkFileServerOptions& InFileServerOptions)
+		: FNetworkFileServerClientConnection(InFileServerOptions)
+		, Connection(InConnection)
 	{
-		return true; 
 	}
 
-	virtual uint32 Run() override
+	bool ProcessRequest(const UE::Cook::FCookOnTheFlyRequest& Request)
 	{
-		while (!StopRequested.GetValue())
-		{
-			  // read a header and payload pair
-			  FArrayReader Payload; 
-			  if (!FNFSMessageHeader::ReceivePayload(Payload, FSimpleAbstractSocket_FSocket(Socket)))
-				  break; 
-			  // now process the contents of the payload
-			  if ( !FNetworkFileServerClientConnection::ProcessPayload(Payload) )
-			  {
-				  // give the processing of the payload a chance to terminate the connection
-				  // failed to process message
-				  UE_LOG(LogFileServer, Warning, TEXT("Unable to process payload terminating connection"));
-				  break;
-			  }
-		  }
-
-		  return true;
-	}
-
-	virtual bool SendPayload( TArray<uint8> &Out ) override
-	{
-#if USE_MCSOCKET_FOR_NFS
-		return FNFSMessageHeader::WrapAndSendPayload(Out, FSimpleAbstractSocket_FMultichannelTCPSocket(MCSocket, NFS_Channels::Main));
-#else
-		return FNFSMessageHeader::WrapAndSendPayload(Out, FSimpleAbstractSocket_FSocket(Socket));
-#endif
-	}
+		using namespace UE::Cook;
+		using namespace UE::ZenCookOnTheFly::Messaging;
 
-	virtual void Stop() override
-	{
-		StopRequested.Set(true);
-	}
-
-	virtual void Exit() override
-	{
-		Socket->Close();
-		ISocketSubsystem::Get()->DestroySocket(Socket);
-		Running.Set(false); 
-	}
-
-	bool IsRunning()
-	{
-		return  (Running.GetValue() != 0); 
-	}
-
-	void GetAddress( FInternetAddr& Addr )
-	{
-		Socket->GetAddress(Addr);
+		RequestToRespondTo = &Request;
+		TUniquePtr<FArchive> RequestBody = Request.ReadBody();
+		return ProcessPayload(*RequestBody);
 	}
 
-	void GetPeerAddress( FInternetAddr& Addr )
+	virtual bool SendPayload(TArray<uint8>& Out) override
 	{
-		Socket->GetPeerAddress(Addr);
-	}
+		using namespace UE::Cook;
+		using namespace UE::ZenCookOnTheFly::Messaging;
 
-	~FNetworkFileServerClientConnectionThreaded()
-	{
-		WorkerThread->Kill(true);
+		if (RequestToRespondTo)
+		{
+			FCookOnTheFlyResponse Response(*RequestToRespondTo);
+			RequestToRespondTo = nullptr;
+			TUniquePtr<FArchive> PayloadWriter = Response.WriteBody();
+			PayloadWriter->Serialize(Out.GetData(), Out.Num());
+			return Connection.SendMessage(Response);
+		}
+		else
+		{
+			FCookOnTheFlyMessage Message(ECookOnTheFlyMessage::NetworkPlatformFile);
+			TUniquePtr<FArchive> PayloadWriter = Message.WriteBody();
+			PayloadWriter->Serialize(Out.GetData(), Out.Num());
+			return Connection.SendMessage(Message);
+		}
 	}
 
-private: 
-	FSocket* Socket; 
-	FThreadSafeCounter StopRequested;
-	FThreadSafeCounter Running;
-	FRunnableThread* WorkerThread; 
+private:
+	UE::Cook::ICookOnTheFlyClientConnection& Connection;
+	const UE::Cook::FCookOnTheFlyRequest* RequestToRespondTo = nullptr;
 };
 
-
-
 /* FNetworkFileServer constructors
  *****************************************************************************/
 
-FNetworkFileServer::FNetworkFileServer(FNetworkFileServerOptions InFileServerOptions)
+FNetworkFileServer::FNetworkFileServer(FNetworkFileServerOptions InFileServerOptions, TSharedRef<UE::Cook::ICookOnTheFlyNetworkServer> InCookOnTheFlyNetworkServer)
 	: FileServerOptions(MoveTemp(InFileServerOptions))
+	, CookOnTheFlyServer(InCookOnTheFlyNetworkServer)
 {
-	if (FileServerOptions.Port < 0 )
-	{
-		FileServerOptions.Port = DEFAULT_TCP_FILE_SERVING_PORT;
-	}
-
-	Running.Set(false);
-	StopRequested.Set(false);
-	UE_LOG(LogFileServer, Warning, TEXT("Unreal Network File Server starting up..."));
-
-	// make sure sockets are going
-	ISocketSubsystem* SocketSubsystem = ISocketSubsystem::Get();
-
-	if(!SocketSubsystem)
-	{
-		UE_LOG(LogFileServer, Error, TEXT("Could not get socket subsystem."));
-	}
-	else
-	{
-		// listen on any IP address
-		ListenAddr = SocketSubsystem->GetLocalBindAddr(*GLog);
-		ListenAddr->SetPort(FileServerOptions.Port);
-
-		// create a server TCP socket
-		Socket = SocketSubsystem->CreateSocket(NAME_Stream, TEXT("FNetworkFileServer tcp-listen"), ListenAddr->GetProtocolType());
-
-		if(!Socket)
-		{
-			UE_LOG(LogFileServer, Error, TEXT("Could not create listen socket."));
-		}
-		else
-		{
-			Socket->SetReuseAddr();
-
-			// bind to the address
-			if (!Socket->Bind(*ListenAddr))
-			{
-				UE_LOG(LogFileServer, Warning, TEXT("Failed to bind listen socket %s in FNetworkFileServer"), *ListenAddr->ToString(true));
-			}
-			// listen for connections
-			else if (!Socket->Listen(16))
-			{
-				UE_LOG(LogFileServer, Warning, TEXT("Failed to listen on socket %s in FNetworkFileServer"), *ListenAddr->ToString(true));
-			}
-			else
-			{
-				// set the port on the listen address to be the same as the port on the socket
-				int32 port = Socket->GetPortNo();
-				check((FileServerOptions.Port == 0 && port != 0) || port == FileServerOptions.Port);
-				ListenAddr->SetPort(port);
-
-				// now create a thread to accept connections
-				Thread = FRunnableThread::Create(this, TEXT("FNetworkFileServer"), 8 * 1024, TPri_AboveNormal);
-
-				UE_LOG(LogFileServer, Display, TEXT("Unreal Network File Server is ready for client connections on %s!"), *ListenAddr->ToString(true));
-			}
-		}
-	}
+	using namespace UE::Cook;
 
+	CookOnTheFlyServer->OnClientConnected().AddRaw(this, &FNetworkFileServer::OnClientConnected);
+	CookOnTheFlyServer->OnClientDisconnected().AddRaw(this, &FNetworkFileServer::OnClientDisconnected);
+	CookOnTheFlyServer->OnRequest(ECookOnTheFlyMessage::NetworkPlatformFile).BindRaw(this, &FNetworkFileServer::HandleRequest);
 }
 
 FNetworkFileServer::~FNetworkFileServer()
 {
-	// Kill the running thread.
-	if( Thread != NULL )
-	{
-		Thread->Kill(true);
-
-		delete Thread;
-		Thread = NULL;
-	}
-
-	// We are done with the socket.
-	Socket->Close();
-	ISocketSubsystem::Get()->DestroySocket(Socket);
-	Socket = NULL;
+	Shutdown();
 }
 
+void FNetworkFileServer::OnClientConnected(UE::Cook::ICookOnTheFlyClientConnection& Connection)
+{
+	FCookOnTheFlyNetworkFileServerConnection* InternalConnection = new FCookOnTheFlyNetworkFileServerConnection(Connection, FileServerOptions);
+	FScopeLock _(&ConnectionsCritical);
+	Connections.Add(&Connection, InternalConnection);
+}
 
-/* FRunnable overrides
- *****************************************************************************/
-
-uint32 FNetworkFileServer::Run( )
+void FNetworkFileServer::OnClientDisconnected(UE::Cook::ICookOnTheFlyClientConnection& Connection)
 {
-	Running.Set(true); 
-	// go until requested to be done
-	while (!StopRequested.GetValue())
+	FNetworkFileServerClientConnection* InternalConnection;
 	{
-		bool bReadReady = false;
-
-		// clean up closed connections
-		for (int32 ConnectionIndex = 0; ConnectionIndex < Connections.Num(); ++ConnectionIndex)
-		{
-			FNetworkFileServerClientConnectionThreaded* Connection = Connections[ConnectionIndex];
-
-			if (!Connection->IsRunning())
-			{
-				UE_LOG(LogFileServer, Display, TEXT( "Client %s disconnected." ), *Connection->GetDescription() );
-				Connections.RemoveAtSwap(ConnectionIndex);
-				delete Connection;
-			}
-		}
-
-		// check for incoming connections
-		if (Socket->WaitForPendingConnection(bReadReady, FTimespan::FromSeconds(0.25f)))
-		{
-			if (bReadReady)
-			{
-				FSocket* ClientSocket = Socket->Accept(TEXT("Remote Console Connection"));
-
-				if (ClientSocket != NULL)
-				{
-					TSharedPtr<FInternetAddr> Addr = ISocketSubsystem::Get(PLATFORM_SOCKETSUBSYSTEM)->CreateInternetAddr();
-					ClientSocket->GetAddress(*Addr);
-					TSharedPtr<FInternetAddr> PeerAddr = ISocketSubsystem::Get(PLATFORM_SOCKETSUBSYSTEM)->CreateInternetAddr();
-					ClientSocket->GetPeerAddress(*PeerAddr);
-
-					for (auto PreviousConnection : Connections)
-					{
-						TSharedPtr<FInternetAddr> PreviousAddr = ISocketSubsystem::Get(PLATFORM_SOCKETSUBSYSTEM)->CreateInternetAddr();;
-						PreviousConnection->GetAddress(*PreviousAddr);
-						TSharedPtr<FInternetAddr> PreviousPeerAddr = ISocketSubsystem::Get(PLATFORM_SOCKETSUBSYSTEM)->CreateInternetAddr();;
-						PreviousConnection->GetPeerAddress(*PreviousPeerAddr);
-						if ((*Addr == *PreviousAddr) &&
-							(*PeerAddr == *PreviousPeerAddr))
-						{
-							// kill the connection 
-							PreviousConnection->Stop();
-							UE_LOG(LogFileServer, Warning, TEXT("Killing client connection %s because new client connected from same address."), *PreviousConnection->GetDescription());
-						}
-					}
-
-					FNetworkFileServerClientConnectionThreaded* Connection = new FNetworkFileServerClientConnectionThreaded(ClientSocket, FileServerOptions);
-					Connections.Add(Connection);
-					UE_LOG(LogFileServer, Display, TEXT("Client %s connected."), *Connection->GetDescription());
-				}
-			}
-		}
-		else
-		{
-			FPlatformProcess::Sleep(0.25f);
-		}
+		FScopeLock _(&ConnectionsCritical);
+		InternalConnection = Connections.FindRef(&Connection);
+		Connections.Remove(&Connection);
 	}
-
-	return 0;
+	delete InternalConnection;
 }
 
-
-void FNetworkFileServer::Exit( )
+bool FNetworkFileServer::HandleRequest(UE::Cook::ICookOnTheFlyClientConnection& Connection, const UE::Cook::FCookOnTheFlyRequest& Request)
 {
-	// close all connections
-	for (int32 ConnectionIndex = 0; ConnectionIndex < Connections.Num(); ConnectionIndex++)
+	FCookOnTheFlyNetworkFileServerConnection* InternalConnection;
 	{
-		delete Connections[ConnectionIndex];
+		FScopeLock _(&ConnectionsCritical);
+		InternalConnection = Connections.FindRef(&Connection);
 	}
-
-	Connections.Empty();
+	if (!InternalConnection)
+	{
+		return false;
+	}
+	return InternalConnection->ProcessRequest(Request);
 }
 
 
@@ -279,47 +116,30 @@ void FNetworkFileServer::Exit( )
 
 FString FNetworkFileServer::GetSupportedProtocol() const
 {
-	return FString("tcp");
+	return CookOnTheFlyServer->GetSupportedProtocol();
 }
 
 
 bool FNetworkFileServer::GetAddressList( TArray<TSharedPtr<FInternetAddr> >& OutAddresses ) const
 {
-	if (ListenAddr.IsValid())
-	{
-		FString ListenAddressString = ListenAddr->ToString(true);
-
-		if (ListenAddressString.StartsWith(TEXT("0.0.0.0")))
-		{
-			if (ISocketSubsystem::Get()->GetLocalAdapterAddresses(OutAddresses))
-			{
-				for (int32 AddressIndex = 0; AddressIndex < OutAddresses.Num(); ++AddressIndex)
-				{
-					OutAddresses[AddressIndex]->SetPort(ListenAddr->GetPort());
-				}
-			}
-		}
-		else
-		{
-			OutAddresses.Add(ListenAddr);
-		}
-	}
-
-	return (OutAddresses.Num() > 0);
+	return CookOnTheFlyServer->GetAddressList(OutAddresses);
 }
 
 
 bool FNetworkFileServer::IsItReadyToAcceptConnections(void) const
 {
-	return (Running.GetValue() != 0); 
+	return CookOnTheFlyServer->IsReadyToAcceptConnections();
 }
 
 int32 FNetworkFileServer::NumConnections(void) const
 {
-	return Connections.Num(); 
+	return CookOnTheFlyServer->NumConnections();
 }
 
 void FNetworkFileServer::Shutdown(void)
 {
-	Stop();
+	CookOnTheFlyServer->OnClientConnected().RemoveAll(this);
+	CookOnTheFlyServer->OnClientDisconnected().RemoveAll(this);
+	CookOnTheFlyServer->OnRequest(UE::Cook::ECookOnTheFlyMessage::NetworkPlatformFile).Unbind();
+	CookOnTheFlyServer.Reset();
 }
diff --git a/Engine/Source/Runtime/NetworkFileSystem/Private/NetworkFileServer.h b/Engine/Source/Runtime/NetworkFileSystem/Private/NetworkFileServer.h
--- a/Engine/Source/Runtime/NetworkFileSystem/Private/NetworkFileServer.h
+++ b/Engine/Source/Runtime/NetworkFileSystem/Private/NetworkFileServer.h
@@ -11,14 +11,18 @@
 class FInternetAddr;
 class FSocket;
 class ITargetPlatform;
-
+namespace UE::Cook
+{
+	class ICookOnTheFlyNetworkServer;
+	class ICookOnTheFlyClientConnection;
+	class FCookOnTheFlyRequest;
+}
 
 /**
  * This class wraps the server thread and network connection
  */
 class FNetworkFileServer
-	: public FRunnable
-	, public INetworkFileServer
+	: public INetworkFileServer
 {
 public:
 
@@ -27,33 +31,13 @@ public:
 	 *
 	 * @param InFileServerOptions Network file server options
 	 */
-	FNetworkFileServer(FNetworkFileServerOptions InFileServerOptions);
+	FNetworkFileServer(FNetworkFileServerOptions InFileServerOptions, TSharedRef<UE::Cook::ICookOnTheFlyNetworkServer> InCookOnTheFlyNetworkServer);
 
 	/**
 	 * Destructor.
 	 */
 	~FNetworkFileServer( );
 
-public:
-
-	// FRunnable Interface
-
-	virtual bool Init( ) override
-	{
-		return true;
-	}
-
-	virtual uint32 Run( ) override;
-
-	virtual void Stop( ) override
-	{
-		StopRequested.Set(true);
-	}
-
-	virtual void Exit( ) override;
-
-public:
-
 	// INetworkFileServer interface
 
 	virtual bool IsItReadyToAcceptConnections(void) const override;
@@ -62,25 +46,16 @@ public:
 	virtual int32 NumConnections() const override;
 	virtual void Shutdown() override;
 private:
+	void OnClientConnected(UE::Cook::ICookOnTheFlyClientConnection& Connection);
+	void OnClientDisconnected(UE::Cook::ICookOnTheFlyClientConnection& Connection);
+	bool HandleRequest(UE::Cook::ICookOnTheFlyClientConnection& Connection, const UE::Cook::FCookOnTheFlyRequest& Request);
 
 	// File server options
 	FNetworkFileServerOptions FileServerOptions;
 
-	// Holds the server (listening) socket.
-	FSocket* Socket;
-
-	// Holds the server thread object.
-	FRunnableThread* Thread;
-
-	// Holds the list of all client connections.
-	TArray< class FNetworkFileServerClientConnectionThreaded*> Connections;
-
-	// Holds a flag indicating whether the thread should stop executing
-	FThreadSafeCounter StopRequested;
-
-	// Is the Listner thread up and running. 
-	FThreadSafeCounter Running;
-
-	// Holds the address that the server is bound to.
-	TSharedPtr<FInternetAddr> ListenAddr;
+	TSharedPtr<UE::Cook::ICookOnTheFlyNetworkServer> CookOnTheFlyServer;
+	
+	// Holds all the client connections.
+	FCriticalSection ConnectionsCritical;
+	TMap<UE::Cook::ICookOnTheFlyClientConnection*, class FCookOnTheFlyNetworkFileServerConnection*> Connections;
 };
diff --git a/Engine/Source/Runtime/NetworkFileSystem/Private/NetworkFileServerConnection.cpp b/Engine/Source/Runtime/NetworkFileSystem/Private/NetworkFileServerConnection.cpp
--- a/Engine/Source/Runtime/NetworkFileSystem/Private/NetworkFileServerConnection.cpp
+++ b/Engine/Source/Runtime/NetworkFileSystem/Private/NetworkFileServerConnection.cpp
@@ -368,10 +368,6 @@ bool FNetworkFileServerClientConnection::ProcessPayload(FArchive& Ar)
 			bSendUnsolicitedFiles = true;
 			break;
 
-		case NFS_Messages::RecompileShaders:
-			ProcessRecompileShaders(Ar, Out);
-			break;
-
 		default:
 
 			UE_LOG(LogFileServer, Error, TEXT("Bad incomming message tag (%d)."), (int32)Msg);
@@ -1011,14 +1007,6 @@ bool FNetworkFileServerClientConnection::ProcessGetFileList( FArchive& In, FArch
 
 	GetSandboxRootDirectories(Sandbox.Get(), SandboxEngine, SandboxProject, SandboxEnginePlatformExtensions, SandboxProjectPlatformExtensions, LocalEngineDir, LocalProjectDir, LocalEnginePlatformExtensionsDir, LocalProjectPlatformExtensionsDir);
 
-
-	// make sure the global shaders are up to date before letting the client read any shaders
-	// @todo: This will probably add about 1/2 second to the boot-up time of the client while the server does this
-	// @note: We assume the delegate will write to the proper sandbox directory, should we pass in SandboxDirectory, or Sandbox?
-	TArray<uint8> GlobalShaderMap;
-	FShaderRecompileData RecompileData(ConnectedPlatformName, SP_NumPlatforms, ODSCRecompileCommand::Global, nullptr, nullptr, &GlobalShaderMap);
-	NetworkFileDelegates->RecompileShadersDelegate.ExecuteIfBound(RecompileData);
-
 	UE_LOG(LogFileServer, Display, TEXT("Getting files for %d directories, game = %s, platform = %s"), RootDirectories.Num(), *GameName, *ConnectedPlatformName);
 	UE_LOG(LogFileServer, Display, TEXT("    Sandbox dir = %s"), *SandboxDirectory);
 
@@ -1288,24 +1276,6 @@ bool FNetworkFileServerClientConnection::PackageFile( FString& Filename, FString
 	return bRetVal;
 }
 
-
-void FNetworkFileServerClientConnection::ProcessRecompileShaders( FArchive& In, FArchive& Out )
-{
-	TArray<FString> RecompileModifiedFiles;
-	TArray<uint8> MeshMaterialMaps;
-	TArray<uint8> GlobalShaderMap;
-	FShaderRecompileData RecompileData(ConnectedPlatformName, &RecompileModifiedFiles, &MeshMaterialMaps, &GlobalShaderMap);
-
-	In << RecompileData;
-
-	NetworkFileDelegates->RecompileShadersDelegate.ExecuteIfBound(RecompileData);
-
-	// tell other side what to do!
-	Out << MeshMaterialMaps;
-	Out << GlobalShaderMap;
-}
-
-
 bool FNetworkFileServerClientConnection::ProcessSyncFile( FArchive& In, FArchive& Out )
 {
 
diff --git a/Engine/Source/Runtime/NetworkFileSystem/Private/NetworkFileServerConnection.h b/Engine/Source/Runtime/NetworkFileSystem/Private/NetworkFileServerConnection.h
--- a/Engine/Source/Runtime/NetworkFileSystem/Private/NetworkFileServerConnection.h
+++ b/Engine/Source/Runtime/NetworkFileSystem/Private/NetworkFileServerConnection.h
@@ -145,14 +145,6 @@ protected:
 
 	bool PackageFile( FString& Filename, FString& TargetFilename, FArchive& Out);
 
-	/**
-	 * Processes a RecompileShaders message.
-	 *
-	 * @param In -
-	 * @param Out -
-	 */
-	void ProcessRecompileShaders( FArchive& In, FArchive& Out );
-
 	/**
 	 * Processes a heartbeat message.
 	 *
diff --git a/Engine/Source/Runtime/NetworkFileSystem/Private/NetworkFileServerPlatformProtocol.cpp b/Engine/Source/Runtime/NetworkFileSystem/Private/NetworkFileServerPlatformProtocol.cpp
deleted file mode 100644
diff --git a/Engine/Source/Runtime/NetworkFileSystem/Private/NetworkFileServerPlatformProtocol.h b/Engine/Source/Runtime/NetworkFileSystem/Private/NetworkFileServerPlatformProtocol.h
deleted file mode 100644
diff --git a/Engine/Source/Runtime/NetworkFileSystem/Private/NetworkFileSystemModule.cpp b/Engine/Source/Runtime/NetworkFileSystem/Private/NetworkFileSystemModule.cpp
--- a/Engine/Source/Runtime/NetworkFileSystem/Private/NetworkFileSystemModule.cpp
+++ b/Engine/Source/Runtime/NetworkFileSystem/Private/NetworkFileSystemModule.cpp
@@ -7,9 +7,9 @@
 #include "NetworkFileSystemLog.h"
 #include "NetworkFileServer.h"
 #include "NetworkFileServerHttp.h"
-#include "NetworkFileServerPlatformProtocol.h"
 #include "Interfaces/ITargetPlatformManagerModule.h"
-
+#include "CookOnTheFlyNetServer.h"
+#include "IPAddress.h"
 
 DEFINE_LOG_CATEGORY(LogFileServer);
 
@@ -60,14 +60,57 @@ public:
 			return new FNetworkFileServerHttp(MoveTemp(FileServerOptions));
 #endif
 		case NFSP_Tcp:
-			return new FNetworkFileServer(MoveTemp(FileServerOptions));
-
 		case NFSP_Platform:
-			return new FNetworkFileServerPlatformProtocol(MoveTemp(FileServerOptions));
+			UE::Cook::FCookOnTheFlyNetworkServerOptions CookOnTheFlyServerOptions;
+			if (FileServerOptions.Protocol == NFSP_Tcp)
+			{
+				CookOnTheFlyServerOptions.Protocol = UE::Cook::ECookOnTheFlyNetworkServerProtocol::Tcp;
+			}
+			else
+			{
+				CookOnTheFlyServerOptions.Protocol = UE::Cook::ECookOnTheFlyNetworkServerProtocol::Platform;
+			}
+			CookOnTheFlyServerOptions.TargetPlatforms = FileServerOptions.TargetPlatforms;
+			CookOnTheFlyServerOptions.Port = FileServerOptions.Port;
+			UE::Cook::ICookOnTheFlyNetworkServerModule& CookOnTheFlyNetworkServerModule = FModuleManager::LoadModuleChecked<UE::Cook::ICookOnTheFlyNetworkServerModule>(TEXT("CookOnTheFlyNetServer"));
+			TSharedPtr<UE::Cook::ICookOnTheFlyNetworkServer> CookOnTheFlyNetworkServer = CookOnTheFlyNetworkServerModule.CreateServer(CookOnTheFlyServerOptions);
+			if (!CookOnTheFlyNetworkServer)
+			{
+				return nullptr;
+			}
+			UE_LOG(LogFileServer, Display, TEXT("Unreal Network File Server starting up..."));
+			if (CookOnTheFlyNetworkServer->Start())
+			{
+				TArray<TSharedPtr<FInternetAddr>> ListenAddresses;
+				if (CookOnTheFlyNetworkServer->GetAddressList(ListenAddresses))
+				{
+					UE_LOG(LogFileServer, Display, TEXT("Unreal Network File Server is ready for client connections on %s!"), *ListenAddresses[0]->ToString(true));
+				}
+				return new FNetworkFileServer(MoveTemp(FileServerOptions), CookOnTheFlyNetworkServer.ToSharedRef());
+			}
+			break;
 		}
- 
+
 		return nullptr;
 	}
+
+	virtual INetworkFileServer* CreateNetworkFileServer(TSharedRef<UE::Cook::ICookOnTheFlyNetworkServer> CookOnTheFlyNetworkServer, FNetworkFileDelegateContainer Delegates) const override
+	{
+		FNetworkFileServerOptions FileServerOptions;
+		FileServerOptions.Delegates = MoveTemp(Delegates);
+		FileServerOptions.bRestrictPackageAssetsToSandbox = true;
+		ITargetPlatformManagerModule& TPM = GetTargetPlatformManagerRef();
+		FString Platforms;
+		if (FParse::Value(FCommandLine::Get(), TEXT("TARGETPLATFORM="), Platforms))
+		{
+			FileServerOptions.TargetPlatforms = TPM.GetActiveTargetPlatforms();
+		}
+		else
+		{
+			FileServerOptions.TargetPlatforms = TPM.GetTargetPlatforms();
+		}
+		return new FNetworkFileServer(MoveTemp(FileServerOptions), CookOnTheFlyNetworkServer);
+	}
 };
 
 
diff --git a/Engine/Source/Runtime/NetworkFileSystem/Public/INetworkFileSystemModule.h b/Engine/Source/Runtime/NetworkFileSystem/Public/INetworkFileSystemModule.h
--- a/Engine/Source/Runtime/NetworkFileSystem/Public/INetworkFileSystemModule.h
+++ b/Engine/Source/Runtime/NetworkFileSystem/Public/INetworkFileSystemModule.h
@@ -7,6 +7,10 @@
 #include "ShaderCompiler.h"
 
 class INetworkFileServer;
+namespace UE::Cook
+{
+	class ICookOnTheFlyNetworkServer;
+}
 
 /**
  * Delegate type for handling file requests from a network client.
@@ -124,6 +128,19 @@ public:
 	 */
 	virtual INetworkFileServer* CreateNetworkFileServer(FNetworkFileServerOptions FileServerOptions, bool bLoadTargetPlatforms) const = 0;
 
+	/**
+	 * Creates a new network file server for COTF.
+	 *
+	 * @param CookOnTheFlyNetworkServer CookOnTheFly network server
+	 * @param Delegates Delegates to be invoked when a file is requested by a client.
+	 *
+	 * @return The new file server, or nullptr if creation failed.
+	*/
+	virtual INetworkFileServer* CreateNetworkFileServer(TSharedRef<UE::Cook::ICookOnTheFlyNetworkServer> CookOnTheFlyNetworkServer, FNetworkFileDelegateContainer Delegates) const
+	{
+		return nullptr;
+	}
+
 public:
 
 	/**
diff --git a/Engine/Source/Runtime/StorageServerClient/Private/CookOnTheFlyPackageStore.cpp b/Engine/Source/Runtime/StorageServerClient/Private/CookOnTheFlyPackageStore.cpp
--- a/Engine/Source/Runtime/StorageServerClient/Private/CookOnTheFlyPackageStore.cpp
+++ b/Engine/Source/Runtime/StorageServerClient/Private/CookOnTheFlyPackageStore.cpp
@@ -8,6 +8,7 @@
 #include "Containers/ChunkedArray.h"
 #include "Serialization/MemoryReader.h"
 #include "CookOnTheFlyMessages.h"
+#include "CookOnTheFly.h"
 
 FCookOnTheFlyPackageStore::FCookOnTheFlyPackageStore(UE::Cook::ICookOnTheFlyServerConnection& InCookOnTheFlyServerConnection)
 	: CookOnTheFlyServerConnection(InCookOnTheFlyServerConnection)
@@ -20,6 +21,7 @@ FCookOnTheFlyPackageStore::FCookOnTheFlyPackageStore(UE::Cook::ICookOnTheFlyServ
 	using namespace UE::ZenCookOnTheFly::Messaging;
 
 	FCookOnTheFlyRequest Request(ECookOnTheFlyMessage::GetCookedPackages);
+	
 	FCookOnTheFlyResponse Response = CookOnTheFlyServerConnection.SendRequest(Request).Get();
 
 	if (Response.IsOk())
diff --git a/Engine/Source/Runtime/StorageServerClient/Private/CookOnTheFlyPackageStore.h b/Engine/Source/Runtime/StorageServerClient/Private/CookOnTheFlyPackageStore.h
--- a/Engine/Source/Runtime/StorageServerClient/Private/CookOnTheFlyPackageStore.h
+++ b/Engine/Source/Runtime/StorageServerClient/Private/CookOnTheFlyPackageStore.h
@@ -9,11 +9,6 @@
 #include "IO/PackageStore.h"
 #include "CookOnTheFly.h"
 
-namespace UE { namespace Cook
-{
-	class ICookOnTheFlyServerConnection;
-}}
-
 class FCookOnTheFlyPackageStore final
 	: public FPackageStoreBase
 {
diff --git a/Engine/Source/Runtime/StorageServerClient/Private/StorageServerPlatformFile.cpp b/Engine/Source/Runtime/StorageServerClient/Private/StorageServerPlatformFile.cpp
--- a/Engine/Source/Runtime/StorageServerClient/Private/StorageServerPlatformFile.cpp
+++ b/Engine/Source/Runtime/StorageServerClient/Private/StorageServerPlatformFile.cpp
@@ -13,11 +13,8 @@
 #include "StorageServerPackageStore.h"
 #include "CookOnTheFlyPackageStore.h"
 #include "Misc/CoreDelegates.h"
-
-#if WITH_COTF
 #include "Modules/ModuleManager.h"
 #include "CookOnTheFly.h"
-#endif
 
 DEFINE_LOG_CATEGORY_STATIC(LogStorageServerPlatformFile, Log, All);
 
@@ -244,22 +241,10 @@ public:
 
 FStorageServerPlatformFile::FStorageServerPlatformFile()
 {
-	FCoreDelegates::OnEnginePreExit.AddRaw(this, &FStorageServerPlatformFile::OnEnginePreExit);
 }
 
 FStorageServerPlatformFile::~FStorageServerPlatformFile()
 {
-	FCoreDelegates::OnEnginePreExit.RemoveAll(this);
-}
-
-void FStorageServerPlatformFile::OnEnginePreExit()
-{
-#if WITH_COTF
-	if (CookOnTheFlyServerConnection)
-	{
-		CookOnTheFlyServerConnection->Disconnect();
-	}
-#endif
 }
 
 bool FStorageServerPlatformFile::ShouldBeUsed(IPlatformFile* Inner, const TCHAR* CmdLine) const
@@ -272,6 +257,17 @@ bool FStorageServerPlatformFile::ShouldBeUsed(IPlatformFile* Inner, const TCHAR*
 			HostAddrs.Add(Host);
 		}
 	}
+	else 
+	{
+#if WITH_COTF
+		UE::Cook::ICookOnTheFlyModule& CookOnTheFlyModule = FModuleManager::LoadModuleChecked<UE::Cook::ICookOnTheFlyModule>(TEXT("CookOnTheFly"));
+		TSharedPtr<UE::Cook::ICookOnTheFlyServerConnection> DefaultConnection = CookOnTheFlyModule.GetDefaultServerConnection();
+		if (DefaultConnection.IsValid() && !DefaultConnection->GetZenProjectName().IsEmpty())
+		{
+			HostAddrs.Add(DefaultConnection->GetHost());
+		}
+#endif
+	}
 
 	return HostAddrs.Num() > 0;
 }
@@ -291,32 +287,19 @@ bool FStorageServerPlatformFile::Initialize(IPlatformFile* Inner, const TCHAR* C
 
 void FStorageServerPlatformFile::InitializeAfterProjectFilePath()
 {
-	Connection.Reset(new FStorageServerConnection());
-	const TCHAR* ProjectOverride = ServerProject.IsEmpty() ? nullptr : *ServerProject;
-	const TCHAR* PlatformOverride = ServerPlatform.IsEmpty() ? nullptr : *ServerPlatform;
 #if WITH_COTF
-	if (FParse::Param(FCommandLine::Get(), TEXT("CookOnTheFly")))
+	UE::Cook::ICookOnTheFlyModule& CookOnTheFlyModule = FModuleManager::LoadModuleChecked<UE::Cook::ICookOnTheFlyModule>(TEXT("CookOnTheFly"));
+	CookOnTheFlyServerConnection = CookOnTheFlyModule.GetDefaultServerConnection();
+	if (CookOnTheFlyServerConnection)
 	{
-		UE::Cook::FCookOnTheFlyHostOptions CookOnTheFlyHostOptions;
-		// Cook-on-the-fly expects the same host as the Zen storage server
-		CookOnTheFlyHostOptions.Hosts = HostAddrs;
-		double ServerWaitTimeInSeconds;
-		if (FParse::Value(FCommandLine::Get(), TEXT("-CookOnTheFlyServerWaitTime="), ServerWaitTimeInSeconds))
-		{
-			CookOnTheFlyHostOptions.ServerStartupWaitTime = FTimespan::FromSeconds(ServerWaitTimeInSeconds);
-		}
-		UE::Cook::ICookOnTheFlyModule& CookOnTheFlyModule = FModuleManager::LoadModuleChecked<UE::Cook::ICookOnTheFlyModule>(TEXT("CookOnTheFly"));
-		CookOnTheFlyServerConnection = CookOnTheFlyModule.ConnectToServer(CookOnTheFlyHostOptions);
-		if (CookOnTheFlyServerConnection)
-		{
-			CookOnTheFlyServerConnection->OnMessage().AddRaw(this, &FStorageServerPlatformFile::OnCookOnTheFlyMessage);
-		}
-		else
-		{
-			UE_LOG(LogStorageServerPlatformFile, Fatal, TEXT("Failed to connect to cook on the fly server"));
-		}
+		CookOnTheFlyServerConnection->OnMessage().AddRaw(this, &FStorageServerPlatformFile::OnCookOnTheFlyMessage);
+		ServerProject = CookOnTheFlyServerConnection->GetZenProjectName();
+		ServerPlatform = CookOnTheFlyServerConnection->GetPlatformName();
 	}
 #endif
+	Connection.Reset(new FStorageServerConnection());
+	const TCHAR* ProjectOverride = ServerProject.IsEmpty() ? nullptr : *ServerProject;
+	const TCHAR* PlatformOverride = ServerPlatform.IsEmpty() ? nullptr : *ServerPlatform;
 	if (Connection->Initialize(HostAddrs, 1337, ProjectOverride, PlatformOverride))
 	{
 		if (SendGetFileListMessage())
diff --git a/Engine/Source/Runtime/StorageServerClient/Private/StorageServerPlatformFile.h b/Engine/Source/Runtime/StorageServerClient/Private/StorageServerPlatformFile.h
--- a/Engine/Source/Runtime/StorageServerClient/Private/StorageServerPlatformFile.h
+++ b/Engine/Source/Runtime/StorageServerClient/Private/StorageServerPlatformFile.h
@@ -102,7 +102,6 @@ public:
 private:
 	friend class FStorageServerFileHandle;
 
-	void OnEnginePreExit();
 	bool MakeStorageServerPath(const TCHAR* LocalFilenameOrDirectory, FStringBuilderBase& OutPath) const;
 	bool MakeLocalPath(const TCHAR* ServerFilenameOrDirectory, FStringBuilderBase& OutPath) const;
 	IFileHandle* InternalOpenFile(const FIoChunkId& FileChunkId, const TCHAR* LocalFilename);
@@ -118,7 +117,7 @@ private:
 	FStringView ServerProjectDirView = FStringView(TEXT("/{project}/"));
 	TUniquePtr<FStorageServerConnection> Connection;
 #if WITH_COTF
-	TUniquePtr<UE::Cook::ICookOnTheFlyServerConnection> CookOnTheFlyServerConnection;
+	TSharedPtr<UE::Cook::ICookOnTheFlyServerConnection> CookOnTheFlyServerConnection;
 #endif
 	FStorageServerFileSystemTOC ServerToc;
 	FString ServerProject;
