ITargetPlatform: Delete RefreshSettings. All overrides were removed previously and it is now an empty virtual function and can be removed.
#rnx

diff --git a/Engine/Source/Developer/TargetPlatform/Public/Common/TargetPlatformBase.h b/Engine/Source/Developer/TargetPlatform/Public/Common/TargetPlatformBase.h
--- a/Engine/Source/Developer/TargetPlatform/Public/Common/TargetPlatformBase.h
+++ b/Engine/Source/Developer/TargetPlatform/Public/Common/TargetPlatformBase.h
@@ -170,10 +170,6 @@ public:
 		// do nothing in the base class
 	}
 
-	virtual void RefreshSettings() override
-	{
-	}
-
 	virtual int32 GetPlatformOrdinal() const override
 	{
 		return PlatformOrdinal;
diff --git a/Engine/Source/Developer/TargetPlatform/Public/Interfaces/ITargetPlatform.h b/Engine/Source/Developer/TargetPlatform/Public/Interfaces/ITargetPlatform.h
--- a/Engine/Source/Developer/TargetPlatform/Public/Interfaces/ITargetPlatform.h
+++ b/Engine/Source/Developer/TargetPlatform/Public/Interfaces/ITargetPlatform.h
@@ -674,11 +674,6 @@ public:
 	 */
 	virtual void GetBuildProjectSettingKeys(FString& OutSection, TArray<FString>& InBoolKeys, TArray<FString>& InIntKeys, TArray<FString>& InStringKeys) const = 0;
 
-	/**
-	 * Give the platform a chance to refresh internal settings before a cook, etc
-	 */
-	virtual void RefreshSettings() = 0;
-
 	/**
 	 * Get unique integer identifier for this platform.
 	 *
diff --git a/Engine/Source/Editor/UnrealEd/Classes/CookOnTheSide/CookOnTheFlyServer.h b/Engine/Source/Editor/UnrealEd/Classes/CookOnTheSide/CookOnTheFlyServer.h
--- a/Engine/Source/Editor/UnrealEd/Classes/CookOnTheSide/CookOnTheFlyServer.h
+++ b/Engine/Source/Editor/UnrealEd/Classes/CookOnTheSide/CookOnTheFlyServer.h
@@ -881,9 +881,6 @@ private:
 	*/
 	void ClearPackageStoreContexts();
 
-	/** Give TargetPlatforms a chance to take action at the beginning of a cook. */
-	void BeginCookTargetPlatforms(const TArrayView<ITargetPlatform* const>& NewTargetPlatforms);
-
 	/**
 	* Initialize shaders for the specified platforms when running cook on the fly.
 	*/
diff --git a/Engine/Source/Editor/UnrealEd/Private/Commandlets/CookGlobalShaderCommandlet.cpp b/Engine/Source/Editor/UnrealEd/Private/Commandlets/CookGlobalShaderCommandlet.cpp
--- a/Engine/Source/Editor/UnrealEd/Private/Commandlets/CookGlobalShaderCommandlet.cpp
+++ b/Engine/Source/Editor/UnrealEd/Private/Commandlets/CookGlobalShaderCommandlet.cpp
@@ -273,8 +273,6 @@ int32 UCookGlobalShadersCommandlet::Main(const FString& Params)
 			}
 			return 1;
 		}
-
-		TargetPlatform->RefreshSettings();
 	}
 
 	// Get target device
diff --git a/Engine/Source/Editor/UnrealEd/Private/CookOnTheFlyServer.cpp b/Engine/Source/Editor/UnrealEd/Private/CookOnTheFlyServer.cpp
--- a/Engine/Source/Editor/UnrealEd/Private/CookOnTheFlyServer.cpp
+++ b/Engine/Source/Editor/UnrealEd/Private/CookOnTheFlyServer.cpp
@@ -709,7 +709,6 @@ void UCookOnTheFlyServer::AddCookOnTheFlyPlatformFromGameThread(ITargetPlatform*
 	// Initialize systems that nothing in AddCookOnTheFlyPlatformFromGameThread references
 	// Functions in this section are not dependent upon each other and can be ordered arbitrarily or for async performance
 	BeginCookPackageWriters(BeginContext);
-	BeginCookTargetPlatforms(BeginContext.TargetPlatforms);
 
 	// SaveCurrentIniSettings is done in CookByTheBookFinished for CBTB, but we don't have a definite end point in CookOnTheFly so we write it at the start
 	// This will miss settings that are accessed during the cook
@@ -7916,15 +7915,6 @@ void UCookOnTheFlyServer::ClearPackageStoreContexts()
 	SavePackageContexts.Empty();
 }
 
-void UCookOnTheFlyServer::BeginCookTargetPlatforms(const TArrayView<ITargetPlatform* const>& NewTargetPlatforms)
-{
-	//allow each platform to update its internals before cooking
-	for (ITargetPlatform* TargetPlatform : NewTargetPlatforms)
-	{
-		TargetPlatform->RefreshSettings();
-	}
-}
-
 void UCookOnTheFlyServer::DiscoverPlatformSpecificNeverCookPackages(
 	const TArrayView<const ITargetPlatform* const>& TargetPlatforms, const TArray<FString>& UBTPlatformStrings)
 {
@@ -8017,7 +8007,6 @@ void UCookOnTheFlyServer::StartCookByTheBook( const FCookByTheBookStartupOptions
 	BeginCookEditorSystems();
 	BeginCookEDLCookInfo(BeginContext);
 	BeginCookPackageWriters(BeginContext);
-	BeginCookTargetPlatforms(BeginContext.TargetPlatforms);
 	GenerateInitialRequests(BeginContext);
 	GenerateLocalizationReferences(BeginContext.StartupOptions->CookCultures);
 	InitializePollables();
