World Partition Convert Commandlet: make sure to destroy the persistent level model before saving sublevels with level blueprints.

#rnx

diff --git a/Engine/Source/Editor/UnrealEd/Private/Commandlets/WorldPartitionConvertCommandlet.cpp b/Engine/Source/Editor/UnrealEd/Private/Commandlets/WorldPartitionConvertCommandlet.cpp
--- a/Engine/Source/Editor/UnrealEd/Private/Commandlets/WorldPartitionConvertCommandlet.cpp
+++ b/Engine/Source/Editor/UnrealEd/Private/Commandlets/WorldPartitionConvertCommandlet.cpp
@@ -1189,6 +1189,9 @@ int32 UWorldPartitionConvertCommandlet::Main(const FString& Params)
 		UWorld* SubWorld = SubLevel->GetTypedOuter<UWorld>();
 		UPackage* SubPackage = SubLevel->GetPackage();
 
+		SubWorld->PersistentLevel->Model->MarkAsGarbage();
+		SubWorld->PersistentLevel->Model = nullptr;
+
 		RemapSoftObjectPaths.Add(FSoftObjectPath(SubWorld).ToString(), FSoftObjectPath(MainWorld).ToString());
 		RemapSoftObjectPaths.Add(FSoftObjectPath(SubLevel).ToString(), FSoftObjectPath(MainLevel).ToString());
 		RemapSoftObjectPaths.Add(FSoftObjectPath(SubPackage).ToString(), FSoftObjectPath(MainPackage).ToString());
