Cooker Stats: Fix the "Packages Cooked" stat line to use the number of packages cooked rather than the number of packages multiplied by the number of platforms.
#rnx

diff --git a/Engine/Source/Editor/UnrealEd/Private/CookOnTheFlyServer.cpp b/Engine/Source/Editor/UnrealEd/Private/CookOnTheFlyServer.cpp
--- a/Engine/Source/Editor/UnrealEd/Private/CookOnTheFlyServer.cpp
+++ b/Engine/Source/Editor/UnrealEd/Private/CookOnTheFlyServer.cpp
@@ -4659,6 +4659,7 @@ void FSaveCookedPackageContext::FinishPackage()
 #endif
 			}
 		}
+		COOK_STAT(++DetailedCookStats::NumPackagesSavedForCook);
 	}
 	else if (bReferencedOnlyByEditorOnlyData)
 	{
@@ -7342,7 +7343,7 @@ void UCookOnTheFlyServer::PrintFinishStats()
 	UE_LOG(LogCook, Display, TEXT("Peak Used virtual %u MiB Peak Used physical %u MiB"), MemStats.PeakUsedVirtual / 1024 / 1024, MemStats.PeakUsedPhysical / 1024 / 1024);
 
 	COOK_STAT(UE_LOG(LogCook, Display, TEXT("Packages Cooked: %d, Packages Iteratively Skipped: %d, Total Packages: %d"),
-		UE::SavePackageUtilities::GetNumPackagesSaved(), DetailedCookStats::NumPackagesIterativelySkipped, PackageDatas->GetNumCooked()));
+		DetailedCookStats::NumPackagesSavedForCook, DetailedCookStats::NumPackagesIterativelySkipped, PackageDatas->GetNumCooked()));
 }
 
 TMap<FName, TSet<FName>> UCookOnTheFlyServer::BuildMapDependencyGraph(const ITargetPlatform* TargetPlatform)
@@ -7831,6 +7832,8 @@ void UCookOnTheFlyServer::BeginCookSandbox(FBeginCookContext& BeginContext)
 		}
 	}
 
+	COOK_STAT(DetailedCookStats::NumPackagesSavedForCook = 0);
+
 #if OUTPUT_COOKTIMING
 	FString PlatformNames;
 	for (const ITargetPlatform* Target : BeginContext.TargetPlatforms)
diff --git a/Engine/Source/Editor/UnrealEd/Private/Cooker/CookProfiling.cpp b/Engine/Source/Editor/UnrealEd/Private/Cooker/CookProfiling.cpp
--- a/Engine/Source/Editor/UnrealEd/Private/Cooker/CookProfiling.cpp
+++ b/Engine/Source/Editor/UnrealEd/Private/Cooker/CookProfiling.cpp
@@ -187,6 +187,7 @@ namespace DetailedCookStats
 	int32 PeakSaveQueueSize = 0;
 	uint32 NumPreloadedDependencies = 0;
 	uint32 NumPackagesIterativelySkipped = 0;
+	uint32 NumPackagesSavedForCook = 0;
 	FCookStatsManager::FAutoRegisterCallback RegisterCookOnTheFlyServerStats([](FCookStatsManager::AddStatFuncRef AddStat)
 		{
 			AddStat(TEXT("Package.Load"), FCookStatsManager::CreateKeyValueArray(TEXT("NumPreloadedDependencies"), NumPreloadedDependencies));
diff --git a/Engine/Source/Editor/UnrealEd/Private/Cooker/CookProfiling.h b/Engine/Source/Editor/UnrealEd/Private/Cooker/CookProfiling.h
--- a/Engine/Source/Editor/UnrealEd/Private/Cooker/CookProfiling.h
+++ b/Engine/Source/Editor/UnrealEd/Private/Cooker/CookProfiling.h
@@ -132,6 +132,8 @@ namespace DetailedCookStats
 	// Stats tracked through FAutoRegisterCallback
 	extern uint32 NumPreloadedDependencies;
 	extern uint32 NumPackagesIterativelySkipped;
+	/** Tracks the number of packages cooked - counted only once even if multiple platforms cooked. */
+	extern uint32 NumPackagesSavedForCook;
 	extern int32 PeakRequestQueueSize;
 	extern int32 PeakLoadQueueSize;
 	extern int32 PeakSaveQueueSize;
