Basic shapes can now be searched for in the Add Component list.

diff --git a/Engine/Source/Editor/UnrealEd/Private/SComponentClassCombo.cpp b/Engine/Source/Editor/UnrealEd/Private/SComponentClassCombo.cpp
--- a/Engine/Source/Editor/UnrealEd/Private/SComponentClassCombo.cpp
+++ b/Engine/Source/Editor/UnrealEd/Private/SComponentClassCombo.cpp
@@ -214,7 +214,7 @@
 					if(!bAllowEntry)
 					{
 						const FString FriendlyComponentName = GetSanitizedComponentName(CurrentEntry);
-						TextFilter->TestTextFilter(FBasicStringFilterExpressionContext(FriendlyComponentName));
+						bAllowEntry = TextFilter->TestTextFilter(FBasicStringFilterExpressionContext(FriendlyComponentName));
 					}
 				}
 			}
