Move TemplateMapInfo to a different header so UnrealEdEngine.h does not need to be included

diff --git a/Engine/Source/Editor/UnrealEd/Classes/Editor/TemplateMapInfo.h b/Engine/Source/Editor/UnrealEd/Classes/Editor/TemplateMapInfo.h
new file mode 100644
--- /dev/null
+++ b/Engine/Source/Editor/UnrealEd/Classes/Editor/TemplateMapInfo.h
@@ -0,0 +1,36 @@
+// Copyright Epic Games, Inc. All Rights Reserved.
+
+
+#pragma once
+
+#include "CoreMinimal.h"
+#include "UObject/ObjectMacros.h"
+#include "UObject/Object.h"
+#include "TemplateMapInfo.generated.h"
+
+/** Used by new level dialog. */
+USTRUCT()
+struct FTemplateMapInfo
+{
+	GENERATED_USTRUCT_BODY()
+
+	/** The Texture2D associated with this map template */
+	UPROPERTY()
+	TSoftObjectPtr<class UTexture2D> ThumbnailTexture;
+
+	/** The Texture associated with this map template */
+	UPROPERTY(config)
+	FString Thumbnail;
+
+	/** The object path to the template map */
+	UPROPERTY(config)
+	FString Map;
+
+	/** Optional display name override for this map template  */
+	UPROPERTY()
+	FText DisplayName;
+
+	FTemplateMapInfo()
+	{
+	}
+};
diff --git a/Engine/Source/Editor/UnrealEd/Classes/Editor/UnrealEdEngine.h b/Engine/Source/Editor/UnrealEd/Classes/Editor/UnrealEdEngine.h
--- a/Engine/Source/Editor/UnrealEd/Classes/Editor/UnrealEdEngine.h
+++ b/Engine/Source/Editor/UnrealEd/Classes/Editor/UnrealEdEngine.h
@@ -12,6 +12,7 @@
 #include "ISourceControlProvider.h"
 #include "ComponentVisualizer.h"
 #include "ComponentVisualizerManager.h"
+#include "TemplateMapInfo.h"
 #include "UnrealEdEngine.generated.h"
 
 class AGroupActor;
@@ -79,37 +80,6 @@ struct FClassMoveInfo
 
 };
 
-/** Used by new level dialog. */
-USTRUCT()
-struct FTemplateMapInfo
-{
-	GENERATED_USTRUCT_BODY()
-
-	/** The Texture2D associated with this map template */
-	UPROPERTY()
-	TSoftObjectPtr<UTexture2D> ThumbnailTexture;
-
-	/** The Texture associated with this map template */
-	UPROPERTY(config, meta = (AllowedClasses = "Texture2D"))
-	FSoftObjectPath Thumbnail;
-
-	/** The object path to the template map */
-	UPROPERTY(config, meta = (AllowedClasses = "World"))
-	FSoftObjectPath Map;
-
-	/** Optional display name override for this map template  */
-	UPROPERTY(config)
-	FText DisplayName;
-
-	/* Optional category used for sorting */
-	UPROPERTY(config)
-	FString Category;
-
-	FTemplateMapInfo()
-	{
-	}
-};
-
 class FPerformanceMonitor;
 
 
