Post merge cleanup

diff --git a/Engine/Source/Editor/UnrealEd/Classes/Editor/TemplateMapInfo.h b/Engine/Source/Editor/UnrealEd/Classes/Editor/TemplateMapInfo.h
--- a/Engine/Source/Editor/UnrealEd/Classes/Editor/TemplateMapInfo.h
+++ b/Engine/Source/Editor/UnrealEd/Classes/Editor/TemplateMapInfo.h
@@ -19,18 +19,22 @@ struct FTemplateMapInfo
 	TSoftObjectPtr<class UTexture2D> ThumbnailTexture;
 
 	/** The Texture associated with this map template */
-	UPROPERTY(config)
-	FString Thumbnail;
+	UPROPERTY(config, meta = (AllowedClasses = "Texture2D"))
+	FSoftObjectPath Thumbnail;
 
 	/** The object path to the template map */
-	UPROPERTY(config)
-	FString Map;
+	UPROPERTY(config, meta = (AllowedClasses = "World"))
+	FSoftObjectPath Map;
 
 	/** Optional display name override for this map template  */
-	UPROPERTY()
+	UPROPERTY(config)
 	FText DisplayName;
 
+	/* Optional category used for sorting */
+	UPROPERTY(config)
+	FString Category;
+
 	FTemplateMapInfo()
 	{
 	}
-};
+};
\ No newline at end of file
