Make sure we dirty external actors packages when performing a Precomputed Static Visibility build on an OFPA level.

diff --git a/Engine/Source/Editor/UnrealEd/Private/StaticLightingSystem/StaticLightingSystem.cpp b/Engine/Source/Editor/UnrealEd/Private/StaticLightingSystem/StaticLightingSystem.cpp
--- a/Engine/Source/Editor/UnrealEd/Private/StaticLightingSystem/StaticLightingSystem.cpp
+++ b/Engine/Source/Editor/UnrealEd/Private/StaticLightingSystem/StaticLightingSystem.cpp
@@ -517,10 +517,16 @@ bool FStaticLightingSystem::BeginLightmassProcess()
 		GShadowmapTotalSize = 0;
 		GShadowmapTotalStreamingSize = 0;
 
+		TSet<UPackage*> PackagesToDirty;
+
 		for( TObjectIterator<UPrimitiveComponent> It ; It ; ++It )
 		{
 			UPrimitiveComponent* Component = *It;
-			Component->VisibilityId = INDEX_NONE;
+			if (Component->VisibilityId != INDEX_NONE)
+			{
+				Component->VisibilityId = INDEX_NONE;
+				PackagesToDirty.Add(Component->GetPackage());
+			}
 		}
 
 		{
@@ -541,7 +547,7 @@ bool FStaticLightingSystem::BeginLightmassProcess()
 							->AddToken(FTextToken::Create(LOCTEXT("LightmassError_DuplicatedLevelGuids2", ". A new GUID is assigned to the later one. All previously built lighting is invalidated and the level needs to be resaved.")));
 
 						Level->LevelBuildDataId = FGuid::NewGuid();
-						Level->MarkPackageDirty();
+						PackagesToDirty.Add(Level->GetPackage());
 					}
 
 					LevelGuids.Add(Level->LevelBuildDataId, Level);
@@ -549,6 +555,12 @@ bool FStaticLightingSystem::BeginLightmassProcess()
 			}
 		}
 
+		// Mark package(s) as dirty
+		for (UPackage* Package : PackagesToDirty)
+		{
+			Package->MarkPackageDirty();
+		}
+
 		FString SkippedLevels;
 		for ( int32 LevelIndex=0; LevelIndex < World->GetNumLevels(); LevelIndex++ )
 		{
@@ -910,7 +922,8 @@ void FStaticLightingSystem::GatherStaticLightingInfo(bool bRebuildDirtyGeometryF
 	// Gather static lighting info from actor components.
 	for (int32 LevelIndex = 0; LevelIndex < World->GetNumLevels(); LevelIndex++)
 	{
-		bool bMarkLevelDirty = false;
+		TSet<UPackage*> PackagesToDirty;
+
 		ULevel* Level = World->GetLevel(LevelIndex);
 
 		if (!ShouldOperateOnLevel(Level))
@@ -1164,8 +1177,8 @@ void FStaticLightingSystem::GatherStaticLightingInfo(bool bRebuildDirtyGeometryF
 						{
 							if (World->GetWorldSettings()->bPrecomputeVisibility)
 							{
-								// Make sure the level gets dirtied since we are changing the visibility Id of a component in it
-								bMarkLevelDirty = true;
+								// Make sure packages gets dirtied since we are changing the visibility Id of a component in them
+								PackagesToDirty.Add(Primitive->GetPackage());
 							}
 
 							PrimitiveInfo.VisibilityId = Primitive->VisibilityId = NextVisibilityId;
@@ -1214,9 +1227,10 @@ void FStaticLightingSystem::GatherStaticLightingInfo(bool bRebuildDirtyGeometryF
 			}
 		}
 
-		if (bMarkLevelDirty)
+		// Mark package(s) as dirty
+		for (UPackage* Package : PackagesToDirty)
 		{
-			Level->MarkPackageDirty();
+			Package->MarkPackageDirty();
 		}
 	}
 
@@ -1625,8 +1639,9 @@ void FStaticLightingSystem::AddBSPStaticLightingInfo(ULevel* Level, bool bBuildL
 	// create all NodeGroups
 	Model->GroupAllNodes(Level, Lights);
 
+	TSet<UPackage*> PackagesToDirty;
+
 	// now we need to make the mappings/meshes
-	bool bMarkLevelDirty = false;
 	for (TMap<int32, FNodeGroup*>::TIterator It(Model->NodeGroups); It; ++It)
 	{
 		FNodeGroup* NodeGroup = It.Value();
@@ -1700,8 +1715,8 @@ void FStaticLightingSystem::AddBSPStaticLightingInfo(ULevel* Level, bool bBuildL
 				{
 					if (World->GetWorldSettings()->bPrecomputeVisibility)
 					{
-						// Make sure the level gets dirtied since we are changing the visibility Id of a component in it
-						bMarkLevelDirty = true;
+						// Make sure packages gets dirtied since we are changing the visibility Id of a component in them
+						PackagesToDirty.Add(Component->GetPackage());
 					}
 					Component->VisibilityId = NextVisibilityId;
 					NextVisibilityId++;
@@ -1760,9 +1775,10 @@ void FStaticLightingSystem::AddBSPStaticLightingInfo(ULevel* Level, bool bBuildL
 		}
 	}
 
-	if (bMarkLevelDirty)
+	// Mark package(s) as dirty
+	for (UPackage* Package : PackagesToDirty)
 	{
-		Level->MarkPackageDirty();
+		Package->MarkPackageDirty();
 	}
 }
 
