Unify PIE world leak reporting with LoadMap reporting.

diff --git a/Engine/Source/Editor/UnrealEd/Private/PlayLevel.cpp b/Engine/Source/Editor/UnrealEd/Private/PlayLevel.cpp
--- a/Engine/Source/Editor/UnrealEd/Private/PlayLevel.cpp
+++ b/Engine/Source/Editor/UnrealEd/Private/PlayLevel.cpp
@@ -412,53 +412,72 @@
 	}
 
 	// Make sure that all objects in the temp levels were entirely garbage collected.
+	TSet<UObject*> LeakedObjectsSet;
+	TSet<UPackage*> LeakedPackages;
 	for(FThreadSafeObjectIterator ObjectIt; ObjectIt; ++ObjectIt )
 	{
 		UObject* Object = *ObjectIt;
-		if( Object->GetOutermost()->HasAnyPackageFlags(PKG_PlayInEditor))
+		UPackage* ObjectPackage = Object->GetOutermost();
+		if (ObjectPackage->HasAnyPackageFlags(PKG_PlayInEditor))
 		{
-			UWorld* TheWorld = UWorld::FindWorldInPackage(Object->GetOutermost());
-			if ( TheWorld )
+			LeakedPackages.Add(ObjectPackage);
+			if (UWorld* TheWorld = UWorld::FindWorldInPackage(ObjectPackage))
 			{
-				StaticExec(nullptr, *FString::Printf(TEXT("OBJ REFS CLASS=WORLD NAME=%s"), *TheWorld->GetPathName()));
+				LeakedObjectsSet.Add(TheWorld);
 			}
 			else
 			{
-				UE_LOG(LogPlayLevel, Error, TEXT("No PIE world was found when attempting to gather references after GC."));
-			}
-
-			FReferenceChainSearch RefChainSearch(Object, EReferenceChainSearchMode::Shortest);
-
-			FFormatNamedArguments Arguments;
-			Arguments.Add(TEXT("Path"), FText::FromString(RefChainSearch.GetRootPath()));
-				
-			// We cannot safely recover from this.
-			if (UObjectBaseUtility::IsPendingKillEnabled())
-			{
-				FMessageLog(NAME_CategoryPIE).CriticalError()
-					->AddToken(FUObjectToken::Create(Object, FText::FromString(Object->GetFullName())))
-					->AddToken(FTextToken::Create(FText::Format(LOCTEXT("PIEObjectStillReferenced", "Object from PIE level still referenced. Shortest path from root: {Path}"), Arguments)));
-			}
-			else
-			{
-				// And yet we still try and recover
-				UPackage* ObjectPackage = Object->GetOutermost();
-				ObjectPackage->ClearPackageFlags(PKG_PlayInEditor);
-				ObjectPackage->ClearFlags(RF_Standalone);
-				// We let it leak but it needs to be renamed to not collide with future attempts at creating the same object(s)
-				ObjectPackage->Rename(*MakeUniqueObjectName(nullptr, UPackage::StaticClass()).ToString(), nullptr, REN_ForceNoResetLoaders | REN_DontCreateRedirectors | REN_DoNotDirty | REN_NonTransactional);
-
-				FText ErrorMessage = FText::Format(LOCTEXT("PIEObjectStillReferenced", "Object from PIE level still referenced. Shortest path from root: {Path}"), Arguments);
-				// Throw an ensure so that crash reporter logs this error
-				ensureAlwaysMsgf(false, TEXT("%s"), *ErrorMessage.ToString());
-				// Display the error in the editor window
-				FMessageLog(NAME_CategoryPIE).Error()
-					->AddToken(FUObjectToken::Create(Object, FText::FromString(Object->GetFullName())))
-					->AddToken(FTextToken::Create(ErrorMessage));
+				LeakedObjectsSet.Add(ObjectPackage);
 			}
 		}
 	}
 
+	TArray<UObject*> LeakedObjects = LeakedObjectsSet.Array();
+	TArray<FString> Paths;
+	if (LeakedObjects.Num() > 0)
+	{
+		Paths = FindAndPrintStaleReferencesToObjects(LeakedObjects, EPrintStaleReferencesOptions::Ensure);
+	}
+	for (int32 i = 0; i < LeakedObjects.Num(); ++i)
+	{
+		const FString& Path = Paths[i];
+		UObject* Object = LeakedObjects[i];
+
+		FFormatNamedArguments Arguments;
+		Arguments.Add(TEXT("Path"), FText::FromString(Paths[i]));
+
+		// We cannot safely recover from this.
+		if (UObjectBaseUtility::IsPendingKillEnabled())
+		{
+			FMessageLog(NAME_CategoryPIE).CriticalError()
+				->AddToken(FUObjectToken::Create(Object, FText::FromString(Object->GetFullName())))
+				->AddToken(FTextToken::Create(FText::Format(LOCTEXT("PIEObjectStillReferenced", "Object from PIE level still referenced. Shortest path from root: {Path}"), Arguments)));
+		}
+		else
+		{
+			// Nonfatal error, we will rename objects to try and recover. 
+			FText ErrorMessage = FText::Format(LOCTEXT("PIEObjectStillReferenced", "Object from PIE level still referenced. Shortest path from root: {Path}"), Arguments);
+			FMessageLog(NAME_CategoryPIE).Error()
+				->AddToken(FUObjectToken::Create(Object, FText::FromString(Object->GetFullName())))
+				->AddToken(FTextToken::Create(ErrorMessage));
+		}
+	}
+
+	// Try and recover by renaming leaked packages 
+	if (!UObjectBaseUtility::IsPendingKillEnabled())
+	{
+		for (UPackage* ObjectPackage : LeakedPackages)
+		{
+			ObjectPackage->ClearPackageFlags(PKG_PlayInEditor);
+			ObjectPackage->ClearFlags(RF_Standalone);
+			// We let it leak but it needs to be renamed to not collide with future attempts at creating the same object(s)
+			FName NewName = MakeUniqueObjectName(nullptr, UPackage::StaticClass());
+			UE_LOG(LogTemp, Log, TEXT("Renaming PIE package from '%s' to '%s' to prevent future name collisions."),
+				*ObjectPackage->GetName(), *NewName.ToString());
+			ObjectPackage->Rename(*NewName.ToString(), nullptr, REN_ForceNoResetLoaders | REN_DontCreateRedirectors | REN_DoNotDirty | REN_NonTransactional);
+		}
+	}
+
 	// Final cleanup/reseting
 	FWorldContext& EditorWorldContext = GEditor->GetEditorWorldContext();
 	UPackage* Package = EditorWorldContext.World()->GetOutermost();
