Improved the performance of the TypedElementSelectionSet Selection.
From 119 Seconds to 269 ms when measuring the time redo a large marquee selection (30 384 static mesh instances).

diff --git a/Engine/Source/Editor/UnrealEd/Classes/Editor/Transactor.h b/Engine/Source/Editor/UnrealEd/Classes/Editor/Transactor.h
--- a/Engine/Source/Editor/UnrealEd/Classes/Editor/Transactor.h
+++ b/Engine/Source/Editor/UnrealEd/Classes/Editor/Transactor.h
@@ -641,6 +641,7 @@ public:
 	virtual void StoreUndo( UObject* Object, TUniquePtr<FChange> UndoChange ) override;
 	virtual void SetPrimaryObject(UObject* InObject) override;
 	virtual void SnapshotObject( UObject* InObject, TArrayView<const FProperty*> Properties ) override;
+	virtual bool ContainsObject(const UObject* Object) const override;
 
 	/** BeginOperation should be called when a transaction or undo/redo starts */
 	virtual void BeginOperation() override;
diff --git a/Engine/Source/Editor/UnrealEd/Private/DragTool_BoxSelect.cpp b/Engine/Source/Editor/UnrealEd/Private/DragTool_BoxSelect.cpp
--- a/Engine/Source/Editor/UnrealEd/Private/DragTool_BoxSelect.cpp
+++ b/Engine/Source/Editor/UnrealEd/Private/DragTool_BoxSelect.cpp
@@ -224,6 +224,8 @@ void FDragTool_ActorBoxSelect::EndDrag()
 		bool bSelectionChanged = false;
 		UWorld* IteratorWorld = GWorld;
 		const TArray<FName>& HiddenLayers = LevelViewportClient->ViewHiddenLayers;
+		TArray<FTypedElementHandle> Handles;
+
 		for( FActorIterator It(IteratorWorld); It; ++It )
 		{
 			AActor* Actor = *It;
@@ -242,18 +244,19 @@ void FDragTool_ActorBoxSelect::EndDrag()
 			// Select the actor or its child elements
 			if( bActorIsVisible )
 			{
-				TArray<FTypedElementHandle> Handles = UE::LevelEditor::Private::GetElementsIntersectingBox( Actor, SelBBox, LevelViewportClient, SeletionArgs );
-				if (bShouldSelect)
-				{
-					SelectionSet->SelectElements(Handles, ElementSelectionOption);
-				}
-				else
-				{
-					SelectionSet->DeselectElements(Handles, ElementSelectionOption);
-				}
+				Handles.Append( UE::LevelEditor::Private::GetElementsIntersectingBox( Actor, SelBBox, LevelViewportClient, SeletionArgs ) );
 			}
 		}
 
+		if (bShouldSelect)
+		{
+			SelectionSet->SelectElements( Handles, ElementSelectionOption );
+		}
+		else
+		{
+			SelectionSet->DeselectElements( Handles, ElementSelectionOption );
+		}
+
 		// Check every model to see if its BSP surfaces should be selected
 		for( int32 ModelIndex = 0; ModelIndex < ModelsToCheck.Num(); ++ModelIndex )
 		{
diff --git a/Engine/Source/Editor/UnrealEd/Private/DragTool_FrustumSelect.cpp b/Engine/Source/Editor/UnrealEd/Private/DragTool_FrustumSelect.cpp
--- a/Engine/Source/Editor/UnrealEd/Private/DragTool_FrustumSelect.cpp
+++ b/Engine/Source/Editor/UnrealEd/Private/DragTool_FrustumSelect.cpp
@@ -125,15 +125,6 @@ void FDragTool_ActorFrustumSelect::EndDrag()
 			bGeometryMode
 		};
 
-		 FTypedElementSelectionOptions ElementSelectionOption;
-
-
-		if( !bShiftDown )
-		{
-			// If the user is selecting, but isn't hold down SHIFT, remove all current selections.
-			SelectionSet->ClearSelection(ElementSelectionOption);
-		}
-
 		const int32 ViewportSizeX = LevelViewportClient->Viewport->GetSizeXY().X;
 		const int32 ViewportSizeY = LevelViewportClient->Viewport->GetSizeXY().Y;
 
@@ -147,6 +138,7 @@ void FDragTool_ActorFrustumSelect::EndDrag()
 			Swap( Start.Y, End.Y );
 		}
 
+		TArray<FTypedElementHandle> ElementsToSelect;
 		const bool bTransparentBoxSelection = GetDefault<ULevelEditorViewportSettings>()->bTransparentBoxSelection;
 		if (bTransparentBoxSelection)
 		{
@@ -154,7 +146,7 @@ void FDragTool_ActorFrustumSelect::EndDrag()
 			for(FActorIterator It(LevelViewportClient->GetWorld()); It; ++It)
 			{
 				AActor* Actor = *It;
-				SelectionSet->SelectElements(UE::LevelEditor::Private::GetElementsIntersectingFrustum(Actor, Frustum, LevelViewportClient, SeletionArgs), ElementSelectionOption);
+				ElementsToSelect.Append(UE::LevelEditor::Private::GetElementsIntersectingFrustum(Actor, Frustum, LevelViewportClient, SeletionArgs));
 			}
 		}
 		else
@@ -170,12 +162,12 @@ void FDragTool_ActorFrustumSelect::EndDrag()
 
 				if (bStrictDragSelection)
 				{
-					ElementList->ForEachElement<ITypedElementWorldInterface>([bStrictDragSelection, &Frustum, &SelectionSet, &ElementSelectionOption]
+					ElementList->ForEachElement<ITypedElementWorldInterface>([bStrictDragSelection, &Frustum, &SelectionSet, &ElementsToSelect]
 						(const TTypedElement<ITypedElementWorldInterface>& InElement)
 						{
 							if (InElement.IsElementInConvexVolume(Frustum, bStrictDragSelection))
 							{
-								SelectionSet->SelectElement(SelectionSet->GetSelectionElement(InElement, ETypedElementSelectionMethod::Primary), ElementSelectionOption);
+								ElementsToSelect.Add(SelectionSet->GetSelectionElement(InElement, ETypedElementSelectionMethod::Primary));
 							}
 
 							return true;
@@ -184,9 +176,9 @@ void FDragTool_ActorFrustumSelect::EndDrag()
 				else
 				{
 					// Grab only the selectable handles (this remove the components from the selection and select the actor instead)
-					ElementList->ForEachElementHandle([&SelectionSet, &ElementSelectionOption](const FTypedElementHandle& InHandle)
+					ElementList->ForEachElementHandle([&SelectionSet, &ElementsToSelect](const FTypedElementHandle& InHandle)
 						{
-							SelectionSet->SelectElement(SelectionSet->GetSelectionElement(InHandle, ETypedElementSelectionMethod::Primary), ElementSelectionOption);
+							ElementsToSelect.Add(SelectionSet->GetSelectionElement(InHandle, ETypedElementSelectionMethod::Primary));
 							return true;
 						});
 				}
@@ -226,15 +218,27 @@ void FDragTool_ActorFrustumSelect::EndDrag()
 					AActor* Actor = *It;
 					if (bStrictDragSelection)
 					{
-						SelectionSet->SelectElements(UE::LevelEditor::Private::GetElementsIntersectingFrustum(Actor, Frustum, LevelViewportClient, SeletionArgs), ElementSelectionOption);
+						ElementsToSelect.Append(UE::LevelEditor::Private::GetElementsIntersectingFrustum(Actor, Frustum, LevelViewportClient, SeletionArgs));
 					}
 					else
 					{
-						SelectionSet->SelectElement(UEngineElementsLibrary::AcquireEditorActorElementHandle(Actor), ElementSelectionOption);
+						ElementsToSelect.Add(UEngineElementsLibrary::AcquireEditorActorElementHandle(Actor));
 					}
 				}
 			}
 		}
+
+		 FTypedElementSelectionOptions ElementSelectionOption;
+
+		if (!bShiftDown)
+		{
+			// If the user is selecting, but isn't hold down SHIFT, remove the previous selections.
+			SelectionSet->SetSelection(MoveTemp(ElementsToSelect), ElementSelectionOption);
+		}
+		else
+		{
+			SelectionSet->SelectElements(MoveTemp(ElementsToSelect), ElementSelectionOption);
+		}
 	}
 
 	// Clear any hovered objects that might have been created while dragging
diff --git a/Engine/Source/Editor/UnrealEd/Private/EditorTransaction.cpp b/Engine/Source/Editor/UnrealEd/Private/EditorTransaction.cpp
--- a/Engine/Source/Editor/UnrealEd/Private/EditorTransaction.cpp
+++ b/Engine/Source/Editor/UnrealEd/Private/EditorTransaction.cpp
@@ -910,6 +910,23 @@ void FTransaction::SnapshotObject( UObject* InObject, TArrayView<const FProperty
 	}
 }
 
+bool FTransaction::ContainsObject(const UObject* Object) const
+{
+	FPersistentObjectRef PersistentObjectRef(const_cast<UObject*>(Object));
+	if (const FObjectRecords* ObjectRecords = ObjectRecordsMap.Find(PersistentObjectRef))
+	{
+		for (FObjectRecord* Record : ObjectRecords->Records)
+		{
+			if (Record->Object == PersistentObjectRef)
+			{
+				return true;
+			}
+		}
+	}
+
+	return false;
+}
+
 void FTransaction::BeginOperation()
 {
 	check(!OperationId.IsValid());
diff --git a/Engine/Source/Runtime/Core/Public/Misc/ITransaction.h b/Engine/Source/Runtime/Core/Public/Misc/ITransaction.h
--- a/Engine/Source/Runtime/Core/Public/Misc/ITransaction.h
+++ b/Engine/Source/Runtime/Core/Public/Misc/ITransaction.h
@@ -397,4 +397,9 @@ public:
 	 * @param Property	The optional list of properties that have potentially changed on the object (to avoid snapshotting the entire object).
 	 */
 	virtual void SnapshotObject( UObject* Object, TArrayView<const FProperty*> Properties ) = 0;
+
+	/**
+	 * Does the transaction know that the object is being modified.
+	 */
+	virtual bool ContainsObject(const UObject* Object) const = 0;
 };
diff --git a/Engine/Source/Runtime/Engine/Private/Instances/InstancedPlacementClientInfo.cpp b/Engine/Source/Runtime/Engine/Private/Instances/InstancedPlacementClientInfo.cpp
--- a/Engine/Source/Runtime/Engine/Private/Instances/InstancedPlacementClientInfo.cpp
+++ b/Engine/Source/Runtime/Engine/Private/Instances/InstancedPlacementClientInfo.cpp
@@ -400,8 +400,16 @@ void FClientPlacementInfo::NotifyISMPartitionInstanceSelectionChanged(const FISM
 	check(IsInitialized());
 	check(InstanceId.Handle == ClientHandle);
 
-	ParentPartitionActor->Modify(false);
-	ParentPartitionActor->SelectISMInstances(ClientHandle, bIsSelected, { InstanceId.Index });
+	if (AInstancedPlacementPartitionActor* ParentPartitionActorPtr = ParentPartitionActor.Get())
+	{
+		if (!GUndo || !GUndo->ContainsObject(ParentPartitionActorPtr))
+		{
+			// Actor modify can be a bit slow since when there is a transaction active it need to iterate the properties
+			ParentPartitionActorPtr->Modify(false);
+		}
+
+		ParentPartitionActorPtr->SelectISMInstances(ClientHandle, bIsSelected, { InstanceId.Index });
+	}
 #endif
 }
 
diff --git a/Engine/Source/Runtime/TypedElementFramework/Private/Elements/Framework/TypedElementList.cpp b/Engine/Source/Runtime/TypedElementFramework/Private/Elements/Framework/TypedElementList.cpp
--- a/Engine/Source/Runtime/TypedElementFramework/Private/Elements/Framework/TypedElementList.cpp
+++ b/Engine/Source/Runtime/TypedElementFramework/Private/Elements/Framework/TypedElementList.cpp
@@ -360,39 +360,38 @@ bool TTypedElementList<HandleType>::RemoveElementImpl(const FTypedElementId& InE
 template<class HandleType>
 int32 TTypedElementList<HandleType>::RemoveAllElementsImpl(TFunctionRef<bool(const HandleType&)> InPredicate)
 {
-	int32 RemovedCount = 0;
-
 	if (ElementHandles.Num() > 0)
 	{
 		FLegacySyncScopedBatch LegacySyncBatch(*this);
 
 		NoteListMayChange();
 
-		for (int32 Index = ElementHandles.Num() - 1; Index >= 0; --Index)
-		{
-			if (InPredicate(ElementHandles[Index]))
+		return ElementHandles.RemoveAll([this, InPredicate](HandleType& InHandle)
 			{
-				HandleType RemovedElementHandle = MoveTemp(ElementHandles[Index]);
-				ElementCombinedIds.Remove(RemovedElementHandle.GetId().GetCombinedId());
-				ElementHandles.RemoveAt(Index, 1, /*bAllowShrinking*/false);
-
-				if constexpr (std::is_same<HandleType, FTypedElementHandle>::value)
+				if (InPredicate(InHandle))
 				{
-					ElementCounts.RemoveElement(RemovedElementHandle);
-				}
-				else
-				{
-					ElementCounts.RemoveElement(RemovedElementHandle.GetTypedElementHandle());
-				}
+					HandleType RemovedElementHandle = MoveTemp(InHandle);
+					ElementCombinedIds.Remove(RemovedElementHandle.GetId().GetCombinedId());
+
+					if constexpr (std::is_same<HandleType, FTypedElementHandle>::value)
+					{
+						ElementCounts.RemoveElement(RemovedElementHandle);
+					}
+					else
+					{
+						ElementCounts.RemoveElement(RemovedElementHandle.GetTypedElementHandle());
+					}
 
-				NoteListChanged(EChangeType::Removed, RemovedElementHandle);
+					NoteListChanged(EChangeType::Removed, RemovedElementHandle);
 
-				++RemovedCount;
-			}
-		}
+					return true;
+				}
+
+				return false;
+			});
 	}
 
-	return RemovedCount;
+	return 0;
 }
 
 template<class HandleType>
diff --git a/Engine/Source/Runtime/TypedElementRuntime/Private/Elements/Framework/TypedElementSelectionSet.cpp b/Engine/Source/Runtime/TypedElementRuntime/Private/Elements/Framework/TypedElementSelectionSet.cpp
--- a/Engine/Source/Runtime/TypedElementRuntime/Private/Elements/Framework/TypedElementSelectionSet.cpp
+++ b/Engine/Source/Runtime/TypedElementRuntime/Private/Elements/Framework/TypedElementSelectionSet.cpp
@@ -6,9 +6,10 @@
 #include "Elements/Framework/TypedElementUtil.h"
 #include "Elements/Interfaces/TypedElementHierarchyInterface.h"
 
-#include "UObject/GCObjectScopeGuard.h"
+#include "Misc/ITransaction.h"
 #include "Serialization/ObjectReader.h"
 #include "Serialization/ObjectWriter.h"
+#include "UObject/GCObjectScopeGuard.h"
 
 #if WITH_EDITORONLY_DATA
 #include "Serialization/TextReferenceCollector.h"
@@ -76,7 +77,7 @@ void UTypedElementSelectionSet::PostEditUndo()
 
 bool UTypedElementSelectionSet::Modify(bool bAlwaysMarkDirty)
 {
-	if (GUndo && CanModify())
+	if (GUndo && CanModify() && !GUndo->ContainsObject(this))
 	{
 		bool bCanModify = true;
 		ElementList->ForEachElement<ITypedElementSelectionInterface>([&bCanModify](const TTypedElement<ITypedElementSelectionInterface>& InSelectionElement)
@@ -225,6 +226,7 @@ bool UTypedElementSelectionSet::SelectElements(TArrayView<const FTypedElementHan
 	FTypedElementList::FLegacySyncScopedBatch LegacySyncBatch(*ElementList, InSelectionOptions.AllowLegacyNotifications());
 
 	bool bSelectionChanged = false;
+	ElementList->Reserve(ElementList->Num() + InElementHandles.Num());
 
 	for (const FTypedElementHandle& ElementHandle : InElementHandles)
 	{
@@ -239,6 +241,7 @@ bool UTypedElementSelectionSet::SelectElements(FTypedElementListConstRef InEleme
 	FTypedElementList::FLegacySyncScopedBatch LegacySyncBatch(*ElementList, InSelectionOptions.AllowLegacyNotifications());
 
 	bool bSelectionChanged = false;
+	ElementList->Reserve(ElementList->Num() + InElementList->Num());
 
 	InElementList->ForEachElementHandle([this, &bSelectionChanged, &InSelectionOptions](const FTypedElementHandle& ElementHandle)
 	{
@@ -317,7 +320,7 @@ bool UTypedElementSelectionSet::DeselectElements(FTypedElementListConstRef InEle
 bool UTypedElementSelectionSet::ClearSelection(const FTypedElementSelectionOptions InSelectionOptions)
 {
 	FTypedElementList::FLegacySyncScopedBatch LegacySyncBatch(*ElementList, InSelectionOptions.AllowLegacyNotifications());
-
+	
 	bool bSelectionChanged = false;
 
 	// Run deselection via the selection interface where possible
