Cooker PakChunk files: Suppress the creation of empty pakchunkN.txt files for empty chunks. It causes a performance problem when the project is using high chunkid values for organization while leaving lower values unused.
#rnx

diff --git a/Engine/Source/Editor/UnrealEd/Private/Commandlets/AssetRegistryGenerator.cpp b/Engine/Source/Editor/UnrealEd/Private/Commandlets/AssetRegistryGenerator.cpp
--- a/Engine/Source/Editor/UnrealEd/Private/Commandlets/AssetRegistryGenerator.cpp
+++ b/Engine/Source/Editor/UnrealEd/Private/Commandlets/AssetRegistryGenerator.cpp
@@ -554,10 +554,18 @@ bool FAssetRegistryGenerator::GenerateStreamingInstallManifest(int64 InOverrideC
 			FinalChunkManifests[PakchunkIndex].Reset(new FChunkPackageSet());
 			Manifest = FinalChunkManifests[PakchunkIndex].Get();
 		}
-
+		
+		// Split the chunk into subchunks as necessary and create and register a PakListFile for each subchunk
 		int32 FilenameIndex = 0;
 		TArray<FString> ChunkFilenames;
 		Manifest->GenerateValueArray(ChunkFilenames);
+
+		// Do not create any files if the chunk is empty
+		if (ChunkFilenames.IsEmpty())
+		{
+			continue;
+		}
+
 		bool bFinishedAllFiles = false;
 		for (int32 SubChunkIndex = 0; !bFinishedAllFiles; ++SubChunkIndex)
 		{
