Cooker: Multiprocess support: Remove some deprecated code and move some out of place functions that should not be executed on CookWorker into the proper calling function that will be called on the director only.
#rnx

diff --git a/Engine/Config/BaseEditor.ini b/Engine/Config/BaseEditor.ini
--- a/Engine/Config/BaseEditor.ini
+++ b/Engine/Config/BaseEditor.ini
@@ -315,9 +315,6 @@ BulkDataRegistryEnabled=true
 ; HybridIterative uses the TargetDomain and BuildDependencies to decide when packages need to be recooked. If enabled,
 ; every cook is iterative for packages with allowed classes (IterativeClassAllowList) unless -noiterative is passed.
 HybridIterativeEnabled=false
-; PreexploreDependenciesEnabled can cause editor-only packages to be cooked due to lack of robust editor-only reference detection.
-; Projects can disable it to prevent cooking those packages, at the cost of lower cook performance.
-PreexploreDependenciesEnabled=true
 PackagesPerGC=0
 IdleTimeToGC=20
 MemoryMaxUsedPhysical=16384
diff --git a/Engine/Source/Editor/UnrealEd/Classes/CookOnTheSide/CookOnTheFlyServer.h b/Engine/Source/Editor/UnrealEd/Classes/CookOnTheSide/CookOnTheFlyServer.h
--- a/Engine/Source/Editor/UnrealEd/Classes/CookOnTheSide/CookOnTheFlyServer.h
+++ b/Engine/Source/Editor/UnrealEd/Classes/CookOnTheSide/CookOnTheFlyServer.h
@@ -302,6 +302,8 @@ private:
 
 	/* Initializing Platforms must be done on the tickloop thread; Platform data is read only on other threads */
 	void AddCookOnTheFlyPlatformFromGameThread(ITargetPlatform* TargetPlatform);
+	/* Start the session for the given platform in response to the first client connection. */
+	void StartCookOnTheFlySessionFromGameThread(ITargetPlatform* TargetPlatform);
 
 	/* Callback to recalculate all ITargetPlatform* pointers when they change due to modules reloading */
 	void OnTargetPlatformsInvalidated();
@@ -514,8 +516,7 @@ public:
 	TArray<UPackage*> GetUnsolicitedPackages(const TArray<const ITargetPlatform*>& TargetPlatforms) const;
 
 	/** Execute class-specific special case cook postloads and reference discovery on a given package. */
-	void PostLoadPackageFixup(UPackage* Package, TArray<FName>* OutDiscoveredPackageNames,
-		TMap<FName, UE::Cook::FInstigator>* OutInstigators);
+	void PostLoadPackageFixup(UE::Cook::FPackageData& PackageData, UPackage* Package);
 
 	/** Tick CBTB until it finishes or needs to yield. Should only be called when in CookByTheBook Mode. */
 	uint32 TickCookByTheBook(const float TimeSlice, uint32& CookedPackageCount, ECookTickFlags TickFlags = ECookTickFlags::None, int32 InMaxNumPackagesToSave = 50);
@@ -854,11 +855,10 @@ private:
 	 *
 	 * @param BuildFilename long package name of the package to load 
 	 * @param OutPackage UPackage of the package loaded, non-null on success, may be non-null on failure if the UPackage existed but had another failure
-	 * @param OverrideFileName Filename to use instead of PackageData.GetFileName(), for e.g. generated packages
 	 * @param ReportingPackageData PackageData which caused the load, for e.g. generated packages
 	 * @return Whether LoadPackage was completely successful and the package can be cooked
 	 */
-	bool LoadPackageForCooking(UE::Cook::FPackageData& PackageData, UPackage*& OutPackage, FString* LoadFromFileName = nullptr,
+	bool LoadPackageForCooking(UE::Cook::FPackageData& PackageData, UPackage*& OutPackage,
 		UE::Cook::FPackageData* ReportingPackageData = nullptr);
 
 	/** Read information about the previous cook from disk and set whether the current cook is iterative based on previous cook, config, and commandline. */
@@ -1100,8 +1100,6 @@ private:
 
 	/** Waits for the AssetRegistry to complete so that we know any missing assets are missing on disk */
 	void BlockOnAssetRegistry();
-	/** Setup necessary only once for CookOnTheFly, but that are not required until the first request. */
-	void CookOnTheFlyDeferredInitialize();
 
 	/** Construct or refresh-for-filechanges the platform-specific asset registry for the given platforms */
 	void RefreshPlatformAssetRegistries(const TArrayView<const ITargetPlatform* const>& TargetPlatforms);
@@ -1164,13 +1162,10 @@ private:
 	bool bPreloadingEnabled = false;
 	/** If enabled, we load/save TargetDomainKey hashes and use those to test whether packages have already been cooked in hybrid-iterative builds. */
 	bool bHybridIterativeEnabled = true;
-	/** Whether to explore transitive dependencies up front, or discover them as we load packages. */
-	bool bPreexploreDependenciesEnabled = true;
 	/** Test mode for the debug of hybrid iterative dependencies. */
 	bool bHybridIterativeDebug = false;
 	bool bFirstCookInThisProcessInitialized = false;
 	bool bFirstCookInThisProcess = true;
-	bool bHasDeferredInitializeCookOnTheFly = false;
 	bool bImportBehaviorCallbackInstalled = false;
 	/** Cancel has been queued and will be processed next tick */
 	bool bCancelSession = false;
diff --git a/Engine/Source/Editor/UnrealEd/Private/Commandlets/AssetRegistryGenerator.cpp b/Engine/Source/Editor/UnrealEd/Private/Commandlets/AssetRegistryGenerator.cpp
--- a/Engine/Source/Editor/UnrealEd/Private/Commandlets/AssetRegistryGenerator.cpp
+++ b/Engine/Source/Editor/UnrealEd/Private/Commandlets/AssetRegistryGenerator.cpp
@@ -89,11 +89,30 @@ FName GetPackageNameFromDependencyPackageName(const FName RawPackageFName)
 //////////////////////////////////////////////////////////////////////////
 // FAssetRegistryGenerator
 
+void FAssetRegistryGenerator::UpdateAssetManagerDatabase()
+{
+	InitializeUseAssetManager();
+	if (bUseAssetManager)
+	{
+		UAssetManager::Get().UpdateManagementDatabase();
+	}
+}
+
+void FAssetRegistryGenerator::InitializeUseAssetManager()
+{
+	PRAGMA_DISABLE_DEPRECATION_WARNINGS;
+	FAssignStreamingChunkDelegate& AssignStreamingChunkDelegate = FGameDelegates::Get().GetAssignStreamingChunkDelegate();
+	FGetPackageDependenciesForManifestGeneratorDelegate& GetPackageDependenciesForManifestGeneratorDelegate = FGameDelegates::Get().GetGetPackageDependenciesForManifestGeneratorDelegate();
+	PRAGMA_ENABLE_DEPRECATION_WARNINGS;
+	bUseAssetManager = UAssetManager::IsValid() && !AssignStreamingChunkDelegate.IsBound() && !GetPackageDependenciesForManifestGeneratorDelegate.IsBound();
+}
+
+bool FAssetRegistryGenerator::bUseAssetManager = true;
+
 FAssetRegistryGenerator::FAssetRegistryGenerator(const ITargetPlatform* InPlatform)
 	: AssetRegistry(FModuleManager::LoadModuleChecked<FAssetRegistryModule>(TEXT("AssetRegistry")).Get())
 	, TargetPlatform(InPlatform)
 	, bGenerateChunks(false)
-	, bUseAssetManager(false)
 	, HighestChunkId(0)
 	, DependencyInfo(*GetMutableDefault<UChunkDependencyInfo>())
 {
@@ -106,17 +125,7 @@ FAssetRegistryGenerator::FAssetRegistryGenerator(const ITargetPlatform* InPlatfo
 
 	DependencyQuery = bOnlyHardReferences ? UE::AssetRegistry::EDependencyQuery::Hard : UE::AssetRegistry::EDependencyQuery::NoRequirements;
 
-	PRAGMA_DISABLE_DEPRECATION_WARNINGS
-	FAssignStreamingChunkDelegate& AssignStreamingChunkDelegate = FGameDelegates::Get().GetAssignStreamingChunkDelegate();
-	FGetPackageDependenciesForManifestGeneratorDelegate& GetPackageDependenciesForManifestGeneratorDelegate = FGameDelegates::Get().GetGetPackageDependenciesForManifestGeneratorDelegate();
-	PRAGMA_ENABLE_DEPRECATION_WARNINGS
-	if (UAssetManager::IsValid() && !AssignStreamingChunkDelegate.IsBound() && !GetPackageDependenciesForManifestGeneratorDelegate.IsBound())
-	{
-		bUseAssetManager = true;
-
-		UAssetManager::Get().UpdateManagementDatabase();
-	}
-
+	InitializeUseAssetManager();
 	InitializeChunkIdPakchunkIndexMapping();
 }
 
diff --git a/Engine/Source/Editor/UnrealEd/Private/CookOnTheFlyServer.cpp b/Engine/Source/Editor/UnrealEd/Private/CookOnTheFlyServer.cpp
--- a/Engine/Source/Editor/UnrealEd/Private/CookOnTheFlyServer.cpp
+++ b/Engine/Source/Editor/UnrealEd/Private/CookOnTheFlyServer.cpp
@@ -337,7 +337,6 @@ public:
 			return false;
 		}
 
-		Cooker.CookOnTheFlyDeferredInitialize();
 		FName StandardFileName(*FPaths::CreateStandardFilename(CookPackageRequest.Filename));
 		UE_LOG(LogCook, Verbose, TEXT("Enqueing cook request, Filename='%s', Platform='%s'"), *CookPackageRequest.Filename, *CookPackageRequest.PlatformName.ToString());
 		Cooker.ExternalRequests->EnqueueUnique(
@@ -659,16 +658,6 @@ bool UCookOnTheFlyServer::StartCookOnTheFly(FCookOnTheFlyStartupOptions InCookOn
 #endif
 }
 
-void UCookOnTheFlyServer::CookOnTheFlyDeferredInitialize()
-{
-	if (bHasDeferredInitializeCookOnTheFly)
-	{
-		return;
-	}
-	bHasDeferredInitializeCookOnTheFly = true;
-	BlockOnAssetRegistry();
-}
-
 void UCookOnTheFlyServer::InitializeShadersForCookOnTheFly(const TArrayView<ITargetPlatform* const>& NewTargetPlatforms)
 {
 	UE_LOG(LogCook, Display, TEXT("Initializing shaders for cook-on-the-fly"));
@@ -716,6 +705,17 @@ void UCookOnTheFlyServer::AddCookOnTheFlyPlatformFromGameThread(ITargetPlatform*
 	SaveCurrentIniSettings(TargetPlatform);
 }
 
+void UCookOnTheFlyServer::StartCookOnTheFlySessionFromGameThread(ITargetPlatform* TargetPlatform)
+{
+	PlatformManager->AddSessionPlatform(*this, TargetPlatform);
+	bPackageFilterDirty = true;
+
+	// Blocking on the AssetRegistry needs to wait until the session starts because it needs all plugins loaded.
+	// AddCookOnTheFlyPlatformFromGameThread can be called on cooker startup which occurs in UUnrealEdEngine::Init
+	// before all plugins are loaded.
+	BlockOnAssetRegistry();
+}
+
 void UCookOnTheFlyServer::OnTargetPlatformsInvalidated()
 {
 	check(IsInGameThread());
@@ -2043,7 +2043,6 @@ void UCookOnTheFlyServer::PumpRequests(UE::Cook::FTickStackData& StackData, int3
 	{
 		FPackageData* PackageData = RequestQueue.PopReadyRequest();
 		FPoppedPackageDataScope Scope(*PackageData);
-		// Exploring dependencies in ReadRequests is legacy behavior and is allowed even if bPreexploreDependenciesEnabled  is false
 		if (TryCreateRequestCluster(*PackageData))
 		{
 			continue;
@@ -2104,7 +2103,7 @@ void UCookOnTheFlyServer::PumpLoads(UE::Cook::FTickStackData& StackData, uint32
 
 		FPackageData& PackageData(*LoadReadyQueue.PopFrontValue());
 		FPoppedPackageDataScope Scope(PackageData);
-		if (bPreexploreDependenciesEnabled && TryCreateRequestCluster(PackageData))
+		if (TryCreateRequestCluster(PackageData))
 		{
 			continue;
 		}
@@ -2155,7 +2154,7 @@ void UCookOnTheFlyServer::PumpPreloadStarts()
 	while (!EntryQueue.IsEmpty() && Monitor.GetNumPreloadAllocated() < static_cast<int32>(MaxPreloadAllocated))
 	{
 		FPackageData* PackageData = EntryQueue.PopFrontValue();
-		if (bPreexploreDependenciesEnabled  && TryCreateRequestCluster(*PackageData))
+		if (TryCreateRequestCluster(*PackageData))
 		{
 			continue;
 		}
@@ -2248,11 +2247,11 @@ void UCookOnTheFlyServer::LoadPackageInQueue(UE::Cook::FPackageData& PackageData
 
 		FPackageData& OwnerPackageData = const_cast<FPackageData&>(GeneratorStruct->GetOwner());
 		UPackage* OwnerPackage;
-		bool bLoadFullySuccessful = LoadPackageForCooking(OwnerPackageData, OwnerPackage, nullptr, &PackageData);
+		bool bLoadFullySuccessful = LoadPackageForCooking(OwnerPackageData, OwnerPackage, &PackageData);
 		if (!bLoadFullySuccessful)
 		{
 			ResultFlags |= COSR_ErrorLoadingPackage;
-			UE_LOG(LogCook, Error, TEXT("Package %s is a generated package and we could not load it's generator package %s. It can not be loaded."),
+			UE_LOG(LogCook, Error, TEXT("Package %s is a generated package and we could not load its generator package %s. It can not be loaded."),
 				*PackageFileName.ToString(), *OwnerPackageData.GetFileName().ToString());
 			RejectPackageToLoad(PackageData, TEXT("is a generated package which could not load its generator"));
 			return;
@@ -2272,13 +2271,13 @@ void UCookOnTheFlyServer::LoadPackageInQueue(UE::Cook::FPackageData& PackageData
 		// Already cooked. This can happen if we needed to load a package that was previously cooked and garbage collected because it is a loaddependency of a new request.
 		// Send the package back to idle, nothing further to do with it.
 		PackageData.SendToState(EPackageState::Idle, ESendFlags::QueueAdd);
+		return;
 	}
-	else
-	{
-		PackageData.SetPackage(LoadedPackage);
-		PackageData.SendToState(EPackageState::Save, ESendFlags::QueueAdd);
-		++OutNumPushed;
-	}
+
+	PostLoadPackageFixup(PackageData, LoadedPackage);
+	PackageData.SetPackage(LoadedPackage);
+	PackageData.SendToState(EPackageState::Save, ESendFlags::QueueAdd);
+	++OutNumPushed;
 }
 
 void UCookOnTheFlyServer::RejectPackageToLoad(UE::Cook::FPackageData& PackageData, const TCHAR* Reason)
@@ -3025,7 +3024,7 @@ void UCookOnTheFlyServer::TickCancels()
 	PackageDatas->PollPendingCookedPlatformDatas(false);
 }
 
-bool UCookOnTheFlyServer::LoadPackageForCooking(UE::Cook::FPackageData& PackageData, UPackage*& OutPackage, FString* LoadFromFilename,
+bool UCookOnTheFlyServer::LoadPackageForCooking(UE::Cook::FPackageData& PackageData, UPackage*& OutPackage,
 	UE::Cook::FPackageData* ReportingPackageData)
 {
 	UE_SCOPED_HIERARCHICAL_COOKTIMER_AND_DURATION(LoadPackageForCooking, DetailedCookStats::TickCookOnTheSideLoadPackagesTimeSec);
@@ -3041,7 +3040,7 @@ bool UCookOnTheFlyServer::LoadPackageForCooking(UE::Cook::FPackageData& PackageD
 	FString PackageName = PackageData.GetPackageName().ToString();
 	OutPackage = FindObject<UPackage>(nullptr, *PackageName);
 
-	FString FileName(LoadFromFilename ? *LoadFromFilename : PackageData.GetFileName().ToString());
+	FString FileName(PackageData.GetFileName().ToString());
 	FString ReportingFileName(ReportingPackageData ? ReportingPackageData->GetFileName().ToString() : FileName);
 #if DEBUG_COOKONTHEFLY
 	UE_LOG(LogCook, Display, TEXT("Processing request %s"), *ReportingFileName);
@@ -3059,24 +3058,10 @@ bool UCookOnTheFlyServer::LoadPackageForCooking(UE::Cook::FPackageData& PackageD
 	{
 		bool bWasPartiallyLoaded = OutPackage != nullptr;
 		GIsCookerLoadingPackage = true;
-		UPackage* LoadIntoPackage = nullptr;
-		if (LoadFromFilename)
-		{
-			// When loading from a specified filename, we provide the packagename to load into via LoadPackage's first argument - the InOuter package
-			if (OutPackage)
-			{
-				LoadIntoPackage = OutPackage;
-			}
-			else
-			{
-				LoadIntoPackage = CreatePackage(*PackageName);
-			}
-		}
-
 		UPackage* LoadedPackage;
 		{
 			LLM_SCOPE(ELLMTag::Untagged); // Reset the scope so that untagged memory in the package shows up as Untagged rather than Cooker
-			LoadedPackage = LoadPackage(LoadIntoPackage, *FileName, LOAD_None);
+			LoadedPackage = LoadPackage(nullptr, *FileName, LOAD_None);
 		}
 		if (IsValid(LoadedPackage) && LoadedPackage->IsFullyLoaded())
 		{
@@ -3121,33 +3106,22 @@ void UCookOnTheFlyServer::ProcessUnsolicitedPackages(TArray<FName>* OutDiscovere
 {
 	using namespace UE::Cook;
 
-	// Ensure sublevels are loaded by iterating all recently loaded packages and invoking
-	// PostLoadPackageFixup
+	TMap<UPackage*, UE::Cook::FInstigator> NewPackages = PackageTracker->GetNewPackages();
 
+	for (auto& PackageWithInstigator : NewPackages)
 	{
-		UE_SCOPED_HIERARCHICAL_COOKTIMER(PostLoadPackageFixup);
+		UPackage* Package = PackageWithInstigator.Key;
+		FInstigator& Instigator = PackageWithInstigator.Value;
 
-		TMap<UPackage*, UE::Cook::FInstigator> NewPackages = PackageTracker->GetNewPackages();
-
-		for (auto& PackageWithInstigator : NewPackages)
+		bool bWasInProgress;
+		FPackageData* PackageData = QueueDiscoveredPackage(Package, FInstigator(Instigator), &bWasInProgress);
+		if (PackageData && OutDiscoveredPackageNames && !bWasInProgress)
 		{
-			UPackage* Package = PackageWithInstigator.Key;
-			FInstigator& Instigator = PackageWithInstigator.Value;
-			if (!CookByTheBookOptions->bSkipSoftReferences)
-			{
-				PostLoadPackageFixup(Package, OutDiscoveredPackageNames, OutInstigators);
-			}
-
-			bool bWasInProgress;
-			FPackageData* PackageData = QueueDiscoveredPackage(Package, FInstigator(Instigator), &bWasInProgress);
-			if (OutDiscoveredPackageNames && PackageData && !bWasInProgress)
+			FInstigator& Existing = OutInstigators->FindOrAdd(PackageData->GetPackageName());
+			if (Existing.Category == EInstigator::InvalidCategory)
 			{
-				FInstigator& Existing = OutInstigators->FindOrAdd(PackageData->GetPackageName());
-				if (Existing.Category == EInstigator::InvalidCategory)
-				{
-					OutDiscoveredPackageNames->Add(PackageData->GetPackageName());
-					Existing = Instigator;
-				}
+				OutDiscoveredPackageNames->Add(PackageData->GetPackageName());
+				Existing = Instigator;
 			}
 		}
 	}
@@ -3232,7 +3206,7 @@ void UCookOnTheFlyServer::PumpSaves(UE::Cook::FTickStackData& StackData, uint32
 	while (SaveQueue.Num() > static_cast<int32>(DesiredQueueLength))
 	{
 		FPackageData& PackageData(*SaveQueue.PopFrontValue());
-		if (bPreexploreDependenciesEnabled  && TryCreateRequestCluster(PackageData))
+		if (TryCreateRequestCluster(PackageData))
 		{
 			continue;
 		}
@@ -3440,8 +3414,7 @@ void UCookOnTheFlyServer::PumpSaves(UE::Cook::FTickStackData& StackData, uint32
 	}
 }
 
-void UCookOnTheFlyServer::PostLoadPackageFixup(UPackage* Package, TArray<FName>* OutDiscoveredPackageNames,
-	TMap<FName, UE::Cook::FInstigator>* OutInstigators)
+void UCookOnTheFlyServer::PostLoadPackageFixup(UE::Cook::FPackageData& PackageData, UPackage* Package)
 {
 	if (Package->ContainsMap() == false)
 	{
@@ -3453,17 +3426,12 @@ void UCookOnTheFlyServer::PostLoadPackageFixup(UPackage* Package, TArray<FName>*
 		return;
 	}
 
-	// Ensure we only process the package once
-	if (PackageTracker->PostLoadFixupPackages.Find(Package) != nullptr)
-	{
-		return;
-	}
-	PackageTracker->PostLoadFixupPackages.Add(Package);
+	UE_SCOPED_HIERARCHICAL_COOKTIMER(PostLoadPackageFixup);
 
 	// Perform special processing for UWorld
 	World->PersistentLevel->HandleLegacyMapBuildData();
 
-	if (IsCookByTheBookMode() == false)
+	if (IsCookByTheBookMode() == false || CookByTheBookOptions->bSkipSoftReferences)
 	{
 		return;
 	}
@@ -3496,18 +3464,9 @@ void UCookOnTheFlyServer::PostLoadPackageFixup(UPackage* Package, TArray<FName>*
 		UE::Cook::FPackageData* NewPackageData = PackageDatas->TryAddPackageDataByPackageName(PackageFName);
 		if (NewPackageData && !NewPackageData->IsInProgress())
 		{
-			bool bIsUrgent = false;
-			NewPackageData->UpdateRequestData(PlatformManager->GetSessionPlatforms(), bIsUrgent, UE::Cook::FCompletionCallback(),
+			NewPackageData->UpdateRequestData(PlatformManager->GetSessionPlatforms(), false /* IsUrgent */,
+				UE::Cook::FCompletionCallback(),
 				UE::Cook::FInstigator(UE::Cook::EInstigator::Dependency, OwnerName));
-			if (OutDiscoveredPackageNames)
-			{
-				UE::Cook::FInstigator& Existing = OutInstigators->FindOrAdd(PackageFName);
-				if (Existing.Category == UE::Cook::EInstigator::InvalidCategory)
-				{
-					OutDiscoveredPackageNames->Add(PackageFName);
-					Existing = UE::Cook::FInstigator(UE::Cook::EInstigator::Dependency, OwnerName);
-				}
-			}
 		}
 	}
 }
@@ -4364,8 +4323,8 @@ void FSaveCookedPackageContext::SetupPackage()
 	check(!Filename.IsEmpty()); // PackageData guarantees FileName is non-empty; if not found it should never make it into save state
 	if (Package->HasAnyPackageFlags(PKG_ReloadingForCooker))
 	{
-		UE_LOG(LogCook, Warning, TEXT("Package %s marked as reloading for cook by was requested to save"), *PackageName);
-		UE_LOG(LogCook, Fatal, TEXT("Package %s marked as reloading for cook by was requested to save"), *PackageName);
+		UE_LOG(LogCook, Warning, TEXT("Package %s marked as reloading for cook was requested to save"), *PackageName);
+		UE_LOG(LogCook, Fatal, TEXT("Package %s marked as reloading for cook was requested to save"), *PackageName);
 	}
 
 	SaveFlags = SAVE_KeepGUID | SAVE_Async
@@ -4454,13 +4413,6 @@ void FSaveCookedPackageContext::SetupPlatform(const ITargetPlatform* InTargetPla
 		Package->ClearPackageFlags(PKG_FilterEditorOnly);
 	}
 
-	if (World)
-	{
-		// Fixup legacy lightmaps before saving
-		// This should be done after loading, but Core loads UWorlds with LoadObject so there's no opportunity to handle this fixup on load
-		World->PersistentLevel->HandleLegacyMapBuildData();
-	}
-
 	CookContext = &COTFS.FindOrCreateSaveContext(TargetPlatform);
 	SavePackageContext = &CookContext->SaveContext;
 	PackageWriter = CookContext->PackageWriter;
@@ -4780,7 +4732,6 @@ void UCookOnTheFlyServer::Initialize( ECookMode::Type DesiredCookMode, ECookInit
 		bPreloadingEnabled = true;
 		FLinkerLoad::SetPreloadingEnabled(true);
 	}
-	GConfig->GetBool(TEXT("CookSettings"), TEXT("PreexploreDependenciesEnabled"), bPreexploreDependenciesEnabled, GEditorIni);
 
 	PumpPollablesTimeSlice = 1.0f;
 	PumpPollablesMinPeriod = 0.1f;
@@ -6044,6 +5995,8 @@ void UCookOnTheFlyServer::BlockOnAssetRegistry()
 	}
 	TRACE_CPUPROFILER_EVENT_SCOPE(UCookOnTheFlyServer::BlockOnAssetRegistry);
 	UE_LOG(LogCook, Display, TEXT("Waiting for Asset Registry"));
+	// Blocking on the AssetRegistry has to be done on the game thread since some AssetManager functions require it
+	check(IsInGameThread());
 	if (ShouldPopulateFullAssetRegistry())
 	{
 		// Trigger or wait for completion the primary AssetRegistry scan.
@@ -6074,6 +6027,8 @@ void UCookOnTheFlyServer::BlockOnAssetRegistry()
 		DumpAssetRegistryForCooker(AssetRegistry);
 	}
 #endif
+
+	FAssetRegistryGenerator::UpdateAssetManagerDatabase();
 }
 
 void UCookOnTheFlyServer::RefreshPlatformAssetRegistries(const TArrayView<const ITargetPlatform* const>& TargetPlatforms)
@@ -6090,7 +6045,6 @@ void UCookOnTheFlyServer::RefreshPlatformAssetRegistries(const TArrayView<const
 		{
 			RegistryGenerator = new FAssetRegistryGenerator(TargetPlatform);
 			PlatformData->RegistryGenerator = TUniquePtr<FAssetRegistryGenerator>(RegistryGenerator);
-			RegistryGenerator->CleanManifestDirectories();
 		}
 
 		// if we are cooking DLC, we will just spend a lot of time removing the shipped packages from the AR,
@@ -7644,6 +7598,19 @@ void UCookOnTheFlyServer::SetBeginCookConfigSettings(FBeginCookContext& BeginCon
 		// We have cooked before; set bFirstCookInThisProcess=false
 		bFirstCookInThisProcess = false;
 	}
+
+	for (const ITargetPlatform* TargetPlatform : BeginContext.TargetPlatforms)
+	{
+		FConfigFile PlatformEngineIni;
+		FConfigCacheIni::LoadLocalIniFile(PlatformEngineIni, TEXT("Engine"), true, *TargetPlatform->IniPlatformName());
+
+		bool bLegacyBulkDataOffsets = false;
+		PlatformEngineIni.GetBool(TEXT("Core.System"), TEXT("LegacyBulkDataOffsets"), bLegacyBulkDataOffsets);
+		if (bLegacyBulkDataOffsets)
+		{
+			UE_LOG(LogCook, Warning, TEXT("Engine.ini:[Core.System]:LegacyBulkDataOffsets is no longer supported in UE5. The intended use was to reduce patch diffs, but UE5 changed cooked bytes in every package for other reasons, so removing support for this flag does not cause additional patch diffs."));
+		}
+	}
 }
 
 void UCookOnTheFlyServer::SetNeverCookPackageConfigSettings(FBeginCookContext& BeginContext)
@@ -7796,6 +7763,8 @@ void UCookOnTheFlyServer::BeginCookSandbox(FBeginCookContext& BeginContext)
 			CookInfo.bFullBuild = PlatformContext.bFullBuild;
 			CookInfo.bIterateSharedBuild = PlatformContext.bIterateSharedBuild;
 			PackageWriter.Initialize(CookInfo);
+			// Clean the Manifest directory even on iterative builds; it is written from scratch each time
+			PlatformData->RegistryGenerator->CleanManifestDirectories();
 
 			if (PlatformContext.bPopulateMemoryResultsFromDiskResults)
 			{
@@ -7880,16 +7849,6 @@ UE::Cook::FCookSavePackageContext* UCookOnTheFlyServer::CreateSaveContext(const
 
 	DiffModeHelper->InitializePackageWriter(PackageWriter);
 
-	FConfigFile PlatformEngineIni;
-	FConfigCacheIni::LoadLocalIniFile(PlatformEngineIni, TEXT("Engine"), true, *TargetPlatform->IniPlatformName());
-		
-	bool bLegacyBulkDataOffsets = false;
-	PlatformEngineIni.GetBool(TEXT("Core.System"), TEXT("LegacyBulkDataOffsets"), bLegacyBulkDataOffsets);
-	if (bLegacyBulkDataOffsets)
-	{
-		UE_LOG(LogCook, Warning, TEXT("Engine.ini:[Core.System]:LegacyBulkDataOffsets is no longer supported in UE5. The intended use was to reduce patch diffs, but UE5 changed cooked bytes in every package for other reasons, so removing support for this flag does not cause additional patch diffs."));
-	}
-
 	FCookSavePackageContext* Context = new FCookSavePackageContext(TargetPlatform, PackageWriter, WriterDebugName);
 	Context->SaveContext.SetValidator(MakeUnique<FCookedSavePackageValidator>(TargetPlatform, *this));
 	return Context;
@@ -8114,7 +8073,7 @@ void UCookOnTheFlyServer::GenerateInitialRequests(FBeginCookContext& BeginContex
 	CollectFilesToCook(FilesInPath, FilesInPathInstigators, CookMaps, CookDirectories, IniMapSections, CookOptions, TargetPlatforms, GameDefaultObjects);
 
 	// Add soft/hard startup references after collecting requested files and handling empty requests
-	if (bPreexploreDependenciesEnabled && !CookByTheBookOptions->bSkipHardReferences && !CookByTheBookOptions->bFullLoadAndSave)
+	if (!CookByTheBookOptions->bSkipHardReferences && !CookByTheBookOptions->bFullLoadAndSave)
 	{
 		ProcessUnsolicitedPackages(&FilesInPath, &FilesInPathInstigators);
 	}
diff --git a/Engine/Source/Editor/UnrealEd/Private/Cooker/CookPlatformManager.cpp b/Engine/Source/Editor/UnrealEd/Private/Cooker/CookPlatformManager.cpp
--- a/Engine/Source/Editor/UnrealEd/Private/Cooker/CookPlatformManager.cpp
+++ b/Engine/Source/Editor/UnrealEd/Private/Cooker/CookPlatformManager.cpp
@@ -279,13 +279,12 @@ namespace Cook
 
 		if (!HasSessionPlatform(PlatformData->TargetPlatform))
 		{
-			CookOnTheFlyServer.ExternalRequests->AddCallback([this, PlatformName, LocalCookOnTheFlyServer = &CookOnTheFlyServer]()
+			CookOnTheFlyServer.ExternalRequests->AddCallback([PlatformName, LocalCookOnTheFlyServer = &CookOnTheFlyServer]()
 				{
-					const ITargetPlatform* TargetPlatform = GetTargetPlatformManager()->FindTargetPlatform(PlatformName.ToString());
+					ITargetPlatform* TargetPlatform = GetTargetPlatformManager()->FindTargetPlatform(PlatformName.ToString());
 					if (TargetPlatform)
 					{
-						AddSessionPlatform(*LocalCookOnTheFlyServer, TargetPlatform);
-						LocalCookOnTheFlyServer->bPackageFilterDirty = true;
+						LocalCookOnTheFlyServer->StartCookOnTheFlySessionFromGameThread(TargetPlatform);
 					}
 				});
 		}
diff --git a/Engine/Source/Editor/UnrealEd/Private/Cooker/CookRequestCluster.cpp b/Engine/Source/Editor/UnrealEd/Private/Cooker/CookRequestCluster.cpp
--- a/Engine/Source/Editor/UnrealEd/Private/Cooker/CookRequestCluster.cpp
+++ b/Engine/Source/Editor/UnrealEd/Private/Cooker/CookRequestCluster.cpp
@@ -201,7 +201,6 @@
 		bAllowSoftDependencies = false;
 	}
 	bHybridIterativeEnabled = COTFS.bHybridIterativeEnabled;
-	bPreexploreDependenciesEnabled = COTFS.bPreexploreDependenciesEnabled;
 	if (bErrorOnEngineContentUse)
 	{
 		DLCPath = FPaths::Combine(*COTFS.GetBaseDirectoryForDLC(), TEXT("Content"));
@@ -691,7 +690,7 @@
 		// We always skip assetregistry soft dependencies if the cook commandline is set to skip soft references.
 		// We also need to skip them if the project has problems with editor-only robustness and has turned
 		// ExploreDependencies off
-		if (Cluster.bAllowSoftDependencies && Cluster.bPreexploreDependenciesEnabled)
+		if (Cluster.bAllowSoftDependencies)
 		{
 			Cluster.AssetRegistry.GetDependencies(PackageName, SoftDependencies, EDependencyCategory::Package,
 				DependencyQuery | EDependencyQuery::Soft);
@@ -732,7 +731,7 @@
 					UE::SavePackageUtilities::EDLCookInfoAddIterativelySkippedPackage(PackageName);
 				}
 				HardDependencies.Append(PlatformAttachments.BuildDependencies);
-				if (Cluster.bAllowSoftDependencies && Cluster.bPreexploreDependenciesEnabled)
+				if (Cluster.bAllowSoftDependencies)
 				{
 					SoftDependencies.Append(PlatformAttachments.RuntimeOnlyDependencies);
 				}
@@ -902,19 +901,6 @@
 	OutNewVertex->bInitialRequest = bInitialRequest;
 	OutNewVertex->bCookable = bCookable;
 	OutNewVertex->bExploreDependencies = bCookable;
-	if (bCookable && !Cluster.bPreexploreDependenciesEnabled)
-	{
-		// TODO: Editor-only packages may be loaded at editor startup, which makes them skip the request state
-		// and go straight to load and save, where they are culled and returned to idle.
-		// If we add dependencies from these startup packages, in some projects we add transitive dependency
-		// packages that are not otherwise cooked. To prevent breaking projects we are temporarily skipping the
-		// dependency search from these packages.
-		bool bInProgress = false;
-		if (FindPackage(nullptr, WriteToString<256>(PackageName).ToString()))
-		{
-			OutNewVertex->bExploreDependencies = false;
-		}
-	}
 
 	return PackageData;
 }
diff --git a/Engine/Source/Editor/UnrealEd/Private/Cooker/CookRequestCluster.h b/Engine/Source/Editor/UnrealEd/Private/Cooker/CookRequestCluster.h
--- a/Engine/Source/Editor/UnrealEd/Private/Cooker/CookRequestCluster.h
+++ b/Engine/Source/Editor/UnrealEd/Private/Cooker/CookRequestCluster.h
@@ -315,7 +315,6 @@ private:
 	int32 NextRequest = 0;
 	bool bAllowHardDependencies = true;
 	bool bAllowSoftDependencies = true;
-	bool bPreexploreDependenciesEnabled = true;
 	bool bHybridIterativeEnabled = true;
 	bool bErrorOnEngineContentUse = false;
 	bool bPackageNamesComplete = false;
diff --git a/Engine/Source/Editor/UnrealEd/Private/Cooker/PackageTracker.cpp b/Engine/Source/Editor/UnrealEd/Private/Cooker/PackageTracker.cpp
--- a/Engine/Source/Editor/UnrealEd/Private/Cooker/PackageTracker.cpp
+++ b/Engine/Source/Editor/UnrealEd/Private/Cooker/PackageTracker.cpp
@@ -131,7 +131,6 @@ namespace UE::Cook
 
 			LoadedPackages.Remove(Package);
 			NewPackages.Remove(Package);
-			PostLoadFixupPackages.Remove(Package);
 		}
 	}
 
diff --git a/Engine/Source/Editor/UnrealEd/Private/Cooker/PackageTracker.h b/Engine/Source/Editor/UnrealEd/Private/Cooker/PackageTracker.h
--- a/Engine/Source/Editor/UnrealEd/Private/Cooker/PackageTracker.h
+++ b/Engine/Source/Editor/UnrealEd/Private/Cooker/PackageTracker.h
@@ -187,9 +187,6 @@ namespace Cook
 		/** Swap all ITargetPlatform* stored on this instance according to the mapping in @param Remap. */
 		void RemapTargetPlatforms(const TMap<ITargetPlatform*, ITargetPlatform*>& Remap);
 
-		// This is the set of packages which have already had PostLoadFixup called 
-		TSet<UPackage*> PostLoadFixupPackages;
-
 		// This is a complete list of currently loaded UPackages
 		TFastPointerSet<UPackage*> LoadedPackages;
 
diff --git a/Engine/Source/Editor/UnrealEd/Public/Commandlets/AssetRegistryGenerator.h b/Engine/Source/Editor/UnrealEd/Public/Commandlets/AssetRegistryGenerator.h
--- a/Engine/Source/Editor/UnrealEd/Public/Commandlets/AssetRegistryGenerator.h
+++ b/Engine/Source/Editor/UnrealEd/Public/Commandlets/AssetRegistryGenerator.h
@@ -198,6 +198,12 @@ public:
 	 */
 	void UpdateAssetRegistryPackageData(const UPackage& Package, FSavePackageResultStruct& SavePackageResult, FCookTagList&& InArchiveCookTagList);
 
+	/**
+	 * Check config to see whether chunk assignments use the AssetManager. If so, run the once-per-process construction
+	 * of ManageReferences and store them in the global AssetRegistry.
+	 */
+	static void UpdateAssetManagerDatabase();
+
 private:
 	/**
 	 * Ensures all assets in the input package are present in the registry
@@ -222,6 +228,8 @@ private:
 	 */
 	void UpdateKeptPackages();
 
+	static void InitializeUseAssetManager();
+
 	/** State of the asset registry that is being built for this platform */
 	FAssetRegistryState State;
 	
@@ -254,8 +262,6 @@ private:
 	TSet<FName> PackagesContainingMaps;
 	/** Should the chunks be generated or only asset registry */
 	bool bGenerateChunks;
-	/** True if we should use the AssetManager, false to use the deprecated path */
-	bool bUseAssetManager;
 	/** Highest chunk id, being used for geneating dependency tree */
 	int32 HighestChunkId;
 	/** Array of Maps with chunks<->packages assignments */
@@ -279,6 +285,9 @@ private:
 	/** Mapping from chunk id to pakchunk file index. If not defined, Pakchunk index will be the same as chunk id by default */
 	TMap<int32, int32> ChunkIdPakchunkIndexMapping;
 
+	/** True if we should use the AssetManager, false to use the deprecated path */
+	static bool bUseAssetManager;
+
 	struct FReferencePair
 	{
 		FReferencePair() {}
