Horde: Use a parsable format for localization errors, allowing us to link files for build health issues.

diff --git a/Engine/Source/Developer/Localization/Private/LocTextHelper.cpp b/Engine/Source/Developer/Localization/Private/LocTextHelper.cpp
--- a/Engine/Source/Developer/Localization/Private/LocTextHelper.cpp
+++ b/Engine/Source/Developer/Localization/Private/LocTextHelper.cpp
@@ -955,18 +955,18 @@ bool FLocTextHelper::AddSourceText(const FLocKey& InNamespace, const FLocItem& I
 			const FManifestContext* ConflictingContext = ExistingEntry->FindContext(InContext.Key, InContext.KeyMetadataObj);
 			const FString& ExistingEntrySourceLocation = (!ExistingEntryFileName.IsEmpty()) ? ExistingEntryFileName : ConflictingContext->SourceLocation;
 
-			UE_LOG(LogLocTextHelper, Warning, TEXT("Text conflict%s for namespace \"%s\" and key \"%s\"%s at \"%s\" and \"%s\". The conflicting sources are \"%s\"%s and \"%s\"%s."), 
+			UE_LOG(LogLocTextHelper, Warning, TEXT("%s: Text conflict%s for namespace \"%s\" and key \"%s\"%s. The conflicting sources are \"%s\"%s and \"%s\"%s."), 
+				*InContext.SourceLocation,
 				(InDescription ? *FString::Printf(TEXT(" from %s"), **InDescription) : TEXT("")),
 				*SanitizeLogOutput(InNamespace.GetString()),
 				*SanitizeLogOutput(InContext.Key.GetString()),
 				(InContext.KeyMetadataObj ? *FString::Printf(TEXT(" (meta-data \"%s\")"), *SanitizeLogOutput(FJsonInternationalizationMetaDataSerializer::MetadataToString(InContext.KeyMetadataObj))) : TEXT("")),
-				*InContext.SourceLocation,
-				*ExistingEntrySourceLocation,
 				*SanitizeLogOutput(InSource.Text),
 				(InSource.MetadataObj ? *FString::Printf(TEXT(" (meta-data \"%s\")"), *SanitizeLogOutput(FJsonInternationalizationMetaDataSerializer::MetadataToString(InSource.MetadataObj))) : TEXT("")),
 				*SanitizeLogOutput(ExistingEntry->Source.Text),
 				(ExistingEntry->Source.MetadataObj ? *FString::Printf(TEXT(" (meta-data \"%s\")"), *SanitizeLogOutput(FJsonInternationalizationMetaDataSerializer::MetadataToString(ExistingEntry->Source.MetadataObj))) : TEXT(""))
 				);
+			UE_LOG(LogLocTextHelper, Warning, TEXT("%s: See conflicting location."), *ExistingEntrySourceLocation);
 
 			ConflictTracker.AddConflict(InNamespace, InContext.Key, InContext.KeyMetadataObj, InSource, InContext.SourceLocation);
 			ConflictTracker.AddConflict(InNamespace, InContext.Key, InContext.KeyMetadataObj, ExistingEntry->Source, ExistingEntrySourceLocation);
@@ -977,12 +977,12 @@ bool FLocTextHelper::AddSourceText(const FLocKey& InNamespace, const FLocItem& I
 		bAddSuccessful = Manifest->AddSource(InNamespace, InSource, InContext);
 		if (!bAddSuccessful)
 		{
-			UE_LOG(LogLocTextHelper, Error, TEXT("Failed to add text%s for namespace \"%s\" and key \"%s\"%s from \"%s\" with source \"%s\"%s."),
+			UE_LOG(LogLocTextHelper, Error, TEXT("%s: Failed to add text%s for namespace \"%s\" and key \"%s\"%s with source \"%s\"%s."),
+				*InContext.SourceLocation,
 				(InDescription ? *FString::Printf(TEXT(" from %s"), **InDescription) : TEXT("")),
 				*InNamespace.GetString(),
 				*InContext.Key.GetString(),
 				(InContext.KeyMetadataObj ? *FString::Printf(TEXT(" (meta-data \"%s\")"), *FJsonInternationalizationMetaDataSerializer::MetadataToString(InContext.KeyMetadataObj)) : TEXT("")),
-				*InContext.SourceLocation,
 				*InSource.Text,
 				(InSource.MetadataObj ? *FString::Printf(TEXT(" (meta-data \"%s\")"), *FJsonInternationalizationMetaDataSerializer::MetadataToString(InSource.MetadataObj)) : TEXT(""))
 				);
diff --git a/Engine/Source/Editor/UnrealEd/Classes/Commandlets/GatherTextFromSourceCommandlet.h b/Engine/Source/Editor/UnrealEd/Classes/Commandlets/GatherTextFromSourceCommandlet.h
--- a/Engine/Source/Editor/UnrealEd/Classes/Commandlets/GatherTextFromSourceCommandlet.h
+++ b/Engine/Source/Editor/UnrealEd/Classes/Commandlets/GatherTextFromSourceCommandlet.h
@@ -52,7 +52,7 @@ private:
 
 		FString ToString() const
 		{
-			return (Line == INDEX_NONE) ? File : FString::Printf(TEXT("%s:%d"), *File, Line);
+			return (Line == INDEX_NONE) ? File : FString::Printf(TEXT("%s(%d)"), *File, Line);
 		}
 
 		FString File;
diff --git a/Engine/Source/Editor/UnrealEd/Private/Commandlets/GatherTextFromSourceCommandlet.cpp b/Engine/Source/Editor/UnrealEd/Private/Commandlets/GatherTextFromSourceCommandlet.cpp
--- a/Engine/Source/Editor/UnrealEd/Private/Commandlets/GatherTextFromSourceCommandlet.cpp
+++ b/Engine/Source/Editor/UnrealEd/Private/Commandlets/GatherTextFromSourceCommandlet.cpp
@@ -309,19 +309,19 @@ int32 UGatherTextFromSourceCommandlet::Main( const FString& Params )
 		FString SourceFileText;
 		if (!FFileHelper::LoadFileToString(SourceFileText, *SourceFile))
 		{
-			UE_LOG(LogGatherTextFromSourceCommandlet, Error, TEXT("GatherTextSource failed to open file %s"), *ParseCtxt.Filename);
+			UE_LOG(LogGatherTextFromSourceCommandlet, Error, TEXT("%s: GatherTextSource failed to open file"), *ParseCtxt.Filename);
 		}
 		else
 		{
 			if (!ParseSourceText(SourceFileText, Parsables, ParseCtxt))
 			{
-				UE_LOG(LogGatherTextFromSourceCommandlet, Warning, TEXT("GatherTextSource error(s) parsing source file %s"), *ParseCtxt.Filename);
+				UE_LOG(LogGatherTextFromSourceCommandlet, Warning, TEXT("%s: GatherTextSource error(s) parsing source file"), *ParseCtxt.Filename);
 			}
 			else
 			{
 				if (ParseCtxt.WithinNamespaceDefineLineNumber != INDEX_NONE)
 				{
-					UE_LOG(LogGatherTextFromSourceCommandlet, Warning, TEXT("Missing '#undef LOCTEXT_NAMESPACE' for '#define LOCTEXT_NAMESPACE' at %s:%d"), *ParseCtxt.Filename, ParseCtxt.WithinNamespaceDefineLineNumber);
+					UE_LOG(LogGatherTextFromSourceCommandlet, Warning, TEXT("%s(%d): Missing '#undef LOCTEXT_NAMESPACE' for '#define LOCTEXT_NAMESPACE'"), *ParseCtxt.Filename, ParseCtxt.WithinNamespaceDefineLineNumber);
 				}
 			}
 		}
@@ -1060,7 +1060,7 @@ bool UGatherTextFromSourceCommandlet::ParseSourceText(const FString& Text, const
 	// Handle a string C++11 raw string literal that was never closed as this is likely a false positive that needs to be fixed in the parser
 	if (ParseCtxt.WithinStringLiteral && !ParseCtxt.RawStringLiteralClosingDelim.IsEmpty())
 	{
-		UE_LOG(LogGatherTextFromSourceCommandlet, Warning, TEXT("A C++11 raw string literal was not correctly terminated. File %s, starting line: %s"), *ParseCtxt.Filename, ParseCtxt.WithinStartingLine);
+		UE_LOG(LogGatherTextFromSourceCommandlet, Warning, TEXT("%s(%s): A C++11 raw string literal was not correctly terminated."), *ParseCtxt.Filename, ParseCtxt.WithinStartingLine);
 	}
 
 	return true;
@@ -1193,7 +1193,7 @@ void UGatherTextFromSourceCommandlet::FSourceFileParseContext::SetDefine(const F
 		// #define LOC_DEFINE_REGION
 		if (ExcludedRegion)
 		{
-			UE_LOG(LogGatherTextFromSourceCommandlet, Warning, TEXT("Found a '#define LOC_DEFINE_REGION' within another '#define LOC_DEFINE_REGION' while parsing %s:%d"), *Filename, LineNumber);
+			UE_LOG(LogGatherTextFromSourceCommandlet, Warning, TEXT("%s(%d): Found a '#define LOC_DEFINE_REGION' within another '#define LOC_DEFINE_REGION'"), *Filename, LineNumber);
 		}
 		else
 		{
@@ -1208,7 +1208,7 @@ void UGatherTextFromSourceCommandlet::FSourceFileParseContext::SetDefine(const F
 			// #define LOCTEXT_NAMESPACE <namespace>
 			if (WithinNamespaceDefineLineNumber != INDEX_NONE)
 			{
-				UE_LOG(LogGatherTextFromSourceCommandlet, Warning, TEXT("Found a '#define LOCTEXT_NAMESPACE' within another '#define LOCTEXT_NAMESPACE' while parsing %s:%d"), *Filename, LineNumber);
+				UE_LOG(LogGatherTextFromSourceCommandlet, Warning, TEXT("%s(%d): Found a '#define LOCTEXT_NAMESPACE' within another '#define LOCTEXT_NAMESPACE'"), *Filename, LineNumber);
 			}
 			else
 			{
@@ -1239,7 +1239,7 @@ void UGatherTextFromSourceCommandlet::FSourceFileParseContext::RemoveDefine(cons
 		// #undef LOC_DEFINE_REGION
 		if (!ExcludedRegion)
 		{
-			UE_LOG(LogGatherTextFromSourceCommandlet, Warning, TEXT("Found an '#undef LOC_DEFINE_REGION' without a corresponding '#define LOC_DEFINE_REGION' while parsing %s:%d"), *Filename, LineNumber);
+			UE_LOG(LogGatherTextFromSourceCommandlet, Warning, TEXT("%s(%d): Found an '#undef LOC_DEFINE_REGION' without a corresponding '#define LOC_DEFINE_REGION'"), *Filename, LineNumber);
 		}
 		else
 		{
@@ -1254,7 +1254,7 @@ void UGatherTextFromSourceCommandlet::FSourceFileParseContext::RemoveDefine(cons
 			// #undef LOCTEXT_NAMESPACE
 			if (WithinNamespaceDefineLineNumber == INDEX_NONE)
 			{
-				UE_LOG(LogGatherTextFromSourceCommandlet, Warning, TEXT("Found an '#undef LOCTEXT_NAMESPACE' without a corresponding '#define LOCTEXT_NAMESPACE' while parsing %s:%d"), *Filename, LineNumber);
+				UE_LOG(LogGatherTextFromSourceCommandlet, Warning, TEXT("%s(%d): Found an '#undef LOCTEXT_NAMESPACE' without a corresponding '#define LOCTEXT_NAMESPACE'"), *Filename, LineNumber);
 			}
 			else
 			{
@@ -1272,7 +1272,7 @@ bool UGatherTextFromSourceCommandlet::FSourceFileParseContext::AddStringTableImp
 	FParsedStringTable& ParsedStringTable = ParsedStringTables.FindOrAdd(InTableId);
 	if (ParsedStringTable.SourceLocation.Line != INDEX_NONE)
 	{
-		UE_LOG(LogGatherTextFromSourceCommandlet, Warning, TEXT("String table with ID \"%s\" at \"%s\" was already parsed at \"%s\". Ignoring additional definition."), *InTableId.ToString(), *FSourceLocation(Filename, LineNumber).ToString(), *ParsedStringTable.SourceLocation.ToString());
+		UE_LOG(LogGatherTextFromSourceCommandlet, Warning, TEXT("%s: String table with ID \"%s\" was already parsed at \"%s\". Ignoring additional definition."), *FSourceLocation(Filename, LineNumber).ToString(), *InTableId.ToString(), *ParsedStringTable.SourceLocation.ToString());
 		return false;
 	}
 
@@ -1298,7 +1298,7 @@ bool UGatherTextFromSourceCommandlet::FSourceFileParseContext::AddStringTableEnt
 		}
 		else
 		{
-			UE_LOG(LogGatherTextFromSourceCommandlet, Warning, TEXT("String table entry with ID \"%s\" and key \"%s\" at \"%s\" was already parsed at \"%s\". Ignoring additional definition."), *InTableId.ToString(), *InKey, *FSourceLocation(Filename, LineNumber).ToString(), *ExistingEntry->SourceLocation.ToString());
+			UE_LOG(LogGatherTextFromSourceCommandlet, Warning, TEXT("%s: String table entry with ID \"%s\" and key \"%s\" was already parsed at \"%s\". Ignoring additional definition."), *FSourceLocation(Filename, LineNumber).ToString(), *InTableId.ToString(), *InKey, *ExistingEntry->SourceLocation.ToString());
 			return false;
 		}
 	}
@@ -1331,7 +1331,7 @@ bool UGatherTextFromSourceCommandlet::FSourceFileParseContext::AddStringTableEnt
 		}
 		else
 		{
-			UE_LOG(LogGatherTextFromSourceCommandlet, Warning, TEXT("String table entry meta-data with ID \"%s\" and key \"%s\" at \"%s\" was already parsed at \"%s\". Ignoring additional definition."), *InTableId.ToString(), *InKey, *FSourceLocation(Filename, LineNumber).ToString(), *ExistingMetaData->SourceLocation.ToString());
+			UE_LOG(LogGatherTextFromSourceCommandlet, Warning, TEXT("%s: String table entry meta-data with ID \"%s\" and key \"%s\" was already parsed at \"%s\". Ignoring additional definition."), *FSourceLocation(Filename, LineNumber).ToString(), *InTableId.ToString(), *InKey, *ExistingMetaData->SourceLocation.ToString());
 			return false;
 		}
 	}
@@ -1377,7 +1377,7 @@ void UGatherTextFromSourceCommandlet::FSourceFileParseContext::AddStringTableFro
 		}
 		else
 		{
-			UE_LOG(LogGatherTextFromSourceCommandlet, Warning, TEXT("String table with ID \"%s\" at \"%s\" failed to import strings from \"%s\"."), *InTableId.ToString(), *FSourceLocation(Filename, LineNumber).ToString(), *FullImportPath);
+			UE_LOG(LogGatherTextFromSourceCommandlet, Warning, TEXT("%s: String table with ID \"%s\" failed to import strings from \"%s\"."), *FSourceLocation(Filename, LineNumber).ToString(), *InTableId.ToString(), *FullImportPath);
 		}
 	}
 }
@@ -1525,7 +1525,7 @@ bool UGatherTextFromSourceCommandlet::FMacroDescriptor::ParseArgumentString(cons
 				// Here, we found the end of the line, but we don't know if it's the end of the argument.
 				if (!ArgsGatherer.Gather(ArgStart, Cursor - ArgStart))
 				{
-					UE_LOG(LogGatherTextFromSourceCommandlet, Warning, TEXT("Parsing Arguments failed in %s macro while parsing %s:%d. %s"), *GetToken(), *Context.Filename, Context.LineNumber, *FLocTextHelper::SanitizeLogOutput(Context.LineText.TrimStartAndEnd()));
+					UE_LOG(LogGatherTextFromSourceCommandlet, Warning, TEXT("%s(%d): Parsing Arguments failed in %s macro. %s"), *Context.Filename, Context.LineNumber, *GetToken(), *FLocTextHelper::SanitizeLogOutput(Context.LineText.TrimStartAndEnd()));
 					return false;
 				}
 			}
@@ -1547,7 +1547,7 @@ bool UGatherTextFromSourceCommandlet::FMacroDescriptor::ParseArgumentString(cons
 				{
 					if (!ArgsGatherer.Gather(ArgStart, Cursor - ArgStart - 1))
 					{
-						UE_LOG(LogGatherTextFromSourceCommandlet, Warning, TEXT("Parsing Arguments failed in %s macro while parsing %s:%d. %s"), *GetToken(), *Context.Filename, Context.LineNumber, *FLocTextHelper::SanitizeLogOutput(Context.LineText.TrimStartAndEnd()));
+						UE_LOG(LogGatherTextFromSourceCommandlet, Warning, TEXT("%s(%d): Parsing Arguments failed in %s macro. %s"), *Context.Filename, Context.LineNumber, *GetToken(), *FLocTextHelper::SanitizeLogOutput(Context.LineText.TrimStartAndEnd()));
 						return false;
 					}
 				}
@@ -1592,7 +1592,7 @@ bool UGatherTextFromSourceCommandlet::FMacroDescriptor::ParseArgumentString(cons
 
 			if (0 > BracketStack)
 			{
-				UE_LOG(LogGatherTextFromSourceCommandlet, Warning, TEXT("Unexpected bracket ')' in %s macro while parsing %s:%d. %s"), *GetToken(), *Context.Filename, Context.LineNumber, *FLocTextHelper::SanitizeLogOutput(Context.LineText.TrimStartAndEnd()));
+				UE_LOG(LogGatherTextFromSourceCommandlet, Warning, TEXT("%s(%d): Unexpected bracket ')' in %s macro while parsing. %s"), *Context.Filename, Context.LineNumber, *GetToken(), *FLocTextHelper::SanitizeLogOutput(Context.LineText.TrimStartAndEnd()));
 				return false;
 			}
 		}		
@@ -1605,12 +1605,12 @@ bool UGatherTextFromSourceCommandlet::FMacroDescriptor::ParseArgumentString(cons
 					// We have a single bracket open and we found a ',' this is the end of the argument. If Bracket stack is > 1 it means that we are in a function call in one of the parameters.
 					if (!ArgsGatherer.Gather(ArgStart, Cursor - ArgStart))
 					{
-						UE_LOG(LogGatherTextFromSourceCommandlet, Warning, TEXT("Parsing Arguments failed in %s macro while parsing %s:%d. %s"), *GetToken(), *Context.Filename, Context.LineNumber, *FLocTextHelper::SanitizeLogOutput(Context.LineText.TrimStartAndEnd()));
+						UE_LOG(LogGatherTextFromSourceCommandlet, Warning, TEXT("%s(%d): Parsing Arguments failed in %s macro. %s"), *Context.Filename, Context.LineNumber, *GetToken(), *FLocTextHelper::SanitizeLogOutput(Context.LineText.TrimStartAndEnd()));
 						return false;
 					}
 					if (!ArgsGatherer.EndArgument())
 					{
-						UE_LOG(LogGatherTextFromSourceCommandlet, Warning, TEXT("Parsing Arguments failed in %s macro while parsing %s:%d. %s"), *GetToken(), *Context.Filename, Context.LineNumber, *FLocTextHelper::SanitizeLogOutput(Context.LineText.TrimStartAndEnd()));
+						UE_LOG(LogGatherTextFromSourceCommandlet, Warning, TEXT("%s(%d): Parsing Arguments failed in %s macro. %s"), *Context.Filename, Context.LineNumber, *GetToken(), *FLocTextHelper::SanitizeLogOutput(Context.LineText.TrimStartAndEnd()));
 						return false;
 					}
 
@@ -1623,7 +1623,7 @@ bool UGatherTextFromSourceCommandlet::FMacroDescriptor::ParseArgumentString(cons
 					//	    , "Param3")
 					if (!ArgsGatherer.EndArgument())
 					{
-						UE_LOG(LogGatherTextFromSourceCommandlet, Warning, TEXT("Parsing Arguments failed in %s macro while parsing %s:%d. %s"), *GetToken(), *Context.Filename, Context.LineNumber, *FLocTextHelper::SanitizeLogOutput(Context.LineText.TrimStartAndEnd()));
+						UE_LOG(LogGatherTextFromSourceCommandlet, Warning, TEXT("%s(%d): Parsing Arguments failed in %s macro. %s"), *Context.Filename, Context.LineNumber, *GetToken(), *FLocTextHelper::SanitizeLogOutput(Context.LineText.TrimStartAndEnd()));
 						return false;
 					}
 					ArgStart = Cursor + 1;
@@ -1640,7 +1640,7 @@ bool UGatherTextFromSourceCommandlet::FMacroDescriptor::ParseArgumentString(cons
 				{
 					if (!ArgsGatherer.Gather(ArgStart, Cursor - ArgStart - 1))
 					{
-						UE_LOG(LogGatherTextFromSourceCommandlet, Warning, TEXT("Parsing Arguments failed in %s macro while parsing %s:%d. %s"), *GetToken(), *Context.Filename, Context.LineNumber, *FLocTextHelper::SanitizeLogOutput(Context.LineText.TrimStartAndEnd()));
+						UE_LOG(LogGatherTextFromSourceCommandlet, Warning, TEXT("%s(%d): Parsing Arguments failed in %s macro. %s"), *Context.Filename, Context.LineNumber, *GetToken(), *FLocTextHelper::SanitizeLogOutput(Context.LineText.TrimStartAndEnd()));
 						return false;
 					}
 				}
@@ -1656,13 +1656,13 @@ bool UGatherTextFromSourceCommandlet::FMacroDescriptor::ParseArgumentString(cons
 			{
 				if (!ArgsGatherer.Gather(ArgStart, Cursor - ArgStart))
 				{
-					UE_LOG(LogGatherTextFromSourceCommandlet, Warning, TEXT("Parsing Arguments failed in %s macro while parsing %s:%d. %s"), *GetToken(), *Context.Filename, Context.LineNumber, *FLocTextHelper::SanitizeLogOutput(Context.LineText.TrimStartAndEnd()));
+					UE_LOG(LogGatherTextFromSourceCommandlet, Warning, TEXT("%s(%d): Parsing Arguments failed in %s macro. %s"), *Context.Filename, Context.LineNumber, *GetToken(), *FLocTextHelper::SanitizeLogOutput(Context.LineText.TrimStartAndEnd()));
 					return false;
 				}
 			}
 			if (!ArgsGatherer.EndArgument())
 			{
-				UE_LOG(LogGatherTextFromSourceCommandlet, Warning, TEXT("Parsing Arguments failed in %s macro while parsing %s:%d. %s"), *GetToken(), *Context.Filename, Context.LineNumber, *FLocTextHelper::SanitizeLogOutput(Context.LineText.TrimStartAndEnd()));
+				UE_LOG(LogGatherTextFromSourceCommandlet, Warning, TEXT("%s(%d): Parsing Arguments failed in %s macro. %s"), *Context.Filename, Context.LineNumber, *GetToken(), *FLocTextHelper::SanitizeLogOutput(Context.LineText.TrimStartAndEnd()));
 				return false;
 			}
 			break;
@@ -1728,7 +1728,7 @@ bool UGatherTextFromSourceCommandlet::FMacroDescriptor::ParseArgsFromNextLines(F
 	{
 		if(i >= Context.TextLines.Num())
 		{
-			UE_LOG(LogGatherTextFromSourceCommandlet, Warning, TEXT("We reached end of file while parsing specified %s macro for arguments. %s:%d. %s"), *GetToken(), *Context.Filename, Context.LineNumber, *FLocTextHelper::SanitizeLogOutput(Context.LineText.TrimStartAndEnd()));
+			UE_LOG(LogGatherTextFromSourceCommandlet, Warning, TEXT("%s(%d): We reached end of file while parsing specified %s macro for arguments. %s"), *Context.Filename, Context.LineNumber, *GetToken(), *FLocTextHelper::SanitizeLogOutput(Context.LineText.TrimStartAndEnd()));
 			return false;
 		}
 
@@ -1748,7 +1748,7 @@ bool UGatherTextFromSourceCommandlet::FMacroDescriptor::ParseArgsFromNextLines(F
 			// We did not find the opening bracket on the first line and on the second non empty line, we give up
 			if (OpenBracketIdx == INDEX_NONE)
 			{
-				UE_LOG(LogGatherTextFromSourceCommandlet, Warning, TEXT("Opening bracket '(' not found while parsing specified %s macro for arguments. %s:%d. %s"), *GetToken(), *Context.Filename, Context.LineNumber, *FLocTextHelper::SanitizeLogOutput(Context.LineText.TrimStartAndEnd()));
+				UE_LOG(LogGatherTextFromSourceCommandlet, Warning, TEXT("%s(%d): Opening bracket '(' not found while parsing specified %s macro for arguments. %s"), *Context.Filename, Context.LineNumber, *GetToken(), *FLocTextHelper::SanitizeLogOutput(Context.LineText.TrimStartAndEnd()));
 				return false;
 			}
 			BracketStack = 1;
@@ -1794,7 +1794,7 @@ void UGatherTextFromSourceCommandlet::FUICommandMacroDescriptor::TryParseArgs(co
 	if (Identifier.IsEmpty())
 	{
 		//The command doesn't have an identifier so we can't gather it
-		UE_LOG(LogGatherTextFromSourceCommandlet, Warning, TEXT("%s macro at %s has an empty identifier and cannot be gathered."), *GetToken(), *SourceLocation);
+		UE_LOG(LogGatherTextFromSourceCommandlet, Warning, TEXT("%s: %s macro has an empty identifier and cannot be gathered."), *SourceLocation , *GetToken());
 		return;
 	}
 
@@ -1805,7 +1805,7 @@ void UGatherTextFromSourceCommandlet::FUICommandMacroDescriptor::TryParseArgs(co
 	FString Namespace = Context.WithinNamespaceDefineLineNumber != INDEX_NONE && !Context.Namespace.IsEmpty() ? FString::Printf(TEXT("%s.%s"), *UICommandRootNamespace, *Context.Namespace) : UICommandRootNamespace;
 
 	// parse DefaultLangString argument - this arg will be in quotes without TEXT macro
-	FString MacroDesc = FString::Printf(TEXT("\"FriendlyName\" argument in %s macro at %s:%d."), *GetToken(), *Context.Filename, Context.LineNumber);
+	FString MacroDesc = FString::Printf(TEXT("%s(%d): \"FriendlyName\" argument in %s macro."), *Context.Filename, Context.LineNumber, *GetToken());
 	if (PrepareArgument(SourceText, true, MacroDesc, HasQuotes))
 	{
 		if (HasQuotes && !Identifier.IsEmpty() && !SourceText.IsEmpty())
@@ -1821,7 +1821,7 @@ void UGatherTextFromSourceCommandlet::FUICommandMacroDescriptor::TryParseArgs(co
 			// parse DefaultLangTooltipString argument - this arg will be in quotes without TEXT macro
 			FString TooltipSourceText = Arguments[ArgIndexOffset + 2];
 			TooltipSourceText.TrimStartInline();
-			MacroDesc = FString::Printf(TEXT("\"InDescription\" argument in %s macro at %s:%d"), *GetToken(), *Context.Filename, Context.LineNumber);
+			MacroDesc = FString::Printf(TEXT("%s(%d): \"InDescription\" argument in %s macro"), *Context.Filename, Context.LineNumber, *GetToken());
 			if (PrepareArgument(TooltipSourceText, true, MacroDesc, HasQuotes))
 			{
 				if (HasQuotes && !TooltipSourceText.IsEmpty())
@@ -1852,7 +1852,7 @@ void UGatherTextFromSourceCommandlet::FUICommandMacroDescriptor::TryParse(const
 			// Validate that we got the rightnumber of Arguments
 			if (Arguments.Num() < GetMinNumberOfArgument())
 			{
-				UE_LOG(LogGatherTextFromSourceCommandlet, Warning, TEXT("Expected at least %d arguments for %s macro, but got %d while parsing %s:%d. %s"), GetMinNumberOfArgument(), *GetToken(), Arguments.Num(), *Context.Filename, Context.LineNumber, *FLocTextHelper::SanitizeLogOutput(Context.LineText.TrimStartAndEnd()));
+				UE_LOG(LogGatherTextFromSourceCommandlet, Warning, TEXT("%s(%d): Expected at least %d arguments for %s macro, but got %d. %s"), *Context.Filename, Context.LineNumber, GetMinNumberOfArgument(), *GetToken(), Arguments.Num(), *FLocTextHelper::SanitizeLogOutput(Context.LineText.TrimStartAndEnd()));
 				return;
 			}
 			
@@ -1875,7 +1875,7 @@ void UGatherTextFromSourceCommandlet::FUICommandExtMacroDescriptor::TryParse(con
 			// Need at least 5 arguments
 			if (Arguments.Num() < GetMinNumberOfArgument())
 			{
-				UE_LOG(LogGatherTextFromSourceCommandlet, Warning, TEXT("Expected at least %d arguments for %s macro, but got %d while parsing %s:%d. %s"), GetMinNumberOfArgument(), *GetToken(), Arguments.Num(), *Context.Filename, Context.LineNumber, *FLocTextHelper::SanitizeLogOutput(Context.LineText.TrimStartAndEnd()));
+				UE_LOG(LogGatherTextFromSourceCommandlet, Warning, TEXT("%s(%d): Expected at least %d arguments for %s macro, but got %d. %s"), *Context.Filename, Context.LineNumber, GetMinNumberOfArgument(), *GetToken(), Arguments.Num(), *FLocTextHelper::SanitizeLogOutput(Context.LineText.TrimStartAndEnd()));
 				return;
 			}
 			TryParseArgs(Text, Context, Arguments, 2);
@@ -1891,19 +1891,19 @@ void UGatherTextFromSourceCommandlet::FMetasoundParamMacroDescriptor::TryParseAr
 
 	// NAME should neverbe be in quotes 
 	bool HasQuotes = false;
-	FString MacroDesc = FString::Printf(TEXT("\"InDescription\" argument in %s macro at %s:%d"), *GetToken(), *Context.Filename, Context.LineNumber);
+	FString MacroDesc = FString::Printf(TEXT("%s(%d): \"InDescription\" argument in %s macro"), *Context.Filename, Context.LineNumber, *GetToken());
 	FString SourceLocation = FSourceLocation(Context.Filename, Context.LineNumber).ToString();
 	if (PrepareArgument(NameArgument, true, MacroDesc, HasQuotes))
 	{
 		if (HasQuotes)
 		{
-			UE_LOG(LogGatherTextFromSourceCommandlet, Warning, TEXT("%s macro at \'%s\' has NAME argument enclosed in double quotes and cannot be gathered. Please make sure the argument is not enclosed in double quotes."), *GetToken(), *SourceLocation);
+			UE_LOG(LogGatherTextFromSourceCommandlet, Warning, TEXT("%s: %s macro has NAME argument enclosed in double quotes and cannot be gathered. Please make sure the argument is not enclosed in double quotes."), *SourceLocation, *GetToken());
 			return;
 		}
 		else if (NameArgument.IsEmpty())
 		{
 			//The metasound param does not have a NAME so we cannot gather it 
-			UE_LOG(LogGatherTextFromSourceCommandlet, Warning, TEXT("%s macro at \'%s\' has an empty NAME argument and cannot be gathered."), *GetToken(), *SourceLocation);
+			UE_LOG(LogGatherTextFromSourceCommandlet, Warning, TEXT("%s: %s macro has an empty NAME argument and cannot be gathered."), *SourceLocation, *GetToken());
 			return;
 		}
 	}
@@ -1918,13 +1918,13 @@ void UGatherTextFromSourceCommandlet::FMetasoundParamMacroDescriptor::TryParseAr
 	{
 		if (!HasQuotes)
 		{
-			UE_LOG(LogGatherTextFromSourceCommandlet, Warning, TEXT("%s macro at \'%s\' has NAME_TEXT argument not enclosed in double quotes and cannot be gathered. Please make sure the argument is enclosed in double quotes."), *GetToken(), *SourceLocation);
+			UE_LOG(LogGatherTextFromSourceCommandlet, Warning, TEXT("%s: %s macro has NAME_TEXT argument not enclosed in double quotes and cannot be gathered. Please make sure the argument is enclosed in double quotes."), *SourceLocation, *GetToken());
 			return;
 		}
 		else if (NameText.IsEmpty())
 		{
 			//The metasound param does not have a NAME_TEXT so we cannot gather it 
-			UE_LOG(LogGatherTextFromSourceCommandlet, Warning, TEXT("%s macro at \'%s\' has an empty NAME_TEXT argumentand cannot be gathered."), *GetToken(), *SourceLocation);
+			UE_LOG(LogGatherTextFromSourceCommandlet, Warning, TEXT("%s: %s macro has an empty NAME_TEXT argumentand cannot be gathered."), *SourceLocation, *GetToken());
 			return;
 		}
 	}
@@ -1932,17 +1932,17 @@ void UGatherTextFromSourceCommandlet::FMetasoundParamMacroDescriptor::TryParseAr
 	// parse TOOLTIP_TEXT argument- this arg will be in quotes without TEXT macro
 	FString TooltipSourceText = Arguments[ArgIndexOffset + 2];
 	TooltipSourceText.TrimStartInline();
-	MacroDesc = FString::Printf(TEXT("\"InDescription\" argument in %s macro at %s:%d"), *GetToken(), *Context.Filename, Context.LineNumber);
+	MacroDesc = FString::Printf(TEXT("%s(%d): \"InDescription\" argument in %s macro"), *Context.Filename, Context.LineNumber, *GetToken());
 	if (PrepareArgument(TooltipSourceText, true, MacroDesc, HasQuotes))
 	{
 		if (!HasQuotes)
 		{
-			UE_LOG(LogGatherTextFromSourceCommandlet, Warning, TEXT("%s macro at \'%s\' has a TOOLTIP_TEXT argument not enclosed in double quotes and cannot be gathered. Please ensure the argument is enclosed in double quotes."), *GetToken(), *SourceLocation);
+			UE_LOG(LogGatherTextFromSourceCommandlet, Warning, TEXT("%s: %s macro has a TOOLTIP_TEXT argument not enclosed in double quotes and cannot be gathered. Please ensure the argument is enclosed in double quotes."), *SourceLocation, *GetToken());
 			return;
 		}
 		else if (TooltipSourceText.IsEmpty())
 		{
-			UE_LOG(LogGatherTextFromSourceCommandlet, Warning, TEXT("%s macro at \'%s\' has an empty TOOLTIP_TEXT and cannot be gathered."), *GetToken(), *SourceLocation);
+			UE_LOG(LogGatherTextFromSourceCommandlet, Warning, TEXT("%s: %s macro has an empty TOOLTIP_TEXT and cannot be gathered."), *SourceLocation, *GetToken());
 			return;
 		}
 		else
@@ -1983,7 +1983,7 @@ void UGatherTextFromSourceCommandlet::FMetasoundParamMacroDescriptor::TryParse(c
 			// Validate that we got the rightnumber of Arguments
 			if (Arguments.Num() < GetMinNumberOfArgument())
 			{
-				UE_LOG(LogGatherTextFromSourceCommandlet, Warning, TEXT("Expected at least %d arguments for %s macro, but got %d while parsing %s:%d. %s"), GetMinNumberOfArgument(), *GetToken(), Arguments.Num(), *Context.Filename, Context.LineNumber, *FLocTextHelper::SanitizeLogOutput(Context.LineText.TrimStartAndEnd()));
+				UE_LOG(LogGatherTextFromSourceCommandlet, Warning, TEXT("%s(%d): Expected at least %d arguments for %s macro, but got %d. %s"), *Context.Filename, Context.LineNumber, GetMinNumberOfArgument(), *GetToken(), Arguments.Num(), *FLocTextHelper::SanitizeLogOutput(Context.LineText.TrimStartAndEnd()));
 				return;
 			}
 
@@ -2006,7 +2006,7 @@ void UGatherTextFromSourceCommandlet::FStringMacroDescriptor::TryParse(const FSt
 
 			if (NumArgs != Arguments.Num())
 			{
-				UE_LOG(LogGatherTextFromSourceCommandlet, Warning, TEXT("Expected %d arguments for %s macro, but got %d while parsing %s:%d. %s"), Arguments.Num(), *GetToken(), NumArgs, *Context.Filename, Context.LineNumber, *FLocTextHelper::SanitizeLogOutput(Context.LineText.TrimStartAndEnd()));
+				UE_LOG(LogGatherTextFromSourceCommandlet, Warning, TEXT("%s(%d): Expected %d arguments for %s macro, but got %d. %s"), *Context.Filename, Context.LineNumber, Arguments.Num(), *GetToken(), NumArgs, *FLocTextHelper::SanitizeLogOutput(Context.LineText.TrimStartAndEnd()));
 			}
 			else
 			{
@@ -2028,7 +2028,7 @@ void UGatherTextFromSourceCommandlet::FStringMacroDescriptor::TryParse(const FSt
 					FString ArgText = ArgArray[ArgIdx];
 
 					bool HasQuotes;
-					FString MacroDesc = FString::Printf(TEXT("argument %d of %d in %s macro at %s:%d"), ArgIdx+1, Arguments.Num(), *GetToken(), *Context.Filename, Context.LineNumber);
+					FString MacroDesc = FString::Printf(TEXT("%s(%d): argument %d of %d in %s macro"), *Context.Filename, Context.LineNumber, ArgIdx+1, Arguments.Num(), *GetToken());
 					if (!PrepareArgument(ArgText, Arg.IsAutoText, MacroDesc, HasQuotes))
 					{
 						ArgParseError = true;
@@ -2058,7 +2058,7 @@ void UGatherTextFromSourceCommandlet::FStringMacroDescriptor::TryParse(const FSt
 				if ( Identifier.IsEmpty() )
 				{
 					//The command doesn't have an identifier so we can't gather it
-					UE_LOG(LogGatherTextFromSourceCommandlet, Warning, TEXT("%s macro at %s has an empty identifier and cannot be gathered."), *GetToken(), *SourceLocation );
+					UE_LOG(LogGatherTextFromSourceCommandlet, Warning, TEXT("%s: %s macro has an empty identifier and cannot be gathered."), *SourceLocation, *GetToken() );
 					return;
 				}
 
@@ -2066,7 +2066,7 @@ void UGatherTextFromSourceCommandlet::FStringMacroDescriptor::TryParse(const FSt
 				{
 					if (!Namespace.IsSet())
 					{
-						UE_LOG(LogGatherTextFromSourceCommandlet, Warning, TEXT("%s macro at %s doesn't define a namespace and no external namespace was set. An empty namspace will be used."), *GetToken(), *SourceLocation );
+						UE_LOG(LogGatherTextFromSourceCommandlet, Warning, TEXT("%s: %s macro doesn't define a namespace and no external namespace was set. An empty namspace will be used."), *SourceLocation, *GetToken() );
 						Namespace = FString();
 					}
 
@@ -2100,7 +2100,7 @@ void UGatherTextFromSourceCommandlet::FStringTableMacroDescriptor::TryParse(cons
 		{
 			if (Arguments.Num() != GetMinNumberOfArgument())
 			{
-				UE_LOG(LogGatherTextFromSourceCommandlet, Warning, TEXT("Expected %d arguments for %s macro, but got %d while parsing %s:%d. %s"), GetMinNumberOfArgument(), *GetToken(), Arguments.Num(), *Context.Filename, Context.LineNumber, *FLocTextHelper::SanitizeLogOutput(Context.LineText.TrimStartAndEnd()));
+				UE_LOG(LogGatherTextFromSourceCommandlet, Warning, TEXT("%s(%d): Expected %d arguments for %s macro, but got %d. %s"), *Context.Filename, Context.LineNumber, GetMinNumberOfArgument(), *GetToken(), Arguments.Num(), *FLocTextHelper::SanitizeLogOutput(Context.LineText.TrimStartAndEnd()));
 			}
 			else
 			{
@@ -2109,8 +2109,8 @@ void UGatherTextFromSourceCommandlet::FStringTableMacroDescriptor::TryParse(cons
 				Arguments[1].TrimStartInline();
 				FString TableNamespace = Arguments[1];
 
-				const FString TableIdMacroDesc = FString::Printf(TEXT("\"Id\" argument in %s macro at %s:%d"), *GetToken(), *Context.Filename, Context.LineNumber);
-				const FString TableNamespaceMacroDesc = FString::Printf(TEXT("\"Namespace\" argument in %s macro at %s:%d"), *GetToken(), *Context.Filename, Context.LineNumber);
+				const FString TableIdMacroDesc = FString::Printf(TEXT("%s(%d): \"Id\" argument in %s macro"), *Context.Filename, Context.LineNumber, *GetToken());
+				const FString TableNamespaceMacroDesc = FString::Printf(TEXT("%s(%d): \"Namespace\" argument in %s macro"), *Context.Filename, Context.LineNumber, *GetToken());
 
 				bool HasQuotes;
 				if (PrepareArgument(TableId, true, TableIdMacroDesc, HasQuotes) && PrepareArgument(TableNamespace, true, TableNamespaceMacroDesc, HasQuotes))
@@ -2119,7 +2119,7 @@ void UGatherTextFromSourceCommandlet::FStringTableMacroDescriptor::TryParse(cons
 
 					if (TableIdName.IsNone())
 					{
-						UE_LOG(LogGatherTextFromSourceCommandlet, Warning, TEXT("%s macro at %s:%d has an empty identifier and cannot be gathered."), *GetToken(), *Context.Filename, Context.LineNumber);
+						UE_LOG(LogGatherTextFromSourceCommandlet, Warning, TEXT("%s(%d): %s macro has an empty identifier and cannot be gathered."), *Context.Filename, Context.LineNumber, *GetToken());
 					}
 					else
 					{
@@ -2143,7 +2143,7 @@ void UGatherTextFromSourceCommandlet::FStringTableFromFileMacroDescriptor::TryPa
 		{
 			if (Arguments.Num() != GetMinNumberOfArgument())
 			{
-				UE_LOG(LogGatherTextFromSourceCommandlet, Warning, TEXT("Expected %d arguments for %s macro, but got %d while parsing %s:%d. %s"), GetMinNumberOfArgument(), *GetToken(), Arguments.Num(), *Context.Filename, Context.LineNumber, *FLocTextHelper::SanitizeLogOutput(Context.LineText.TrimStartAndEnd()));
+				UE_LOG(LogGatherTextFromSourceCommandlet, Warning, TEXT("%s(%d): Expected %d arguments for %s macro, but got %d. %s"), *Context.Filename, Context.LineNumber, GetMinNumberOfArgument(), *GetToken(), Arguments.Num(), *FLocTextHelper::SanitizeLogOutput(Context.LineText.TrimStartAndEnd()));
 			}
 			else
 			{
@@ -2154,9 +2154,9 @@ void UGatherTextFromSourceCommandlet::FStringTableFromFileMacroDescriptor::TryPa
 				Arguments[2].TrimStartInline();
 				FString TableFilename = Arguments[2];
 
-				const FString TableIdMacroDesc = FString::Printf(TEXT("\"Id\" argument in %s macro at %s:%d"), *GetToken(), *Context.Filename, Context.LineNumber);
-				const FString TableNamespaceMacroDesc = FString::Printf(TEXT("\"Namespace\" argument in %s macro at %s:%d"), *GetToken(), *Context.Filename, Context.LineNumber);
-				const FString TableFilenameMacroDesc = FString::Printf(TEXT("\"FilePath\" argument in %s macro at %s:%d"), *GetToken(), *Context.Filename, Context.LineNumber);
+				const FString TableIdMacroDesc = FString::Printf(TEXT("%s(%d): \"Id\" argument in %s macro"), *Context.Filename, Context.LineNumber, *GetToken());
+				const FString TableNamespaceMacroDesc = FString::Printf(TEXT("%s(%d): \"Namespace\" argument in %s macro"), *Context.Filename, Context.LineNumber, *GetToken());
+				const FString TableFilenameMacroDesc = FString::Printf(TEXT("%s(%d): \"FilePath\" argument in %s macro"), *Context.Filename, Context.LineNumber, *GetToken());
 
 				bool HasQuotes;
 				if (PrepareArgument(TableId, true, TableIdMacroDesc, HasQuotes) && PrepareArgument(TableNamespace, true, TableNamespaceMacroDesc, HasQuotes) && PrepareArgument(TableFilename, true, TableFilenameMacroDesc, HasQuotes))
@@ -2165,7 +2165,7 @@ void UGatherTextFromSourceCommandlet::FStringTableFromFileMacroDescriptor::TryPa
 
 					if (TableIdName.IsNone())
 					{
-						UE_LOG(LogGatherTextFromSourceCommandlet, Warning, TEXT("%s macro at %s:%d has an empty identifier and cannot be gathered."), *GetToken(), *Context.Filename, Context.LineNumber);
+						UE_LOG(LogGatherTextFromSourceCommandlet, Warning, TEXT("%s(%d): %s macro has an empty identifier and cannot be gathered."), *Context.Filename, Context.LineNumber, *GetToken());
 					}
 					else
 					{
@@ -2189,7 +2189,7 @@ void UGatherTextFromSourceCommandlet::FStringTableEntryMacroDescriptor::TryParse
 		{
 			if (Arguments.Num() != GetMinNumberOfArgument())
 			{
-				UE_LOG(LogGatherTextFromSourceCommandlet, Warning, TEXT("Expected %d arguments for %s macro, but got %d while parsing %s:%d. %s"), GetMinNumberOfArgument(), *GetToken(), Arguments.Num(), *Context.Filename, Context.LineNumber, *FLocTextHelper::SanitizeLogOutput(Context.LineText.TrimStartAndEnd()));
+				UE_LOG(LogGatherTextFromSourceCommandlet, Warning, TEXT("%s(%d): Expected %d arguments for %s macro, but got %d. %s"), *Context.Filename, Context.LineNumber, GetMinNumberOfArgument(), *GetToken(), Arguments.Num(), *FLocTextHelper::SanitizeLogOutput(Context.LineText.TrimStartAndEnd()));
 			}
 			else
 			{
@@ -2200,9 +2200,9 @@ void UGatherTextFromSourceCommandlet::FStringTableEntryMacroDescriptor::TryParse
 				Arguments[2].TrimStartInline();
 				FString SourceString = Arguments[2];
 
-				const FString TableIdMacroDesc = FString::Printf(TEXT("\"Id\" argument in %s macro at %s:%d"), *GetToken(), *Context.Filename, Context.LineNumber);
-				const FString KeyMacroDesc = FString::Printf(TEXT("\"Key\" argument in %s macro at %s:%d"), *GetToken(), *Context.Filename, Context.LineNumber);
-				const FString SourceStringMacroDesc = FString::Printf(TEXT("\"SourceString\" argument in %s macro at %s:%d"), *GetToken(), *Context.Filename, Context.LineNumber);
+				const FString TableIdMacroDesc = FString::Printf(TEXT("%s(%d): \"Id\" argument in %s macro"), *Context.Filename, Context.LineNumber, *GetToken());
+				const FString KeyMacroDesc = FString::Printf(TEXT("%s(%d): \"Key\" argument in %s macro"), *Context.Filename, Context.LineNumber, *GetToken());
+				const FString SourceStringMacroDesc = FString::Printf(TEXT("%s(%d): \"SourceString\" argument in %s macro"), *Context.Filename, Context.LineNumber, *GetToken());
 
 				bool HasQuotes;
 				if (PrepareArgument(TableId, true, TableIdMacroDesc, HasQuotes) && PrepareArgument(Key, true, KeyMacroDesc, HasQuotes) && PrepareArgument(SourceString, true, SourceStringMacroDesc, HasQuotes))
@@ -2211,7 +2211,7 @@ void UGatherTextFromSourceCommandlet::FStringTableEntryMacroDescriptor::TryParse
 
 					if (TableIdName.IsNone() || Key.IsEmpty())
 					{
-						UE_LOG(LogGatherTextFromSourceCommandlet, Warning, TEXT("%s macro at %s:%d has an empty identifier and cannot be gathered."), *GetToken(), *Context.Filename, Context.LineNumber);
+						UE_LOG(LogGatherTextFromSourceCommandlet, Warning, TEXT("%s(%d): %s macro has an empty identifier and cannot be gathered."), *Context.Filename, Context.LineNumber, *GetToken());
 					}
 					else if (!SourceString.IsEmpty())
 					{
@@ -2235,7 +2235,7 @@ void UGatherTextFromSourceCommandlet::FStringTableEntryMetaDataMacroDescriptor::
 		{
 			if (Arguments.Num() != GetMinNumberOfArgument())
 			{
-				UE_LOG(LogGatherTextFromSourceCommandlet, Warning, TEXT("Expected %d arguments for %s macro, but got %d while parsing %s:%d. %s"), GetMinNumberOfArgument(), *GetToken(), Arguments.Num(), *Context.Filename, Context.LineNumber, *FLocTextHelper::SanitizeLogOutput(Context.LineText.TrimStartAndEnd()));
+				UE_LOG(LogGatherTextFromSourceCommandlet, Warning, TEXT("%s(%d): Expected %d arguments for %s macro, but got %d. %s"), *Context.Filename, Context.LineNumber, GetMinNumberOfArgument(), *GetToken(), Arguments.Num(), *FLocTextHelper::SanitizeLogOutput(Context.LineText.TrimStartAndEnd()));
 			}
 			else
 			{
@@ -2248,10 +2248,10 @@ void UGatherTextFromSourceCommandlet::FStringTableEntryMetaDataMacroDescriptor::
 				Arguments[3].TrimStartInline();
 				FString MetaData = Arguments[3];
 
-				const FString TableIdMacroDesc = FString::Printf(TEXT("\"Id\" argument in %s macro at %s:%d"), *GetToken(), *Context.Filename, Context.LineNumber);
-				const FString KeyMacroDesc = FString::Printf(TEXT("\"Key\" argument in %s macro at %s:%d"), *GetToken(), *Context.Filename, Context.LineNumber);
-				const FString MetaDataIdMacroDesc = FString::Printf(TEXT("\"MetaDataId\" argument in %s macro at %s:%d"), *GetToken(), *Context.Filename, Context.LineNumber);
-				const FString MetaDataMacroDesc = FString::Printf(TEXT("\"MetaData\" argument in %s macro at %s:%d"), *GetToken(), *Context.Filename, Context.LineNumber);
+				const FString TableIdMacroDesc = FString::Printf(TEXT("%s(%d): \"Id\" argument in %s macro"), *Context.Filename, Context.LineNumber, *GetToken());
+				const FString KeyMacroDesc = FString::Printf(TEXT("%s(%d): \"Key\" argument in %s macro"), *Context.Filename, Context.LineNumber, *GetToken());
+				const FString MetaDataIdMacroDesc = FString::Printf(TEXT("%s(%d): \"MetaDataId\" argument in %s macro"), *Context.Filename, Context.LineNumber, *GetToken());
+				const FString MetaDataMacroDesc = FString::Printf(TEXT("%s(%d): \"MetaData\" argument in %s macro"), *Context.Filename, Context.LineNumber, *GetToken());
 
 				bool HasQuotes;
 				if (PrepareArgument(TableId, true, TableIdMacroDesc, HasQuotes) && PrepareArgument(Key, true, KeyMacroDesc, HasQuotes) && PrepareArgument(MetaDataId, true, MetaDataIdMacroDesc, HasQuotes) && PrepareArgument(MetaData, true, MetaDataMacroDesc, HasQuotes))
@@ -2261,7 +2261,7 @@ void UGatherTextFromSourceCommandlet::FStringTableEntryMetaDataMacroDescriptor::
 
 					if (TableIdName.IsNone() || Key.IsEmpty() || MetaDataIdName.IsNone())
 					{
-						UE_LOG(LogGatherTextFromSourceCommandlet, Warning, TEXT("%s macro at %s:%d has an empty identifier and cannot be gathered."), *GetToken(), *Context.Filename, Context.LineNumber);
+						UE_LOG(LogGatherTextFromSourceCommandlet, Warning, TEXT("%s(%d): %s macro has an empty identifier and cannot be gathered."), *Context.Filename, Context.LineNumber, *GetToken());
 					}
 					else if (!MetaData.IsEmpty())
 					{
diff --git a/Engine/Source/Programs/AutomationTool/AutomationUtils/Matchers/LocalizationEventMatcher.cs b/Engine/Source/Programs/AutomationTool/AutomationUtils/Matchers/LocalizationEventMatcher.cs
new file mode 100644
--- /dev/null
+++ b/Engine/Source/Programs/AutomationTool/AutomationUtils/Matchers/LocalizationEventMatcher.cs
@@ -0,0 +1,57 @@
+// Copyright Epic Games, Inc. All Rights Reserved.
+
+using System.Text.RegularExpressions;
+using EpicGames.Core;
+using Microsoft.Extensions.Logging;
+
+#nullable enable
+
+namespace AutomationUtils.Matchers
+{
+	/// <summary>
+	/// Matches events formatted as UE log channel output by the localization commandlet
+	/// </summary>
+	class LocalizationEventMatcher : ILogEventMatcher
+	{
+		readonly Regex s_pattern = new Regex(
+			@"^(\s*)" +
+			@"(?<channel>LogLocTextHelper|LogGatherTextFromSourceCommandlet):\s*" +
+			@"(?<severity>Error|Warning|Display):\s+" +
+			@"(?<file>([a-zA-Z]:)?[^:/\\]*[/\\][^:]+[^\)])" +
+			@"(?:\((?<line>\d+)\))?" +
+			@":");
+
+		/// <inheritdoc/>
+		public LogEventMatch? Match(ILogCursor input)
+		{
+			Match? match = s_pattern.Match(input.CurrentLine!);
+			if (match.Success)
+			{
+				LogEventBuilder builder = new LogEventBuilder(input);
+				builder.Annotate(match.Groups["channel"], LogEventMarkup.Channel);
+				builder.Annotate(match.Groups["severity"], LogEventMarkup.Severity);
+				builder.AnnotateSourceFile(match.Groups["file"], "Engine");
+				builder.Annotate(match.Groups["line"], LogEventMarkup.LineNumber);
+
+				if (builder.Next.TryMatch(@"^\s+", out Match? indent))
+				{
+					while (builder.Next.IsMatch($"^{indent.Value}"))
+					{
+						builder.MoveNext();
+					}
+				}
+
+				LogLevel level = match.Groups["severity"].Value switch
+				{
+					"Error" => LogLevel.Error,
+					"Warning" => LogLevel.Warning,
+					_ => LogLevel.Information,
+				};
+
+				return builder.ToMatch(LogEventPriority.High, level, KnownLogEvents.Engine_Localization);
+
+			}
+			return null;
+		}
+	}
+}
diff --git a/Engine/Source/Programs/Horde/Horde.Agent.Tests/LogParserTests.cs b/Engine/Source/Programs/Horde/Horde.Agent.Tests/LogParserTests.cs
--- a/Engine/Source/Programs/Horde/Horde.Agent.Tests/LogParserTests.cs
+++ b/Engine/Source/Programs/Horde/Horde.Agent.Tests/LogParserTests.cs
@@ -757,6 +757,27 @@ namespace Horde.Agent.Tests
 			}
 		}
 
+		[TestMethod]
+		public void LocalizationChannelMatcher()
+		{
+			string[] lines =
+			{
+				@"LogLocTextHelper: Warning: Source/Runtime/Engine/Private/StaticMesh.cpp(5649): Text conflict from LOCTEXT macro for namespace"
+			};
+
+			{
+				List<LogEvent> logEvents = Parse(lines);
+				Assert.AreEqual(1, logEvents.Count);
+				CheckEventGroup(logEvents, 0, 1, LogLevel.Warning, KnownLogEvents.Engine_Localization);
+
+				LogEvent logEvent = logEvents[0];
+				Assert.AreEqual("LogLocTextHelper", logEvent.GetProperty("channel").ToString());
+				Assert.AreEqual("Warning", logEvent.GetProperty("severity").ToString());
+				Assert.AreEqual("Engine/Source/Runtime/Engine/Private/StaticMesh.cpp", logEvent.GetProperty<LogValue>("file")!.Properties!["relativePath"].ToString());
+				Assert.AreEqual("5649", logEvent.GetProperty("line").ToString());
+			}
+		}
+
 		static List<LogEvent> Parse(IEnumerable<string> lines)
 		{
 			return Parse(String.Join("\n", lines));
diff --git a/Engine/Source/Programs/Horde/Horde.Build/Issues/Impl/Handlers/LocalizationIssueHandler.cs b/Engine/Source/Programs/Horde/Horde.Build/Issues/Impl/Handlers/LocalizationIssueHandler.cs
new file mode 100644
--- /dev/null
+++ b/Engine/Source/Programs/Horde/Horde.Build/Issues/Impl/Handlers/LocalizationIssueHandler.cs
@@ -0,0 +1,102 @@
+// Copyright Epic Games, Inc. All Rights Reserved.
+
+using System;
+using System.Collections.Generic;
+using System.Diagnostics.CodeAnalysis;
+using System.Linq;
+using EpicGames.Core;
+using Horde.Build.Collections;
+using Horde.Build.Models;
+using Horde.Build.Utilities;
+using Microsoft.Extensions.Logging;
+using MongoDB.Driver;
+
+namespace Horde.Build.IssueHandlers.Impl
+{
+	/// <summary>
+	/// Instance of a localization error
+	/// </summary>
+	class LocalizationIssueHandler : IIssueHandler
+	{
+		/// <inheritdoc/>
+		public string Type => "Localization";
+
+		/// <inheritdoc/>
+		public int Priority => 10;
+
+		/// <summary>
+		/// Filenames containing errors or warnings
+		/// </summary>
+		public HashSet<string> FileNames { get; set; } = new HashSet<string>();
+
+		/// <summary>
+		/// Determines if the given event id matches
+		/// </summary>
+		/// <param name="eventId">The event id to compare</param>
+		/// <returns>True if the given event id matches</returns>
+		public static bool IsMatchingEventId(EventId? eventId)
+		{
+			return eventId == KnownLogEvents.Engine_Localization;
+		}
+
+		/// <summary>
+		/// Extracts a list of source files from an event
+		/// </summary>
+		/// <param name="logEventData">The event data</param>
+		/// <param name="sourceFiles">List of source files</param>
+		public static void GetSourceFiles(ILogEventData logEventData, HashSet<string> sourceFiles)
+		{
+			foreach (ILogEventLine line in logEventData.Lines)
+			{
+				string? relativePath;
+				if (line.Data.TryGetNestedProperty("properties.file.relativePath", out relativePath) || line.Data.TryGetNestedProperty("properties.file", out relativePath))
+				{
+					int endIdx = relativePath.LastIndexOfAny(new char[] { '/', '\\' }) + 1;
+					string fileName = relativePath.Substring(endIdx);
+					sourceFiles.Add(fileName);
+				}
+			}
+		}
+
+		/// <inheritdoc/>
+		public bool TryGetFingerprint(IJob job, INode node, ILogEventData eventData, [NotNullWhen(true)] out NewIssueFingerprint? fingerprint)
+		{
+			if (!IsMatchingEventId(eventData.EventId))
+			{
+				fingerprint = null;
+				return false;
+			}
+
+			HashSet<string> newFileNames = new HashSet<string>(StringComparer.OrdinalIgnoreCase);
+			GetSourceFiles(eventData, newFileNames);
+			fingerprint = new NewIssueFingerprint(Type, newFileNames, null);
+			return true;
+		}
+
+		/// <inheritdoc/>
+		public void RankSuspects(IIssueFingerprint fingerprint, List<SuspectChange> suspects)
+		{
+			foreach (SuspectChange suspect in suspects)
+			{
+				if (suspect.ContainsCode)
+				{
+					if (FileNames.Any(x => suspect.ModifiesFile(x)))
+					{
+						suspect.Rank += 20;
+					}
+					else
+					{
+						suspect.Rank += 10;
+					}
+				}
+			}
+		}
+
+		/// <inheritdoc/>
+		public string GetSummary(IIssueFingerprint fingerprint, IssueSeverity severity)
+		{
+			string type = (severity == IssueSeverity.Warning)? "warnings" : "errors";
+			return $"Localization {type} in {StringUtils.FormatList(fingerprint.Keys.ToArray(), 2)}";
+		}
+	}
+}
diff --git a/Engine/Source/Programs/Shared/EpicGames.Core/KnownLogEvents.cs b/Engine/Source/Programs/Shared/EpicGames.Core/KnownLogEvents.cs
--- a/Engine/Source/Programs/Shared/EpicGames.Core/KnownLogEvents.cs
+++ b/Engine/Source/Programs/Shared/EpicGames.Core/KnownLogEvents.cs
@@ -70,6 +70,11 @@ namespace EpicGames.Core
 		/// </summary>
 		public static EventId Engine_AssetLog { get; } = new EventId(303);
 
+		/// <summary>
+		/// Engine: Localization commandlet output
+		/// </summary>
+		public static EventId Engine_Localization { get; } = new EventId(304);
+
 		/// <summary>
 		/// UAT error
 		/// </summary>
