Converting hardcoded short class/enum names to pathnames ahead of ANY_PACKAGE removal

diff --git a/Engine/Plugins/FX/Niagara/Source/Niagara/Classes/NiagaraLensEffectBase.h b/Engine/Plugins/FX/Niagara/Source/Niagara/Classes/NiagaraLensEffectBase.h
--- a/Engine/Plugins/FX/Niagara/Source/Niagara/Classes/NiagaraLensEffectBase.h
+++ b/Engine/Plugins/FX/Niagara/Source/Niagara/Classes/NiagaraLensEffectBase.h
@@ -55,7 +55,7 @@ protected:
 	 *  If an effect class in this array is currently playing, do not play this effect.
 	 *  Useful for preventing multiple similar or expensive camera effects from playing simultaneously.
 	 */
-	UPROPERTY(EditDefaultsOnly, Category = "Effect Activation", meta = (MustImplement = "CameraLensEffectInterface"))
+	UPROPERTY(EditDefaultsOnly, Category = "Effect Activation", meta = (MustImplement = "/Script/Engine.CameraLensEffectInterface"))
 	TArray<TSubclassOf<AActor>> EmittersToTreatAsSame;
 
 	UPROPERTY()
diff --git a/Engine/Source/Editor/UnrealEd/Classes/Editor/PropertyEditorTestObject.h b/Engine/Source/Editor/UnrealEd/Classes/Editor/PropertyEditorTestObject.h
--- a/Engine/Source/Editor/UnrealEd/Classes/Editor/PropertyEditorTestObject.h
+++ b/Engine/Source/Editor/UnrealEd/Classes/Editor/PropertyEditorTestObject.h
@@ -621,7 +621,7 @@ class UPropertyEditorTestObject : public UObject
 	TScriptInterface<IAnimClassInterface> AnimClassInterface;
 
 	// This is an IBlendableInterface that only allows for ULightPropagationVolumeBlendable objects
-	UPROPERTY(EditAnywhere, Category=ScriptInterfaces, meta=(AllowedClasses="LightPropagationVolumeBlendable"))
+	UPROPERTY(EditAnywhere, Category=ScriptInterfaces, meta=(AllowedClasses="/Script/Engine.LightPropagationVolumeBlendable"))
 	TScriptInterface<IBlendableInterface> LightPropagationVolumeBlendable;
 
 	// Allows either an object that's derived from UTexture or IBlendableInterface, to ensure that Object Property handles know how to
diff --git a/Engine/Source/Runtime/Engine/Classes/Camera/CameraLensEffectInterface.h b/Engine/Source/Runtime/Engine/Classes/Camera/CameraLensEffectInterface.h
--- a/Engine/Source/Runtime/Engine/Classes/Camera/CameraLensEffectInterface.h
+++ b/Engine/Source/Runtime/Engine/Classes/Camera/CameraLensEffectInterface.h
@@ -72,7 +72,7 @@ struct FCameraLensInterfaceClassSupport
 	GENERATED_BODY()
 
 	/** The class to spawn/reference. Must implement CameraLensEffectInterface! */
-	UPROPERTY(EditAnywhere, Category="Lens Effect", meta=(MustImplement="CameraLensEffectInterface"))
+	UPROPERTY(EditAnywhere, Category="Lens Effect", meta=(MustImplement= "/Script/Engine.CameraLensEffectInterface"))
 	TSubclassOf<AActor> Class;
 };
 
diff --git a/Engine/Source/Runtime/Engine/Classes/Particles/EmitterCameraLensEffectBase.h b/Engine/Source/Runtime/Engine/Classes/Particles/EmitterCameraLensEffectBase.h
--- a/Engine/Source/Runtime/Engine/Classes/Particles/EmitterCameraLensEffectBase.h
+++ b/Engine/Source/Runtime/Engine/Classes/Particles/EmitterCameraLensEffectBase.h
@@ -54,7 +54,7 @@ public:
 	 *  If an emitter class in this array is currently playing, do not play this effect.
 	 *  Useful for preventing multiple similar or expensive camera effects from playing simultaneously.
 	 */
-	UPROPERTY(EditDefaultsOnly, Category = EmitterCameraLensEffectBase, meta = (MustImplement = "CameraLensEffectInterface"))
+	UPROPERTY(EditDefaultsOnly, Category = EmitterCameraLensEffectBase, meta = (MustImplement = "/Script/Engine.CameraLensEffectInterface"))
 	TArray<TSubclassOf<AActor>> EmittersToTreatAsSame;
 public:
 	//~ Begin AActor Interface
diff --git a/Engine/Source/Runtime/UMG/Public/Components/ListViewBase.h b/Engine/Source/Runtime/UMG/Public/Components/ListViewBase.h
--- a/Engine/Source/Runtime/UMG/Public/Components/ListViewBase.h
+++ b/Engine/Source/Runtime/UMG/Public/Components/ListViewBase.h
@@ -642,7 +642,7 @@ protected:
 
 	// Note: Options for this property can be configured via class and property metadata. See class declaration comment above.
 	/** The type of widget to create for each entry displayed in the list. */
-	UPROPERTY(EditAnywhere, BlueprintReadOnly, Category = ListEntries, meta = (DesignerRebuild, AllowPrivateAccess = true, MustImplement = UserListEntry))
+	UPROPERTY(EditAnywhere, BlueprintReadOnly, Category = ListEntries, meta = (DesignerRebuild, AllowPrivateAccess = true, MustImplement = "/Script/UMG.UserListEntry"))
 	TSubclassOf<UUserWidget> EntryWidgetClass;
 
 	/** The multiplier to apply when wheel scrolling */
