Try to generate some import options base on the mesh when importing a lod on a static mesh without import options. (This should probably be revisited at some point)

diff --git a/Engine/Source/Editor/UnrealEd/Private/FbxMeshUtils.cpp b/Engine/Source/Editor/UnrealEd/Private/FbxMeshUtils.cpp
--- a/Engine/Source/Editor/UnrealEd/Private/FbxMeshUtils.cpp
+++ b/Engine/Source/Editor/UnrealEd/Private/FbxMeshUtils.cpp
@@ -160,12 +160,26 @@ namespace FbxMeshUtils
 			// Import data already exists, apply it to the fbx import options
 			ReimportUI->StaticMeshImportData = ImportData;
 			ApplyImportUIToImportOptions(ReimportUI, *ImportOptions);
-			ImportOptions->bIsImportCancelable = false;
-			ImportOptions->bImportMaterials = false;
-			ImportOptions->bImportTextures = false;
-			//Make sure the LODGroup do not change when re-importing a mesh
-			ImportOptions->StaticMeshLODGroup = BaseStaticMesh->LODGroup;
 		}
+		else if (BaseStaticMesh->IsSourceModelValid(0))
+		{
+			// Use the lod 0 to set the import settings
+			FStaticMeshSourceModel& SourceModel = BaseStaticMesh->GetSourceModel(0);
+			ImportOptions->NormalGenerationMethod = SourceModel.BuildSettings.bUseMikkTSpace ? EFBXNormalGenerationMethod::MikkTSpace : EFBXNormalGenerationMethod::BuiltIn;
+			ImportOptions->bComputeWeightedNormals = SourceModel.BuildSettings.bComputeWeightedNormals;
+			ImportOptions->DistanceFieldResolutionScale = SourceModel.BuildSettings.DistanceFieldResolutionScale;
+			ImportOptions->bRemoveDegenerates = SourceModel.BuildSettings.bRemoveDegenerates;
+			ImportOptions->bBuildReversedIndexBuffer = SourceModel.BuildSettings.bBuildReversedIndexBuffer;
+			ImportOptions->bGenerateLightmapUVs = SourceModel.BuildSettings.bGenerateLightmapUVs;
+		}
+
+		// Set a couple of settings that shouldn't change while importing a lod
+		ImportOptions->bBuildNanite = BaseStaticMesh->NaniteSettings.bEnabled;
+		ImportOptions->StaticMeshLODGroup = BaseStaticMesh->LODGroup;
+		ImportOptions->bIsImportCancelable = false;
+		ImportOptions->bImportMaterials = false;
+		ImportOptions->bImportTextures = false;
+
 		ImportOptions->bAutoComputeLodDistances = true; //Setting auto compute distance to true will avoid changing the staticmesh flag
 		if ( !FFbxImporter->ImportFromFile( *Filename, FPaths::GetExtension( Filename ), true ) )
 		{
