Simply add the name and location of generated INI files to the World Partion Convert Commandlet. Make sure the log on Display to see them in the editor when using Tool -> Convert level option

#rnx

diff --git a/Engine/Source/Editor/UnrealEd/Classes/Commandlets/WorldPartitionConvertCommandlet.h b/Engine/Source/Editor/UnrealEd/Classes/Commandlets/WorldPartitionConvertCommandlet.h
--- a/Engine/Source/Editor/UnrealEd/Classes/Commandlets/WorldPartitionConvertCommandlet.h
+++ b/Engine/Source/Editor/UnrealEd/Classes/Commandlets/WorldPartitionConvertCommandlet.h
@@ -12,7 +12,7 @@
 #include "WorldPartition/HLOD/HLODLayer.h"
 #include "WorldPartitionConvertCommandlet.generated.h"
 
-UNREALED_API DECLARE_LOG_CATEGORY_EXTERN(LogWorldPartitionConvertCommandlet, Log, All);
+UNREALED_API DECLARE_LOG_CATEGORY_EXTERN(LogWorldPartitionConvertCommandlet, Display, All);
 
 class ULevelStreaming;
 class UWorldPartition;
diff --git a/Engine/Source/Editor/UnrealEd/Private/Commandlets/WorldPartitionConvertCommandlet.cpp b/Engine/Source/Editor/UnrealEd/Private/Commandlets/WorldPartitionConvertCommandlet.cpp
--- a/Engine/Source/Editor/UnrealEd/Private/Commandlets/WorldPartitionConvertCommandlet.cpp
+++ b/Engine/Source/Editor/UnrealEd/Private/Commandlets/WorldPartitionConvertCommandlet.cpp
@@ -174,7 +174,7 @@ UWorld* UWorldPartitionConvertCommandlet::LoadWorld(const FString& LevelToLoad)
 	TRACE_CPUPROFILER_EVENT_SCOPE(UWorldPartitionConvertCommandlet::LoadWorld);
 
 	SET_WARN_COLOR(COLOR_WHITE);
-	UE_LOG(LogWorldPartitionConvertCommandlet, Log, TEXT("Loading level %s."), *LevelToLoad);
+	UE_LOG(LogWorldPartitionConvertCommandlet, Display, TEXT("Loading level %s."), *LevelToLoad);
 	CLEAR_WARN_COLOR();
 
 	UPackage* MapPackage = LoadPackage(nullptr, *LevelToLoad, LOAD_None);
@@ -192,7 +192,7 @@ ULevel* UWorldPartitionConvertCommandlet::InitWorld(UWorld* World)
 	TRACE_CPUPROFILER_EVENT_SCOPE(UWorldPartitionConvertCommandlet::InitWorld);
 
 	SET_WARN_COLOR(COLOR_WHITE);
-	UE_LOG(LogWorldPartitionConvertCommandlet, Log, TEXT("Initializing level %s."), *World->GetName());
+	UE_LOG(LogWorldPartitionConvertCommandlet, Display, TEXT("Initializing level %s."), *World->GetName());
 	CLEAR_WARN_COLOR();
 
 	// Setup the world.
@@ -273,7 +273,7 @@ void UWorldPartitionConvertCommandlet::GatherAndPrepareSubLevelsToConvert(ULevel
 		}
 		else
 		{
-			UE_LOG(LogWorldPartitionConvertCommandlet, Log, TEXT("Skipping conversion of streaming Level %s"), *StreamingLevel->GetWorldAssetPackageName());
+			UE_LOG(LogWorldPartitionConvertCommandlet, Display, TEXT("Skipping conversion of streaming Level %s"), *StreamingLevel->GetWorldAssetPackageName());
 		}
 	}
 
@@ -304,7 +304,7 @@ bool UWorldPartitionConvertCommandlet::PrepareStreamingLevelForConversion(ULevel
 	if (bOnlyMergeSubLevels || StreamingLevel->ShouldBeAlwaysLoaded() || StreamingLevel->bDisableDistanceStreaming)
 	{
 		FString WorldPath = SubLevel->GetPackage()->GetName();
-		UE_LOG(LogWorldPartitionConvertCommandlet, Log, TEXT("Converting %s streaming level %s"), StreamingLevel->bDisableDistanceStreaming ? TEXT("non distance-based") : TEXT("always loaded"), *StreamingLevel->GetWorldAssetPackageName());
+		UE_LOG(LogWorldPartitionConvertCommandlet, Display, TEXT("Converting %s streaming level %s"), StreamingLevel->bDisableDistanceStreaming ? TEXT("non distance-based") : TEXT("always loaded"), *StreamingLevel->GetWorldAssetPackageName());
 
 		for (AActor* Actor: SubLevel->Actors)
 		{
@@ -1307,7 +1307,7 @@ int32 UWorldPartitionConvertCommandlet::Main(const FString& Params)
 			ActorsToConvert = SubLevel->Actors;
 		}
 
-		UE_LOG(LogWorldPartitionConvertCommandlet, Log, TEXT("Converting %s"), *SubWorld->GetName());
+		UE_LOG(LogWorldPartitionConvertCommandlet, Display, TEXT("Converting %s"), *SubWorld->GetName());
 
 		PrepareLevelActors(SubLevel, ActorsToConvert, false);
 
@@ -1352,7 +1352,7 @@ int32 UWorldPartitionConvertCommandlet::Main(const FString& Params)
 				if (!ActorClass->IsNative() && (ActorClass->GetPackage() == SubPackage))
 				{
 					ChangeObjectOuter(ActorClass, MainPackage);
-					UE_LOG(LogWorldPartitionConvertCommandlet, Log, TEXT("Extracted non-native class %s"), *ActorClass->GetName());
+					UE_LOG(LogWorldPartitionConvertCommandlet, Display, TEXT("Extracted non-native class %s"), *ActorClass->GetName());
 				}
 			}
 		}
@@ -1470,7 +1470,7 @@ int32 UWorldPartitionConvertCommandlet::Main(const FString& Params)
 			UPackage* ActorPackage = Actor->GetExternalPackage();
 			PackagesToSave.Add(ActorPackage);
 
-			UE_LOG(LogWorldPartitionConvertCommandlet, Log, TEXT("Extracted actor %s(guid:%s) in %s"), *Actor->GetName(), *Actor->GetActorGuid().ToString(EGuidFormats::Digits), *ActorPackage->GetName());
+			UE_LOG(LogWorldPartitionConvertCommandlet, Display, TEXT("Extracted actor %s(guid:%s) in %s"), *Actor->GetName(), *Actor->GetActorGuid().ToString(EGuidFormats::Digits), *ActorPackage->GetName());
 		}
 
 		// Required to clear any deleted actors from the level
@@ -1514,7 +1514,7 @@ int32 UWorldPartitionConvertCommandlet::Main(const FString& Params)
 		{
 			TRACE_CPUPROFILER_EVENT_SCOPE(CheckoutPackages);
 
-			UE_LOG(LogWorldPartitionConvertCommandlet, Log, TEXT("Checking out %d packages."), PackagesToSave.Num());
+			UE_LOG(LogWorldPartitionConvertCommandlet, Display, TEXT("Checking out %d packages."), PackagesToSave.Num());
 			for(UPackage* Package: PackagesToSave)
 			{
 				FString PackageFileName = SourceControlHelpers::PackageFilename(Package);
@@ -1541,7 +1541,7 @@ int32 UWorldPartitionConvertCommandlet::Main(const FString& Params)
 		{
 			TRACE_CPUPROFILER_EVENT_SCOPE(SavePackages);
 
-			UE_LOG(LogWorldPartitionConvertCommandlet, Log, TEXT("Saving %d packages."), PackagesToSave.Num());
+			UE_LOG(LogWorldPartitionConvertCommandlet, Display, TEXT("Saving %d packages."), PackagesToSave.Num());
 			for (UPackage* PackageToSave : PackagesToSave)
 			{
 				FString PackageFileName = SourceControlHelpers::PackageFilename(PackageToSave);
@@ -1576,7 +1576,7 @@ int32 UWorldPartitionConvertCommandlet::Main(const FString& Params)
 
 		UPackage::WaitForAsyncFileWrites();
 
-		UE_LOG(LogWorldPartitionConvertCommandlet, Log, TEXT("######## CONVERSION COMPLETED SUCCESSFULLY ########"));
+		UE_LOG(LogWorldPartitionConvertCommandlet, Display, TEXT("######## CONVERSION COMPLETED SUCCESSFULLY ########"));
 	}
 
 	if (bGenerateIni || !bReportOnly)
@@ -1595,6 +1595,8 @@ int32 UWorldPartitionConvertCommandlet::Main(const FString& Params)
 			{
 				Pair.Value->SaveConfig(CPF_Config, *LevelConfigFilename);
 			}
+
+			UE_LOG(LogWorldPartitionConvertCommandlet, Display, TEXT("Generated ini file: %s"), *LevelConfigFilename);
 		}
 	}
 
