Cooker PumpPollables: Add support for manually triggered pollables. Change Pollables into an indirect array, to allow the pointer be used as a handle for the manual trigger. Change pollables into a heap so we don't have to view the entire array each tick.
#rnx

diff --git a/Engine/Source/Editor/UnrealEd/Classes/CookOnTheSide/CookOnTheFlyServer.h b/Engine/Source/Editor/UnrealEd/Classes/CookOnTheSide/CookOnTheFlyServer.h
--- a/Engine/Source/Editor/UnrealEd/Classes/CookOnTheSide/CookOnTheFlyServer.h
+++ b/Engine/Source/Editor/UnrealEd/Classes/CookOnTheSide/CookOnTheFlyServer.h
@@ -12,6 +12,7 @@
 #include "IPlatformFileSandboxWrapper.h"
 #include "Misc/EnumClassFlags.h"
 #include "Templates/Function.h"
+#include "Templates/RefCounting.h"
 #include "Templates/UniquePtr.h"
 #include "TickableEditorObject.h"
 #include "UObject/Object.h"
@@ -154,16 +155,58 @@ class UNREALED_API UCookOnTheFlyServer : public UObject, public FTickableEditorO
 
 private:
 	using FPollFunction = TUniqueFunction<void(UE::Cook::FTickStackData&)>;
-	struct FPollable
+	/**
+	 * Wrapper around a function for cooker tasks that need to be ticked on a schedule.
+	 * Includes data to support calling it on a schedule or triggering it manually.
+	 */
+	struct FPollable : public FRefCountBase
 	{
-		FPollable(float InPeriodSeconds, float InPeriodIdleSeconds, FPollFunction&& InFunction);
+	public:
+		enum EManualTrigger {};
+		FPollable(const TCHAR* InDebugName, float InPeriodSeconds, float InPeriodIdleSeconds, FPollFunction&& InFunction);
+		FPollable(const TCHAR* InDebugName, EManualTrigger, FPollFunction&& InFunction);
 
+	public:
+		/** Trigger the pollable to execute on the next PumpPollables */
+		void Trigger(UCookOnTheFlyServer& COTFS);
+		/**
+		 * Run the pollable's callback function now. Note this means running it outside of its normal location which is not valid for all pollables.
+		 * PumpPollables is not thread safe, so this must be called on the scheduler thread.
+		 */
+		void RunNow(UCookOnTheFlyServer& COTFS);
+
+		/** Run the pollable's callback function, from PumpPollables. Calculate the new time and store it in the output. */
+		void RunDuringPump(UE::Cook::FTickStackData& StackData, double& OutNewCurrentTime, double& OutNextTimeSeconds);
+
+	public:
+		const TCHAR* DebugName = TEXT("");
 		FPollFunction PollFunction;
-		double NextTimeSeconds = 0.;
+		/** Time when this should be next called, if the cooker is idle. (See also NextTimeSeconds). */
 		double NextTimeIdleSeconds = 0.;
 		float PeriodSeconds = 60.f;
 		float PeriodIdleSeconds = 5.f;
 	};
+	/** A key/value pair for storing Pollables in a priority queue, keyed by next call time. */
+	struct FPollableQueueKey
+	{
+	public:
+		FPollableQueueKey() = default;
+		explicit FPollableQueueKey(FPollable* InPollable);
+		explicit FPollableQueueKey(const TRefCountPtr<FPollable>& InPollable);
+		explicit FPollableQueueKey(TRefCountPtr<FPollable>&& InPollable);
+		bool operator<(const FPollableQueueKey& Other) const
+		{
+			return NextTimeSeconds < Other.NextTimeSeconds;
+		}
+
+	public:
+		TRefCountPtr<FPollable> Pollable;
+		/**
+		 * Time when the pollable be next called, if the cooker is not idle. Stored here rather than
+		 * on the FPollable to support fast access in the queue.
+		 */
+		double NextTimeSeconds = 0.;
+	};
 
 	//////////////////////////////////////////////////////////////////////////
 	// Cook on the fly server interface adapter
@@ -385,8 +428,8 @@ private:
 	void InitializePollables();
 	void PumpPollables(UE::Cook::FTickStackData& StackData, bool bIsIdle);
 	void PollFlushRenderingCommands();
-	TOptional<FPollable> CreatePollableLLM();
-	TOptional<FPollable> CreatePollableTriggerGC();
+	TRefCountPtr<FPollable> CreatePollableLLM();
+	TRefCountPtr<FPollable> CreatePollableTriggerGC();
 	void PollPackagesPerGC(UE::Cook::FTickStackData& StackData);
 	void WaitForAsync(UE::Cook::FTickStackData& StackData);
 
@@ -1195,12 +1238,10 @@ private:
 	/** Helper struct for running cooking in diagnostic modes */
 	TUniquePtr<FDiffModeCookServerUtils> DiffModeHelper;
 
-	TArray<FPollable> Pollables;
+	TArray<FPollableQueueKey> Pollables;
+	FCriticalSection PollablesLock;
 	double PollNextTimeSeconds = 0.;
 	double PollNextTimeIdleSeconds = 0.;
-	int32 PollStartIndex = 0;
-	float PumpPollablesTimeSlice = 0.f;
-	float PumpPollablesMinPeriod = 0.f;
 	uint32 CookedPackageCountSinceLastGC = 0;
 	float WaitForAsyncSleepSeconds = 0.0f;
 	friend UE::Cook::FPackageData;
diff --git a/Engine/Source/Editor/UnrealEd/Private/CookOnTheFlyServer.cpp b/Engine/Source/Editor/UnrealEd/Private/CookOnTheFlyServer.cpp
--- a/Engine/Source/Editor/UnrealEd/Private/CookOnTheFlyServer.cpp
+++ b/Engine/Source/Editor/UnrealEd/Private/CookOnTheFlyServer.cpp
@@ -1609,62 +1609,189 @@ void UCookOnTheFlyServer::UpdateDisplay(ECookTickFlags TickFlags, bool bForceDis
 	}
 }
 
-UCookOnTheFlyServer::FPollable::FPollable(float InPeriodSeconds, float InPeriodIdleSeconds, UCookOnTheFlyServer::FPollFunction&& InFunction)
-	: PollFunction(MoveTemp(InFunction))
-	, NextTimeSeconds(0)
+namespace UE::Cook::Pollable
+{
+constexpr double TimePeriodNever = MAX_flt / 2;
+constexpr int32 ExpectedMaxNum = 10; // Used to size inline arrays
+
+}
+
+UCookOnTheFlyServer::FPollable::FPollable(const TCHAR* InDebugName, float InPeriodSeconds, float InPeriodIdleSeconds,
+	UCookOnTheFlyServer::FPollFunction&& InFunction)
+	: DebugName(InDebugName)
+	, PollFunction(MoveTemp(InFunction))
 	, NextTimeIdleSeconds(0)
 	, PeriodSeconds(InPeriodSeconds)
 	, PeriodIdleSeconds(InPeriodIdleSeconds)
 {
+	check(DebugName);
+}
+
+UCookOnTheFlyServer::FPollable::FPollable(const TCHAR* InDebugName, EManualTrigger,
+	UCookOnTheFlyServer::FPollFunction&& InFunction)
+	: DebugName(InDebugName)
+	, PollFunction(MoveTemp(InFunction))
+	, NextTimeIdleSeconds(MAX_flt)
+	, PeriodSeconds(UE::Cook::Pollable::TimePeriodNever)
+	, PeriodIdleSeconds(UE::Cook::Pollable::TimePeriodNever)
+{
+	check(DebugName);
+}
+
+UCookOnTheFlyServer::FPollableQueueKey::FPollableQueueKey(FPollable* InPollable)
+	: FPollableQueueKey(TRefCountPtr<FPollable>(InPollable))
+{
+}
+UCookOnTheFlyServer::FPollableQueueKey::FPollableQueueKey(const TRefCountPtr<FPollable>& InPollable)
+	: FPollableQueueKey(TRefCountPtr<FPollable>(InPollable))
+{
+}
+UCookOnTheFlyServer::FPollableQueueKey::FPollableQueueKey(TRefCountPtr<FPollable>&& InPollable)
+	: Pollable(MoveTemp(InPollable))
+{
+	if (Pollable->PeriodSeconds < UE::Cook::Pollable::TimePeriodNever)
+	{
+		NextTimeSeconds = 0;
+	}
+	else
+	{
+		NextTimeSeconds = MAX_flt;
+	}
+}
+
+void UCookOnTheFlyServer::FPollable::Trigger(UCookOnTheFlyServer& COTFS)
+{
+	FScopeLock PollablesScopeLock(&COTFS.PollablesLock);
+
+	FPollableQueueKey* KeyInQueue = COTFS.Pollables.FindByPredicate(
+		[this](const FPollableQueueKey& Existing) { return Existing.Pollable.GetReference() == this; });
+	if (ensure(KeyInQueue))
+	{
+		FPollableQueueKey LocalQueueKey;
+		LocalQueueKey.Pollable = MoveTemp(KeyInQueue->Pollable);
+		// If the top of the heap is already triggered, put this after the top of the heap to
+		// avoid excessive triggering causing starvation for other pollables
+		// Note that the top of the heap might be this. Otherwise put this at the top of the
+		// heap by setting its time to CurrentTime
+		double CurrentTime = FPlatformTime::Seconds();
+		double TimeAfterHeapTop = COTFS.Pollables.HeapTop().NextTimeSeconds + .001f;
+		LocalQueueKey.NextTimeSeconds = FMath::Min(CurrentTime, TimeAfterHeapTop);
+		this->NextTimeIdleSeconds = LocalQueueKey.NextTimeSeconds;
+
+		int32 Index = KeyInQueue - COTFS.Pollables.GetData();
+		COTFS.Pollables.HeapRemoveAt(Index, false /* bAllowShrinking */);
+		COTFS.Pollables.HeapPush(MoveTemp(LocalQueueKey));
+		COTFS.PollNextTimeSeconds = 0;
+		COTFS.PollNextTimeIdleSeconds = 0;
+	}
+}
+
+void UCookOnTheFlyServer::FPollable::RunNow(UCookOnTheFlyServer& COTFS)
+{
+	FScopeLock PollablesScopeLock(&COTFS.PollablesLock);
+
+	UE::Cook::FTickStackData StackData(MAX_flt, COTFS.IsRealtimeMode(), ECookTickFlags::None, MAX_int32);
+	PollFunction(StackData);
+
+	FPollableQueueKey* KeyInQueue = COTFS.Pollables.FindByPredicate(
+		[this](const FPollableQueueKey& Existing) { return Existing.Pollable.GetReference() == this; });
+	if (ensure(KeyInQueue))
+	{
+		FPollableQueueKey LocalQueueKey;
+		LocalQueueKey.Pollable = MoveTemp(KeyInQueue->Pollable);
+		double CurrentTime = FPlatformTime::Seconds();
+		LocalQueueKey.NextTimeSeconds = CurrentTime + this->PeriodSeconds;
+		this->NextTimeIdleSeconds = CurrentTime + this->PeriodIdleSeconds;
+
+		int32 Index = KeyInQueue - COTFS.Pollables.GetData();
+		COTFS.Pollables.HeapRemoveAt(Index, false /* bAllowShrinking */);
+		COTFS.Pollables.HeapPush(MoveTemp(LocalQueueKey));
+		COTFS.PollNextTimeSeconds = FMath::Min(LocalQueueKey.NextTimeSeconds, COTFS.PollNextTimeSeconds);
+		COTFS.PollNextTimeIdleSeconds = FMath::Min(LocalQueueKey.Pollable->NextTimeIdleSeconds, COTFS.PollNextTimeIdleSeconds);
+	}
+}
+
+void UCookOnTheFlyServer::FPollable::RunDuringPump(UE::Cook::FTickStackData& StackData, double& OutNewCurrentTime, double& OutNextTimeSeconds)
+{
+	PollFunction(StackData);
+	OutNewCurrentTime = FPlatformTime::Seconds();
+	OutNextTimeSeconds = OutNewCurrentTime + PeriodSeconds;
+	NextTimeIdleSeconds = OutNewCurrentTime + PeriodIdleSeconds;
 }
 
 void UCookOnTheFlyServer::PumpPollables(UE::Cook::FTickStackData& StackData, bool bIsIdle)
 {
 	UE_SCOPED_HIERARCHICAL_COOKTIMER(PumpPollables);
-	PollNextTimeSeconds = MAX_flt;
-	PollNextTimeIdleSeconds = MAX_flt;
+	FScopeLock PollablesScopeLock(&PollablesLock);
+
 	int32 NumPollables = Pollables.Num();
 	if (NumPollables == 0)
 	{
+		PollNextTimeSeconds = MAX_flt;
+		PollNextTimeIdleSeconds = MAX_flt;
 		return;
 	}
 
 	double CurrentTime = FPlatformTime::Seconds();
-	double MaxTime = CurrentTime + PumpPollablesTimeSlice;
-
-	PollStartIndex = PollStartIndex % NumPollables;
-	int32 PollEndIndex = PollStartIndex;
-	do
+	if (!bIsIdle)
 	{
-		FPollable& Pollable = Pollables[PollStartIndex++];
-		PollStartIndex = PollStartIndex >= NumPollables ? 0 : PollStartIndex;
-
-		double NextTimeSeconds = bIsIdle ? Pollable.NextTimeIdleSeconds : Pollable.NextTimeSeconds;
-		if (NextTimeSeconds <= CurrentTime)
+		// To avoid an infinite loop, we keep the popped pollables in a separate list to readd afterwards
+		// rather than readding them as soon as we know their new time
+		TArray<FPollableQueueKey, TInlineAllocator<UE::Cook::Pollable::ExpectedMaxNum>> PoppedQueueKeys;
+		while (!Pollables.IsEmpty() && Pollables.HeapTop().NextTimeSeconds <= CurrentTime)
 		{
-			Pollable.PollFunction(StackData);
-			double NewCurrentTime = FPlatformTime::Seconds();
-			Pollable.NextTimeSeconds = NewCurrentTime + Pollable.PeriodSeconds;
-			Pollable.NextTimeIdleSeconds = NewCurrentTime + Pollable.PeriodIdleSeconds;
-			CurrentTime = NewCurrentTime;
+			FPollableQueueKey QueueKey;
+			Pollables.HeapPop(QueueKey, false /* bAllowShrinking */);
+			QueueKey.Pollable->RunDuringPump(StackData, CurrentTime, QueueKey.NextTimeSeconds);
+			PoppedQueueKeys.Add(MoveTemp(QueueKey));
+			if (StackData.Timer.IsTimeUp(CurrentTime))
+			{
+				break;
+			}
 		}
-		PollNextTimeSeconds = FMath::Min(Pollable.NextTimeSeconds, PollNextTimeSeconds);
-		PollNextTimeIdleSeconds = FMath::Min(Pollable.NextTimeIdleSeconds, PollNextTimeIdleSeconds);
-		if (CurrentTime >= MaxTime)
+		for (FPollableQueueKey& QueueKey: PoppedQueueKeys)
 		{
-			break;
+			Pollables.HeapPush(MoveTemp(QueueKey));
 		}
-	} while (PollStartIndex != PollEndIndex);
-
-	// If we early exited, finish calculating PollNextTimeSeconds from the remaining members we didn't reach
-	for (int32 Index = PollStartIndex; Index != PollEndIndex; Index = (Index+1) % NumPollables)
+		PollNextTimeSeconds = Pollables.HeapTop().NextTimeSeconds;
+		// We don't know the real value of PollNextTimeIdleSeconds because we didn't look at the entire heap.
+		// Mark that it needs to run next time we're idle, which will also make it recalculate PollNextTimeIdleSeconds
+		PollNextTimeIdleSeconds = 0;
+	}
+	else
 	{
-		FPollable& Pollable = Pollables[Index];
-		PollNextTimeSeconds = FMath::Min(Pollable.NextTimeSeconds, PollNextTimeSeconds);
-		PollNextTimeIdleSeconds = FMath::Min(Pollable.NextTimeIdleSeconds, PollNextTimeIdleSeconds);
+		// Since Idle times are not heap sorted, we have to look at all elements in the heap.
+		bool bUpdated = false;
+		PollNextTimeSeconds = MAX_flt;
+		PollNextTimeIdleSeconds = MAX_flt;
+		int32 PollIndex = 0;
+		for (; PollIndex < NumPollables; ++PollIndex)
+		{
+			FPollableQueueKey& QueueKey= Pollables[PollIndex];
+			if (QueueKey.Pollable->NextTimeIdleSeconds <= CurrentTime)
+			{
+				QueueKey.Pollable->RunDuringPump(StackData, CurrentTime, QueueKey.NextTimeSeconds);
+				bUpdated = true;
+			}
+			PollNextTimeSeconds = FMath::Min(QueueKey.NextTimeSeconds, PollNextTimeSeconds);
+			PollNextTimeIdleSeconds = FMath::Min(QueueKey.Pollable->NextTimeIdleSeconds, PollNextTimeIdleSeconds);
+			if (StackData.Timer.IsTimeUp(CurrentTime))
+			{
+				break;
+			}
+		}
+		// If we early exited, finish calculating PollNextTimeSeconds from the remaining members we didn't reach
+		for (;PollIndex < NumPollables; ++PollIndex)
+		{
+			FPollableQueueKey& QueueKey = Pollables[PollIndex];
+			PollNextTimeSeconds = FMath::Min(QueueKey.NextTimeSeconds, PollNextTimeSeconds);
+			PollNextTimeIdleSeconds = FMath::Min(QueueKey.Pollable->NextTimeIdleSeconds, PollNextTimeIdleSeconds);
+		}
+		if (bUpdated)
+		{
+			Pollables.Heapify();
+		}
 	}
-	PollNextTimeSeconds = FMath::Max(PollNextTimeSeconds, CurrentTime + PumpPollablesMinPeriod);
-	// PollNextTimeIdleSeconds does not have a min period
 }
 
 void UCookOnTheFlyServer::PollFlushRenderingCommands()
@@ -1676,21 +1803,21 @@ void UCookOnTheFlyServer::PollFlushRenderingCommands()
 	FlushRenderingCommands();
 }
 
-TOptional<UCookOnTheFlyServer::FPollable> UCookOnTheFlyServer::CreatePollableLLM()
+TRefCountPtr<UCookOnTheFlyServer::FPollable> UCookOnTheFlyServer::CreatePollableLLM()
 {
 #if ENABLE_LOW_LEVEL_MEM_TRACKER
 	if (FLowLevelMemTracker::Get().IsEnabled())
 	{
 		float PeriodSeconds = 120.0f;
 		FParse::Value(FCommandLine::Get(), TEXT("-CookLLMPeriod="), PeriodSeconds);
-		return TOptional<FPollable>(FPollable(PeriodSeconds, PeriodSeconds,
+		return TRefCountPtr<FPollable>(new FPollable(TEXT("LLM"), PeriodSeconds, PeriodSeconds,
 			[](UE::Cook::FTickStackData&) { FLowLevelMemTracker::Get().UpdateStatsPerFrame(); }));
 	}
 #endif
-	return TOptional<FPollable>();
+	return TRefCountPtr<FPollable>();
 }
 
-TOptional<UCookOnTheFlyServer::FPollable> UCookOnTheFlyServer::CreatePollableTriggerGC()
+TRefCountPtr<UCookOnTheFlyServer::FPollable> UCookOnTheFlyServer::CreatePollableTriggerGC()
 {
 	bool bTestCook = IsCookFlagSet(ECookInitializationFlags::TestCook);
 	if (PackagesPerGC > 0 || bTestCook)
@@ -1708,10 +1835,10 @@ TOptional<UCookOnTheFlyServer::FPollable> UCookOnTheFlyServer::CreatePollableTri
 			NumPackagesForTimeEstimate = FMath::Min(NumPackagesForTimeEstimate, PackagesPerGC);
 		}
 		float PeriodSeconds = FMath::Min(NumPackagesForTimeEstimate * .01f, 10.f);
-		return TOptional<FPollable>(FPollable(PeriodSeconds, PeriodSeconds,
+		return TRefCountPtr<FPollable>(new FPollable(TEXT("TimeForGC"), PeriodSeconds, PeriodSeconds,
 			[this](UE::Cook::FTickStackData& StackData) { PollPackagesPerGC(StackData); }));
 	}
-	return TOptional<FPollable>();
+	return TRefCountPtr<FPollable>();
 }
 
 void UCookOnTheFlyServer::PollPackagesPerGC(UE::Cook::FTickStackData& StackData)
@@ -1781,32 +1908,24 @@ void UCookOnTheFlyServer::InitializePollables()
 
 	if (IsCookByTheBookMode() && !IsCookingInEditor())
 	{
-		Pollables.Add(FPollable(60.f, 5.f, [this](FTickStackData&) { PollFlushRenderingCommands(); }));
-		if (TOptional<FPollable> Pollable = CreatePollableLLM())
+		Pollables.Emplace(new FPollable(TEXT("FlushRenderingCommands"), 60.f, 5.f, [this](FTickStackData&) { PollFlushRenderingCommands(); }));
+		if (TRefCountPtr<FPollable> Pollable = CreatePollableLLM())
 		{
-			Pollables.Add(MoveTemp(*Pollable));
+			Pollables.Emplace(MoveTemp(Pollable));
 		}
-		if (TOptional<FPollable> Pollable = CreatePollableTriggerGC())
+		if (TRefCountPtr<FPollable> Pollable = CreatePollableTriggerGC())
 		{
-			Pollables.Add(MoveTemp(*Pollable));
+			Pollables.Emplace(MoveTemp(Pollable));
 		}
-		Pollables.Add(FPollable(60.f, 5.f, [this](FTickStackData&) { CBTBProcessDeferredCommands(); }));
-		Pollables.Add(FPollable(60.f, 5.f, [this](FTickStackData&) { CBTBTickCommandletStats(); }));
-		Pollables.Add(FPollable(60.f, 5.f, [this](FTickStackData&) { CBTBTickShaderCompilingManager(); }));
-		Pollables.Add(FPollable(60.f, 5.f, [this](FTickStackData&) { CBTBTickAssetRegistry(); }));
+		Pollables.Emplace(new FPollable(TEXT("ProcessDeferredCommands"), 60.f, 5.f, [this](FTickStackData&) { CBTBProcessDeferredCommands(); }));
+		Pollables.Emplace(new FPollable(TEXT("CommandletStats"), 60.f, 5.f, [this](FTickStackData&) { CBTBTickCommandletStats(); }));
+		Pollables.Emplace(new FPollable(TEXT("ShaderCompilingManager"), 60.f, 5.f, [this](FTickStackData&) { CBTBTickShaderCompilingManager(); }));
+		Pollables.Emplace(new FPollable(TEXT("AssetRegistry"), 60.f, 5.f, [this](FTickStackData&) { CBTBTickAssetRegistry(); }));
 	}
+	Pollables.Heapify();
 
-	PollStartIndex = 0;
-	if (Pollables.Num())
-	{
-		PollNextTimeSeconds = 0.;
-		PollNextTimeIdleSeconds = 0.;
-	}
-	else
-	{
-		PollNextTimeSeconds = MAX_flt;
-		PollNextTimeIdleSeconds = MAX_flt;
-	}
+	PollNextTimeSeconds = 0.;
+	PollNextTimeIdleSeconds = 0.;
 }
 
 void UCookOnTheFlyServer::WaitForAsync(UE::Cook::FTickStackData& StackData)
@@ -4733,8 +4852,6 @@ void UCookOnTheFlyServer::Initialize( ECookMode::Type DesiredCookMode, ECookInit
 		FLinkerLoad::SetPreloadingEnabled(true);
 	}
 
-	PumpPollablesTimeSlice = 1.0f;
-	PumpPollablesMinPeriod = 0.1f;
 	PollNextTimeSeconds = MAX_flt;
 	PollNextTimeIdleSeconds = MAX_flt;
 	WaitForAsyncSleepSeconds = 1.0f;
