Deprecating ANY_PACKAGE.

diff --git a/Engine/Source/Editor/MovieSceneTools/Private/TrackEditors/CinematicShotTrackEditor.cpp b/Engine/Source/Editor/MovieSceneTools/Private/TrackEditors/CinematicShotTrackEditor.cpp
--- a/Engine/Source/Editor/MovieSceneTools/Private/TrackEditors/CinematicShotTrackEditor.cpp
+++ b/Engine/Source/Editor/MovieSceneTools/Private/TrackEditors/CinematicShotTrackEditor.cpp
@@ -676,7 +676,7 @@ TSharedRef<SWidget> FCinematicShotTrackEditor::HandleAddCinematicShotComboButton
 		AssetPickerConfig.OnAssetEnterPressed = FOnAssetEnterPressed::CreateRaw( this, &FCinematicShotTrackEditor::HandleAddCinematicShotComboButtonMenuEntryEnterPressed);
 		AssetPickerConfig.bAllowNullSelection = false;
 		AssetPickerConfig.InitialAssetViewType = EAssetViewType::List;
-		AssetPickerConfig.Filter.ClassNames.Add(TEXT("LevelSequence"));
+		AssetPickerConfig.Filter.ClassPaths.Add(FTopLevelAssetPath(TEXT("/Script/LevelSequence"), TEXT("LevelSequence")));
 		AssetPickerConfig.SaveSettingsName = TEXT("SequencerAssetPicker");
 		AssetPickerConfig.AdditionalReferencingAssets.Add(FAssetData(Sequence));
 	}
diff --git a/Engine/Source/Editor/MovieSceneTools/Private/TrackEditors/MaterialParameterCollectionTrackEditor.cpp b/Engine/Source/Editor/MovieSceneTools/Private/TrackEditors/MaterialParameterCollectionTrackEditor.cpp
--- a/Engine/Source/Editor/MovieSceneTools/Private/TrackEditors/MaterialParameterCollectionTrackEditor.cpp
+++ b/Engine/Source/Editor/MovieSceneTools/Private/TrackEditors/MaterialParameterCollectionTrackEditor.cpp
@@ -46,7 +46,7 @@ TSharedRef<SWidget> CreateAssetPicker(FOnAssetSelected OnAssetSelected, FOnAsset
 		AssetPickerConfig.bAllowNullSelection = false;
 		AssetPickerConfig.InitialAssetViewType = EAssetViewType::List;
 		AssetPickerConfig.Filter.bRecursiveClasses = true;
-		AssetPickerConfig.Filter.ClassNames.Add(UMaterialParameterCollection::StaticClass()->GetFName());
+		AssetPickerConfig.Filter.ClassPaths.Add(UMaterialParameterCollection::StaticClass()->GetClassPathName());
 		AssetPickerConfig.SaveSettingsName = TEXT("SequencerAssetPicker");
 		AssetPickerConfig.AdditionalReferencingAssets.Add(FAssetData(Sequence));
 	}
diff --git a/Engine/Source/Editor/MovieSceneTools/Private/TrackEditors/SkeletalAnimationTrackEditor.cpp b/Engine/Source/Editor/MovieSceneTools/Private/TrackEditors/SkeletalAnimationTrackEditor.cpp
--- a/Engine/Source/Editor/MovieSceneTools/Private/TrackEditors/SkeletalAnimationTrackEditor.cpp
+++ b/Engine/Source/Editor/MovieSceneTools/Private/TrackEditors/SkeletalAnimationTrackEditor.cpp
@@ -1646,7 +1646,7 @@ void FSkeletalAnimationTrackEditor::BuildObjectBindingTrackMenu(FMenuBuilder& Me
 
 			// Collect a full list of assets with the specified class
 			TArray<FAssetData> AssetDataList;
-			AssetRegistryModule.Get().GetAssetsByClass(UAnimSequenceBase::StaticClass()->GetFName(), AssetDataList, true);
+			AssetRegistryModule.Get().GetAssetsByClass(UAnimSequenceBase::StaticClass()->GetClassPathName(), AssetDataList, true);
 
 			if (AssetDataList.Num())
 			{
@@ -1676,7 +1676,7 @@ TSharedRef<SWidget> FSkeletalAnimationTrackEditor::BuildAnimationSubMenu(FGuid O
 bool FSkeletalAnimationTrackEditor::ShouldFilterAsset(const FAssetData& AssetData)
 {
 	// we don't want montage
-	if (AssetData.AssetClass == UAnimMontage::StaticClass()->GetFName())
+	if (AssetData.AssetClassPath == UAnimMontage::StaticClass()->GetClassPathName())
 	{
 		return true;
 	}
@@ -1703,7 +1703,7 @@ void FSkeletalAnimationTrackEditor::AddAnimationSubMenu(FMenuBuilder& MenuBuilde
 		AssetPickerConfig.InitialAssetViewType = EAssetViewType::List;
 		AssetPickerConfig.OnShouldFilterAsset = FOnShouldFilterAsset::CreateRaw(this, &FSkeletalAnimationTrackEditor::ShouldFilterAsset);
 		AssetPickerConfig.Filter.bRecursiveClasses = true;
-		AssetPickerConfig.Filter.ClassNames.Add(UAnimSequenceBase::StaticClass()->GetFName());
+		AssetPickerConfig.Filter.ClassPaths.Add(UAnimSequenceBase::StaticClass()->GetClassPathName());
 		AssetPickerConfig.OnShouldFilterAsset.BindRaw(this, &FSkeletalAnimationTrackEditor::FilterAnimSequences, Skeleton);
 		AssetPickerConfig.SaveSettingsName = TEXT("SequencerAssetPicker");
 		AssetPickerConfig.AdditionalReferencingAssets.Add(FAssetData(Sequence));
diff --git a/Engine/Source/Editor/MovieSceneTools/Private/TrackEditors/SubTrackEditor.cpp b/Engine/Source/Editor/MovieSceneTools/Private/TrackEditors/SubTrackEditor.cpp
--- a/Engine/Source/Editor/MovieSceneTools/Private/TrackEditors/SubTrackEditor.cpp
+++ b/Engine/Source/Editor/MovieSceneTools/Private/TrackEditors/SubTrackEditor.cpp
@@ -495,7 +495,7 @@ TSharedRef<SWidget> FSubTrackEditor::HandleAddSubSequenceComboButtonGetMenuConte
 			AssetPickerConfig.OnAssetEnterPressed = FOnAssetEnterPressed::CreateRaw( this, &FSubTrackEditor::HandleAddSubSequenceComboButtonMenuEntryEnterPressed, InTrack);
 			AssetPickerConfig.bAllowNullSelection = false;
 			AssetPickerConfig.InitialAssetViewType = EAssetViewType::List;
-			AssetPickerConfig.Filter.ClassNames.Add(TEXT("LevelSequence"));
+			AssetPickerConfig.Filter.ClassPaths.Add(FTopLevelAssetPath(TEXT("/Script/LevelSequence"), TEXT("LevelSequence")));
 			AssetPickerConfig.SaveSettingsName = TEXT("SequencerAssetPicker");
 			AssetPickerConfig.AdditionalReferencingAssets.Add(FAssetData(Sequence));
 		}
diff --git a/Engine/Source/Editor/MovieSceneTools/Private/TrackEditors/TransformTrackEditor.cpp b/Engine/Source/Editor/MovieSceneTools/Private/TrackEditors/TransformTrackEditor.cpp
--- a/Engine/Source/Editor/MovieSceneTools/Private/TrackEditors/TransformTrackEditor.cpp
+++ b/Engine/Source/Editor/MovieSceneTools/Private/TrackEditors/TransformTrackEditor.cpp
@@ -148,7 +148,7 @@ void F3DTransformTrackEditor::BuildTrackContextMenu( FMenuBuilder& MenuBuilder,
 
 		FAssetPickerConfig AssetPickerConfig;
 		AssetPickerConfig.SelectionMode = ESelectionMode::Single;
-		AssetPickerConfig.Filter.ClassNames.Add(UAnimSequence::StaticClass()->GetFName());
+		AssetPickerConfig.Filter.ClassPaths.Add(UAnimSequence::StaticClass()->GetClassPathName());
 		AssetPickerConfig.OnAssetSelected = FOnAssetSelected::CreateStatic(&F3DTransformTrackEditor::ImportAnimSequenceTransforms, InSequencer, InTransformTrack);
 		AssetPickerConfig.OnAssetEnterPressed = FOnAssetEnterPressed::CreateStatic(&F3DTransformTrackEditor::ImportAnimSequenceTransformsEnterPressed, InSequencer, InTransformTrack);
 		AssetPickerConfig.SaveSettingsName = TEXT("SequencerAssetPicker");
diff --git a/Engine/Source/Editor/MovieSceneTools/Public/PropertyTrackEditor.h b/Engine/Source/Editor/MovieSceneTools/Public/PropertyTrackEditor.h
--- a/Engine/Source/Editor/MovieSceneTools/Public/PropertyTrackEditor.h
+++ b/Engine/Source/Editor/MovieSceneTools/Public/PropertyTrackEditor.h
@@ -193,7 +193,7 @@ private:
 		const FString& MetaSequencerTrackClass = Property->GetMetaData( TEXT( "SequencerTrackClass" ) );
 		if ( !MetaSequencerTrackClass.IsEmpty() )
 		{
-			UClass* MetaClass = FindObject<UClass>( ANY_PACKAGE, *MetaSequencerTrackClass );
+			UClass* MetaClass = UClass::TryFindTypeSlow<UClass>(MetaSequencerTrackClass);
 			if ( !MetaClass )
 			{
 				MetaClass = LoadObject<UClass>( nullptr, *MetaSequencerTrackClass );
diff --git a/Engine/Source/Editor/Persona/Private/AnimGraphNodeDetails.cpp b/Engine/Source/Editor/Persona/Private/AnimGraphNodeDetails.cpp
--- a/Engine/Source/Editor/Persona/Private/AnimGraphNodeDetails.cpp
+++ b/Engine/Source/Editor/Persona/Private/AnimGraphNodeDetails.cpp
@@ -96,7 +96,7 @@ void FAnimGraphNodeDetails::CustomizeDetails(class IDetailLayoutBuilder& DetailB
 	}
 
 	TargetSkeleton = AnimGraphNode->GetAnimBlueprint()->TargetSkeleton;
-	TargetSkeletonName = TargetSkeleton ? FString::Printf(TEXT("%s'%s'"), *TargetSkeleton->GetClass()->GetName(), *TargetSkeleton->GetPathName()) : FString(TEXT(""));
+	TargetSkeletonName = TargetSkeleton ? FObjectPropertyBase::GetExportPath(TargetSkeleton) : FString(TEXT(""));
 
 	// Get the node property
 	const FStructProperty* NodeProperty = AnimGraphNode->GetFNodeProperty();
diff --git a/Engine/Source/Editor/Persona/Private/PersonaAssetFamily.cpp b/Engine/Source/Editor/Persona/Private/PersonaAssetFamily.cpp
--- a/Engine/Source/Editor/Persona/Private/PersonaAssetFamily.cpp
+++ b/Engine/Source/Editor/Persona/Private/PersonaAssetFamily.cpp
@@ -63,14 +63,14 @@ static void FindAssets(const USkeleton* InSkeleton, TArray<FAssetData>& OutAsset
 	FAssetRegistryModule& AssetRegistryModule = FModuleManager::LoadModuleChecked<FAssetRegistryModule>("AssetRegistry");
 	FARFilter Filter;
 	Filter.bRecursiveClasses = true;
-	Filter.ClassNames.Add(AssetType::StaticClass()->GetFName());
+	Filter.ClassPaths.Add(AssetType::StaticClass()->GetClassPathName());
 	Filter.TagsAndValues.Add(SkeletonTag, FAssetData(InSkeleton).GetExportTextName());
 	
 	// Also include all compatible assets.
 	FString CompatibleTagValue;
 	for (const auto& CompatibleSkeleton : InSkeleton->GetCompatibleSkeletons())
 	{
-		CompatibleTagValue = FString::Format(TEXT("{0}'{1}'"), { *USkeleton::StaticClass()->GetName(), *CompatibleSkeleton.ToString() });
+		CompatibleTagValue = FString::Format(TEXT("{0}'{1}'"), { *USkeleton::StaticClass()->GetPathName(), *CompatibleSkeleton.ToString() });
 		Filter.TagsAndValues.Add(SkeletonTag, CompatibleTagValue);
 	}
 
@@ -147,7 +147,7 @@ FAssetData FPersonaAssetFamily::FindAssetOfType(UClass* InAssetClass) const
 				FAssetRegistryModule& AssetRegistryModule = FModuleManager::LoadModuleChecked<FAssetRegistryModule>("AssetRegistry");
 				FARFilter Filter;
 				Filter.bRecursiveClasses = true;
-				Filter.ClassNames.Add(UPhysicsAsset::StaticClass()->GetFName());
+				Filter.ClassPaths.Add(UPhysicsAsset::StaticClass()->GetClassPathName());
 				if(Mesh.IsValid())
 				{
 					Filter.TagsAndValues.Add(GET_MEMBER_NAME_CHECKED(UPhysicsAsset, PreviewSkeletalMesh), FAssetData(Mesh.Get()).ObjectPath.ToString());
@@ -192,7 +192,7 @@ void FPersonaAssetFamily::FindAssetsOfType(UClass* InAssetClass, TArray<FAssetDa
 			FAssetRegistryModule& AssetRegistryModule = FModuleManager::LoadModuleChecked<FAssetRegistryModule>("AssetRegistry");
 			FARFilter Filter;
 			Filter.bRecursiveClasses = true;
-			Filter.ClassNames.Add(UPhysicsAsset::StaticClass()->GetFName());
+			Filter.ClassPaths.Add(UPhysicsAsset::StaticClass()->GetClassPathName());
 			if(Mesh != nullptr)
 			{
 				Filter.TagsAndValues.Add(GET_MEMBER_NAME_CHECKED(UPhysicsAsset, PreviewSkeletalMesh), FAssetData(Mesh.Get()).ObjectPath.ToString());
diff --git a/Engine/Source/Editor/Persona/Private/PersonaMeshDetails.cpp b/Engine/Source/Editor/Persona/Private/PersonaMeshDetails.cpp
--- a/Engine/Source/Editor/Persona/Private/PersonaMeshDetails.cpp
+++ b/Engine/Source/Editor/Persona/Private/PersonaMeshDetails.cpp
@@ -457,15 +457,15 @@ void FSkeletalMeshReductionSettingsLayout::GenerateChildContent(IDetailChildrenB
 	//Get the Enums
 	if (EnumReductionMethod == nullptr)
 	{
-		EnumReductionMethod = FindObject<UEnum>(ANY_PACKAGE, TEXT("SkeletalMeshOptimizationType"), true);
+		EnumReductionMethod = FindObject<UEnum>(nullptr, TEXT("/Script/Engine.SkeletalMeshOptimizationType"), true);
 	}
 	if (EnumImportance == nullptr)
 	{
-		EnumImportance = FindObject<UEnum>(ANY_PACKAGE, TEXT("SkeletalMeshOptimizationImportance"), true);
+		EnumImportance = FindObject<UEnum>(nullptr, TEXT("/Script/Engine.SkeletalMeshOptimizationImportance"), true);
 	}
 	if (EnumTerminationCriterion == nullptr)
 	{
-		EnumTerminationCriterion = FindObject<UEnum>(ANY_PACKAGE, TEXT("SkeletalMeshTerminationCriterion"), true);
+		EnumTerminationCriterion = FindObject<UEnum>(nullptr, TEXT("/Script/Engine.SkeletalMeshTerminationCriterion"), true);
 	}
 
 	bool bUseThirdPartyUI = !UseNativeReductionTool();
@@ -3374,7 +3374,7 @@ FReply FPersonaMeshDetails::OnSaveLODSettings()
 		SaveAssetDialogConfig.DefaultPath = DefaultPath;
 		SaveAssetDialogConfig.DefaultAssetName = DefaultName;
 		SaveAssetDialogConfig.ExistingAssetPolicy = ESaveAssetDialogExistingAssetPolicy::AllowButWarn;
-		SaveAssetDialogConfig.AssetClassNames.Add(USkeletalMeshLODSettings::StaticClass()->GetFName());
+		SaveAssetDialogConfig.AssetClassNames.Add(USkeletalMeshLODSettings::StaticClass()->GetClassPathName());
 
 		FContentBrowserModule& ContentBrowserModule = FModuleManager::LoadModuleChecked<FContentBrowserModule>("ContentBrowser");
 		FString SaveObjectPath = ContentBrowserModule.Get().CreateModalSaveAssetDialog(SaveAssetDialogConfig);
diff --git a/Engine/Source/Editor/Persona/Private/PersonaModule.cpp b/Engine/Source/Editor/Persona/Private/PersonaModule.cpp
--- a/Engine/Source/Editor/Persona/Private/PersonaModule.cpp
+++ b/Engine/Source/Editor/Persona/Private/PersonaModule.cpp
@@ -512,12 +512,12 @@ void FPersonaModule::ImportNewAsset(USkeleton* InSkeleton, EFBXImportType Defaul
 	}
 }
 
-void PopulateWithAssets(FName ClassName, FName SkeletonMemberName, const FString& SkeletonString, TArray<FAssetData>& OutAssets)
+void PopulateWithAssets(const FTopLevelAssetPath& ClassName, FName SkeletonMemberName, const FString& SkeletonString, TArray<FAssetData>& OutAssets)
 {
 	FAssetRegistryModule& AssetRegistryModule = FModuleManager::LoadModuleChecked<FAssetRegistryModule>(TEXT("AssetRegistry"));
 
 	FARFilter Filter;
-	Filter.ClassNames.Add(ClassName);
+	Filter.ClassPaths.Add(ClassName);
 	Filter.TagsAndValues.Add(SkeletonMemberName, SkeletonString);
 
 	AssetRegistryModule.Get().GetAssets(Filter, OutAssets);
@@ -532,9 +532,9 @@ void FPersonaModule::TestSkeletonCurveNamesForUse(const TSharedRef<IEditableSkel
 		const FString SkeletonString = FAssetData(&Skeleton).GetExportTextName();
 
 		TArray<FAssetData> SkeletalMeshes;
-		PopulateWithAssets(USkeletalMesh::StaticClass()->GetFName(), USkeletalMesh::GetSkeletonMemberName(), SkeletonString, SkeletalMeshes);
+		PopulateWithAssets(USkeletalMesh::StaticClass()->GetClassPathName(), USkeletalMesh::GetSkeletonMemberName(), SkeletonString, SkeletalMeshes);
 		TArray<FAssetData> Animations;
-		PopulateWithAssets(UAnimSequence::StaticClass()->GetFName(), USkeletalMesh::GetSkeletonMemberName(), SkeletonString, Animations);
+		PopulateWithAssets(UAnimSequence::StaticClass()->GetClassPathName(), USkeletalMesh::GetSkeletonMemberName(), SkeletonString, Animations);
 
 		FText TimeTakenMessage = FText::Format(LOCTEXT("TimeTakenWarning", "In order to verify curve usage all Skeletal Meshes and Animations that use this skeleton will be loaded, this may take some time.\n\nProceed?\n\nNumber of Meshes: {0}\nNumber of Animations: {1}"), FText::AsNumber(SkeletalMeshes.Num()), FText::AsNumber(Animations.Num()));
 
@@ -945,7 +945,7 @@ void FPersonaModule::AddCommonToolbarExtensions(FToolBarBuilder& InToolbarBuilde
 				AssetPickerConfig.bAllowNullSelection = false;
 				AssetPickerConfig.InitialAssetViewType = EAssetViewType::List;
 				AssetPickerConfig.Filter.bRecursiveClasses = false;
-				AssetPickerConfig.Filter.ClassNames.Add(USkeletalMesh::StaticClass()->GetFName());
+				AssetPickerConfig.Filter.ClassPaths.Add(USkeletalMesh::StaticClass()->GetClassPathName());
 				AssetPickerConfig.OnShouldFilterAsset = FOnShouldFilterAsset::CreateLambda([WeakPersonaToolkit](const FAssetData& AssetData)
 				{
 					if (WeakPersonaToolkit.IsValid())
@@ -1022,7 +1022,7 @@ void FPersonaModule::AddCommonToolbarExtensions(FToolBarBuilder& InToolbarBuilde
 				AssetPickerConfig.bAllowNullSelection = false;
 				AssetPickerConfig.InitialAssetViewType = EAssetViewType::List;
 				AssetPickerConfig.Filter.bRecursiveClasses = true;
-				AssetPickerConfig.Filter.ClassNames.Add(UAnimationAsset::StaticClass()->GetFName());
+				AssetPickerConfig.Filter.ClassPaths.Add(UAnimationAsset::StaticClass()->GetClassPathName());
 				AssetPickerConfig.OnShouldFilterAsset = FOnShouldFilterAsset::CreateLambda([WeakPersonaToolkit](const FAssetData& AssetData)
 				{
 					if (WeakPersonaToolkit.IsValid())
@@ -1307,7 +1307,7 @@ void FPersonaModule::FillInsertPoseMenu(FMenuBuilder& MenuBuilder, TWeakPtr<IPer
 	USkeleton* Skeleton = PersonaToolkit->GetSkeleton();
 
 	/** The asset picker will only show skeletons */
-	AssetPickerConfig.Filter.ClassNames.Add(*UPoseAsset::StaticClass()->GetName());
+	AssetPickerConfig.Filter.ClassPaths.Add(UPoseAsset::StaticClass()->GetClassPathName());
 	AssetPickerConfig.Filter.bRecursiveClasses = false;
 	AssetPickerConfig.bAllowNullSelection = false;
 	AssetPickerConfig.Filter.TagsAndValues.Add(TEXT("Skeleton"), FAssetData(Skeleton).GetExportTextName());
diff --git a/Engine/Source/Editor/Persona/Private/PreviewSceneCustomizations.cpp b/Engine/Source/Editor/Persona/Private/PreviewSceneCustomizations.cpp
--- a/Engine/Source/Editor/Persona/Private/PreviewSceneCustomizations.cpp
+++ b/Engine/Source/Editor/Persona/Private/PreviewSceneCustomizations.cpp
@@ -35,7 +35,7 @@
 #define LOCTEXT_NAMESPACE "PreviewSceneCustomizations"
 
 // static list that contains available classes, so that we can only allow these classes
-TArray<FName> FPreviewSceneDescriptionCustomization::AvailableClassNameList;
+TArray<FTopLevelAssetPath> FPreviewSceneDescriptionCustomization::AvailableClassNameList;
 
 FPreviewSceneDescriptionCustomization::FPreviewSceneDescriptionCustomization(const FString& InSkeletonName, const TSharedRef<class IPersonaToolkit>& InPersonaToolkit)
 	: SkeletonName(InSkeletonName)
@@ -54,7 +54,7 @@ FPreviewSceneDescriptionCustomization::FPreviewSceneDescriptionCustomization(con
 		{
 			if (ClassIt->IsChildOf(UDataAsset::StaticClass()) && ClassIt->ImplementsInterface(UPreviewCollectionInterface::StaticClass()))
 			{
-				AvailableClassNameList.Add(ClassIt->GetFName());
+				AvailableClassNameList.Add(ClassIt->GetClassPathName());
 			}
 		}
 	}
@@ -523,9 +523,9 @@ bool FPreviewSceneDescriptionCustomization::HandleShouldFilterAdditionalMesh(con
 	bool bValidClass = false;
 
 	// first to see if it's allowed class
-	for (FName& ClassName: AvailableClassNameList)
+	for (FTopLevelAssetPath ClassName: AvailableClassNameList)
 	{
-		if (ClassName == InAssetData.AssetClass)
+		if (ClassName == InAssetData.AssetClassPath)
 		{
 			bValidClass = true;
 			break;
diff --git a/Engine/Source/Editor/Persona/Private/PreviewSceneCustomizations.h b/Engine/Source/Editor/Persona/Private/PreviewSceneCustomizations.h
--- a/Engine/Source/Editor/Persona/Private/PreviewSceneCustomizations.h
+++ b/Engine/Source/Editor/Persona/Private/PreviewSceneCustomizations.h
@@ -119,7 +119,7 @@ private:
 	TArray<TSharedPtr<FPersonaModeComboEntry>> ControllerItems;
 
 	/** This is list of class available to filter asset by. This list doesn't change once loaded, so only collect once */
-	static TArray<FName> AvailableClassNameList;
+	static TArray<FTopLevelAssetPath> AvailableClassNameList;
 
 #if CHAOS_SIMULATION_DETAIL_VIEW_FACTORY_SELECTOR
 	/** List of available cloth simulation factories. */
diff --git a/Engine/Source/Editor/Persona/Private/SAnimCurvePicker.cpp b/Engine/Source/Editor/Persona/Private/SAnimCurvePicker.cpp
--- a/Engine/Source/Editor/Persona/Private/SAnimCurvePicker.cpp
+++ b/Engine/Source/Editor/Persona/Private/SAnimCurvePicker.cpp
@@ -73,7 +73,7 @@ void SAnimCurvePicker::RefreshListItems()
 
 	FARFilter Filter;
 	Filter.bRecursiveClasses = true;
-	Filter.ClassNames.Add(UAnimationAsset::StaticClass()->GetFName());
+	Filter.ClassPaths.Add(UAnimationAsset::StaticClass()->GetClassPathName());
 	Filter.TagsAndValues.Add(TEXT("Skeleton"), FAssetData(&EditableSkeleton.Pin()->GetSkeleton()).GetExportTextName());
 
 	TArray<FAssetData> FoundAssetData;
diff --git a/Engine/Source/Editor/Persona/Private/SAnimNotifyPanel.cpp b/Engine/Source/Editor/Persona/Private/SAnimNotifyPanel.cpp
--- a/Engine/Source/Editor/Persona/Private/SAnimNotifyPanel.cpp
+++ b/Engine/Source/Editor/Persona/Private/SAnimNotifyPanel.cpp
@@ -4804,7 +4804,7 @@ void SAnimNotifyPanel::OnGetNotifyBlueprintData(TArray<FAssetData>& OutNotifyDat
 	FAssetRegistryModule& AssetRegistryModule = FModuleManager::LoadModuleChecked<FAssetRegistryModule>(TEXT("AssetRegistry"));
 
 	// Collect a full list of assets with the specified class
-	AssetRegistryModule.Get().GetAssetsByClass(UBlueprint::StaticClass()->GetFName(), AssetDataList);
+	AssetRegistryModule.Get().GetAssetsByClass(UBlueprint::StaticClass()->GetClassPathName(), AssetDataList);
 
 
 	int32 BeginClassCount = InOutAllowedClassNames->Num();
@@ -4880,7 +4880,7 @@ void SAnimNotifyPanel::OnGetNativeNotifyData(TArray<UClass*>& OutClasses, UClass
 		{
 			OutClasses.Add(Class);
 			// Form class name to search later
-			FString ClassName = FString::Printf(TEXT("%s'%s'"), *Class->GetClass()->GetName(), *Class->GetPathName());
+			FString ClassName = FObjectPropertyBase::GetExportPath(Class);
 			OutAllowedBlueprintClassNames->AddUnique(ClassName);
 		}
 	}
diff --git a/Engine/Source/Editor/Persona/Private/SAnimSegmentsPanel.cpp b/Engine/Source/Editor/Persona/Private/SAnimSegmentsPanel.cpp
--- a/Engine/Source/Editor/Persona/Private/SAnimSegmentsPanel.cpp
+++ b/Engine/Source/Editor/Persona/Private/SAnimSegmentsPanel.cpp
@@ -409,7 +409,7 @@ void SAnimSegmentsPanel::FillSubMenu(FMenuBuilder& MenuBuilder, int32 AnimSegmen
 			FAssetPickerConfig AssetPickerConfig;
 
 			/** The asset picker will only show skeletons */
-			AssetPickerConfig.Filter.ClassNames.Add(*OldSequenceBase->GetClass()->GetName());
+			AssetPickerConfig.Filter.ClassPaths.Add(OldSequenceBase->GetClass()->GetClassPathName());
 			AssetPickerConfig.Filter.bRecursiveClasses = false;
 			AssetPickerConfig.bAllowNullSelection = false;
 
diff --git a/Engine/Source/Editor/Persona/Private/SAnimViewportToolBar.cpp b/Engine/Source/Editor/Persona/Private/SAnimViewportToolBar.cpp
--- a/Engine/Source/Editor/Persona/Private/SAnimViewportToolBar.cpp
+++ b/Engine/Source/Editor/Persona/Private/SAnimViewportToolBar.cpp
@@ -803,7 +803,7 @@ void SAnimViewportToolBar::FillCharacterMirrorMenu(FMenuBuilder& MenuBuilder) co
 	{
 		USkeleton* Skeleton = Mesh->GetSkeleton();
 	
-		AssetPickerConfig.Filter.ClassNames.Add(*UMirrorDataTable::StaticClass()->GetName());
+		AssetPickerConfig.Filter.ClassPaths.Add(UMirrorDataTable::StaticClass()->GetClassPathName());
 		AssetPickerConfig.Filter.bRecursiveClasses = false;
 		AssetPickerConfig.bAllowNullSelection = true;
 		AssetPickerConfig.Filter.TagsAndValues.Add(TEXT("Skeleton"), FAssetData(Skeleton).GetExportTextName());
diff --git a/Engine/Source/Editor/Persona/Private/SAnimationSequenceBrowser.cpp b/Engine/Source/Editor/Persona/Private/SAnimationSequenceBrowser.cpp
--- a/Engine/Source/Editor/Persona/Private/SAnimationSequenceBrowser.cpp
+++ b/Engine/Source/Editor/Persona/Private/SAnimationSequenceBrowser.cpp
@@ -847,8 +847,8 @@ void SAnimationSequenceBrowser::Construct(const FArguments& InArgs, const TShare
 
 	// Configure filter for asset picker
 	Filter.bRecursiveClasses = true;
-	Filter.ClassNames.Add(UAnimationAsset::StaticClass()->GetFName());
-	Filter.ClassNames.Add(USoundWave::StaticClass()->GetFName());
+	Filter.ClassPaths.Add(UAnimationAsset::StaticClass()->GetClassPathName());
+	Filter.ClassPaths.Add(USoundWave::StaticClass()->GetClassPathName());
 
 	FAssetPickerConfig Config;
 	Config.Filter = Filter;
@@ -1245,7 +1245,7 @@ TSharedRef<SToolTip> SAnimationSequenceBrowser::CreateCustomAssetToolTip(FAssetD
 {
 	// Make a list of tags to show
 	TArray<UObject::FAssetRegistryTag> Tags;
-	UClass* AssetClass = FindObject<UClass>(ANY_PACKAGE, *AssetData.AssetClass.ToString());
+	UClass* AssetClass = FindObject<UClass>(AssetData.AssetClassPath);
 	check(AssetClass);
 	AssetClass->GetDefaultObject()->GetAssetRegistryTags(Tags);
 
@@ -1438,7 +1438,7 @@ bool SAnimationSequenceBrowser::OnVisualizeAssetToolTip(const TSharedPtr<SWidget
 {
 	// Resolve the asset
 	USkeletalMesh* MeshToUse = nullptr;
-	UClass* AssetClass = FindObject<UClass>(ANY_PACKAGE, *AssetData.AssetClass.ToString());
+	UClass* AssetClass = FindObject<UClass>(AssetData.AssetClassPath);
 	if(AssetClass->IsChildOf(UAnimationAsset::StaticClass()) && AssetData.IsAssetLoaded() && AssetData.GetAsset())
 	{
 		// Set up the viewport to show the asset. Catching the visualize allows us to use
diff --git a/Engine/Source/Editor/Persona/Private/SAssetFamilyShortcutBar.cpp b/Engine/Source/Editor/Persona/Private/SAssetFamilyShortcutBar.cpp
--- a/Engine/Source/Editor/Persona/Private/SAssetFamilyShortcutBar.cpp
+++ b/Engine/Source/Editor/Persona/Private/SAssetFamilyShortcutBar.cpp
@@ -279,7 +279,7 @@ public:
 				UClass* FilterClass = AssetFamily->GetAssetFamilyClass(AssetData.GetClass());
 				if (FilterClass != nullptr)
 				{
-					AssetPickerConfig.Filter.ClassNames.Add(FilterClass->GetFName());
+					AssetPickerConfig.Filter.ClassPaths.Add(FilterClass->GetClassPathName());
 					AssetPickerConfig.Filter.bRecursiveClasses = true;
 				}
 
diff --git a/Engine/Source/Editor/Persona/Private/SRetargetSourceWindow.cpp b/Engine/Source/Editor/Persona/Private/SRetargetSourceWindow.cpp
--- a/Engine/Source/Editor/Persona/Private/SRetargetSourceWindow.cpp
+++ b/Engine/Source/Editor/Persona/Private/SRetargetSourceWindow.cpp
@@ -431,7 +431,7 @@ void SRetargetSourceWindow::OnAddRetargetSource()
 	FContentBrowserModule& ContentBrowserModule = FModuleManager::Get().LoadModuleChecked<FContentBrowserModule>(TEXT("ContentBrowser"));
 
 	FAssetPickerConfig AssetPickerConfig;
-	AssetPickerConfig.Filter.ClassNames.Add(USkeletalMesh::StaticClass()->GetFName());
+	AssetPickerConfig.Filter.ClassPaths.Add(USkeletalMesh::StaticClass()->GetClassPathName());
 	AssetPickerConfig.OnAssetSelected = FOnAssetSelected::CreateSP(this, &SRetargetSourceWindow::OnAssetSelectedFromMeshPicker);
 	AssetPickerConfig.bAllowNullSelection = false;
 	AssetPickerConfig.InitialAssetViewType = EAssetViewType::Tile;
@@ -674,7 +674,7 @@ FReply SCompatibleSkeletons::OnAddSkeletonClicked()
 	FContentBrowserModule& ContentBrowserModule = FModuleManager::Get().LoadModuleChecked<FContentBrowserModule>(TEXT("ContentBrowser"));
 
 	FAssetPickerConfig AssetPickerConfig;
-	AssetPickerConfig.Filter.ClassNames.Add(USkeleton::StaticClass()->GetFName());
+	AssetPickerConfig.Filter.ClassPaths.Add(USkeleton::StaticClass()->GetClassPathName());
 	AssetPickerConfig.OnAssetSelected = FOnAssetSelected::CreateSP(this, &SCompatibleSkeletons::OnAssetSelectedFromSkeletonPicker);
 	AssetPickerConfig.bAllowNullSelection = false;
 	AssetPickerConfig.InitialAssetViewType = EAssetViewType::Tile;
diff --git a/Engine/Source/Editor/Persona/Private/SRigPicker.cpp b/Engine/Source/Editor/Persona/Private/SRigPicker.cpp
--- a/Engine/Source/Editor/Persona/Private/SRigPicker.cpp
+++ b/Engine/Source/Editor/Persona/Private/SRigPicker.cpp
@@ -55,7 +55,7 @@ void SRigPicker::Construct(const FArguments& InArgs)
 		FContentBrowserModule& ContentBrowserModule = FModuleManager::Get().LoadModuleChecked<FContentBrowserModule>(TEXT("ContentBrowser"));
 
 		FAssetPickerConfig AssetPickerConfig;
-		AssetPickerConfig.Filter.ClassNames.Add(URig::StaticClass()->GetFName());
+		AssetPickerConfig.Filter.ClassPaths.Add(URig::StaticClass()->GetClassPathName());
 		// Allow child classes
 		AssetPickerConfig.Filter.bRecursiveClasses = true;
 		// Set a delegate for setting the asset from the picker
diff --git a/Engine/Source/Editor/Persona/Private/SSkeletonSlotNames.cpp b/Engine/Source/Editor/Persona/Private/SSkeletonSlotNames.cpp
--- a/Engine/Source/Editor/Persona/Private/SSkeletonSlotNames.cpp
+++ b/Engine/Source/Editor/Persona/Private/SSkeletonSlotNames.cpp
@@ -477,7 +477,7 @@ void SSkeletonSlotNames::GetCompatibleAnimBlueprints( TArray<FAssetData>& OutAss
 	FAssetRegistryModule& AssetRegistryModule = FModuleManager::LoadModuleChecked<FAssetRegistryModule>(TEXT("AssetRegistry"));
 
 	TArray<FAssetData> AssetDataList;
-	AssetRegistryModule.Get().GetAssetsByClass(UAnimBlueprint::StaticClass()->GetFName(), AssetDataList, true);
+	AssetRegistryModule.Get().GetAssetsByClass(UAnimBlueprint::StaticClass()->GetClassPathName(), AssetDataList, true);
 
 	OutAssets.Empty(AssetDataList.Num());
 
@@ -560,7 +560,7 @@ void SSkeletonSlotNames::GetCompatibleAnimMontages(TArray<struct FAssetData>& Ou
 	FAssetRegistryModule& AssetRegistryModule = FModuleManager::LoadModuleChecked<FAssetRegistryModule>(TEXT("AssetRegistry"));
 
 	TArray<FAssetData> AssetDataList;
-	AssetRegistryModule.Get().GetAssetsByClass(UAnimMontage::StaticClass()->GetFName(), AssetDataList, true);
+	AssetRegistryModule.Get().GetAssetsByClass(UAnimMontage::StaticClass()->GetClassPathName(), AssetDataList, true);
 
 	OutAssets.Empty(AssetDataList.Num());
 
diff --git a/Engine/Source/Editor/PhysicsAssetEditor/Private/PhysicsAssetEditor.cpp b/Engine/Source/Editor/PhysicsAssetEditor/Private/PhysicsAssetEditor.cpp
--- a/Engine/Source/Editor/PhysicsAssetEditor/Private/PhysicsAssetEditor.cpp
+++ b/Engine/Source/Editor/PhysicsAssetEditor/Private/PhysicsAssetEditor.cpp
@@ -1634,7 +1634,7 @@ bool FPhysicsAssetEditor::ShouldFilterAssetBasedOnSkeleton( const FAssetData& As
 		{
 			USkeleton* Skeleton = EditorSkelMesh->GetSkeleton();
 
-			if ( Skeleton && (*SkeletonName) == FString::Printf(TEXT("%s'%s'"), *Skeleton->GetClass()->GetName(), *Skeleton->GetPathName()) )
+			if ( Skeleton && (*SkeletonName) == FObjectPropertyBase::GetExportPath(Skeleton) )
 			{
 				return false;
 			}
@@ -2847,7 +2847,7 @@ TSharedRef<SWidget> FPhysicsAssetEditor::BuildStaticMeshAssetPicker()
 	FContentBrowserModule& ContentBrowserModule = FModuleManager::Get().LoadModuleChecked<FContentBrowserModule>(TEXT("ContentBrowser"));
 
 	FAssetPickerConfig AssetPickerConfig;
-	AssetPickerConfig.Filter.ClassNames.Add(UStaticMesh::StaticClass()->GetFName());
+	AssetPickerConfig.Filter.ClassPaths.Add(UStaticMesh::StaticClass()->GetClassPathName());
 	AssetPickerConfig.OnAssetSelected = FOnAssetSelected::CreateSP(this, &FPhysicsAssetEditor::OnAssetSelectedFromStaticMeshAssetPicker);
 	AssetPickerConfig.bAllowNullSelection = true;
 	AssetPickerConfig.InitialAssetViewType = EAssetViewType::List;
@@ -2933,7 +2933,7 @@ TSharedRef<SWidget> FPhysicsAssetEditor::BuildPhysicalMaterialAssetPicker(bool b
 	FContentBrowserModule& ContentBrowserModule = FModuleManager::Get().LoadModuleChecked<FContentBrowserModule>(TEXT("ContentBrowser"));
 
 	FAssetPickerConfig AssetPickerConfig;
-	AssetPickerConfig.Filter.ClassNames.Add(UPhysicalMaterial::StaticClass()->GetFName());
+	AssetPickerConfig.Filter.ClassPaths.Add(UPhysicalMaterial::StaticClass()->GetClassPathName());
 	AssetPickerConfig.OnAssetSelected = FOnAssetSelected::CreateSP(this, &FPhysicsAssetEditor::OnAssetSelectedFromPhysicalMaterialAssetPicker, bForAllBodies);
 	AssetPickerConfig.bAllowNullSelection = true;
 	AssetPickerConfig.InitialAssetViewType = EAssetViewType::List;
diff --git a/Engine/Source/Editor/PlacementMode/Private/PlacementModeModule.cpp b/Engine/Source/Editor/PlacementMode/Private/PlacementModeModule.cpp
--- a/Engine/Source/Editor/PlacementMode/Private/PlacementModeModule.cpp
+++ b/Engine/Source/Editor/PlacementMode/Private/PlacementModeModule.cpp
@@ -779,10 +779,7 @@ bool FPlacementModeModule::PassesFilters(const TSharedPtr<FPlaceableItem>& Item)
 				if (Item->AssetData.GetTagValue(FBlueprintTags::NativeParentClassPath, TagValue) && !TagValue.IsEmpty())
 				{
 					// If the native parent class can't be placed, neither can the blueprint
-
-					UObject* Outer = nullptr;
-					ResolveName(Outer, TagValue, false, false);
-					UClass* NativeParentClass = FindObject<UClass>(ANY_PACKAGE, *TagValue);
+					UClass* NativeParentClass = UClass::TryFindTypeSlow<UClass>(FPackageName::ExportTextPathToObjectPath(TagValue));
 
 					bPlaceable = AssetSelectionUtils::IsChildBlueprintPlaceable(NativeParentClass);
 				}
diff --git a/Engine/Source/Editor/PropertyEditor/Private/EditConditionContext.cpp b/Engine/Source/Editor/PropertyEditor/Private/EditConditionContext.cpp
--- a/Engine/Source/Editor/PropertyEditor/Private/EditConditionContext.cpp
+++ b/Engine/Source/Editor/PropertyEditor/Private/EditConditionContext.cpp
@@ -451,7 +451,7 @@ TOptional<FString> FEditConditionContext::GetTypeName(const FString& PropertyNam
 
 TOptional<int64> FEditConditionContext::GetIntegerValueOfEnum(const FString& EnumTypeName, const FString& MemberName) const
 {
-	const UEnum* EnumType = FindObject<UEnum>((UObject*) ANY_PACKAGE, *EnumTypeName, true);
+	const UEnum* EnumType = UClass::TryFindTypeSlow<UEnum>(EnumTypeName, EFindFirstObjectOptions::ExactClass);
 	if (EnumType == nullptr)
 	{
 		return TOptional<int64>();
diff --git a/Engine/Source/Editor/PropertyEditor/Private/Presentation/PropertyEditor/PropertyEditor.cpp b/Engine/Source/Editor/PropertyEditor/Private/Presentation/PropertyEditor/PropertyEditor.cpp
--- a/Engine/Source/Editor/PropertyEditor/Private/Presentation/PropertyEditor/PropertyEditor.cpp
+++ b/Engine/Source/Editor/PropertyEditor/Private/Presentation/PropertyEditor/PropertyEditor.cpp
@@ -325,7 +325,7 @@ void FPropertyEditor::MakeNewBlueprint()
 		{
 			if (RequiredInterface != nullptr && FKismetEditorUtilities::CanBlueprintImplementInterface(Blueprint, RequiredInterface))
 			{
-				FBlueprintEditorUtils::ImplementNewInterface(Blueprint, RequiredInterface->GetFName());
+				FBlueprintEditorUtils::ImplementNewInterface(Blueprint, RequiredInterface->GetClassPathName());
 			}
 			
 			PropertyHandle->SetValueFromFormattedString(Blueprint->GeneratedClass->GetPathName());
@@ -653,19 +653,20 @@ void FPropertyEditor::SyncToObjectsInNode( const TWeakPtr< FPropertyNode >& Weak
 		TArray<UObject*> Objects;
 		for ( int32 ObjectIndex = 0 ; ObjectIndex < ObjectNames.Num() ; ++ObjectIndex )
 		{
-
-			UObject* Package = ANY_PACKAGE;
+			UObject* Object = nullptr;
 			if( ObjectNames[ObjectIndex].Contains( TEXT(".")) )
 			{
-				// Formatted text string, use the exact path instead of any package
-				Package = NULL;
+				Object = StaticFindObject(PropertyClass, nullptr, *ObjectNames[ObjectIndex]);
+				if (!Object)
+				{
+					Object = StaticLoadObject(PropertyClass, nullptr, *ObjectNames[ObjectIndex]);
+				}
 			}
-
-			UObject* Object = StaticFindObject( PropertyClass, Package, *ObjectNames[ObjectIndex] );
-			if( !Object && Package != ANY_PACKAGE )
+			else
 			{
-				Object = StaticLoadObject(PropertyClass, Package, *ObjectNames[ObjectIndex]);
+				Object = StaticFindFirstObject(PropertyClass, *ObjectNames[ObjectIndex], EFindFirstObjectOptions::EnsureIfAmbiguous);
 			}
+
 			if ( Object )
 			{
 				// If the selected object is a blueprint generated class, then browsing to it in the content browser should instead point to the blueprint
diff --git a/Engine/Source/Editor/PropertyEditor/Private/PropertyEditorHelpers.cpp b/Engine/Source/Editor/PropertyEditor/Private/PropertyEditorHelpers.cpp
--- a/Engine/Source/Editor/PropertyEditor/Private/PropertyEditorHelpers.cpp
+++ b/Engine/Source/Editor/PropertyEditor/Private/PropertyEditorHelpers.cpp
@@ -532,7 +532,7 @@ namespace PropertyEditorHelpers
 				{
 
 					const FString& EnumName = Property->GetMetaData(TEXT("Enum"));
-					Enum = FindObject<UEnum>(ANY_PACKAGE, *EnumName, true);
+					Enum = UClass::TryFindTypeSlow<UEnum>(EnumName, EFindFirstObjectOptions::ExactClass);
 				}
 
 				if(Enum)
diff --git a/Engine/Source/Editor/PropertyEditor/Private/PropertyHandleImpl.cpp b/Engine/Source/Editor/PropertyEditor/Private/PropertyHandleImpl.cpp
--- a/Engine/Source/Editor/PropertyEditor/Private/PropertyHandleImpl.cpp
+++ b/Engine/Source/Editor/PropertyEditor/Private/PropertyHandleImpl.cpp
@@ -2813,7 +2813,7 @@ bool FPropertyHandleBase::GeneratePossibleValues(TArray< TSharedPtr<FString> >&
 	else if ( Property->IsA(FStrProperty::StaticClass()) && Property->HasMetaData( TEXT("Enum") ) )
 	{
 		const FString& EnumName = Property->GetMetaData(TEXT("Enum"));
-		Enum = FindObject<UEnum>(ANY_PACKAGE, *EnumName, true);
+		Enum = UClass::TryFindTypeSlow<UEnum>(EnumName, EFindFirstObjectOptions::ExactClass);
 		check( Enum );
 	}
 
@@ -3908,7 +3908,16 @@ FPropertyAccess::Result FPropertyHandleObject::SetValueFromFormattedString(const
 			{
 				for (const FString& ClassName : AllowedClassNames)
 				{
-					UClass* AllowedClass = FindObject<UClass>(ANY_PACKAGE, *ClassName);
+					UClass* AllowedClass = nullptr;
+					if (!FPackageName::IsShortPackageName(ClassName))
+					{
+						AllowedClass = FindObject<UClass>(nullptr, *ClassName);
+					}
+					else
+					{ 
+						AllowedClass = FindFirstObject<UClass>(*ClassName, EFindFirstObjectOptions::None, ELogVerbosity::Warning, TEXT("FPropertyHandleObject::SetValueFromFormattedString"));
+					}
+					
 					const bool bIsInterface = AllowedClass && AllowedClass->HasAnyClassFlags(CLASS_Interface);
 				
 					// Check if the object is an allowed class type this property supports
@@ -3932,7 +3941,7 @@ FPropertyAccess::Result FPropertyHandleObject::SetValueFromFormattedString(const
 
 				for (const FString& DisallowedClassName : DisallowedClassNames)
 				{
-					UClass* DisallowedClass = FindObject<UClass>(ANY_PACKAGE, *DisallowedClassName);
+					UClass* DisallowedClass = UClass::TryFindTypeSlow<UClass>(DisallowedClassName);
 					const bool bIsInterface = DisallowedClass && DisallowedClass->HasAnyClassFlags(CLASS_Interface);
 
 					if ((DisallowedClass && QualifiedClass->IsChildOf(DisallowedClass)) || (bIsInterface && QualifiedObject->GetClass()->ImplementsInterface(DisallowedClass)))
diff --git a/Engine/Source/Editor/PropertyEditor/Private/UserInterface/PropertyEditor/SPropertyAssetPicker.cpp b/Engine/Source/Editor/PropertyEditor/Private/UserInterface/PropertyEditor/SPropertyAssetPicker.cpp
--- a/Engine/Source/Editor/PropertyEditor/Private/UserInterface/PropertyEditor/SPropertyAssetPicker.cpp
+++ b/Engine/Source/Editor/PropertyEditor/Private/UserInterface/PropertyEditor/SPropertyAssetPicker.cpp
@@ -68,7 +68,7 @@ TSharedRef<SWidget> SPropertyAssetPicker::OnGenerateAssetPicker()
 	for ( auto ClassIt = AllowedClasses.CreateConstIterator(); ClassIt; ++ClassIt )
 	{
 		const UClass* Class = (*ClassIt);
-		AssetPickerConfig.Filter.ClassNames.Add( Class->GetFName() );
+		AssetPickerConfig.Filter.ClassPaths.Add( Class->GetClassPathName() );
 	}
 	// Allow child classes
 	AssetPickerConfig.Filter.bRecursiveClasses = true;
diff --git a/Engine/Source/Editor/PropertyEditor/Private/UserInterface/PropertyEditor/SPropertyEditorAsset.cpp b/Engine/Source/Editor/PropertyEditor/Private/UserInterface/PropertyEditor/SPropertyEditorAsset.cpp
--- a/Engine/Source/Editor/PropertyEditor/Private/UserInterface/PropertyEditor/SPropertyEditorAsset.cpp
+++ b/Engine/Source/Editor/PropertyEditor/Private/UserInterface/PropertyEditor/SPropertyEditorAsset.cpp
@@ -129,6 +129,16 @@ void SPropertyEditorAsset::InitializeClassFilters(const FProperty* Property)
 
 	bExactClass = GetTagOrBoolMetadata(MetadataProperty, TEXT("ExactClass"), false);
 
+	auto FindClass = [](const FString& InClassName)
+	{
+		UClass* Class = UClass::TryFindTypeSlow<UClass>(InClassName, EFindFirstObjectOptions::EnsureIfAmbiguous);
+		if (!Class)
+		{
+			Class = LoadObject<UClass>(nullptr, *InClassName);
+		}
+		return Class;
+	};
+
 	const FString AllowedClassesFilterString = MetadataProperty->GetMetaData(TEXT("AllowedClasses"));
 	if (!AllowedClassesFilterString.IsEmpty())
 	{
@@ -137,11 +147,7 @@ void SPropertyEditorAsset::InitializeClassFilters(const FProperty* Property)
 
 		for (const FString& ClassName : AllowedClassFilterNames)
 		{
-			UClass* Class = FindObject<UClass>(ANY_PACKAGE, *ClassName);
-			if (!Class)
-			{
-				Class = LoadObject<UClass>(nullptr, *ClassName);
-			}
+			UClass* Class = FindClass(ClassName);
 
 			if (Class)
 			{
@@ -179,11 +185,7 @@ void SPropertyEditorAsset::InitializeClassFilters(const FProperty* Property)
 
 		for (const FString& ClassName : DisallowedClassFilterNames)
 		{
-			UClass* Class = FindObject<UClass>(ANY_PACKAGE, *ClassName);
-			if (!Class)
-			{
-				Class = LoadObject<UClass>(nullptr, *ClassName);
-			}
+			UClass* Class = FindClass(ClassName);
 
 			if (Class)
 			{
diff --git a/Engine/Source/Editor/PropertyEditor/Private/UserInterface/PropertyEditor/SPropertyEditorClass.cpp b/Engine/Source/Editor/PropertyEditor/Private/UserInterface/PropertyEditor/SPropertyEditorClass.cpp
--- a/Engine/Source/Editor/PropertyEditor/Private/UserInterface/PropertyEditor/SPropertyEditorClass.cpp
+++ b/Engine/Source/Editor/PropertyEditor/Private/UserInterface/PropertyEditor/SPropertyEditorClass.cpp
@@ -72,7 +72,7 @@ private:
 
 static UClass* FindOrLoadClass(const FString& ClassName)
 {
-	UClass* Class = FindObject<UClass>(ANY_PACKAGE, *ClassName);
+	UClass* Class = UClass::TryFindTypeSlow<UClass>(ClassName, EFindFirstObjectOptions::EnsureIfAmbiguous);
 
 	if (!Class)
 	{
diff --git a/Engine/Source/Editor/PropertyEditor/Private/UserInterface/PropertyEditor/SPropertyEditorEditInline.cpp b/Engine/Source/Editor/PropertyEditor/Private/UserInterface/PropertyEditor/SPropertyEditorEditInline.cpp
--- a/Engine/Source/Editor/PropertyEditor/Private/UserInterface/PropertyEditor/SPropertyEditorEditInline.cpp
+++ b/Engine/Source/Editor/PropertyEditor/Private/UserInterface/PropertyEditor/SPropertyEditorEditInline.cpp
@@ -254,7 +254,7 @@ void SPropertyEditorEditInline::OnClassPicked(UClass* InClass)
 					NewObjectName = ObjectName;
 
 					ConstructorHelpers::StripObjectClass(DefaultValue);
-					NewObjectTemplate = StaticFindObject(InClass, ANY_PACKAGE, *DefaultValue);
+					NewObjectTemplate = StaticFindObject(InClass, nullptr, *DefaultValue);
 				}
 			}
 		}
@@ -331,7 +331,7 @@ void SPropertyEditorEditInline::OnClassPicked(UClass* InClass)
 				// Move the old subobject to the transient package so GetObjectsWithOuter will not return it
 				// This is particularly important for UActorComponent objects so resetting owned components on the parent doesn't find it
 				ConstructorHelpers::StripObjectClass(PrevPerObjectValues[Index]);
-				if (UObject* SubObject = StaticFindObject(UObject::StaticClass(), ANY_PACKAGE, *PrevPerObjectValues[Index]))
+				if (UObject* SubObject = StaticFindObject(UObject::StaticClass(), nullptr, *PrevPerObjectValues[Index]))
 				{
 					SubObject->Rename(nullptr, GetTransientOuterForRename(SubObject->GetClass()), REN_DontCreateRedirectors);
 				}
diff --git a/Engine/Source/Editor/PropertyEditor/Private/UserInterface/PropertyEditor/SPropertyEditorNumeric.h b/Engine/Source/Editor/PropertyEditor/Private/UserInterface/PropertyEditor/SPropertyEditorNumeric.h
--- a/Engine/Source/Editor/PropertyEditor/Private/UserInterface/PropertyEditor/SPropertyEditorNumeric.h
+++ b/Engine/Source/Editor/PropertyEditor/Private/UserInterface/PropertyEditor/SPropertyEditorNumeric.h
@@ -70,7 +70,7 @@ public:
 				if (!BitmaskEnumName.IsEmpty())
 				{
 					// @TODO: Potentially replace this with a parameter passed in from a member variable on the FProperty (e.g. FByteProperty::Enum)
-					BitmaskEnum = FindObject<UEnum>(ANY_PACKAGE, *BitmaskEnumName);
+					BitmaskEnum = UClass::TryFindTypeSlow<UEnum>(BitmaskEnumName);
 				}
 
 				if (BitmaskEnum)
diff --git a/Engine/Source/Editor/PropertyEditor/Private/UserInterface/PropertyEditor/SPropertyEditorStruct.cpp b/Engine/Source/Editor/PropertyEditor/Private/UserInterface/PropertyEditor/SPropertyEditorStruct.cpp
--- a/Engine/Source/Editor/PropertyEditor/Private/UserInterface/PropertyEditor/SPropertyEditorStruct.cpp
+++ b/Engine/Source/Editor/PropertyEditor/Private/UserInterface/PropertyEditor/SPropertyEditorStruct.cpp
@@ -81,7 +81,7 @@ void SPropertyEditorStruct::Construct(const FArguments& InArgs, const TSharedPtr
 				const FString& MetaStructName = Property->GetOwnerProperty()->GetMetaData(TEXT("MetaStruct"));
 				if (!MetaStructName.IsEmpty())
 				{
-					MetaStruct = FindObject<UScriptStruct>(ANY_PACKAGE, *MetaStructName);
+					MetaStruct = UClass::TryFindTypeSlow<UScriptStruct>(MetaStructName, EFindFirstObjectOptions::EnsureIfAmbiguous);
 					if (!MetaStruct)
 					{
 						MetaStruct = LoadObject<UScriptStruct>(nullptr, *MetaStructName);
@@ -222,7 +222,7 @@ void SPropertyEditorStruct::SendToObjects(const FString& NewValue)
 	}
 	else if (!NewValue.IsEmpty() && NewValue != TEXT("None"))
 	{
-		const UScriptStruct* NewStruct = FindObject<UScriptStruct>(ANY_PACKAGE, *NewValue);
+		const UScriptStruct* NewStruct = FindObject<UScriptStruct>(nullptr, *NewValue);
 		if (!NewStruct)
 		{
 			NewStruct = LoadObject<UScriptStruct>(nullptr, *NewValue);
diff --git a/Engine/Source/Editor/PropertyEditor/Private/UserInterface/PropertyEditor/SPropertyMenuAssetPicker.cpp b/Engine/Source/Editor/PropertyEditor/Private/UserInterface/PropertyEditor/SPropertyMenuAssetPicker.cpp
--- a/Engine/Source/Editor/PropertyEditor/Private/UserInterface/PropertyEditor/SPropertyMenuAssetPicker.cpp
+++ b/Engine/Source/Editor/PropertyEditor/Private/UserInterface/PropertyEditor/SPropertyMenuAssetPicker.cpp
@@ -113,19 +113,19 @@ void SPropertyMenuAssetPicker::Construct( const FArguments& InArgs )
 		// Add filter classes - if we have a single filter class of "Object" then don't set a filter since it would always match everything (but slower!)
 		if (AllowedClasses.Num() == 1 && AllowedClasses[0] == UObject::StaticClass())
 		{
-			AssetPickerConfig.Filter.ClassNames.Reset();
+			AssetPickerConfig.Filter.ClassPaths.Reset();
 		}
 		else
 		{
 			for(int32 i = 0; i < AllowedClasses.Num(); ++i)
 			{
-				AssetPickerConfig.Filter.ClassNames.Add( AllowedClasses[i]->GetFName() );
+				AssetPickerConfig.Filter.ClassPaths.Add( AllowedClasses[i]->GetClassPathName() );
 			}
 		}
 
 		for (int32 i = 0; i < DisallowedClasses.Num(); ++i)
 		{
-			AssetPickerConfig.Filter.RecursiveClassesExclusionSet.Add(DisallowedClasses[i]->GetFName());
+			AssetPickerConfig.Filter.RecursiveClassPathsExclusionSet.Add(DisallowedClasses[i]->GetClassPathName());
 		}
 
 		// Allow child classes
diff --git a/Engine/Source/Editor/SequenceRecorder/Private/SequenceRecorderModule.cpp b/Engine/Source/Editor/SequenceRecorder/Private/SequenceRecorderModule.cpp
--- a/Engine/Source/Editor/SequenceRecorder/Private/SequenceRecorderModule.cpp
+++ b/Engine/Source/Editor/SequenceRecorder/Private/SequenceRecorderModule.cpp
@@ -493,7 +493,7 @@ class FSequenceRecorderModule : public ISequenceRecorder, private FSelfRegisteri
 				}
 				else
 				{
-					UClass* FoundClass = FindObject<UClass>(ANY_PACKAGE, *SpecifierStr);
+					UClass* FoundClass = UClass::TryFindTypeSlow<UClass>(SpecifierStr);
 					if(FoundClass != nullptr)
 					{
 						Settings->ActorFilter.ActorClassesToRecord.Empty();
diff --git a/Engine/Source/Editor/Sequencer/Private/SSequencerPlayRateCombo.cpp b/Engine/Source/Editor/Sequencer/Private/SSequencerPlayRateCombo.cpp
--- a/Engine/Source/Editor/Sequencer/Private/SSequencerPlayRateCombo.cpp
+++ b/Engine/Source/Editor/Sequencer/Private/SSequencerPlayRateCombo.cpp
@@ -464,7 +464,7 @@ void SSequencerPlayRateCombo::PopulateCustomClockSourceMenu(FMenuBuilder& MenuBu
 	auto AssetClockSourceMenu = [this](FMenuBuilder& ActorMenuBuilder)
 	{
 		FAssetPickerConfig AssetPickerConfig;
-		AssetPickerConfig.Filter.ClassNames.Add(UMovieSceneCustomClockSource::StaticClass()->GetFName());
+		AssetPickerConfig.Filter.ClassPaths.Add(UMovieSceneCustomClockSource::StaticClass()->GetClassPathName());
 		AssetPickerConfig.bAllowNullSelection = false;
 		AssetPickerConfig.Filter.bRecursiveClasses = true;
 		AssetPickerConfig.OnAssetSelected = FOnAssetSelected::CreateLambda([this](const FAssetData& In){ this->SetCustomClockSource(In.GetAsset()); });
diff --git a/Engine/Source/Editor/Sequencer/Private/SequencerTrackFilters.h b/Engine/Source/Editor/Sequencer/Private/SequencerTrackFilters.h
--- a/Engine/Source/Editor/Sequencer/Private/SequencerTrackFilters.h
+++ b/Engine/Source/Editor/Sequencer/Private/SequencerTrackFilters.h
@@ -355,8 +355,8 @@ public:
 
 	virtual bool SupportsSequence(UMovieSceneSequence* InSequence) const override
 	{
-		static UClass* LevelSequenceClass = FindObject<UClass>(ANY_PACKAGE, TEXT("LevelSequence"), true);
-		static UClass* WidgetAnimationClass = FindObject<UClass>(ANY_PACKAGE, TEXT("WidgetAnimation"), true);
+		static UClass* LevelSequenceClass = FindObject<UClass>(nullptr, TEXT("/Script/LevelSequence.LevelSequence"), true);
+		static UClass* WidgetAnimationClass = FindObject<UClass>(nullptr, TEXT("/Script/UMG.WidgetAnimation"), true);
 		return InSequence != nullptr &&
 			((LevelSequenceClass != nullptr && InSequence->GetClass()->IsChildOf(LevelSequenceClass)) ||
 			(WidgetAnimationClass != nullptr && InSequence->GetClass()->IsChildOf(WidgetAnimationClass)));
diff --git a/Engine/Source/Editor/SkeletalMeshEditor/Private/MeshHierarchyCmd.cpp b/Engine/Source/Editor/SkeletalMeshEditor/Private/MeshHierarchyCmd.cpp
--- a/Engine/Source/Editor/SkeletalMeshEditor/Private/MeshHierarchyCmd.cpp
+++ b/Engine/Source/Editor/SkeletalMeshEditor/Private/MeshHierarchyCmd.cpp
@@ -19,7 +19,7 @@ bool FMeshHierarchyCmd::Exec(UWorld*, const TCHAR* Cmd, FOutputDevice& Ar)
 	{
 		Ar.Log(TEXT("Starting Mesh Test"));
 		FARFilter Filter;
-		Filter.ClassNames.Add(USkeletalMesh::StaticClass()->GetFName());
+		Filter.ClassPaths.Add(USkeletalMesh::StaticClass()->GetClassPathName());
 
 		TArray<FAssetData> SkeletalMeshes;
 		FAssetRegistryModule& AssetRegistryModule = FModuleManager::LoadModuleChecked<FAssetRegistryModule>(TEXT("AssetRegistry"));
diff --git a/Engine/Source/Editor/SkeletalMeshEditor/Private/SControlRigMappingWindow.cpp b/Engine/Source/Editor/SkeletalMeshEditor/Private/SControlRigMappingWindow.cpp
--- a/Engine/Source/Editor/SkeletalMeshEditor/Private/SControlRigMappingWindow.cpp
+++ b/Engine/Source/Editor/SkeletalMeshEditor/Private/SControlRigMappingWindow.cpp
@@ -265,7 +265,7 @@ void SControlRigMappingWindow::OnAddNodeMapping()
 	FContentBrowserModule& ContentBrowserModule = FModuleManager::Get().LoadModuleChecked<FContentBrowserModule>(TEXT("ContentBrowser"));
 
 	FAssetPickerConfig AssetPickerConfig;
- 	AssetPickerConfig.Filter.ClassNames.Add(UBlueprint::StaticClass()->GetFName());
+ 	AssetPickerConfig.Filter.ClassPaths.Add(UBlueprint::StaticClass()->GetClassPathName());
  	AssetPickerConfig.Filter.bRecursiveClasses = true;
 	AssetPickerConfig.OnAssetSelected = FOnAssetSelected::CreateSP(this, &SControlRigMappingWindow::OnAssetSelectedFromMeshPicker);
 	AssetPickerConfig.bAllowNullSelection = false;
@@ -318,9 +318,7 @@ bool SControlRigMappingWindow::OnShouldFilterAnimAsset(const FAssetData& AssetDa
 	{
 		if (ParentClassName.IsEmpty() == false)
 		{
-			UObject* Outer = nullptr;
-			ResolveName(Outer, ParentClassName, false, false);
-			UClass* ParentClass = FindObject<UClass>(ANY_PACKAGE, *ParentClassName);
+			UClass* ParentClass = UClass::TryFindTypeSlow<UClass>(FPackageName::ExportTextPathToObjectPath(ParentClassName));
 			while (ParentClass && ParentClass != UObject::StaticClass())
 			{
 				if (ParentClass->ImplementsInterface(UNodeMappingProviderInterface::StaticClass()))
diff --git a/Engine/Source/Editor/SkeletonEditor/Private/EditableSkeleton.cpp b/Engine/Source/Editor/SkeletonEditor/Private/EditableSkeleton.cpp
--- a/Engine/Source/Editor/SkeletonEditor/Private/EditableSkeleton.cpp
+++ b/Engine/Source/Editor/SkeletonEditor/Private/EditableSkeleton.cpp
@@ -375,7 +375,7 @@ void FEditableSkeleton::GetAssetsContainingCurves(const FName& InContainerName,
 	const FString CurrentSkeletonName = SkeletonData.GetExportTextName();
 
 	FAssetRegistryModule& AssetModule = FModuleManager::LoadModuleChecked<FAssetRegistryModule>(TEXT("AssetRegistry"));
-	AssetModule.Get().GetAssetsByClass(UAnimationAsset::StaticClass()->GetFName(), OutAssets, true);
+	AssetModule.Get().GetAssetsByClass(UAnimationAsset::StaticClass()->GetClassPathName(), OutAssets, true);
 
 	GWarn->BeginSlowTask(LOCTEXT("CollectAnimationsTaskDesc", "Collecting assets..."), true);
 
@@ -455,16 +455,16 @@ void FEditableSkeleton::RemoveSmartnamesAndFixupAnimations(const FName& InContai
 
 		AnimationAssets.Sort([&](const FAssetData& A, const FAssetData& B)
 		{ 
-			if (A.AssetClass == B.AssetClass)
+			if (A.AssetClassPath == B.AssetClassPath)
 			{
 				return A.AssetName.LexicalLess(B.AssetName);
 			}
-			return A.AssetClass.LexicalLess(B.AssetClass);
+			return A.AssetClassPath.Compare(B.AssetClassPath) < 0;
 		});
 
 		for (FAssetData& Data : AnimationAssets)
 		{
-			AssetMessage += Data.AssetName.ToString() + " (" + Data.AssetClass.ToString() + ")\n";
+			AssetMessage += Data.AssetName.ToString() + " (" + Data.AssetClassPath.ToString() + ")\n";
 		}
 
 		FText AssetTitleText = LOCTEXT("DeleteCurveDialogTitle", "Confirm Deletion");
@@ -1230,7 +1230,7 @@ void FEditableSkeleton::GetCompatibleAnimSequences(TArray<struct FAssetData>& Ou
 	FAssetRegistryModule& AssetRegistryModule = FModuleManager::LoadModuleChecked<FAssetRegistryModule>(TEXT("AssetRegistry"));
 
 	TArray<FAssetData> AssetDataList;
-	AssetRegistryModule.Get().GetAssetsByClass(UAnimSequenceBase::StaticClass()->GetFName(), AssetDataList, true);
+	AssetRegistryModule.Get().GetAssetsByClass(UAnimSequenceBase::StaticClass()->GetClassPathName(), AssetDataList, true);
 
 	OutAssets.Empty(AssetDataList.Num());
 
@@ -1550,7 +1550,7 @@ void FEditableSkeleton::RemoveUnusedBones()
 	}
 
 	FARFilter Filter;
-	Filter.ClassNames.Add(USkeletalMesh::StaticClass()->GetFName());
+	Filter.ClassPaths.Add(USkeletalMesh::StaticClass()->GetClassPathName());
 
 	FString SkeletonString = FAssetData(Skeleton).GetExportTextName();
 	Filter.TagsAndValues.Add(USkeletalMesh::GetSkeletonMemberName(), SkeletonString);
diff --git a/Engine/Source/Editor/SkeletonEditor/Private/SSkeletonTree.cpp b/Engine/Source/Editor/SkeletonEditor/Private/SSkeletonTree.cpp
--- a/Engine/Source/Editor/SkeletonEditor/Private/SSkeletonTree.cpp
+++ b/Engine/Source/Editor/SkeletonEditor/Private/SSkeletonTree.cpp
@@ -1354,7 +1354,7 @@ void SSkeletonTree::FillAttachAssetSubmenu(FMenuBuilder& MenuBuilder, const TSha
 
 	for(int i = 0; i < FilterClasses.Num(); ++i)
 	{
-		AssetPickerConfig.Filter.ClassNames.Add(FilterClasses[i]->GetFName());
+		AssetPickerConfig.Filter.ClassPaths.Add(FilterClasses[i]->GetClassPathName());
 	}
 
 
diff --git a/Engine/Source/Editor/SourceControlWindows/Private/SSourceControlCommon.cpp b/Engine/Source/Editor/SourceControlWindows/Private/SSourceControlCommon.cpp
--- a/Engine/Source/Editor/SourceControlWindows/Private/SSourceControlCommon.cpp
+++ b/Engine/Source/Editor/SourceControlWindows/Private/SSourceControlCommon.cpp
@@ -102,7 +102,7 @@ void FFileTreeItem::RefreshAssetInformation()
 
 			TempAssetName = RetrieveAssetName(AssetData);
 			TempAssetPath = RetrieveAssetPath(AssetData);
-			TempAssetType = AssetData.AssetClass.ToString();
+			TempAssetType = AssetData.AssetClassPath.ToString();
 
 			const FAssetToolsModule& AssetToolsModule = FModuleManager::LoadModuleChecked<FAssetToolsModule>(TEXT("AssetTools"));
 			const TSharedPtr<IAssetTypeActions> AssetTypeActions = AssetToolsModule.Get().GetAssetTypeActionsForClass(AssetData.GetClass()).Pin();
@@ -182,7 +182,6 @@ FText FFileTreeItem::GetAssetName()
 FString FFileTreeItem::RetrieveAssetName(const FAssetData& InAssetData) const
 {
 	static const FName NAME_ActorLabel(TEXT("ActorLabel"));
-	static const FName NAME_ActorFolder(TEXT("ActorFolder"));
 
 	if (InAssetData.FindTag(NAME_ActorLabel))
 	{
@@ -191,7 +190,7 @@ FString FFileTreeItem::RetrieveAssetName(const FAssetData& InAssetData) const
 		InAssetData.GetTagValue(NAME_ActorLabel, ResultAssetName);
 		return ResultAssetName;
 	}
-	else if (InAssetData.AssetClass == NAME_ActorFolder)
+	else if (InAssetData.AssetClassPath == UActorFolder::StaticClass()->GetClassPathName())
 	{
 		FString ActorFolderPath = UActorFolder::GetAssetRegistryInfoFromPackage(InAssetData.PackageName).GetDisplayName();
 		if (!ActorFolderPath.IsEmpty())
@@ -250,7 +249,7 @@ FOfflineFileTreeItem::FOfflineFileTreeItem(const FString& InFilename)
 				AssetName = FText::FromName(AssetData.AssetName);
 			}
 
-			AssetType = FText::FromName(AssetData.AssetClass);
+			AssetType = FText::FromString(AssetData.AssetClassPath.ToString());
 
 			const FAssetToolsModule& AssetToolsModule = FModuleManager::LoadModuleChecked<FAssetToolsModule>(TEXT("AssetTools"));
 			const TSharedPtr<IAssetTypeActions> AssetTypeActions = AssetToolsModule.Get().GetAssetTypeActionsForClass(AssetData.GetClass()).Pin();
diff --git a/Engine/Source/Editor/StructViewer/Private/SStructViewer.cpp b/Engine/Source/Editor/StructViewer/Private/SStructViewer.cpp
--- a/Engine/Source/Editor/StructViewer/Private/SStructViewer.cpp
+++ b/Engine/Source/Editor/StructViewer/Private/SStructViewer.cpp
@@ -915,7 +915,7 @@ void FStructHierarchy::PopulateStructHierarchy()
 		FAssetRegistryModule& AssetRegistryModule = FModuleManager::LoadModuleChecked<FAssetRegistryModule>(TEXT("AssetRegistry"));
 
 		FARFilter Filter;
-		Filter.ClassNames.Add(UUserDefinedStruct::StaticClass()->GetFName());
+		Filter.ClassPaths.Add(UUserDefinedStruct::StaticClass()->GetClassPathName());
 		Filter.bRecursiveClasses = true;
 
 		TArray<FAssetData> UserDefinedStructsList;
diff --git a/Engine/Source/Editor/SubobjectDataInterface/Private/SubobjectDataSubsystem.cpp b/Engine/Source/Editor/SubobjectDataInterface/Private/SubobjectDataSubsystem.cpp
--- a/Engine/Source/Editor/SubobjectDataInterface/Private/SubobjectDataSubsystem.cpp
+++ b/Engine/Source/Editor/SubobjectDataInterface/Private/SubobjectDataSubsystem.cpp
@@ -650,7 +650,7 @@ UClass* USubobjectDataSubsystem::CreateNewBPComponent(TSubclassOf<UActorComponen
 			const FString PackagePath = NewClassPath / NewClassName;
 
 			// Check for an existing object
-			if(UObject* ExistingObject = StaticFindObject(UObject::StaticClass(), ANY_PACKAGE, *PackagePath))
+			if(UObject* ExistingObject = StaticFindObject(UObject::StaticClass(), nullptr, *PackagePath))
 			{
 				UE_LOG(LogSubobjectSubsystem, Warning, TEXT("Failed to Create new BP Component: A class with a name '%s' already exists!"), *PackagePath);
 				return nullptr;
diff --git a/Engine/Source/Editor/UMGEditor/Private/Palette/SPaletteViewModel.cpp b/Engine/Source/Editor/UMGEditor/Private/Palette/SPaletteViewModel.cpp
--- a/Engine/Source/Editor/UMGEditor/Private/Palette/SPaletteViewModel.cpp
+++ b/Engine/Source/Editor/UMGEditor/Private/Palette/SPaletteViewModel.cpp
@@ -138,6 +138,11 @@ namespace UE::Editor::SPaletteViewModel::Private
 		{
 			return ClassPath;
 		}
+
+		virtual FTopLevelAssetPath GetClassPathName() const
+		{
+			return FTopLevelAssetPath(ClassPath.ToString());
+		}
 		// End IUnloadedBlueprintData interface
 
 	private:
@@ -465,7 +470,7 @@ void FWidgetCatalogViewModel::BuildClassWidgetList()
 	// Locate all widget BP assets (include unloaded)
 	const FAssetRegistryModule& AssetRegistryModule = FModuleManager::LoadModuleChecked<FAssetRegistryModule>(TEXT("AssetRegistry"));
 	TArray<FAssetData> AllBPsAssetData;
-	AssetRegistryModule.Get().GetAssetsByClass(UBlueprint::StaticClass()->GetFName(), AllBPsAssetData, true);
+	AssetRegistryModule.Get().GetAssetsByClass(UBlueprint::StaticClass()->GetClassPathName(), AllBPsAssetData, true);
 
 	for (FAssetData& BPAssetData : AllBPsAssetData)
 	{
@@ -478,9 +483,7 @@ void FWidgetCatalogViewModel::BuildClassWidgetList()
 		}
 		if (!ParentClassName.IsEmpty())
 		{
-			UObject* Outer = nullptr;
-			ResolveName(Outer, ParentClassName, false, false);
-			ParentClass = FindObject<UClass>(ANY_PACKAGE, *ParentClassName);
+			ParentClass = UClass::TryFindTypeSlow<UClass>(FPackageName::ExportTextPathToObjectPath(ParentClassName));
 			// UUserWidgets have their own loading section, and we don't want to process any blueprints that don't have UWidget parents
 			if (ParentClass)
 			{
@@ -512,7 +515,7 @@ void FWidgetCatalogViewModel::BuildClassWidgetList()
 	}
 
 	TArray<FAssetData> AllWidgetBPsAssetData;
-	AssetRegistryModule.Get().GetAssetsByClass(UWidgetBlueprint::StaticClass()->GetFName(), AllWidgetBPsAssetData, true);
+	AssetRegistryModule.Get().GetAssetsByClass(UWidgetBlueprint::StaticClass()->GetClassPathName(), AllWidgetBPsAssetData, true);
 
 	FName ActiveWidgetBlueprintName = ActiveWidgetBlueprintClass->ClassGeneratedBy->GetFName();
 	for (FAssetData& WidgetBPAssetData : AllWidgetBPsAssetData)
diff --git a/Engine/Source/Editor/UMGEditor/Private/Templates/WidgetTemplateClass.cpp b/Engine/Source/Editor/UMGEditor/Private/Templates/WidgetTemplateClass.cpp
--- a/Engine/Source/Editor/UMGEditor/Private/Templates/WidgetTemplateClass.cpp
+++ b/Engine/Source/Editor/UMGEditor/Private/Templates/WidgetTemplateClass.cpp
@@ -51,9 +51,7 @@ FWidgetTemplateClass::FWidgetTemplateClass(const FAssetData& InWidgetAssetData,
 			}
 			if (!ParentClassName.IsEmpty())
 			{
-				UObject* Outer = nullptr;
-				ResolveName(Outer, ParentClassName, false, false);
-				CachedParentClass = FindObject<UClass>(ANY_PACKAGE, *ParentClassName);
+				CachedParentClass = UClass::TryFindTypeSlow<UClass>(FPackageName::ExportTextPathToObjectPath(ParentClassName));
 			}
 		}
 	}
diff --git a/Engine/Source/Editor/UMGEditor/Private/Templates/WidgetTemplateImageClass.cpp b/Engine/Source/Editor/UMGEditor/Private/Templates/WidgetTemplateImageClass.cpp
--- a/Engine/Source/Editor/UMGEditor/Private/Templates/WidgetTemplateImageClass.cpp
+++ b/Engine/Source/Editor/UMGEditor/Private/Templates/WidgetTemplateImageClass.cpp
@@ -18,11 +18,11 @@ UWidget* FWidgetTemplateImageClass::Create(UWidgetTree* WidgetTree)
 {
 	UWidget* Widget = FWidgetTemplateClass::Create(WidgetTree);
 
-	if (Cast<UImage>(Widget) != nullptr && Supports(FindObject<UClass>(ANY_PACKAGE, *WidgetAssetData.AssetClass.ToString())) )
+	if (Cast<UImage>(Widget) != nullptr && Supports(FindObject<UClass>(WidgetAssetData.AssetClassPath)) )
 	{
 		UImage* ImageWidget = Cast<UImage>(Widget);
 
-		UObject* ImageResource = FindObject<UObject>(ANY_PACKAGE, *WidgetAssetData.ObjectPath.ToString());
+		UObject* ImageResource = FindObject<UObject>(nullptr, *WidgetAssetData.ObjectPath.ToString());
 		ImageWidget->Brush.SetResourceObject(ImageResource);
 	}
 
@@ -31,7 +31,7 @@ UWidget* FWidgetTemplateImageClass::Create(UWidgetTree* WidgetTree)
 
 bool FWidgetTemplateImageClass::Supports(UClass* InClass)
 {
-	static const UClass* SlateTextureAtlasInterface = FindObject<UClass>(ANY_PACKAGE, TEXT("SlateTextureAtlasInterface"));
+	static const UClass* SlateTextureAtlasInterface = FindObject<UClass>(nullptr, TEXT("/Script/Engine.SlateTextureAtlasInterface"));
 
 	return InClass != nullptr && ( InClass->IsChildOf(UTexture::StaticClass())
 		|| InClass->IsChildOf(UMaterialInterface::StaticClass())
diff --git a/Engine/Source/Editor/UMGEditor/Private/WidgetBlueprintEditor.cpp b/Engine/Source/Editor/UMGEditor/Private/WidgetBlueprintEditor.cpp
--- a/Engine/Source/Editor/UMGEditor/Private/WidgetBlueprintEditor.cpp
+++ b/Engine/Source/Editor/UMGEditor/Private/WidgetBlueprintEditor.cpp
@@ -524,7 +524,7 @@ void FWidgetBlueprintEditor::OpenCreateNativeBaseClassDialog()
 
 void FWidgetBlueprintEditor::OnCreateNativeBaseClassSuccessfully(const FString& InClassName, const FString& InClassPath, const FString& InModuleName)
 {
-	UClass* NewNativeClass = FindObject<UClass>(ANY_PACKAGE, *InClassName);
+	UClass* NewNativeClass = FindObject<UClass>(FTopLevelAssetPath(*InClassPath, *InClassName));
 	if (NewNativeClass)
 	{
 		ReparentBlueprint_NewParentChosen(NewNativeClass);
diff --git a/Engine/Source/Editor/UMGEditor/Private/WidgetBlueprintEditorUtils.cpp b/Engine/Source/Editor/UMGEditor/Private/WidgetBlueprintEditorUtils.cpp
--- a/Engine/Source/Editor/UMGEditor/Private/WidgetBlueprintEditorUtils.cpp
+++ b/Engine/Source/Editor/UMGEditor/Private/WidgetBlueprintEditorUtils.cpp
@@ -1421,9 +1421,9 @@ UWidget* FWidgetBlueprintEditorUtils::GetWidgetTemplateFromDragDrop(UWidgetBluep
 			// Only handle first valid dragged widget, multi widget drag drop is not practically useful
 			const FAssetData& AssetData = AssetDragDropOp->GetAssets()[0];
 
-			bool CodeClass = AssetData.AssetClass == "Class";
-			FName ClassName = CodeClass ? AssetData.ObjectPath : AssetData.AssetClass;
-			UClass* AssetClass = FindObjectChecked<UClass>(ANY_PACKAGE, *ClassName.ToString());
+			bool CodeClass = AssetData.AssetClassPath == FTopLevelAssetPath(TEXT("/Script/CoreUObject"), TEXT("Class"));
+			FString ClassName = CodeClass ? AssetData.ObjectPath.ToString() : AssetData.AssetClassPath.ToString();
+			UClass* AssetClass = FindObjectChecked<UClass>(nullptr, *ClassName);
 
 			if (FWidgetTemplateBlueprintClass::Supports(AssetClass))
 			{
diff --git a/Engine/Source/Editor/UnrealEd/Classes/Commandlets/AssetRegUtilCommandlet.h b/Engine/Source/Editor/UnrealEd/Classes/Commandlets/AssetRegUtilCommandlet.h
--- a/Engine/Source/Editor/UnrealEd/Classes/Commandlets/AssetRegUtilCommandlet.h
+++ b/Engine/Source/Editor/UnrealEd/Classes/Commandlets/AssetRegUtilCommandlet.h
@@ -33,7 +33,7 @@ protected:
 	IAssetRegistry* AssetRegistry;
 
 	void RecursivelyGrabDependencies(TArray<FSortableDependencyEntry>& OutSortableDependencies,
-		const int32& DepSet, int32& DepOrder, int32 DepHierarchy, TSet<FName>& ProcessedFiles, const TSet<FName>& OriginalSet, const FName& FilePath, const FName& PackageFName, const TArray<FName>& FilterByClass);
+		const int32& DepSet, int32& DepOrder, int32 DepHierarchy, TSet<FName>& ProcessedFiles, const TSet<FName>& OriginalSet, const FName& FilePath, const FName& PackageFName, const TArray<FTopLevelAssetPath>& FilterByClass);
 
 	void ReorderOrderFile(const FString& OrderFilePath, const FString& ReorderFileOutPath);
 	// Generate a new FOO which only takes partial update from new FOO files with input of an old FOO file and a new FOO file
diff --git a/Engine/Source/Editor/UnrealEd/Classes/Commandlets/FixupNeedsLoadForEditorGameCommandlet.h b/Engine/Source/Editor/UnrealEd/Classes/Commandlets/FixupNeedsLoadForEditorGameCommandlet.h
--- a/Engine/Source/Editor/UnrealEd/Classes/Commandlets/FixupNeedsLoadForEditorGameCommandlet.h
+++ b/Engine/Source/Editor/UnrealEd/Classes/Commandlets/FixupNeedsLoadForEditorGameCommandlet.h
@@ -13,7 +13,7 @@ class UFixupNeedsLoadForEditorGameCommandlet : public UResavePackagesCommandlet
 	GENERATED_BODY()
 
 	/** Map of class names and their default NeedsLoadForEditorGame() return values. */
-	TMap<FName, bool> ResaveClassNeedsLoadForEditorGameValues;
+	TMap<FString, bool> ResaveClassNeedsLoadForEditorGameValues;
 
 protected:
 
diff --git a/Engine/Source/Editor/UnrealEd/Classes/Commandlets/ResavePackagesCommandlet.h b/Engine/Source/Editor/UnrealEd/Classes/Commandlets/ResavePackagesCommandlet.h
--- a/Engine/Source/Editor/UnrealEd/Classes/Commandlets/ResavePackagesCommandlet.h
+++ b/Engine/Source/Editor/UnrealEd/Classes/Commandlets/ResavePackagesCommandlet.h
@@ -52,7 +52,7 @@ protected:
 	int32 MaxPackagesToResave;
 
 	/** allows users to save only packages with a particular class in them (useful for fixing content) */
-	TArray<FName> ResaveClasses;
+	TArray<FString> ResaveClasses;
 
 	// If non-empty, this substring has to be present in the package name for the commandlet to process it
 	FString PackageSubstring;
diff --git a/Engine/Source/Editor/UnrealEd/Private/AnimationEditorUtils.cpp b/Engine/Source/Editor/UnrealEd/Private/AnimationEditorUtils.cpp
--- a/Engine/Source/Editor/UnrealEd/Private/AnimationEditorUtils.cpp
+++ b/Engine/Source/Editor/UnrealEd/Private/AnimationEditorUtils.cpp
@@ -317,7 +317,7 @@ bool SAnimationCompressionSelectionDialog::IsConfirmButtonEnabled() const
 void SAnimationCompressionSelectionDialog::Construct(const FArguments& InArgs, const FAnimationCompressionSelectionDialogConfig& InConfig)
 {
 	FAssetPickerConfig AssetPickerConfig;
-	AssetPickerConfig.Filter.ClassNames.Push(UAnimBoneCompressionSettings::StaticClass()->GetFName());
+	AssetPickerConfig.Filter.ClassPaths.Push(UAnimBoneCompressionSettings::StaticClass()->GetClassPathName());
 	AssetPickerConfig.Filter.bRecursiveClasses = true;
 	AssetPickerConfig.InitialAssetViewType = EAssetViewType::List;
 	AssetPickerConfig.OnAssetSelected = FOnAssetSelected::CreateSP(this, &SAnimationCompressionSelectionDialog::OnAssetSelected);
diff --git a/Engine/Source/Editor/UnrealEd/Private/AssetDeleteModel.cpp b/Engine/Source/Editor/UnrealEd/Private/AssetDeleteModel.cpp
--- a/Engine/Source/Editor/UnrealEd/Private/AssetDeleteModel.cpp
+++ b/Engine/Source/Editor/UnrealEd/Private/AssetDeleteModel.cpp
@@ -366,7 +366,7 @@ bool FAssetDeleteModel::CanReplaceReferencesWith( const FAssetData& InAssetData
 
 		if (!NativeClassName.IsEmpty())
 		{
-			UClass* NativeParentClassToTest = FindObject<UClass>(ANY_PACKAGE, *NativeClassName);
+			UClass* NativeParentClassToTest = UClass::TryFindTypeSlow<UClass>(NativeClassName);
 			UClass* NativeParentClassToReplace = FBlueprintEditorUtils::FindFirstNativeClass(OriginalBPParentClass);
 
 			if (!NativeParentClassToTest || !NativeParentClassToTest->IsChildOf(NativeParentClassToReplace))
diff --git a/Engine/Source/Editor/UnrealEd/Private/AssetThumbnail.cpp b/Engine/Source/Editor/UnrealEd/Private/AssetThumbnail.cpp
--- a/Engine/Source/Editor/UnrealEd/Private/AssetThumbnail.cpp
+++ b/Engine/Source/Editor/UnrealEd/Private/AssetThumbnail.cpp
@@ -105,7 +105,7 @@ public:
 
 		const FAssetData& AssetData = AssetThumbnail->GetAssetData();
 
-		UClass* Class = FindObjectSafe<UClass>(ANY_PACKAGE, *AssetData.AssetClass.ToString());
+		UClass* Class = FindObjectSafe<UClass>(AssetData.AssetClassPath);
 		static FAssetToolsModule& AssetToolsModule = FModuleManager::LoadModuleChecked<FAssetToolsModule>(TEXT("AssetTools"));
 		TSharedPtr<IAssetTypeActions> AssetTypeActions;
 		if ( Class != NULL )
@@ -284,7 +284,7 @@ public:
 
 		// For non-class types, use the default based upon the actual asset class
 		// This has the side effect of not showing a class icon for assets that don't have a proper thumbnail image available
-		const FName DefaultThumbnail = (bIsClassType) ? NAME_None : FName(*FString::Printf(TEXT("ClassThumbnail.%s"), *AssetThumbnail->GetAssetData().AssetClass.ToString()));
+		const FName DefaultThumbnail = (bIsClassType) ? NAME_None : FName(*FString::Printf(TEXT("ClassThumbnail.%s"), *AssetThumbnail->GetAssetData().AssetClassPath.GetAssetName().ToString()));
 		ThumbnailBrush = FClassIconFinder::FindThumbnailForClass(ThumbnailClass.Get(), DefaultThumbnail);
 
 		ClassIconBrush = FSlateIconFinder::FindIconBrushForClass(ThumbnailClass.Get());
@@ -369,7 +369,7 @@ private:
 
 		const FAssetData& AssetData = AssetThumbnail->GetAssetData();
 
-		UClass* Class = FindObject<UClass>(ANY_PACKAGE, *AssetData.AssetClass.ToString());
+		UClass* Class = FindObject<UClass>(AssetData.AssetClassPath);
 		FAssetToolsModule& AssetToolsModule = FModuleManager::LoadModuleChecked<FAssetToolsModule>(TEXT("AssetTools"));
 		TWeakPtr<IAssetTypeActions> AssetTypeActions;
 		if ( Class != NULL )
@@ -559,11 +559,11 @@ private:
 			return !CachedThumbnail->IsEmpty();
 		}
 
-		if ( AssetData.AssetClass != UBlueprint::StaticClass()->GetFName() )
+		if ( AssetData.AssetClassPath != UBlueprint::StaticClass()->GetClassPathName() )
 		{
 			// If we are not a blueprint, see if the CDO of the asset's class has a rendering info
 			// Blueprints can't do this because the rendering info is based on the generated class
-			UClass* AssetClass = FindObject<UClass>(ANY_PACKAGE, *AssetData.AssetClass.ToString());
+			UClass* AssetClass = FindObject<UClass>(AssetData.AssetClassPath);
 
 			if ( AssetClass )
 			{
@@ -654,15 +654,15 @@ private:
 	FText GetAssetClassDisplayName() const
 	{
 		const FAssetData& AssetData = AssetThumbnail->GetAssetData();
-		FString AssetClass = AssetData.AssetClass.ToString();
-		UClass* Class = FindObjectSafe<UClass>(ANY_PACKAGE, *AssetClass);
+		FTopLevelAssetPath AssetClass = AssetData.AssetClassPath;
+		UClass* Class = FindObjectSafe<UClass>(AssetClass);
 
 		if ( Class )
 		{
 			return GetDisplayNameForClass( Class, &AssetData );
 		}
 
-		return FText::FromString(AssetClass);
+		return FText::FromString(AssetClass.ToString());
 	}
 
 	FText GetAssetDisplayName() const
diff --git a/Engine/Source/Editor/UnrealEd/Private/ClassIconFinder.cpp b/Engine/Source/Editor/UnrealEd/Private/ClassIconFinder.cpp
--- a/Engine/Source/Editor/UnrealEd/Private/ClassIconFinder.cpp
+++ b/Engine/Source/Editor/UnrealEd/Private/ClassIconFinder.cpp
@@ -96,7 +96,7 @@ const UClass* FClassIconFinder::GetIconClassForAssetData(const FAssetData& InAss
 		*bOutIsClassType = false;
 	}
 
-	UClass* AssetClass = FindObjectSafe<UClass>(ANY_PACKAGE, *InAssetData.AssetClass.ToString());
+	UClass* AssetClass = FindObjectSafe<UClass>(InAssetData.AssetClassPath);
 	if ( !AssetClass )
 	{
 		return nullptr;
@@ -109,7 +109,7 @@ const UClass* FClassIconFinder::GetIconClassForAssetData(const FAssetData& InAss
 			*bOutIsClassType = true;
 		}
 
-		return FindObject<UClass>(ANY_PACKAGE, *InAssetData.AssetName.ToString());
+		return FindObject<UClass>(nullptr, *InAssetData.ObjectPath.ToString());
 	}
 	
 	static const FName IgnoreClassThumbnail(TEXT("IgnoreClassThumbnail"));
@@ -131,7 +131,7 @@ const UClass* FClassIconFinder::GetIconClassForAssetData(const FAssetData& InAss
 		{
 			UObject* Outer = nullptr;
 			ResolveName(Outer, ParentClassName, false, false);
-			return FindObject<UClass>(ANY_PACKAGE, *ParentClassName);
+			return FindObject<UClass>(Outer, *ParentClassName);
 		}
 	}
 
diff --git a/Engine/Source/Editor/UnrealEd/Private/Commandlets/AssetRegUtilCommandlet.cpp b/Engine/Source/Editor/UnrealEd/Private/Commandlets/AssetRegUtilCommandlet.cpp
--- a/Engine/Source/Editor/UnrealEd/Private/Commandlets/AssetRegUtilCommandlet.cpp
+++ b/Engine/Source/Editor/UnrealEd/Private/Commandlets/AssetRegUtilCommandlet.cpp
@@ -32,7 +32,7 @@ const static FName NAME_uptnl(TEXT("uptnl"));
 
 struct FSortableDependencyEntry
 {
-	FSortableDependencyEntry(const FName& InLongPackageName, const FName& InFilePath, const FName& InExtension, const int32 InDepSet, const int32 InDepHierarchy, const int32 InDepOrder, bool InHasDependencies, TSet<FName> &&InClasses)
+	FSortableDependencyEntry(const FName& InLongPackageName, const FName& InFilePath, const FName& InExtension, const int32 InDepSet, const int32 InDepHierarchy, const int32 InDepOrder, bool InHasDependencies, TSet<FTopLevelAssetPath> &&InClasses)
 		: LongPackageName(InLongPackageName)
 		, FilePath(InFilePath)
 		, Extension(InExtension)
@@ -59,7 +59,7 @@ struct FSortableDependencyEntry
 	FName LongPackageName;
 	FName FilePath;
 	FName Extension;
-	TSet<FName> Classes;
+	TSet<FTopLevelAssetPath> Classes;
 	int32 DepSet;
 	int32 DepHierarchy;
 	int32 DepOrder;
@@ -100,7 +100,7 @@ struct FSortableDependencySortForExports
 
 struct FSortableDependencySort
 {
-	FSortableDependencySort(const TArray<FName>& InGroupExtensions, const TArray<FName>& InGroupClasses, const TMap<FName, int32> InExtensionPriority)
+	FSortableDependencySort(const TArray<FName>& InGroupExtensions, const TArray<FTopLevelAssetPath>& InGroupClasses, const TMap<FName, int32> InExtensionPriority)
 		: GroupExtensions(InGroupExtensions)
 		, GroupClasses(InGroupClasses)
 		, ExtensionPriority(InExtensionPriority)
@@ -109,7 +109,7 @@ struct FSortableDependencySort
 	}
 
 	const TArray<FName>& GroupExtensions;
-	const TArray<FName>& GroupClasses;
+	const TArray<FTopLevelAssetPath>& GroupClasses;
 	const TMap<FName, int32> ExtensionPriority;
 
 	bool operator()(const FSortableDependencyEntry& A, const FSortableDependencyEntry& B) const
@@ -127,9 +127,9 @@ struct FSortableDependencySort
 
 		bool bIsAClassGrouped = false;
 		bool bIsBClassGrouped = false;
-		FName AClass = NAME_None;
-		FName BClass = NAME_None;
-		for (const FName& Class : GroupClasses)
+		FTopLevelAssetPath AClass;
+		FTopLevelAssetPath BClass;
+		for (const FTopLevelAssetPath& Class : GroupClasses)
 		{
 			if (A.Classes.Contains(Class))
 			{
@@ -149,7 +149,7 @@ struct FSortableDependencySort
 
 		if ((AClass != BClass) && bIsAClassGrouped && bIsBClassGrouped)
 		{
-			return BClass.LexicalLess(AClass);
+			return BClass.Compare(AClass) < 0;
 		}
 
 		if (A.DepSet != B.DepSet)
@@ -201,7 +201,7 @@ UAssetRegUtilCommandlet::UAssetRegUtilCommandlet(const FObjectInitializer& Objec
 }
 
 void UAssetRegUtilCommandlet::RecursivelyGrabDependencies(TArray<FSortableDependencyEntry>& OutSortableDependencies,
-                                                          const int32& DepSet, int32& DepOrder, int32 DepHierarchy, TSet<FName>& ProcessedFiles, const TSet<FName>& OriginalSet, const FName& FilePath, const FName& PackageFName, const TArray<FName>& FilterByClasses )
+                                                          const int32& DepSet, int32& DepOrder, int32 DepHierarchy, TSet<FName>& ProcessedFiles, const TSet<FName>& OriginalSet, const FName& FilePath, const FName& PackageFName, const TArray<FTopLevelAssetPath>& FilterByClasses )
 {
 	bool bHasDependencies = false;
 	//now walk the dependency tree for everything under this package
@@ -210,17 +210,17 @@ void UAssetRegUtilCommandlet::RecursivelyGrabDependencies(TArray<FSortableDepend
 
 	TArray<FAssetData> AssetsData;
 	AssetRegistry->GetAssetsByPackageName(PackageFName, AssetsData, true);
-	TSet<FName> Classes;
+	TSet<FTopLevelAssetPath> Classes;
 	Classes.Reserve(AssetsData.Num());
 	for (const FAssetData& AssetData : AssetsData)
 	{
-		Classes.Add(AssetData.AssetClass);
-		TArray<FName> AncestorClasses;
-		AssetRegistry->GetAncestorClassNames(AssetData.AssetClass, AncestorClasses);
+		Classes.Add(AssetData.AssetClassPath);
+		TArray<FTopLevelAssetPath> AncestorClasses;
+		AssetRegistry->GetAncestorClassNames(AssetData.AssetClassPath, AncestorClasses);
 		Classes.Append(AncestorClasses);
 	}
-	TSet<FName> FilteredClasses;
-	for (const FName& FilterClass : FilterByClasses)
+	TSet<FTopLevelAssetPath> FilteredClasses;
+	for (const FTopLevelAssetPath& FilterClass : FilterByClasses)
 	{
 		if (Classes.Contains(FilterClass))
 		{
@@ -279,13 +279,13 @@ void UAssetRegUtilCommandlet::ReorderOrderFile(const FString& OrderFilePath, con
 	ExtraAssetExtensions.Add(NAME_ubulk);
 	ExtraAssetExtensions.Add(NAME_uptnl);
 
-	TArray<FName> FilterByClasses;
-	FilterByClasses.Add(FName(TEXT("Material")));
-	FilterByClasses.Add(FName(TEXT("MaterialFunction")));
-	FilterByClasses.Add(FName(TEXT("MaterialInstance")));
-	FilterByClasses.Add(FName(TEXT("BlueprintCore")));
-	FilterByClasses.Add(FName(TEXT("ParticleEmitter")));
-	FilterByClasses.Add(FName(TEXT("ParticleModule")));
+	TArray<FTopLevelAssetPath> FilterByClasses;
+	FilterByClasses.Add(FTopLevelAssetPath(TEXT("/Script/Engine"), TEXT("Material")));
+	FilterByClasses.Add(FTopLevelAssetPath(TEXT("/Script/Engine"), TEXT("MaterialFunction")));
+	FilterByClasses.Add(FTopLevelAssetPath(TEXT("/Script/Engine"), TEXT("MaterialInstance")));
+	FilterByClasses.Add(FTopLevelAssetPath(TEXT("/Script/Engine"), TEXT("BlueprintCore")));
+	FilterByClasses.Add(FTopLevelAssetPath(TEXT("/Script/Engine"), TEXT("ParticleEmitter")));
+	FilterByClasses.Add(FTopLevelAssetPath(TEXT("/Script/Engine"), TEXT("ParticleModule")));
 
 
 	int32 DepSet = 0; // this is the root set for the dependency (i.e files with a number probably came from the same core dependency)
@@ -309,7 +309,7 @@ void UAssetRegUtilCommandlet::ReorderOrderFile(const FString& OrderFilePath, con
 				{
 					//special case for packages outside of our mounted paths, pick up the header and the export without any dependency-gathering.
 					ProcessedFiles.Add(FilePath);
-					UnsortedEntries.Add(FSortableDependencyEntry(NAME_UnresolvedPackageName, FilePath, FNameExtension, DepSet, 0, 0, false, TSet<FName>()));
+					UnsortedEntries.Add(FSortableDependencyEntry(NAME_UnresolvedPackageName, FilePath, FNameExtension, DepSet, 0, 0, false, TSet<FTopLevelAssetPath>()));
 				}
 			}
 			else if ( ExtraAssetExtensions.Contains(FNameExtension) == false )
@@ -336,7 +336,7 @@ void UAssetRegUtilCommandlet::ReorderOrderFile(const FString& OrderFilePath, con
 				{
 					check(!ProcessedFiles.Contains(ExtraAssetPathFName));
 					ProcessedFiles.Add(ExtraAssetPathFName);
-					TSet<FName> Classes = DependencyEntry.Classes;
+					TSet<FTopLevelAssetPath> Classes = DependencyEntry.Classes;
 					UnsortedEntries.Add(FSortableDependencyEntry(DependencyEntry.LongPackageName, ExtraAssetPathFName, ExtraAssetExtension, DependencyEntry.DepSet, DependencyEntry.DepHierarchy, DependencyEntry.DepOrder, DependencyEntry.bHasDependencies, MoveTemp(Classes)));
 				}
 			}
diff --git a/Engine/Source/Editor/UnrealEd/Private/Commandlets/AssetRegistryGenerator.cpp b/Engine/Source/Editor/UnrealEd/Private/Commandlets/AssetRegistryGenerator.cpp
--- a/Engine/Source/Editor/UnrealEd/Private/Commandlets/AssetRegistryGenerator.cpp
+++ b/Engine/Source/Editor/UnrealEd/Private/Commandlets/AssetRegistryGenerator.cpp
@@ -916,7 +916,7 @@ void FAssetRegistryGenerator::InjectEncryptionData(FAssetRegistryState& TargetSt
 							{
 								FAssetDataTagMap TagsAndValues = AssetData->TagsAndValues.CopyMap();
 								TagsAndValues.Add(UAssetManager::GetEncryptionKeyAssetTagName(), Guid.ToString());
-								FAssetData NewAssetData = FAssetData(AssetData->PackageName, AssetData->PackagePath, AssetData->AssetName, AssetData->AssetClass, TagsAndValues, AssetData->ChunkIDs, AssetData->PackageFlags);
+								FAssetData NewAssetData = FAssetData(AssetData->PackageName, AssetData->PackagePath, AssetData->AssetName, AssetData->AssetClassPath, TagsAndValues, AssetData->ChunkIDs, AssetData->PackageFlags);
 								NewAssetData.TaggedAssetBundles = AssetData->TaggedAssetBundles;
 								TargetState.UpdateAssetData(AssetData, NewAssetData);
 							}
@@ -1111,7 +1111,7 @@ void FAssetRegistryGenerator::UpdateCollectionAssetData()
 			{
 				TagsAndValues.Add(CollectionTagName, FString(CollectionValue));
 			}
-			State.UpdateAssetData(FAssetData(AssetData->PackageName, AssetData->PackagePath, AssetData->AssetName, AssetData->AssetClass, TagsAndValues, AssetData->ChunkIDs, AssetData->PackageFlags));
+			State.UpdateAssetData(FAssetData(AssetData->PackageName, AssetData->PackagePath, AssetData->AssetName, AssetData->AssetClassPath, TagsAndValues, AssetData->ChunkIDs, AssetData->PackageFlags));
 		}
 	}
 }
@@ -1951,7 +1951,7 @@ bool FAssetRegistryGenerator::GenerateAssetChunkInformationCSV(const FString& Ou
 				TmpString = FString::Printf(TEXT("%d,%s,%s,%s,%lld,%s\n"),
 					PakchunkIndex,
 					*AssetData.PackageName.ToString(),
-					*AssetData.AssetClass.ToString(),
+					*AssetData.AssetClassPath.ToString(),
 					bHardChunk ? TEXT("Hard") : *SoftChain,
 					FileSize,
 					AssetData.ChunkIDs.Num() == 1 ? TEXT("None") : *TmpStringChunks
diff --git a/Engine/Source/Editor/UnrealEd/Private/Commandlets/AssetSizeQueryCommandlet.cpp b/Engine/Source/Editor/UnrealEd/Private/Commandlets/AssetSizeQueryCommandlet.cpp
--- a/Engine/Source/Editor/UnrealEd/Private/Commandlets/AssetSizeQueryCommandlet.cpp
+++ b/Engine/Source/Editor/UnrealEd/Private/Commandlets/AssetSizeQueryCommandlet.cpp
@@ -121,8 +121,8 @@ int32 UAssetSizeQueryCommandlet::Main(const FString& FullCommandLine)
 		FName ObjectPath;
 		int64 CompressedSize;
 	};
-	TMap<FName /* AssetClass */, int64> FilteredClassCompressedSizes;
-	TMap<FName /* AssetClass */, TArray<FMatchedAssetInfo>> FilteredClassMatchedAssets;
+	TMap<FTopLevelAssetPath /* AssetClass */, int64> FilteredClassCompressedSizes;
+	TMap<FTopLevelAssetPath /* AssetClass */, TArray<FMatchedAssetInfo>> FilteredClassMatchedAssets;
 	AssetRegistry.EnumerateAllAssets(TSet<FName>(), 
 		[&ImportantAssetCount,
 		 &MatchedAssetCount,
@@ -154,12 +154,12 @@ int32 UAssetSizeQueryCommandlet::Main(const FString& FullCommandLine)
 		{
 			MatchedAssetCount++;
 
-			FMatchedAssetInfo& Info = FilteredClassMatchedAssets.FindOrAdd(AssetData.AssetClass).AddDefaulted_GetRef();
+			FMatchedAssetInfo& Info = FilteredClassMatchedAssets.FindOrAdd(AssetData.AssetClassPath).AddDefaulted_GetRef();
 			Info.ObjectPath = AssetData.ObjectPath;
 			Info.CompressedSize = AssetCompressedSize;
 			
 			FilteredCompressedSize += AssetCompressedSize;
-			int64& FilteredClassCompressedSize = FilteredClassCompressedSizes.FindOrAdd(AssetData.AssetClass);
+			int64& FilteredClassCompressedSize = FilteredClassCompressedSizes.FindOrAdd(AssetData.AssetClassPath);
 			FilteredClassCompressedSize += AssetCompressedSize;
 		}
 
@@ -190,7 +190,7 @@ int32 UAssetSizeQueryCommandlet::Main(const FString& FullCommandLine)
 	{
 		UE_LOG(LogAssetSize, Display, TEXT("Filtered class sizes:                         bytes          (pct of filtered total)"));
 	}
-	for (TPair<FName, int64>& ClassSizePair : FilteredClassCompressedSizes)
+	for (TPair<FTopLevelAssetPath, int64>& ClassSizePair : FilteredClassCompressedSizes)
 	{
 		UE_LOG(LogAssetSize, Display, TEXT("    %-40s %-14s (%.1f%%)"), *ClassSizePair.Key.ToString(), *FText::AsNumber(ClassSizePair.Value).ToString(), 100.0 * ClassSizePair.Value / FilteredCompressedSize);
 		if (ShowCount)
@@ -210,7 +210,7 @@ int32 UAssetSizeQueryCommandlet::Main(const FString& FullCommandLine)
 		if (OutputCSVType == EOutputCSVType::Classes)
 		{
 			Lines.Add(TEXT("AssetClass,AssetCount,TotalCompressedSize"));
-			for (TPair<FName, int64>& ClassSizePair : FilteredClassCompressedSizes)
+			for (TPair<FTopLevelAssetPath, int64>& ClassSizePair : FilteredClassCompressedSizes)
 			{
 				// we add to both maps at the same time to we know the lookup succeeds.
 				Lines.Add(FString::Printf(TEXT("%s,%lld,%lld"), *ClassSizePair.Key.ToString(), FilteredClassMatchedAssets[ClassSizePair.Key].Num(), ClassSizePair.Value));
@@ -219,7 +219,7 @@ int32 UAssetSizeQueryCommandlet::Main(const FString& FullCommandLine)
 		else if (OutputCSVType == EOutputCSVType::Assets)
 		{
 			Lines.Add(TEXT("AssetName,AssetType,CompressedSize"));
-			for (TPair<FName, TArray<FMatchedAssetInfo>>& ClassAssetsPair : FilteredClassMatchedAssets)
+			for (TPair<FTopLevelAssetPath, TArray<FMatchedAssetInfo>>& ClassAssetsPair : FilteredClassMatchedAssets)
 			{
 				// we add to both maps at the same time to we know the lookup succeeds.
 				for (const FMatchedAssetInfo& AssetInfo : ClassAssetsPair.Value)
diff --git a/Engine/Source/Editor/UnrealEd/Private/Commandlets/CompileAllBlueprintsCommandlet.cpp b/Engine/Source/Editor/UnrealEd/Private/Commandlets/CompileAllBlueprintsCommandlet.cpp
--- a/Engine/Source/Editor/UnrealEd/Private/Commandlets/CompileAllBlueprintsCommandlet.cpp
+++ b/Engine/Source/Editor/UnrealEd/Private/Commandlets/CompileAllBlueprintsCommandlet.cpp
@@ -22,7 +22,7 @@ UCompileAllBlueprintsCommandlet::UCompileAllBlueprintsCommandlet(const FObjectIn
 	bCookedOnly = false;
 	bDirtyOnly = false;
 	bSimpleAssetList = false;
-	BlueprintBaseClassName = UBlueprint::StaticClass()->GetFName();
+	BlueprintBaseClassName = UBlueprint::StaticClass()->GetClassPathName();
 
 	TotalNumFailedLoads = 0;
 	TotalNumFatalIssues = 0;
@@ -84,7 +84,9 @@ void UCompileAllBlueprintsCommandlet::InitCommandLine(const FString& Params)
 
 	if (SwitchParams.Contains(TEXT("BlueprintBaseClass")))
 	{
-		BlueprintBaseClassName = *SwitchParams[TEXT("BlueprintBaseClass")];
+		FString BlueprintBaseClassParam = *SwitchParams[TEXT("BlueprintBaseClass")];
+		BlueprintBaseClassName = UClass::TryConvertShortTypeNameToPathName<UClass>(BlueprintBaseClassParam, ELogVerbosity::Warning, TEXT("UCompileAllBlueprintsCommandlet::InitCommandLine"));
+		UE_CLOG(BlueprintBaseClassName.IsNull(), LogCompileAllBlueprintsCommandlet, Error, TEXT("Failed to convert short class name -BlueprintBaseClass=\"%s\" to path name"), *BlueprintBaseClassParam);
 	}
 }
 
diff --git a/Engine/Source/Editor/UnrealEd/Private/Commandlets/ContentCommandlets.cpp b/Engine/Source/Editor/UnrealEd/Private/Commandlets/ContentCommandlets.cpp
--- a/Engine/Source/Editor/UnrealEd/Private/Commandlets/ContentCommandlets.cpp
+++ b/Engine/Source/Editor/UnrealEd/Private/Commandlets/ContentCommandlets.cpp
@@ -401,7 +401,7 @@ int32 UResavePackagesCommandlet::InitializeResaveParameters( const TArray<FStrin
 		TSet<FString> RedirectPackages;
 		TSet<FString> ReferencerPackages;
 
-		AssetRegistry.GetAssetsByClass(UObjectRedirector::StaticClass()->GetFName(), RedirectAssets);
+		AssetRegistry.GetAssetsByClass(UObjectRedirector::StaticClass()->GetClassPathName(), RedirectAssets);
 
 		for (const FAssetData& AssetData : RedirectAssets)
 		{
@@ -493,9 +493,18 @@ int32 UResavePackagesCommandlet::InitializeResaveParameters( const TArray<FStrin
 		{
 			TArray<FString> ClassNames;
 			ClassList.ParseIntoArray(ClassNames, TEXT(","), true);
-			for (const FString& ClassName : ClassNames)
+			for (FString& ClassName : ClassNames)
 			{
-				ResaveClasses.AddUnique(*ClassName);
+				if (FPackageName::IsShortPackageName(ClassName))
+				{
+					UE_LOG(LogContentCommandlet, Warning, TEXT("RESAVECLASS param requires class path names. Short name provided: %s."), *ClassName);
+					FTopLevelAssetPath ClassPathName = UClass::TryConvertShortTypeNameToPathName<UStruct>(ClassName, ELogVerbosity::Warning, TEXT("ResavePackages"));
+					if (!ClassPathName.IsNull())
+					{
+						ClassName = ClassPathName.ToString();
+					}
+				}
+				ResaveClasses.AddUnique(ClassName);
 			}
 
 			break;
@@ -520,7 +529,7 @@ int32 UResavePackagesCommandlet::InitializeResaveParameters( const TArray<FStrin
 		for (int32 ClassIndex = 0; ClassIndex < NumResaveClasses; ++ClassIndex)
 		{
 			// Find the class object and then all derived classes
-			UClass* ResaveClass = FindObject<UClass>(ANY_PACKAGE, *ResaveClasses[ClassIndex].ToString());
+			UClass* ResaveClass = UClass::TryFindTypeSlow<UClass>(ResaveClasses[ClassIndex]);
 			if (ResaveClass)
 			{
 				for (TObjectIterator<UClass> It; It; ++It)
@@ -528,7 +537,7 @@ int32 UResavePackagesCommandlet::InitializeResaveParameters( const TArray<FStrin
 					UClass* MaybeChildClass = *It;
 					if (MaybeChildClass->IsChildOf(ResaveClass))
 					{
-						ResaveClasses.AddUnique(MaybeChildClass->GetFName());
+						ResaveClasses.AddUnique(MaybeChildClass->GetPathName());
 					}
 				}
 			}
@@ -1533,8 +1542,8 @@ void UResavePackagesCommandlet::PerformPreloadOperations( FLinkerLoad* PackageLi
 		bSavePackage = false;
 		for (int32 ExportIndex = 0; !bSavePackage && ExportIndex < PackageLinker->ExportMap.Num(); ExportIndex++)
 		{
-			FName ExportClassName = PackageLinker->GetExportClassName(ExportIndex);
-			if (ResaveClasses.Contains(ExportClassName))
+			FTopLevelAssetPath ExportClassPathName(PackageLinker->GetExportClassPackage(ExportIndex), PackageLinker->GetExportClassName(ExportIndex));
+			if (ResaveClasses.Contains(ExportClassPathName.ToString()))
 			{
 				bSavePackage = true;
 				break;
@@ -3282,7 +3291,7 @@ int32 UListMaterialsUsedWithMeshEmittersCommandlet::Main( const FString& Params
 
 	// Retrieve list of all assets, used to find unreferenced ones.
 	TArray<FAssetData> AssetList;
-	AssetRegistryModule.Get().GetAssetsByClass(UParticleSystem::StaticClass()->GetFName(), AssetList, true);
+	AssetRegistryModule.Get().GetAssetsByClass(UParticleSystem::StaticClass()->GetClassPathName(), AssetList, true);
 
 	for(int32 AssetIdx = 0; AssetIdx < AssetList.Num(); ++AssetIdx )
 	{
@@ -3368,7 +3377,7 @@ int32 UListStaticMeshesImportedFromSpeedTreesCommandlet::Main(const FString& Par
 
 	// Retrieve list of all assets, used to find unreferenced ones.
 	TArray<FAssetData> AssetList;
-	AssetRegistryModule.Get().GetAssetsByClass(UStaticMesh::StaticClass()->GetFName(), AssetList, true);
+	AssetRegistryModule.Get().GetAssetsByClass(UStaticMesh::StaticClass()->GetClassPathName(), AssetList, true);
 
 	for (int32 AssetIdx = 0; AssetIdx < AssetList.Num(); ++AssetIdx)
 	{
@@ -3498,7 +3507,7 @@ int32 UStaticMeshMinLodCommandlet::Main(const FString& Params)
 
 	// Retrieve list of all assets, used to find unreferenced ones.
 	TArray<FAssetData> AssetList;
-	AssetRegistryModule.Get().GetAssetsByClass(UStaticMesh::StaticClass()->GetFName(), AssetList, true);
+	AssetRegistryModule.Get().GetAssetsByClass(UStaticMesh::StaticClass()->GetClassPathName(), AssetList, true);
 	TArray<UPackage*> PackagesToSave;
 
 	// Platform (group) names
diff --git a/Engine/Source/Editor/UnrealEd/Private/Commandlets/DiffAssetBulkDataCommandlet.cpp b/Engine/Source/Editor/UnrealEd/Private/Commandlets/DiffAssetBulkDataCommandlet.cpp
--- a/Engine/Source/Editor/UnrealEd/Private/Commandlets/DiffAssetBulkDataCommandlet.cpp
+++ b/Engine/Source/Editor/UnrealEd/Private/Commandlets/DiffAssetBulkDataCommandlet.cpp
@@ -204,8 +204,8 @@ int32 UDiffAssetBulkDataCommandlet::Main(const FString& FullCommandLine)
 		FString TagCurrentValue;
 	};
 
-	TMap<FName /* TagName */, TMap<FName /* AssetClass */, TArray<FDiffResult>>> Results;
-	TMap<FName /* AssetClass */, TArray<FName /* PackageName */>> NoTagPackagesByAssumedClass;
+	TMap<FName /* TagName */, TMap<FTopLevelAssetPath /* AssetClass */, TArray<FDiffResult>>> Results;
+	TMap<FTopLevelAssetPath /* AssetClass */, TArray<FName /* PackageName */>> NoTagPackagesByAssumedClass;
 	TArray<FName /* PackageName */> PackagesWithUnassignableDiffsAndUntaggedAssets;
 	TArray<FName /* PackageName */> PackagesWithUnassignableDiffs;
 	
@@ -289,11 +289,11 @@ int32 UDiffAssetBulkDataCommandlet::Main(const FString& FullCommandLine)
 			FAssetData const* RepresentativeAsset = UE::AssetRegistry::GetMostImportantAsset(CurrentAssetDatas, true);
 			if (RepresentativeAsset == nullptr)
 			{
-				NoTagPackagesByAssumedClass.FindOrAdd(NAME_None).Add(ChangedPackageName);
+				NoTagPackagesByAssumedClass.FindOrAdd(FTopLevelAssetPath()).Add(ChangedPackageName);
 			}
 			else
 			{
-				NoTagPackagesByAssumedClass.FindOrAdd(RepresentativeAsset->AssetClass).Add(ChangedPackageName);
+				NoTagPackagesByAssumedClass.FindOrAdd(RepresentativeAsset->AssetClassPath).Add(ChangedPackageName);
 			}
 			
 			continue;
@@ -311,9 +311,9 @@ int32 UDiffAssetBulkDataCommandlet::Main(const FString& FullCommandLine)
 			{
 				if (Tag.BaseValue != Tag.CurrentValue)
 				{
-					TMap<FName, TArray<FDiffResult>>& TagResults = Results.FindOrAdd(Tag.TagName);
+					TMap<FTopLevelAssetPath, TArray<FDiffResult>>& TagResults = Results.FindOrAdd(Tag.TagName);
 
-					TArray<FDiffResult>& ClassResults = TagResults.FindOrAdd(Tag.BaseAssetData->AssetClass);
+					TArray<FDiffResult>& ClassResults = TagResults.FindOrAdd(Tag.BaseAssetData->AssetClassPath);
 
 					FDiffResult& Result = ClassResults.AddDefaulted_GetRef();
 					Result.ChangedAssetObjectPath = Tag.BaseAssetData->ObjectPath.ToString();
@@ -354,14 +354,14 @@ int32 UDiffAssetBulkDataCommandlet::Main(const FString& FullCommandLine)
 		return 0;
 	}
 
-	TArray<FName>& CantDetermineAssetClassPackages = NoTagPackagesByAssumedClass.FindOrAdd(NAME_None);
+	TArray<FName>& CantDetermineAssetClassPackages = NoTagPackagesByAssumedClass.FindOrAdd(FTopLevelAssetPath());
 
 	UE_LOG(LogDiffAssetBulk, Display, TEXT("Changed package breakdown:"));
 	UE_LOG(LogDiffAssetBulk, Display, TEXT("    No blame information available:"));
 	UE_LOG(LogDiffAssetBulk, Display, TEXT("        Can't determine asset class   : %d"), CantDetermineAssetClassPackages.Num());
-	for (TPair<FName, TArray<FName>> ClassPackages : NoTagPackagesByAssumedClass)
+	for (TPair<FTopLevelAssetPath, TArray<FName>> ClassPackages : NoTagPackagesByAssumedClass)
 	{
-		if (ClassPackages.Key == NAME_None)
+		if (ClassPackages.Key == FTopLevelAssetPath())
 		{
 			continue;
 		}
@@ -406,10 +406,10 @@ int32 UDiffAssetBulkDataCommandlet::Main(const FString& FullCommandLine)
 	{
 		UE_LOG(LogDiffAssetBulk, Display, TEXT("    Summary changes by blame tag:"));
 
-		for (TPair<FName, TMap<FName, TArray<FDiffResult>>>& TagResults : Results)
+		for (TPair<FName, TMap<FTopLevelAssetPath, TArray<FDiffResult>>>& TagResults : Results)
 		{
 			uint32 TagCount = 0;
-			for (TPair<FName, TArray<FDiffResult>>& ClassResults : TagResults.Value)
+			for (TPair<FTopLevelAssetPath, TArray<FDiffResult>>& ClassResults : TagResults.Value)
 			{
 				TagCount += ClassResults.Value.Num();
 			}
@@ -427,11 +427,11 @@ int32 UDiffAssetBulkDataCommandlet::Main(const FString& FullCommandLine)
 
 		UE_LOG(LogDiffAssetBulk, Display, TEXT("    Asset changes by blame tag:"));
 
-		for (TPair<FName, TMap<FName, TArray<FDiffResult>>>& TagResults : Results)
+		for (TPair<FName, TMap<FTopLevelAssetPath, TArray<FDiffResult>>>& TagResults : Results)
 		{
 			UE_LOG(LogDiffAssetBulk, Display, TEXT("        %s"), *TagResults.Key.ToString());
 
-			for (TPair<FName, TArray<FDiffResult>>& ClassResults : TagResults.Value)
+			for (TPair<FTopLevelAssetPath, TArray<FDiffResult>>& ClassResults : TagResults.Value)
 			{
 				Algo::SortBy(ClassResults.Value, &FDiffResult::ChangedAssetObjectPath);
 				UE_LOG(LogDiffAssetBulk, Display, TEXT("            %s [%d]"), *ClassResults.Key.ToString(), ClassResults.Value.Num());
diff --git a/Engine/Source/Editor/UnrealEd/Private/Commandlets/DiffAssetRegistriesCommandlet.cpp b/Engine/Source/Editor/UnrealEd/Private/Commandlets/DiffAssetRegistriesCommandlet.cpp
--- a/Engine/Source/Editor/UnrealEd/Private/Commandlets/DiffAssetRegistriesCommandlet.cpp
+++ b/Engine/Source/Editor/UnrealEd/Private/Commandlets/DiffAssetRegistriesCommandlet.cpp
@@ -525,10 +525,10 @@ FName UDiffAssetRegistriesCommandlet::GetClassName(FAssetRegistryState& InRegist
 	{
 		TArrayView<FAssetData const * const> Assets = InRegistryState.GetAssetsByPackageName(InAssetPath);
 
-		FName NewName = NAME_None;
+		FName NewName;
 		if (Assets.Num() > 0)
 		{
-			NewName = Assets[0]->AssetClass;
+			NewName = Assets[0]->AssetClassPath.GetAssetName();
 		}
 		else
 		{
diff --git a/Engine/Source/Editor/UnrealEd/Private/Commandlets/DumpMaterialExpressionsCommandlet.cpp b/Engine/Source/Editor/UnrealEd/Private/Commandlets/DumpMaterialExpressionsCommandlet.cpp
--- a/Engine/Source/Editor/UnrealEd/Private/Commandlets/DumpMaterialExpressionsCommandlet.cpp
+++ b/Engine/Source/Editor/UnrealEd/Private/Commandlets/DumpMaterialExpressionsCommandlet.cpp
@@ -257,7 +257,7 @@ void WriteOutMaterialFunctions(FArchive* FileWriter, const FString& OutputFilePa
 	// See UMaterialGraphSchema::GetMaterialFunctionActions for reference
 	TArray<FAssetData> AssetDataList;
 	FAssetRegistryModule& AssetRegistryModule = FModuleManager::LoadModuleChecked<FAssetRegistryModule>(TEXT("AssetRegistry"));
-	AssetRegistryModule.Get().GetAssetsByClass(UMaterialFunction::StaticClass()->GetFName(), AssetDataList);
+	AssetRegistryModule.Get().GetAssetsByClass(UMaterialFunction::StaticClass()->GetClassPathName(), AssetDataList);
 	for (const FAssetData& AssetData : AssetDataList)
 	{
 		// If this is a function that is selected to be exposed to the library
