Deprecating ANY_PACKAGE.

diff --git a/Engine/Source/Editor/UnrealEd/Private/Commandlets/DumpMaterialShaderTypes.cpp b/Engine/Source/Editor/UnrealEd/Private/Commandlets/DumpMaterialShaderTypes.cpp
--- a/Engine/Source/Editor/UnrealEd/Private/Commandlets/DumpMaterialShaderTypes.cpp
+++ b/Engine/Source/Editor/UnrealEd/Private/Commandlets/DumpMaterialShaderTypes.cpp
@@ -565,16 +565,16 @@ int32 UDumpMaterialShaderTypesCommandlet::Main(const FString& Params)
 				FARFilter Filter;
 				Filter.PackagePaths.Add(FName(TEXT("/Game")));
 				Filter.bRecursivePaths = true;
-				Filter.ClassNames.Add(UMaterial::StaticClass()->GetFName());
+				Filter.ClassPaths.Add(UMaterial::StaticClass()->GetClassPathName());
 
 				FCollectionManagerModule& CollectionManagerModule = FCollectionManagerModule::GetModule();
 				CollectionManagerModule.Get().GetObjectsInCollection(FName(*CollectionName), ECollectionShareType::CST_All, Filter.ObjectPaths, ECollectionRecursionFlags::SelfAndChildren);
 
 				AssetRegistry.GetAssets(Filter, MaterialList);
 
-				Filter.ClassNames.Empty();
-				Filter.ClassNames.Add(UMaterialInstance::StaticClass()->GetFName());
-				Filter.ClassNames.Add(UMaterialInstanceConstant::StaticClass()->GetFName());
+				Filter.ClassPaths.Empty();
+				Filter.ClassPaths.Add(UMaterialInstance::StaticClass()->GetClassPathName());
+				Filter.ClassPaths.Add(UMaterialInstanceConstant::StaticClass()->GetClassPathName());
 
 				AssetRegistry.GetAssets(Filter, MaterialInstanceList);
 			}
@@ -583,8 +583,8 @@ int32 UDumpMaterialShaderTypesCommandlet::Main(const FString& Params)
 		{
 			if (!AssetRegistry.IsLoadingAssets())
 			{
-				AssetRegistry.GetAssetsByClass(UMaterial::StaticClass()->GetFName(), MaterialList, true);
-				AssetRegistry.GetAssetsByClass(UMaterialInstance::StaticClass()->GetFName(), MaterialInstanceList, true);
+				AssetRegistry.GetAssetsByClass(UMaterial::StaticClass()->GetClassPathName(), MaterialList, true);
+				AssetRegistry.GetAssetsByClass(UMaterialInstance::StaticClass()->GetClassPathName(), MaterialInstanceList, true);
 			}
 		}
 	}
diff --git a/Engine/Source/Editor/UnrealEd/Private/Commandlets/ExportDialogueScriptCommandlet.cpp b/Engine/Source/Editor/UnrealEd/Private/Commandlets/ExportDialogueScriptCommandlet.cpp
--- a/Engine/Source/Editor/UnrealEd/Private/Commandlets/ExportDialogueScriptCommandlet.cpp
+++ b/Engine/Source/Editor/UnrealEd/Private/Commandlets/ExportDialogueScriptCommandlet.cpp
@@ -180,7 +180,7 @@ int32 UExportDialogueScriptCommandlet::Main(const FString& Params)
 
 	// We want all the non-localized project specific dialogue waves
 	TArray<FAssetData> AssetDataArrayForDialogueWaves;
-	if (!FLocalizedAssetUtil::GetAssetsByPathAndClass(AssetRegistry, *RootAssetPath, UDialogueWave::StaticClass()->GetFName(), /*bIncludeLocalizedAssets*/false, AssetDataArrayForDialogueWaves))
+	if (!FLocalizedAssetUtil::GetAssetsByPathAndClass(AssetRegistry, *RootAssetPath, UDialogueWave::StaticClass()->GetClassPathName(), /*bIncludeLocalizedAssets*/false, AssetDataArrayForDialogueWaves))
 	{
 		UE_LOG(LogExportDialogueScriptCommandlet, Error, TEXT("Unable to get dialogue wave asset data from asset registry."));
 		return -1;
diff --git a/Engine/Source/Editor/UnrealEd/Private/Commandlets/ExportPakDependencies.cpp b/Engine/Source/Editor/UnrealEd/Private/Commandlets/ExportPakDependencies.cpp
--- a/Engine/Source/Editor/UnrealEd/Private/Commandlets/ExportPakDependencies.cpp
+++ b/Engine/Source/Editor/UnrealEd/Private/Commandlets/ExportPakDependencies.cpp
@@ -155,7 +155,7 @@ bool ExportDependencies(const TCHAR * PakFilename, const TCHAR* GameName, const
 			}
 		}
 
-		TMap<FName,FName> PackageToClassMap;
+		TMap<FName, FTopLevelAssetPath> PackageToClassMap;
 
 		// Combine with dependency information from asset registry.
 		{
@@ -169,7 +169,7 @@ bool ExportDependencies(const TCHAR * PakFilename, const TCHAR* GameName, const
 			for( int i=0; i<AssetData.Num(); i++ )
 			{
 				PackageNames.Add(AssetData[i].PackageName);
-				PackageToClassMap.Add(AssetData[i].PackageName,AssetData[i].AssetClass);
+				PackageToClassMap.Add(AssetData[i].PackageName,AssetData[i].AssetClassPath);
 			}
 
 			for( const auto& PackageName : PackageNames )
@@ -241,7 +241,7 @@ bool ExportDependencies(const TCHAR * PakFilename, const TCHAR* GameName, const
 				CSVFileWriter->Logf(TEXT("class,name,inclusive,exclusive"));
 				for(auto Package : AllPackages)
 				{
-					FName ClassName = PackageToClassMap.FindRef(Package->Name);
+					FTopLevelAssetPath ClassName = PackageToClassMap.FindRef(Package->Name);
 					CSVFileWriter->Logf(TEXT("%s,%s,%i,%i"),*ClassName.ToString(),*Package->Name.ToString(),Package->InclusiveSize,Package->ExclusiveSize);
 				}
 				CSVFileWriter->Close();
diff --git a/Engine/Source/Editor/UnrealEd/Private/Commandlets/FixupNeedsLoadForEditorGameCommandlet.cpp b/Engine/Source/Editor/UnrealEd/Private/Commandlets/FixupNeedsLoadForEditorGameCommandlet.cpp
--- a/Engine/Source/Editor/UnrealEd/Private/Commandlets/FixupNeedsLoadForEditorGameCommandlet.cpp
+++ b/Engine/Source/Editor/UnrealEd/Private/Commandlets/FixupNeedsLoadForEditorGameCommandlet.cpp
@@ -15,16 +15,16 @@ int32 UFixupNeedsLoadForEditorGameCommandlet::InitializeResaveParameters(const T
 	// We need ResaveClasses to be specified, otherwise we won't know what to update
 	if (Result == 0 && !ResaveClasses.Num())
 	{
-		UE_LOG(LogContentCommandlet, Error, TEXT("FixupNeedsLoadForEditorGame commandlet requires at least one resave class name. Use -RESAVECLASS=ClassA,ClassB,ClassC to specify resave classes."));
+		UE_LOG(LogContentCommandlet, Error, TEXT("FixupNeedsLoadForEditorGame commandlet requires at least one resave class name. Use -RESAVECLASS=/Path/To.ClassA,/Path/To.ClassB,/Path/To.ClassC to specify resave classes."));
 		Result = 1;
 	}
 	else
 	{
-		for (FName& ClassName : ResaveClasses)
+		for (FString& ClassName : ResaveClasses)
 		{			
 			if (!ResaveClassNeedsLoadForEditorGameValues.Contains(ClassName))
 			{
-				UClass* ResaveClass = FindObject<UClass>(ANY_PACKAGE, *ClassName.ToString());
+				UClass* ResaveClass = UClass::TryFindTypeSlow<UClass>(ClassName);
 				if (ResaveClass)
 				{
 					UObject* DefaultObject = ResaveClass->GetDefaultObject();
@@ -33,7 +33,7 @@ int32 UFixupNeedsLoadForEditorGameCommandlet::InitializeResaveParameters(const T
 			}
 			else if (Verbosity != UResavePackagesCommandlet::ONLY_ERRORS)
 			{
-				UE_LOG(LogContentCommandlet, Warning, TEXT("Resave Class \"%s\" could not be found. Make sure the class name is valid and that it's a native class."), *ClassName.ToString());
+				UE_LOG(LogContentCommandlet, Warning, TEXT("Resave Class \"%s\" could not be found. Make sure the class name is valid and that it's a native class."), *ClassName);
 			}
 		}
 		if (ResaveClassNeedsLoadForEditorGameValues.Num() == 0)
@@ -55,8 +55,8 @@ void UFixupNeedsLoadForEditorGameCommandlet::PerformPreloadOperations(FLinkerLoa
 		bSavePackage = false;
 		for (int32 ExportIndex = 0; ExportIndex < PackageLinker->ExportMap.Num(); ExportIndex++)
 		{
-			FName ExportClassName = PackageLinker->GetExportClassName(ExportIndex);
-			bool* bNeedsLoadForEditorGameValuePtr = ResaveClassNeedsLoadForEditorGameValues.Find(ExportClassName);
+			FTopLevelAssetPath ExportClassPathName(PackageLinker->GetExportClassPackage(ExportIndex), PackageLinker->GetExportClassName(ExportIndex));
+			bool* bNeedsLoadForEditorGameValuePtr = ResaveClassNeedsLoadForEditorGameValues.Find(ExportClassPathName.ToString());
 			if (bNeedsLoadForEditorGameValuePtr)
 			{
 				FObjectExport& Export = PackageLinker->ExportMap[ExportIndex];
diff --git a/Engine/Source/Editor/UnrealEd/Private/Commandlets/GatherTextCommandlet.cpp b/Engine/Source/Editor/UnrealEd/Private/Commandlets/GatherTextCommandlet.cpp
--- a/Engine/Source/Editor/UnrealEd/Private/Commandlets/GatherTextCommandlet.cpp
+++ b/Engine/Source/Editor/UnrealEd/Private/Commandlets/GatherTextCommandlet.cpp
@@ -152,7 +152,7 @@ int32 UGatherTextCommandlet::ProcessGatherConfig(const FString& GatherTextConfig
 		FString PlatformSplitModeString;
 		if (GetStringFromConfig(TEXT("CommonSettings"), TEXT("PlatformSplitMode"), PlatformSplitModeString, GatherTextConfigPath))
 		{
-			UEnum* PlatformSplitModeEnum = FindObjectChecked<UEnum>(ANY_PACKAGE, TEXT("ELocTextPlatformSplitMode"));
+			UEnum* PlatformSplitModeEnum = FindObjectChecked<UEnum>(nullptr, TEXT("/Script/Localization.ELocTextPlatformSplitMode"));
 			const int64 PlatformSplitModeInt = PlatformSplitModeEnum->GetValueByName(*PlatformSplitModeString);
 			if (PlatformSplitModeInt != INDEX_NONE)
 			{
@@ -219,7 +219,7 @@ int32 UGatherTextCommandlet::ProcessGatherConfig(const FString& GatherTextConfig
 	{
 		FString CommandletClassName = GConfig->GetStr( *StepName, TEXT("CommandletClass"), GatherTextConfigPath ) + TEXT("Commandlet");
 
-		UClass* CommandletClass = FindObject<UClass>(ANY_PACKAGE,*CommandletClassName,false);
+		UClass* CommandletClass = FindFirstObject<UClass>(*CommandletClassName, EFindFirstObjectOptions::None, ELogVerbosity::Warning, TEXT("UGatherTextCommandlet::ProcessGatherConfig"));
 		if (!CommandletClass)
 		{
 			UE_LOG(LogGatherTextCommandlet, Error, TEXT("The commandlet name %s in section %s is invalid."), *CommandletClassName, *StepName);
diff --git a/Engine/Source/Editor/UnrealEd/Private/Commandlets/GatherTextFromAssetsCommandlet.cpp b/Engine/Source/Editor/UnrealEd/Private/Commandlets/GatherTextFromAssetsCommandlet.cpp
--- a/Engine/Source/Editor/UnrealEd/Private/Commandlets/GatherTextFromAssetsCommandlet.cpp
+++ b/Engine/Source/Editor/UnrealEd/Private/Commandlets/GatherTextFromAssetsCommandlet.cpp
@@ -526,11 +526,19 @@ int32 UGatherTextFromAssetsCommandlet::Main(const FString& Params)
 		if (ShouldExcludeDerivedClasses)
 		{
 			FirstPassFilter.bRecursiveClasses = true;
-			FirstPassFilter.ClassNames.Add(TEXT("Object"));
+			FirstPassFilter.ClassPaths.Add(UObject::StaticClass()->GetClassPathName());
 			for (const FString& ExcludeClassName : ExcludeClassNames)
 			{
-				// Note: Can't necessarily validate these class names here, as the class may be a generated blueprint class that hasn't been loaded yet.
-				FirstPassFilter.RecursiveClassesExclusionSet.Add(*ExcludeClassName);
+				FTopLevelAssetPath ExcludedClassPathName = UClass::TryConvertShortTypeNameToPathName<UClass>(ExcludeClassName, ELogVerbosity::Warning, TEXT("GatherTextFromAssetsCommandlet"));
+				if (!ExcludedClassPathName.IsNull())
+				{
+					// Note: Can't necessarily validate these class names here, as the class may be a generated blueprint class that hasn't been loaded yet.
+					FirstPassFilter.RecursiveClassPathsExclusionSet.Add(FTopLevelAssetPath(ExcludeClassName));
+				}
+				else
+				{
+					UE_CLOG(!ExcludeClassName.IsEmpty(), LogGatherTextFromAssetsCommandlet, Error, TEXT("Unable to convert short class name \"%s\" to path name. Please use path names fo ExcludeClassNames"), *ExcludeClassName);
+				}
 			}
 		}
 
@@ -552,8 +560,17 @@ int32 UGatherTextFromAssetsCommandlet::Main(const FString& Params)
 		ExcludeExactClassesFilter.bRecursiveClasses = false;
 		for (const FString& ExcludeClassName : ExcludeClassNames)
 		{
-			// Note: Can't necessarily validate these class names here, as the class may be a generated blueprint class that hasn't been loaded yet.
-			ExcludeExactClassesFilter.ClassNames.Add(*ExcludeClassName);
+			FTopLevelAssetPath ExcludedClassPathName = UClass::TryConvertShortTypeNameToPathName<UClass>(ExcludeClassName, ELogVerbosity::Warning, TEXT("GatherTextFromAssetsCommandlet"));
+			if (!ExcludedClassPathName.IsNull())
+			{
+				// Note: Can't necessarily validate these class names here, as the class may be a generated blueprint class that hasn't been loaded yet.
+				ExcludeExactClassesFilter.ClassPaths.Add(FTopLevelAssetPath(ExcludeClassName));
+			}
+			else
+			{
+				UE_CLOG(!ExcludeClassName.IsEmpty(), LogGatherTextFromAssetsCommandlet, Error, TEXT("Unable to convert short class name \"%s\" to path name. Please use path names fo ExcludeClassNames"), *ExcludeClassName);
+			}
+
 		}
 
 		// Reapply filter over the current set of assets.
@@ -715,7 +732,7 @@ int32 UGatherTextFromAssetsCommandlet::Main(const FString& Params)
 			AssetRegistry.GetAssetsByPackageName(PackagePendingGather.PackageName, AllAssetDataInSamePackage);
 			for (const FAssetData& AssetData : AllAssetDataInSamePackage)
 			{
-				if (AssetData.AssetClass == UDialogueWave::StaticClass()->GetFName())
+				if (AssetData.AssetClassPath == UDialogueWave::StaticClass()->GetClassPathName())
 				{
 					PackagePendingGather.PackageLocCacheState = EPackageLocCacheState::Uncached_TooOld;
 				}
diff --git a/Engine/Source/Editor/UnrealEd/Private/Commandlets/GatherTextFromMetadataCommandlet.cpp b/Engine/Source/Editor/UnrealEd/Private/Commandlets/GatherTextFromMetadataCommandlet.cpp
--- a/Engine/Source/Editor/UnrealEd/Private/Commandlets/GatherTextFromMetadataCommandlet.cpp
+++ b/Engine/Source/Editor/UnrealEd/Private/Commandlets/GatherTextFromMetadataCommandlet.cpp
@@ -165,7 +165,7 @@ int32 UGatherTextFromMetaDataCommandlet::Main( const FString& Params )
 				else
 				{
 					const FFieldClass* FieldClass = FFieldClass::GetNameToFieldClassMap().FindRef(*FieldTypeStr);
-					const UClass* UFieldClass = FindObject<UClass>(ANY_PACKAGE, *FieldTypeStr);
+					const UClass* UFieldClass = FindFirstObject<UClass>(*FieldTypeStr, EFindFirstObjectOptions::None, ELogVerbosity::Warning, TEXT("Looking for field types to include or exclude in GatherTextFromMetadata commandlet"));
 					if (!FieldClass && !UFieldClass)
 					{
 						UE_LOG(LogGatherTextFromMetaDataCommandlet, Warning, TEXT("Field Type %s was not found (from %s in section %s). Did you forget a ModulesToPreload entry?"), *FieldTypeStr, InConfigKey, *SectionName);
@@ -216,7 +216,7 @@ int32 UGatherTextFromMetaDataCommandlet::Main( const FString& Params )
 				}
 				else
 				{
-					const UStruct* FieldOwnerType = FindObject<UStruct>(ANY_PACKAGE, *FieldOwnerTypeStr);
+					const UStruct* FieldOwnerType = FindFirstObject<UStruct>(*FieldOwnerTypeStr, EFindFirstObjectOptions::EnsureIfAmbiguous);
 					if (!FieldOwnerType)
 					{
 						UE_LOG(LogGatherTextFromMetaDataCommandlet, Warning, TEXT("Field Owner Type %s was not found (from %s in section %s). Did you forget a ModulesToPreload entry?"), *FieldOwnerTypeStr, InConfigKey, *SectionName);
diff --git a/Engine/Source/Editor/UnrealEd/Private/Commandlets/GenerateAssetManifestCommandlet.cpp b/Engine/Source/Editor/UnrealEd/Private/Commandlets/GenerateAssetManifestCommandlet.cpp
--- a/Engine/Source/Editor/UnrealEd/Private/Commandlets/GenerateAssetManifestCommandlet.cpp
+++ b/Engine/Source/Editor/UnrealEd/Private/Commandlets/GenerateAssetManifestCommandlet.cpp
@@ -135,7 +135,15 @@ int32 UGenerateAssetManifestCommandlet::Main(const FString& InParams)
 		Filter.bRecursivePaths = true;
 		for (const FString& IncludedClass : IncludedClasses)
 		{
-			Filter.ClassNames.AddUnique(*IncludedClass);
+			FTopLevelAssetPath IncludedClassPathName = UClass::TryConvertShortTypeNameToPathName<UStruct>(IncludedClass, ELogVerbosity::Error, TEXT("UGenerateAssetManifestCommandlet::Main"));
+			if (IncludedClassPathName.IsNull())
+			{
+				UE_LOG(LogGenerateAssetManifestCommandlet, Error, TEXT("Failed to convert short class name \"%s\" to path name. Please use class path names for IncludedClasses."), *IncludedClass);
+			}
+			else
+			{
+				Filter.ClassPaths.AddUnique(IncludedClassPathName);
+			}
 		}
 		TArray<FAssetData> AssetList;
 		AssetRegistryModule.Get().GetAssets(Filter, AssetList);
@@ -169,7 +177,15 @@ int32 UGenerateAssetManifestCommandlet::Main(const FString& InParams)
 		Filter.bRecursivePaths = true;
 		for (const FString& ExcludedClass : ExcludedClasses)
 		{
-			Filter.ClassNames.AddUnique(*ExcludedClass);
+			FTopLevelAssetPath ExcludedClassPathName = UClass::TryConvertShortTypeNameToPathName<UStruct>(ExcludedClass, ELogVerbosity::Error, TEXT("UGenerateAssetManifestCommandlet::Main"));
+			if (ExcludedClassPathName.IsNull())
+			{
+				UE_LOG(LogGenerateAssetManifestCommandlet, Error, TEXT("Failed to convert short class name \"%s\" to path name. Please use class path names for ExcludedClasses."), *ExcludedClass);
+			}
+			else
+			{
+				Filter.ClassPaths.AddUnique(ExcludedClassPathName);
+			}
 		}
 		TArray<FAssetData> AssetList;
 		AssetRegistryModule.Get().GetAssets(Filter, AssetList);
diff --git a/Engine/Source/Editor/UnrealEd/Private/Commandlets/GenerateBlueprintAPICommandlet.cpp b/Engine/Source/Editor/UnrealEd/Private/Commandlets/GenerateBlueprintAPICommandlet.cpp
--- a/Engine/Source/Editor/UnrealEd/Private/Commandlets/GenerateBlueprintAPICommandlet.cpp
+++ b/Engine/Source/Editor/UnrealEd/Private/Commandlets/GenerateBlueprintAPICommandlet.cpp
@@ -285,7 +285,7 @@ GenerateBlueprintAPIUtils::CommandletOptions::CommandletOptions(TArray<FString>
 		{
 			FString ClassSwitch, ClassName;
 			Switch.Split(TEXT("="), &ClassSwitch, &ClassName);
-			BlueprintClass = FindObject<UClass>(ANY_PACKAGE, *ClassName);
+			BlueprintClass = FindFirstObject<UClass>(*ClassName, EFindFirstObjectOptions::None, ELogVerbosity::Warning, TEXT("parsing GenerateBlueprintAPIUtils class"));
 
 			if (BlueprintClass == nullptr)
 			{
@@ -297,7 +297,7 @@ GenerateBlueprintAPIUtils::CommandletOptions::CommandletOptions(TArray<FString>
 		{
 			FString ClassSwitch, ClassName;
 			Switch.Split(TEXT("="), &ClassSwitch, &ClassName);
-			PaletteFilter = FindObject<UClass>(ANY_PACKAGE, *ClassName);
+			PaletteFilter = FindFirstObject<UClass>(*ClassName, EFindFirstObjectOptions::None, ELogVerbosity::Warning, TEXT("parsing GenerateBlueprintAPIUtils palfilter"));
 
 			NewDumpFlags |= BPDUMP_FilteredPalette;
 			if (PaletteFilter == nullptr)
diff --git a/Engine/Source/Editor/UnrealEd/Private/Commandlets/ImportLocalizedDialogueCommandlet.cpp b/Engine/Source/Editor/UnrealEd/Private/Commandlets/ImportLocalizedDialogueCommandlet.cpp
--- a/Engine/Source/Editor/UnrealEd/Private/Commandlets/ImportLocalizedDialogueCommandlet.cpp
+++ b/Engine/Source/Editor/UnrealEd/Private/Commandlets/ImportLocalizedDialogueCommandlet.cpp
@@ -156,7 +156,7 @@ int32 UImportLocalizedDialogueCommandlet::Main(const FString& Params)
 
 	// We want all the non-localized project specific dialogue waves
 	TArray<FAssetData> AssetDataArrayForDialogueWaves;
-	if (!FLocalizedAssetUtil::GetAssetsByPathAndClass(AssetRegistry, *RootAssetPath, UDialogueWave::StaticClass()->GetFName(), /*bIncludeLocalizedAssets*/false, AssetDataArrayForDialogueWaves))
+	if (!FLocalizedAssetUtil::GetAssetsByPathAndClass(AssetRegistry, *RootAssetPath, UDialogueWave::StaticClass()->GetClassPathName(), /*bIncludeLocalizedAssets*/false, AssetDataArrayForDialogueWaves))
 	{
 		UE_LOG(LogImportLocalizedDialogueCommandlet, Error, TEXT("Unable to get dialogue wave asset data from asset registry."));
 		return -1;
@@ -193,8 +193,8 @@ int32 UImportLocalizedDialogueCommandlet::Main(const FString& Params)
 			LocalizedSoundWavePathsToSearch.Add(*CultureImportInfo.LocalizedImportedDialoguePackagePath);
 		}
 
-		FLocalizedAssetUtil::GetAssetsByPathAndClass(AssetRegistry, LocalizedDialogueWavePathsToSearch, UDialogueWave::StaticClass()->GetFName(), /*bIncludeLocalizedAssets*/true, LocalizedAssetsToPotentiallyDelete);
-		FLocalizedAssetUtil::GetAssetsByPathAndClass(AssetRegistry, LocalizedSoundWavePathsToSearch, USoundWave::StaticClass()->GetFName(), /*bIncludeLocalizedAssets*/true, LocalizedAssetsToPotentiallyDelete);
+		FLocalizedAssetUtil::GetAssetsByPathAndClass(AssetRegistry, LocalizedDialogueWavePathsToSearch, UDialogueWave::StaticClass()->GetClassPathName(), /*bIncludeLocalizedAssets*/true, LocalizedAssetsToPotentiallyDelete);
+		FLocalizedAssetUtil::GetAssetsByPathAndClass(AssetRegistry, LocalizedSoundWavePathsToSearch, USoundWave::StaticClass()->GetClassPathName(), /*bIncludeLocalizedAssets*/true, LocalizedAssetsToPotentiallyDelete);
 	}
 
 	// We're going to walk every context from every dialogue wave asset looking to see whether there's new audio to import for each culture we generate for
diff --git a/Engine/Source/Editor/UnrealEd/Private/Commandlets/PackageUtilities.cpp b/Engine/Source/Editor/UnrealEd/Private/Commandlets/PackageUtilities.cpp
--- a/Engine/Source/Editor/UnrealEd/Private/Commandlets/PackageUtilities.cpp
+++ b/Engine/Source/Editor/UnrealEd/Private/Commandlets/PackageUtilities.cpp
@@ -1386,7 +1386,7 @@ void FPkgInfoReporter_Log::GeneratePackageReport( FLinkerLoad* InLinker /*=nullp
 			for (FAssetData* AssetData : AssetDatas)
 			{
 				// Display the asset class and path
-				Out.Logf(ELogVerbosity::Display, TEXT("\t\t%d) %s'%s' (%d Tags)"), AssetIdx++, *AssetData->AssetClass.ToString(), *AssetData->ObjectPath.ToString(), AssetData->TagsAndValues.Num());
+				Out.Logf(ELogVerbosity::Display, TEXT("\t\t%d) %s'%s' (%d Tags)"), AssetIdx++, *AssetData->AssetClassPath.ToString(), *AssetData->ObjectPath.ToString(), AssetData->TagsAndValues.Num());
 
 				// Display all tags on the asset
 				for (const TPair<FName, FAssetTagValueRef>& Pair : AssetData->TagsAndValues)
diff --git a/Engine/Source/Editor/UnrealEd/Private/Commandlets/ParticleSystemAuditCommandlet.cpp b/Engine/Source/Editor/UnrealEd/Private/Commandlets/ParticleSystemAuditCommandlet.cpp
--- a/Engine/Source/Editor/UnrealEd/Private/Commandlets/ParticleSystemAuditCommandlet.cpp
+++ b/Engine/Source/Editor/UnrealEd/Private/Commandlets/ParticleSystemAuditCommandlet.cpp
@@ -77,7 +77,7 @@ bool UParticleSystemAuditCommandlet::ProcessParticleSystems()
 	Filter.PackagePaths = PackagePaths;
 	Filter.bRecursivePaths = true;
 
-	Filter.ClassNames.Add(UParticleSystem::StaticClass()->GetFName());
+	Filter.ClassPaths.Add(UParticleSystem::StaticClass()->GetClassPathName());
 	if (!FilterCollection.IsEmpty())
 	{
 		FCollectionManagerModule& CollectionManagerModule = FCollectionManagerModule::GetModule();
diff --git a/Engine/Source/Editor/UnrealEd/Private/Commandlets/PopulateDialogueWaveFromCharacterSheetCommandlet.cpp b/Engine/Source/Editor/UnrealEd/Private/Commandlets/PopulateDialogueWaveFromCharacterSheetCommandlet.cpp
--- a/Engine/Source/Editor/UnrealEd/Private/Commandlets/PopulateDialogueWaveFromCharacterSheetCommandlet.cpp
+++ b/Engine/Source/Editor/UnrealEd/Private/Commandlets/PopulateDialogueWaveFromCharacterSheetCommandlet.cpp
@@ -88,7 +88,7 @@ int32 UPopulateDialogueWaveFromCharacterSheetCommandlet::Main(const FString& Par
 
 	// We only want dialogue wave assets that exist within the Game content directory.
 	TArray<FAssetData> AllDialogueWaves;
-	if (!FLocalizedAssetUtil::GetAssetsByPathAndClass(AssetRegistry, FName("/Game"), UDialogueWave::StaticClass()->GetFName(), /*bIncludeLocalizedAssets*/false, AllDialogueWaves))
+	if (!FLocalizedAssetUtil::GetAssetsByPathAndClass(AssetRegistry, FName("/Game"), UDialogueWave::StaticClass()->GetClassPathName(), /*bIncludeLocalizedAssets*/false, AllDialogueWaves))
 	{
 		UE_LOG(LogPopulateDialogueWaveFromCharacterSheetCommandlet, Error, TEXT("Unable to get dialogue wave asset data from asset registry."));
 		return -1;
@@ -96,7 +96,7 @@ int32 UPopulateDialogueWaveFromCharacterSheetCommandlet::Main(const FString& Par
 
 	// We only want sound wave assets that exist within the Game content directory.
 	TArray<FAssetData> AllSoundWaves;
-	if (!FLocalizedAssetUtil::GetAssetsByPathAndClass(AssetRegistry, FName("/Game"), USoundWave::StaticClass()->GetFName(), /*bIncludeLocalizedAssets*/false, AllSoundWaves))
+	if (!FLocalizedAssetUtil::GetAssetsByPathAndClass(AssetRegistry, FName("/Game"), USoundWave::StaticClass()->GetClassPathName(), /*bIncludeLocalizedAssets*/false, AllSoundWaves))
 	{
 		UE_LOG(LogPopulateDialogueWaveFromCharacterSheetCommandlet, Error, TEXT("Unable to get sound wave asset data from asset registry."));
 		return -1;
diff --git a/Engine/Source/Editor/UnrealEd/Private/Commandlets/ReplaceAssetsCommandlet.cpp b/Engine/Source/Editor/UnrealEd/Private/Commandlets/ReplaceAssetsCommandlet.cpp
--- a/Engine/Source/Editor/UnrealEd/Private/Commandlets/ReplaceAssetsCommandlet.cpp
+++ b/Engine/Source/Editor/UnrealEd/Private/Commandlets/ReplaceAssetsCommandlet.cpp
@@ -116,7 +116,15 @@ int32 UReplaceAssetsCommandlet::Main(const FString& InParams)
 		Filter.bRecursivePaths = true;
 		for (const FString& ReplacedClass : ReplacedClasses)
 		{
-			Filter.ClassNames.AddUnique(*ReplacedClass);
+			FTopLevelAssetPath ReplacedClassPathName = UClass::TryConvertShortTypeNameToPathName<UStruct>(ReplacedClass, ELogVerbosity::Error, TEXT("UReplaceAssetsCommandlet"));
+			if (ReplacedClassPathName.IsNull())
+			{
+				UE_LOG(LogReplaceAssetsCommandlet, Error, TEXT("Failed to convert short class name \"%s\" to path name. Please use class path names for ReplacedClasses."), *ReplacedClass);
+			}
+			else
+			{
+				Filter.ClassPaths.AddUnique(ReplacedClassPathName);
+			}
 		}
 		TArray<FAssetData> AssetList;
 		AssetRegistryModule.Get().GetAssets(Filter, AssetList);
@@ -150,7 +158,15 @@ int32 UReplaceAssetsCommandlet::Main(const FString& InParams)
 		Filter.bRecursivePaths = true;
 		for (const FString& ExcludedClass : ExcludedClasses)
 		{
-			Filter.ClassNames.AddUnique(*ExcludedClass);
+			FTopLevelAssetPath ExcludedClassPathName = UClass::TryConvertShortTypeNameToPathName<UStruct>(ExcludedClass, ELogVerbosity::Error, TEXT("UReplaceAssetsCommandlet"));
+			if (ExcludedClassPathName.IsNull())
+			{
+				UE_LOG(LogReplaceAssetsCommandlet, Error, TEXT("Failed to convert short class name \"%s\" to path name. Please use class path names for ExcludedClasses."), *ExcludedClass);
+			}
+			else
+			{
+				Filter.ClassPaths.AddUnique(ExcludedClassPathName);
+			}
 		}
 		TArray<FAssetData> AssetList;
 		AssetRegistryModule.Get().GetAssets(Filter, AssetList);
@@ -163,7 +179,7 @@ int32 UReplaceAssetsCommandlet::Main(const FString& InParams)
 		UE_LOG(LogReplaceAssetsCommandlet, Display, TEXT("Converting Package Names to File Paths"));
 		for (FAssetData& RemovedAsset : FinalAssetList)
 		{
-			bool bIsMap = RemovedAsset.AssetClass == UWorld::StaticClass()->GetFName();
+			bool bIsMap = RemovedAsset.AssetClassPath == UWorld::StaticClass()->GetClassPathName();
 			FinalFileList.AddUnique(FPaths::ConvertRelativePathToFull(FPackageName::LongPackageNameToFilename(RemovedAsset.PackageName.ToString(), bIsMap ? FPackageName::GetMapPackageExtension() : FPackageName::GetAssetPackageExtension())));
 		}
 
diff --git a/Engine/Source/Editor/UnrealEd/Private/Commandlets/SwapSoundForDialogueInCuesCommandlet.cpp b/Engine/Source/Editor/UnrealEd/Private/Commandlets/SwapSoundForDialogueInCuesCommandlet.cpp
--- a/Engine/Source/Editor/UnrealEd/Private/Commandlets/SwapSoundForDialogueInCuesCommandlet.cpp
+++ b/Engine/Source/Editor/UnrealEd/Private/Commandlets/SwapSoundForDialogueInCuesCommandlet.cpp
@@ -47,7 +47,7 @@ int32 USwapSoundForDialogueInCuesCommandlet::Main(const FString& Params)
 
 	// We only want dialogue wave assets that exist within the Game content directory.
 	TArray<FAssetData> AssetDataArrayForDialogueWaves;
-	if (!FLocalizedAssetUtil::GetAssetsByPathAndClass(AssetRegistry, FName("/Game"), UDialogueWave::StaticClass()->GetFName(), /*bIncludeLocalizedAssets*/false, AssetDataArrayForDialogueWaves))
+	if (!FLocalizedAssetUtil::GetAssetsByPathAndClass(AssetRegistry, FName("/Game"), UDialogueWave::StaticClass()->GetClassPathName(), /*bIncludeLocalizedAssets*/false, AssetDataArrayForDialogueWaves))
 	{
 		UE_LOG(LogSwapSoundForDialogueInCuesCommandlet, Error, TEXT("Unable to get dialogue wave asset data from asset registry."));
 		return -1;
@@ -106,7 +106,7 @@ int32 USwapSoundForDialogueInCuesCommandlet::Main(const FString& Params)
 
 			// Get sound cue assets that reference the context's sound wave.
 			FARFilter Filter;
-			Filter.ClassNames.Add(USoundCue::StaticClass()->GetFName());
+			Filter.ClassPaths.Add(USoundCue::StaticClass()->GetClassPathName());
 			Filter.bRecursiveClasses = true;
 			Filter.PackageNames = SoundWaveReferencerNames;
 			TArray<FAssetData> ReferencingSoundCueAssetDataArray;
diff --git a/Engine/Source/Editor/UnrealEd/Private/Commandlets/WorldPartitionBuilderCommandlet.cpp b/Engine/Source/Editor/UnrealEd/Private/Commandlets/WorldPartitionBuilderCommandlet.cpp
--- a/Engine/Source/Editor/UnrealEd/Private/Commandlets/WorldPartitionBuilderCommandlet.cpp
+++ b/Engine/Source/Editor/UnrealEd/Private/Commandlets/WorldPartitionBuilderCommandlet.cpp
@@ -57,7 +57,7 @@ int32 UWorldPartitionBuilderCommandlet::Main(const FString& Params)
 	}
 
 	// Find builder class
-	UClass* BuilderClass = FindObject<UClass>(ANY_PACKAGE, *BuilderClassName);
+	UClass* BuilderClass = FindFirstObject<UClass>(*BuilderClassName, EFindFirstObjectOptions::EnsureIfAmbiguous);
 	if (!BuilderClass)
 	{
 		UE_LOG(LogWorldPartitionBuilderCommandlet, Error, TEXT("Unknown builder %s."), *BuilderClassName);
diff --git a/Engine/Source/Editor/UnrealEd/Private/ComponentTypeRegistry.cpp b/Engine/Source/Editor/UnrealEd/Private/ComponentTypeRegistry.cpp
--- a/Engine/Source/Editor/UnrealEd/Private/ComponentTypeRegistry.cpp
+++ b/Engine/Source/Editor/UnrealEd/Private/ComponentTypeRegistry.cpp
@@ -228,10 +228,16 @@ namespace UE::Editor::ComponentTypeRegistry::Private
 			return ClassName;
 		}
 
-		virtual FName GetClassPath() const
+		UE_DEPRECATED(5.1, "Class names are now represented by path names. Please use GetClassPathName.")
+		virtual FName GetClassPath() const override
 		{
 			return ClassPath;
 		}
+		
+		virtual FTopLevelAssetPath GetClassPathName() const override
+		{
+			return FTopLevelAssetPath(ClassPath.ToString());
+		}
 		// End IUnloadedBlueprintData interface
 
 	private:
@@ -500,14 +506,14 @@ void FComponentTypeRegistryData::ForceRefreshComponentList()
 	AddBasicShapeComponents(SortedClassList);
 
 	// Add loaded component classes
-	TSet<FName> InMemoryClassPaths;
+	TSet<FTopLevelAssetPath> InMemoryClassPaths;
 	for (TObjectIterator<UClass> It; It; ++It)
 	{
 		UClass* Class = *It;
 		// If this is a subclass of Actor Component, not abstract, and tagged as spawnable from Kismet
 		if (Class->IsChildOf(UActorComponent::StaticClass()))
 		{
-			InMemoryClassPaths.Add(FName(Class->GetPathName()));
+			InMemoryClassPaths.Add(Class->GetClassPathName());
 
 			const bool bOutOfDateClass = Class->HasAnyClassFlags(CLASS_NewerVersionExists);
 			const bool bBlueprintSkeletonClass = FKismetEditorUtilities::IsClassABlueprintSkeleton(Class);
@@ -552,16 +558,15 @@ void FComponentTypeRegistryData::ForceRefreshComponentList()
 
 	// Add unloaded component classes
 	{
-		auto AddUnloadedComponentClass = [this, &SortedClassList](FName ClassPath, const FAssetData& AssetData)
+		auto AddUnloadedComponentClass = [this, &SortedClassList](const FTopLevelAssetPath& ClassPath, const FAssetData& AssetData)
 		{
 			// The blueprint is unloaded, so we need to work out which icon to use for it using its asset data
 			const UClass* IconClass = FClassIconFinder::GetIconClassForAssetData(AssetData);
 
-			FString ClassPathStr = ClassPath.ToString();
-			FString ClassName = FPackageName::ObjectPathToObjectName(ClassPathStr);
+			FString ClassName = ClassPath.GetAssetName().ToString();
 			ClassName.RemoveFromEnd(TEXT("_C"));
 
-			FComponentClassComboEntryPtr NewEntry(new FComponentClassComboEntry(BlueprintComponents, ClassName, ClassPath, IconClass, /*bIncludeInFilter=*/true));
+			FComponentClassComboEntryPtr NewEntry(new FComponentClassComboEntry(BlueprintComponents, ClassName, ClassPath.GetAssetName(), IconClass, /*bIncludeInFilter=*/true));
 			SortedClassList.Add(NewEntry);
 
 			// Create an unloaded blueprint data object to assist with class filtering
@@ -578,31 +583,31 @@ void FComponentTypeRegistryData::ForceRefreshComponentList()
 				NewEntry->SetUnloadedBlueprintData(UnloadedBlueprintData);
 			}
 
-			FComponentTypeEntry Entry = { MoveTemp(ClassName), MoveTemp(ClassPathStr), nullptr };
+			FComponentTypeEntry Entry = { MoveTemp(ClassName), ClassPath.ToString(), nullptr };
 			ComponentTypeList.Add(MoveTemp(Entry));
 		};
 
 		IAssetRegistry& AssetRegistry = FModuleManager::LoadModuleChecked<FAssetRegistryModule>(AssetRegistryConstants::ModuleName).Get();
 
-		TSet<FName> ActorComponentDerivedClassNames;
+		TSet<FTopLevelAssetPath> ActorComponentDerivedClassNames;
 		{
-			TArray<FName> ActorComponentClassNames;
-			ActorComponentClassNames.Add(UActorComponent::StaticClass()->GetFName());
-			AssetRegistry.GetDerivedClassNames(ActorComponentClassNames, TSet<FName>(), ActorComponentDerivedClassNames);
+			TArray<FTopLevelAssetPath> ActorComponentClassNames;
+			ActorComponentClassNames.Add(UActorComponent::StaticClass()->GetClassPathName());
+			AssetRegistry.GetDerivedClassNames(ActorComponentClassNames, TSet<FTopLevelAssetPath>(), ActorComponentDerivedClassNames);
 		}
 
 		// GetAssetsByClass call is a kludge to get the full asset paths for the blueprints we care about
 		// Bob T. thinks that the Asset Registry could just keep asset paths
 		TArray<FAssetData> Assets;
 
-		AssetRegistry.GetAssetsByClass(UBlueprint::StaticClass()->GetFName(), Assets, true);
+		AssetRegistry.GetAssetsByClass(UBlueprint::StaticClass()->GetClassPathName(), Assets, true);
 		for (FAssetData& BPAsset : Assets)
 		{
-			const FName ClassPath = FEditorClassUtils::GetClassPathFromAssetTag(BPAsset);
-			if (!ClassPath.IsNone())
+			const FTopLevelAssetPath ClassPath(FEditorClassUtils::GetClassPathNameFromAssetTag(BPAsset));
+			if (!ClassPath.IsNull())
 			{
 				if (!InMemoryClassPaths.Contains(ClassPath) && 
-					ActorComponentDerivedClassNames.Contains(FName(FPackageName::ObjectPathToObjectName(FStringView(FNameBuilder(ClassPath))))))
+					ActorComponentDerivedClassNames.Contains(ClassPath))
 				{
 					AddUnloadedComponentClass(ClassPath, BPAsset);
 				}
@@ -615,12 +620,13 @@ void FComponentTypeRegistryData::ForceRefreshComponentList()
 		}
 
 		Assets.Reset();
-		AssetRegistry.GetAssetsByClass(UBlueprintGeneratedClass::StaticClass()->GetFName(), Assets, true);
+		AssetRegistry.GetAssetsByClass(UBlueprintGeneratedClass::StaticClass()->GetClassPathName(), Assets, true);
 		for (FAssetData& BPGCAsset : Assets)
 		{
-			if (ActorComponentDerivedClassNames.Contains(BPGCAsset.AssetName) && !InMemoryClassPaths.Contains(BPGCAsset.ObjectPath))
+			FTopLevelAssetPath BPGCAssetClassPathName(BPGCAsset.PackagePath, BPGCAsset.AssetName);
+			if (ActorComponentDerivedClassNames.Contains(BPGCAssetClassPathName) && !InMemoryClassPaths.Contains(BPGCAssetClassPathName))
 			{
-				AddUnloadedComponentClass(BPGCAsset.ObjectPath, BPGCAsset);
+				AddUnloadedComponentClass(BPGCAssetClassPathName, BPGCAsset);
 			}
 		}
 	}
@@ -675,18 +681,17 @@ void FComponentTypeRegistryData::Tick(float)
 			return;
 		}
 
-		TArray<FName> ClassNames;
-		ClassNames.Add(UActorComponent::StaticClass()->GetFName());
-		TSet<FName> DerivedClassNames;
-		AssetRegistryModule.Get().GetDerivedClassNames(ClassNames, TSet<FName>(), DerivedClassNames);
+		TArray<FTopLevelAssetPath> ClassNames;
+		ClassNames.Add(UActorComponent::StaticClass()->GetClassPathName());
+		TSet<FTopLevelAssetPath> DerivedClassNames;
+		AssetRegistryModule.Get().GetDerivedClassNames(ClassNames, TSet<FTopLevelAssetPath>(), DerivedClassNames);
 
 		for (const FAssetData& Asset : PendingAssetData)
 		{
 			const FName BPParentClassName(GET_MEMBER_NAME_CHECKED(UBlueprint, ParentClass));
 			const FString TagValue = Asset.GetTagValueRef<FString>(BPParentClassName);
-			const FString ObjectPath = FPackageName::ExportTextPathToObjectPath(TagValue);
-			FName ObjectName = FName(*FPackageName::ObjectPathToObjectName(ObjectPath));
-			if (DerivedClassNames.Contains(ObjectName))
+			const FTopLevelAssetPath ObjectPath(FPackageName::ExportTextPathToObjectPath(TagValue));
+			if (DerivedClassNames.Contains(ObjectPath))
 			{
 				bRequiresRefresh = true;
 				break;
diff --git a/Engine/Source/Editor/UnrealEd/Private/DataTableEditorUtils.cpp b/Engine/Source/Editor/UnrealEd/Private/DataTableEditorUtils.cpp
--- a/Engine/Source/Editor/UnrealEd/Private/DataTableEditorUtils.cpp
+++ b/Engine/Source/Editor/UnrealEd/Private/DataTableEditorUtils.cpp
@@ -903,7 +903,7 @@ void FDataTableEditorUtils::GetPossibleStructAssetData(TArray<FAssetData>& Struc
 	// Now get unloaded ones
 	const FAssetRegistryModule& AssetRegistryModule = FModuleManager::LoadModuleChecked<FAssetRegistryModule>(TEXT("AssetRegistry"));
 	TArray<FAssetData> AssetData;
-	AssetRegistryModule.Get().GetAssetsByClass(UUserDefinedStruct::StaticClass()->GetFName(), AssetData);
+	AssetRegistryModule.Get().GetAssetsByClass(UUserDefinedStruct::StaticClass()->GetClassPathName(), AssetData);
 
 	for (int32 AssetIndex = 0; AssetIndex < AssetData.Num(); ++AssetIndex)
 	{
diff --git a/Engine/Source/Editor/UnrealEd/Private/DebugToolExec.cpp b/Engine/Source/Editor/UnrealEd/Private/DebugToolExec.cpp
--- a/Engine/Source/Editor/UnrealEd/Private/DebugToolExec.cpp
+++ b/Engine/Source/Editor/UnrealEd/Private/DebugToolExec.cpp
@@ -98,15 +98,16 @@ bool FDebugToolExec::Exec( UWorld* InWorld, const TCHAR* Cmd, FOutputDevice& Ar
 		// not allowed in the editor as this command can have far reaching effects such as impacting serialization
 		if (!GIsEditor)
 		{
-			UClass* Class = NULL;
-			if( ParseObject<UClass>( Cmd, TEXT("CLASS="), Class, ANY_PACKAGE ) == false )
+			UClass* Class = nullptr;
+			FString ClassName;
+			if (FParse::Value(Cmd, TEXT("CLASS="), ClassName))
 			{
-				TCHAR ClassName[256];
-				if ( FParse::Token(Cmd,ClassName,UE_ARRAY_COUNT(ClassName), 1) )
-				{
-					Class = FindObject<UClass>( ANY_PACKAGE, ClassName);
-				}
+				Class = FindFirstObject<UClass>(*ClassName, EFindFirstObjectOptions::None, ELogVerbosity::Warning, TEXT("parsing FDebugToolExec class"));
 			}
+			else if (FParse::Token(Cmd, ClassName, true))
+			{
+				Class = FindFirstObject<UClass>(*ClassName, EFindFirstObjectOptions::None, ELogVerbosity::Warning, TEXT("parsing FDebugToolExec class"));
+			}			
 
 			if (Class)
 			{
@@ -121,10 +122,15 @@ bool FDebugToolExec::Exec( UWorld* InWorld, const TCHAR* Cmd, FOutputDevice& Ar
 	}
 	else if (FParse::Command(&Cmd,TEXT("EDITOBJECT")))
 	{
-		UClass* SearchClass = NULL;
-		UObject* FoundObj = NULL;
+		UClass* SearchClass = nullptr;
+		UObject* FoundObj = nullptr;
+		FString ClassName;
 		// Search by class.
-		if (ParseObject<UClass>(Cmd, TEXT("CLASS="), SearchClass, ANY_PACKAGE))
+		if (FParse::Value(Cmd, TEXT("CLASS="), ClassName))
+		{
+			SearchClass = FindFirstObject<UClass>(*ClassName, EFindFirstObjectOptions::None, ELogVerbosity::Warning, TEXT("parsing FDebugToolExec class"));
+		}
+		if (SearchClass)
 		{
 			// pick the first valid object
 			for (FThreadSafeObjectIterator It(SearchClass); It && FoundObj == NULL; ++It)
@@ -153,7 +159,7 @@ bool FDebugToolExec::Exec( UWorld* InWorld, const TCHAR* Cmd, FOutputDevice& Ar
 			}
 			else if ( FParse::Token(Cmd,SearchPathName, true) )
 			{
-				FoundObj = FindObject<UObject>(ANY_PACKAGE,*SearchPathName);
+				FoundObj = FindFirstObject<UObject>(*SearchPathName, EFindFirstObjectOptions::None, ELogVerbosity::Warning, TEXT("parsing FDebugToolExec object"));
 			}
 		}
 
@@ -179,7 +185,7 @@ bool FDebugToolExec::Exec( UWorld* InWorld, const TCHAR* Cmd, FOutputDevice& Ar
 		FString SearchPathName;
 		if (FParse::Token(Cmd, SearchPathName, true))
 		{
-			foundObj = FindObject<UObject>(ANY_PACKAGE,*SearchPathName);
+			foundObj = FindFirstObject<UObject>(*SearchPathName, EFindFirstObjectOptions::None, ELogVerbosity::Warning, TEXT("EDITARCHETYPE"));
 		}
 
 		// Bring up an property editing window for the found object.
@@ -200,13 +206,14 @@ bool FDebugToolExec::Exec( UWorld* InWorld, const TCHAR* Cmd, FOutputDevice& Ar
 	// Edits an objects properties or copies them to the clipboard.
 	else if( FParse::Command(&Cmd,TEXT("EDITACTOR")) )
 	{
-		UClass*		Class = NULL;
-		AActor*		Found = NULL;
+		UClass*	Class = nullptr;
+		AActor*	Found = nullptr;
+		FString ClassName;
 
 		if (FParse::Command(&Cmd, TEXT("TRACE")))
 		{
 			APlayerController* PlayerController = InWorld->GetGameInstance() ? InWorld->GetGameInstance()->GetFirstLocalPlayerController() : nullptr;
-			if (PlayerController != NULL)
+			if (PlayerController != nullptr)
 			{
 				// Do a trace in the player's facing direction and edit anything that's hit.
 				FVector PlayerLocation;
@@ -218,29 +225,33 @@ bool FDebugToolExec::Exec( UWorld* InWorld, const TCHAR* Cmd, FOutputDevice& Ar
 			}
 		}
 		// Search by class.
-		else if( ParseObject<UClass>( Cmd, TEXT("CLASS="), Class, ANY_PACKAGE ) && Class->IsChildOf(AActor::StaticClass()) )
+		else if (FParse::Value(Cmd, TEXT("CLASS="), ClassName))
 		{
-			UGameEngine* GameEngine = Cast<UGameEngine>(GEngine);
-			
-			// Look for the closest actor of this class to the player.
-			FVector PlayerLocation(0.0f);
-			APlayerController* PlayerController = InWorld->GetGameInstance() ? InWorld->GetGameInstance()->GetFirstLocalPlayerController() : nullptr;
-			if (PlayerController != NULL)
+			Class = FindFirstObject<UClass>(*ClassName, EFindFirstObjectOptions::None, ELogVerbosity::Warning, TEXT("parsing FDebugToolExec class"));
+			if (Class)
 			{
-				FRotator DummyRotation;
-				PlayerController->GetPlayerViewPoint(PlayerLocation, DummyRotation);
-			}
+				UGameEngine* GameEngine = Cast<UGameEngine>(GEngine);
 
-			float   MinDist = FLT_MAX;
-			for( TActorIterator<AActor> It(InWorld, Class); It; ++It )
-			{
-				if ( IsValid(*It) )
+				// Look for the closest actor of this class to the player.
+				FVector PlayerLocation(0.0f);
+				APlayerController* PlayerController = InWorld->GetGameInstance() ? InWorld->GetGameInstance()->GetFirstLocalPlayerController() : nullptr;
+				if (PlayerController != NULL)
+				{
+					FRotator DummyRotation;
+					PlayerController->GetPlayerViewPoint(PlayerLocation, DummyRotation);
+				}
+
+				float   MinDist = FLT_MAX;
+				for (TActorIterator<AActor> It(InWorld, Class); It; ++It)
 				{
-					float const Dist = (PlayerController && It->GetRootComponent()) ? FVector::Dist(It->GetActorLocation(), PlayerLocation) : 0.f;
-					if (Dist < MinDist)
+					if (IsValid(*It))
 					{
-						MinDist = Dist;
-						Found   = *It;
+						float const Dist = (PlayerController && It->GetRootComponent()) ? FVector::Dist(It->GetActorLocation(), PlayerLocation) : 0.f;
+						if (Dist < MinDist)
+						{
+							MinDist = Dist;
+							Found = *It;
+						}
 					}
 				}
 			}
diff --git a/Engine/Source/Editor/UnrealEd/Private/Dialogs/Dialogs.cpp b/Engine/Source/Editor/UnrealEd/Private/Dialogs/Dialogs.cpp
--- a/Engine/Source/Editor/UnrealEd/Private/Dialogs/Dialogs.cpp
+++ b/Engine/Source/Editor/UnrealEd/Private/Dialogs/Dialogs.cpp
@@ -855,8 +855,8 @@ bool PromptUserIfExistingObject(const FString& Name, const FString& Package, UPa
 	QualifiedName += Name;
 
 	// Check for an existing object
-	UObject* ExistingObject = StaticFindObject( UObject::StaticClass(), ANY_PACKAGE, *QualifiedName );
-	if( ExistingObject != NULL )
+	UObject* ExistingObject = StaticFindObject( UObject::StaticClass(), nullptr, *QualifiedName );
+	if( ExistingObject != nullptr )
 	{
 		// Object already exists in either the specified package or another package.  Check to see if the user wants
 		// to replace the object.
diff --git a/Engine/Source/Editor/UnrealEd/Private/Dialogs/SDeleteAssetsDialog.cpp b/Engine/Source/Editor/UnrealEd/Private/Dialogs/SDeleteAssetsDialog.cpp
--- a/Engine/Source/Editor/UnrealEd/Private/Dialogs/SDeleteAssetsDialog.cpp
+++ b/Engine/Source/Editor/UnrealEd/Private/Dialogs/SDeleteAssetsDialog.cpp
@@ -773,7 +773,7 @@ TSharedRef<SWidget> SDeleteAssetsDialog::MakeAssetViewForReferencerAssets()
 TSharedRef<SWidget> SDeleteAssetsDialog::MakeConsolidationAssetPicker()
 {
 	FAssetPickerConfig AssetPickerConfig;
-	//AssetPickerConfig.Filter.ClassNames.Add( UStaticMesh::StaticClass()->GetFName() );
+	//AssetPickerConfig.Filter.ClassPaths.Add( UStaticMesh::StaticClass()->GetFName() );
 	AssetPickerConfig.OnAssetSelected = FOnAssetSelected::CreateSP( this, &SDeleteAssetsDialog::OnAssetSelectedFromConsolidationPicker );
 	AssetPickerConfig.OnShouldFilterAsset = FOnShouldFilterAsset::CreateSP( this, &SDeleteAssetsDialog::OnShouldConsolidationFilterAsset );
 	AssetPickerConfig.bAllowNullSelection = false;
diff --git a/Engine/Source/Editor/UnrealEd/Private/Editor.cpp b/Engine/Source/Editor/UnrealEd/Private/Editor.cpp
--- a/Engine/Source/Editor/UnrealEd/Private/Editor.cpp
+++ b/Engine/Source/Editor/UnrealEd/Private/Editor.cpp
@@ -922,7 +922,7 @@ void RestoreEditorWorld( UWorld* EditorWorld )
  * Takes an FName and checks to see that it is unique among all loaded objects.
  *
  * @param	InName		The name to check
- * @param	Outer		The context for validating this object name. Should be a group/package, but could be ANY_PACKAGE if you want to check across the whole system (not recommended)
+ * @param	Outer		The context for validating this object name. Should be a group/package
  * @param	InReason	If the check fails, this string is filled in with the reason why.
  *
  * @return	1 if the name is valid, 0 if it is not
@@ -943,11 +943,26 @@ bool IsUniqueObjectName( const FName& InName, UObject* Outer, FText* InReason )
 	return true;
 }
 
+bool IsGloballyUniqueObjectName(const FName& InName, FText* InReason)
+{
+	// See if the name is already in use anywhere in the engine.
+	if (StaticFindFirstObject(UObject::StaticClass(), *InName.ToString()) != NULL)
+	{
+		if (InReason != NULL)
+		{
+			*InReason = NSLOCTEXT("UnrealEd", "NameAlreadyInUse", "Name is already in use by another object.");
+		}
+		return false;
+	}
+
+	return true;
+}
+
 /**
  * Takes an FName and checks to see that it is unique among all loaded objects.
  *
  * @param	InName		The name to check
- * @param	Outer		The context for validating this object name. Should be a group/package, but could be ANY_PACKAGE if you want to check across the whole system (not recommended)
+ * @param	Outer		The context for validating this object name. Should be a group/package.
  * @param	InReason	If the check fails, this string is filled in with the reason why.
  *
  * @return	1 if the name is valid, 0 if it is not
diff --git a/Engine/Source/Editor/UnrealEd/Private/EditorClassUtils.cpp b/Engine/Source/Editor/UnrealEd/Private/EditorClassUtils.cpp
--- a/Engine/Source/Editor/UnrealEd/Private/EditorClassUtils.cpp
+++ b/Engine/Source/Editor/UnrealEd/Private/EditorClassUtils.cpp
@@ -230,12 +230,20 @@ TSharedRef<SWidget> FEditorClassUtils::GetSourceLinkFormatted(const UClass* Clas
 
 UClass* FEditorClassUtils::GetClassFromString(const FString& ClassName)
 {
-	if(ClassName.IsEmpty() || ClassName == "None")
+	if(ClassName.IsEmpty() || ClassName == TEXT("None"))
 	{
 		return nullptr;
 	}
 
-	UClass* Class = FindObject<UClass>(ANY_PACKAGE, *ClassName);
+	UClass* Class = nullptr;
+	if (!FPackageName::IsShortPackageName(ClassName))
+	{
+		Class = FindObject<UClass>(nullptr, *ClassName);
+	}
+	else
+	{
+		Class = FindFirstObject<UClass>(*ClassName, EFindFirstObjectOptions::None, ELogVerbosity::Warning, TEXT("FEditorClassUtils::GetClassFromString"));
+	}
 	if(!Class)
 	{
 		Class = LoadObject<UClass>(nullptr, *ClassName);
@@ -245,20 +253,20 @@ UClass* FEditorClassUtils::GetClassFromString(const FString& ClassName)
 
 bool FEditorClassUtils::IsBlueprintAsset(const FAssetData& InAssetData, bool* bOutIsBPGC /*= nullptr*/)
 {
-	bool bIsBP = (InAssetData.AssetClass == UBlueprint::StaticClass()->GetFName());
-	bool bIsBPGC = (InAssetData.AssetClass == UBlueprintGeneratedClass::StaticClass()->GetFName());
+	bool bIsBP = (InAssetData.AssetClassPath == UBlueprint::StaticClass()->GetClassPathName());
+	bool bIsBPGC = (InAssetData.AssetClassPath == UBlueprintGeneratedClass::StaticClass()->GetClassPathName());
 
 	if (!bIsBP && !bIsBPGC)
 	{
 		IAssetRegistry& AssetRegistry = FModuleManager::LoadModuleChecked<FAssetRegistryModule>(AssetRegistryConstants::ModuleName).Get();
-		TArray<FName> AncestorClassNames;
-		AssetRegistry.GetAncestorClassNames(InAssetData.AssetClass, AncestorClassNames);
+		TArray<FTopLevelAssetPath> AncestorClassNames;
+		AssetRegistry.GetAncestorClassNames(InAssetData.AssetClassPath, AncestorClassNames);
 
-		if (AncestorClassNames.Contains(UBlueprint::StaticClass()->GetFName()))
+		if (AncestorClassNames.Contains(UBlueprint::StaticClass()->GetClassPathName()))
 		{
 			bIsBP = true;
 		}
-		else if (AncestorClassNames.Contains(UBlueprintGeneratedClass::StaticClass()->GetFName()))
+		else if (AncestorClassNames.Contains(UBlueprintGeneratedClass::StaticClass()->GetClassPathName()))
 		{
 			bIsBPGC = true;
 		}
@@ -278,6 +286,12 @@ FName FEditorClassUtils::GetClassPathFromAssetTag(const FAssetData& InAssetData)
 	return FName(FPackageName::ExportTextPathToObjectPath(FStringView(GeneratedClassPath)));
 }
 
+FTopLevelAssetPath FEditorClassUtils::GetClassPathNameFromAssetTag(const FAssetData& InAssetData)
+{
+	const FString GeneratedClassPath = InAssetData.GetTagValueRef<FString>(FBlueprintTags::GeneratedClassPath);
+	return FTopLevelAssetPath(FPackageName::ExportTextPathToObjectPath(FStringView(GeneratedClassPath)));
+}
+
 FName FEditorClassUtils::GetClassPathFromAsset(const FAssetData& InAssetData, bool bGenerateClassPathIfMissing /*= false*/)
 {
 	bool bIsBPGC = false;
@@ -289,7 +303,9 @@ FName FEditorClassUtils::GetClassPathFromAsset(const FAssetData& InAssetData, bo
 	}
 	else if (bIsBP)
 	{
+PRAGMA_DISABLE_DEPRECATION_WARNINGS
 		FName ClassPath = GetClassPathFromAssetTag(InAssetData);
+PRAGMA_ENABLE_DEPRECATION_WARNINGS
 		if (bGenerateClassPathIfMissing && ClassPath.IsNone())
 		{
 			FNameBuilder ClassPathBuilder(InAssetData.ObjectPath);
@@ -301,6 +317,29 @@ FName FEditorClassUtils::GetClassPathFromAsset(const FAssetData& InAssetData, bo
 	return NAME_None;
 }
 
+FTopLevelAssetPath FEditorClassUtils::GetClassPathNameFromAsset(const FAssetData& InAssetData, bool bGenerateClassPathIfMissing /*= false*/)
+{
+	bool bIsBPGC = false;
+	const bool bIsBP = IsBlueprintAsset(InAssetData, &bIsBPGC);
+
+	if (bIsBPGC)
+	{
+		return FTopLevelAssetPath(InAssetData.ObjectPath.ToString());
+	}
+	else if (bIsBP)
+	{
+		FTopLevelAssetPath ClassPath = GetClassPathNameFromAssetTag(InAssetData);
+		if (bGenerateClassPathIfMissing && ClassPath.IsNull())
+		{
+			FString ClassPathString = InAssetData.ObjectPath.ToString();
+			ClassPathString += TEXT("_C");
+			ClassPath = FTopLevelAssetPath(ClassPathString);
+		}
+		return ClassPath;
+	}
+	return FTopLevelAssetPath();
+}
+
 void FEditorClassUtils::GetImplementedInterfaceClassPathsFromAsset(const struct FAssetData& InAssetData, TArray<FString>& OutClassPaths)
 {
 	if (!InAssetData.IsValid())
diff --git a/Engine/Source/Editor/UnrealEd/Private/EditorDomain/EditorDomainUtils.cpp b/Engine/Source/Editor/UnrealEd/Private/EditorDomain/EditorDomainUtils.cpp
--- a/Engine/Source/Editor/UnrealEd/Private/EditorDomain/EditorDomainUtils.cpp
+++ b/Engine/Source/Editor/UnrealEd/Private/EditorDomain/EditorDomainUtils.cpp
@@ -464,7 +464,7 @@
 
 	// Scratch variables usable during GetRecursive; they are invalidated when a recursive call is made
 	FString NameStringBuffer;
-	TArray<FName> AncestorShortNames;
+	TArray<FTopLevelAssetPath> AncestorPathNames;
 };
 
 FClassDigestData* FPrecacheClassDigest::GetRecursive(FName ClassName, bool bAllowRedirects)
@@ -544,16 +544,17 @@
 		FStringView ClassObjectName;
 		FStringView ClassSubObjectName;
 		FPackageName::SplitFullObjectPath(NameStringBuffer, UnusedClassOfClassName, ClassPackageName, ClassObjectName, ClassSubObjectName);
-		FName ClassObjectFName(ClassObjectName);
+		FTopLevelAssetPath ClassObjectPathName = FTopLevelAssetPath(FName(ClassPackageName), FName(ClassObjectName));
+
 		// TODO_EDITORDOMAIN: If the class is not yet present in the assetregistry, or
 		// if its parent classes are not, then we will not be able to propagate information from the parent classes; wait on the class to be parsed
-		AncestorShortNames.Reset();
-		IAssetRegistry::Get()->GetAncestorClassNames(ClassObjectFName, AncestorShortNames);
-		for (FName ShortName : AncestorShortNames)
+		AncestorPathNames.Reset();
+		IAssetRegistry::Get()->GetAncestorClassNames(ClassObjectPathName, AncestorPathNames);
+		for (const FTopLevelAssetPath& PathName : AncestorPathNames)
 		{
 			// TODO_EDITORDOMAIN: For robustness and performance, we need the AssetRegistry to return FullPathNames rather than ShortNames
 			// For now, we lookup each shortname using FindObject, and do not handle propagating data from blueprint classes to child classes
-			if (UStruct* CurrentStruct = FindObjectFast<UStruct>(nullptr, ShortName, false /* ExactClass */, true /* AnyPackage */))
+			if (UStruct* CurrentStruct = FindObject<UStruct>(PathName))
 			{
 				NameStringBuffer.Reset();
 				CurrentStruct->GetPathName(nullptr, NameStringBuffer);
diff --git a/Engine/Source/Editor/UnrealEd/Private/EditorEngine.cpp b/Engine/Source/Editor/UnrealEd/Private/EditorEngine.cpp
--- a/Engine/Source/Editor/UnrealEd/Private/EditorEngine.cpp
+++ b/Engine/Source/Editor/UnrealEd/Private/EditorEngine.cpp
@@ -486,9 +486,9 @@
 
 	for ( auto AssetIt = SelectedAssets.CreateConstIterator(); AssetIt; ++AssetIt )
 	{
-		UClass* AssetClass = FindObject<UClass>(ANY_PACKAGE, *(*AssetIt).AssetClass.ToString());
+		UClass* AssetClass = FindObject<UClass>((*AssetIt).AssetClassPath);
 
-		if ( AssetClass != NULL )
+		if ( AssetClass != nullptr )
 		{
 			Selection.AddUnique(AssetClass);
 		}
diff --git a/Engine/Source/Editor/UnrealEd/Private/EditorObject.cpp b/Engine/Source/Editor/UnrealEd/Private/EditorObject.cpp
--- a/Engine/Source/Editor/UnrealEd/Private/EditorObject.cpp
+++ b/Engine/Source/Editor/UnrealEd/Private/EditorObject.cpp
@@ -326,7 +326,7 @@ static const TCHAR* ImportProperties(
 			UFoliageType* SourceFoliageType;
 			FName ComponentName;
 			if (SubobjectRoot &&
-				ParseObject<UFoliageType>(Str, TEXT("FoliageType="), SourceFoliageType, ANY_PACKAGE) &&
+				ParseObject<UFoliageType>(Str, TEXT("FoliageType="), SourceFoliageType, nullptr) &&
 				FParse::Value(Str, TEXT("Component="), ComponentName) )
 			{
 				UPrimitiveComponent* ActorComponent = FindObjectFast<UPrimitiveComponent>(SubobjectRoot, ComponentName);
@@ -408,7 +408,7 @@ static const TCHAR* ImportProperties(
 			// Note: default properties subobjects have compiled class as their Outer (used for localization).
 			UClass*	TemplateClass = NULL;
 			bool bInvalidClass = false;
-			ParseObject<UClass>(Str, TEXT("Class="), TemplateClass, ANY_PACKAGE, &bInvalidClass);
+			ParseObject<UClass>(Str, TEXT("Class="), TemplateClass, nullptr, &bInvalidClass);
 			
 			if (bInvalidClass)
 			{
@@ -486,15 +486,20 @@ static const TCHAR* ImportProperties(
 						if ( FPackageName::ParseExportTextPath(ArchetypeName, &ObjectClass, &ObjectPath) )
 						{
 							// find the class
-							UClass* ArchetypeClass = (UClass*)StaticFindObject(UClass::StaticClass(), ANY_PACKAGE, *ObjectClass);
+							check(FPackageName::IsValidObjectPath(ObjectClass));
+							UClass* ArchetypeClass = (UClass*)StaticFindObject(UClass::StaticClass(), nullptr, *ObjectClass);
 							if (ArchetypeClass)
 							{
 								ObjectPath = ObjectPath.TrimQuotes();
-
-								UPackage* FindInPackage = (FPackageName::IsShortPackageName(ObjectPath) ? ANY_PACKAGE : nullptr);
-
 								// if we had the class, find the archetype
-								Archetype = StaticFindObject(ArchetypeClass, FindInPackage, *ObjectPath);
+								if (!FPackageName::IsShortPackageName(ObjectPath))
+								{
+									Archetype = StaticFindObject(ArchetypeClass, nullptr, *ObjectPath);
+								}
+								else
+								{
+									Archetype = StaticFindFirstObject(ArchetypeClass, *ObjectPath, EFindFirstObjectOptions::NativeFirst | EFindFirstObjectOptions::EnsureIfAmbiguous);
+								}
 							}
 						}
 					}
diff --git a/Engine/Source/Editor/UnrealEd/Private/EditorServer.cpp b/Engine/Source/Editor/UnrealEd/Private/EditorServer.cpp
--- a/Engine/Source/Editor/UnrealEd/Private/EditorServer.cpp
+++ b/Engine/Source/Editor/UnrealEd/Private/EditorServer.cpp
@@ -363,7 +363,7 @@ bool UEditorEngine::SafeExec( UWorld* InWorld, const TCHAR* InStr, FOutputDevice
 		}
 
 		const FString ClassName     = FParse::Token(Str,0);
-		UClass* Class         = FindObject<UClass>( ANY_PACKAGE, *ClassName );
+		UClass* Class         = UClass::TryFindTypeSlow<UClass>(ClassName);
 		if( !Class )
 		{
 			UE_SUPPRESS(LogExec, Warning, Ar.Logf(TEXT("Unrecognized or missing factor class %s"), *ClassName ));
@@ -590,7 +590,7 @@ bool UEditorEngine::Exec_StaticMesh( UWorld* InWorld, const TCHAR* Str, FOutputD
 	{
 		GetSelectedObjects()->DeselectAll( UStaticMesh::StaticClass() );
 		UStaticMesh* StaticMesh = NULL;
-		bResult = ParseObject<UStaticMesh>(Str,TEXT("NAME="), StaticMesh, ANY_PACKAGE);
+		bResult = ParseObject<UStaticMesh>(Str,TEXT("NAME="), StaticMesh, nullptr);
 		if( bResult && StaticMesh)
 		{
 			GetSelectedObjects()->Select( StaticMesh );
@@ -803,13 +803,14 @@ bool UEditorEngine::Exec_Brush( UWorld* InWorld, const TCHAR* Str, FOutputDevice
 	}
 	else if (FParse::Command(&Str,TEXT("ADDVOLUME"))) // BRUSH ADDVOLUME
 	{
-		AVolume* Actor = NULL;
+		AVolume* Actor = nullptr;
 		{
 			const FScopedTransaction Transaction( NSLOCTEXT("UnrealEd", "BrushAddVolume", "Brush AddVolume") );
 			FinishAllSnaps();
 
-			UClass* VolumeClass = NULL;
-			ParseObject<UClass>( Str, TEXT("CLASS="), VolumeClass, ANY_PACKAGE );
+			FString VolumeClassName;
+			FParse::Value(Str, TEXT("CLASS="), VolumeClassName);
+			UClass* VolumeClass = UClass::TryFindTypeSlow<UClass>(VolumeClassName);
 			if( !VolumeClass || !VolumeClass->IsChildOf(AVolume::StaticClass()) )
 			{
 				VolumeClass = AVolume::StaticClass();
@@ -819,7 +820,7 @@ bool UEditorEngine::Exec_Brush( UWorld* InWorld, const TCHAR* Str, FOutputDevice
 			Actor = InWorld->SpawnActor<AVolume>( VolumeClass, SpawnLoc, FRotator::ZeroRotator );
 			if( Actor )
 			{
-				Actor->PreEditChange(NULL);
+				Actor->PreEditChange(nullptr);
 
 				FBSPOps::csgCopyBrush
 				(
@@ -838,7 +839,7 @@ bool UEditorEngine::Exec_Brush( UWorld* InWorld, const TCHAR* Str, FOutputDevice
 					for( int32 poly = 0 ; poly < Actor->Brush->Polys->Element.Num() ; ++poly )
 					{
 						FPoly* Poly = &(Actor->Brush->Polys->Element[poly]);
-						Poly->Material = NULL;
+						Poly->Material = nullptr;
 					}
 				}
 				Actor->PostEditChange();
@@ -4483,7 +4484,7 @@ bool UEditorEngine::Exec_Poly( UWorld* InWorld, const TCHAR* Str, FOutputDevice&
 	else if( FParse::Command(&Str,TEXT("DEFAULT")) ) // POLY DEFAULT <variable>=<value>...
 	{
 		//CurrentMaterial=NULL;
-		//ParseObject<UMaterial>(Str,TEXT("TEXTURE="),CurrentMaterial,ANY_PACKAGE);
+		//ParseObject<UMaterial>(Str,TEXT("TEXTURE="),CurrentMaterial,nullptr);
 		return true;
 	}
 	else if( FParse::Command(&Str,TEXT("SETMATERIAL")) )
@@ -4639,15 +4640,18 @@ bool UEditorEngine::Exec_Obj( const TCHAR* Str, FOutputDevice& Ar )
 	if( FParse::Command(&Str,TEXT("EXPORT")) )//oldver
 	{
 		FName Package=NAME_None;
-		UClass* Type;
+		
 		UObject* Res;
 		FParse::Value( Str, TEXT("PACKAGE="), Package );
-
+		FString TypeName;
+		FParse::Value(Str, TEXT("TYPE="), TypeName);
+		UClass* Type = UClass::TryFindTypeSlow<UClass>(TypeName);
 		FString TempFname;
+
 		if
-		(	ParseObject<UClass>( Str, TEXT("TYPE="), Type, ANY_PACKAGE )
+		(	Type
 		&&	FParse::Value( Str, TEXT("FILE="), TempFname )
-		&&	ParseObject( Str, TEXT("NAME="), Type, Res, ANY_PACKAGE ) )
+		&&	ParseObject( Str, TEXT("NAME="), Type, Res, nullptr ) )
 		{
 			for (FThreadSafeObjectIterator It; It; ++It)
 			{
@@ -4750,7 +4754,7 @@ bool UEditorEngine::Exec_Obj( const TCHAR* Str, FOutputDevice& Ar )
 
 AActor* UEditorEngine::SelectNamedActor(const TCHAR* TargetActorName)
 {
-	AActor* Actor = FindObject<AActor>( ANY_PACKAGE, TargetActorName, false );
+	AActor* Actor = FindFirstObject<AActor>( TargetActorName, EFindFirstObjectOptions::None, ELogVerbosity::Warning, TEXT("SelectNamedActor") );
 	if( Actor && !Actor->IsA(AWorldSettings::StaticClass()) )
 	{
 		SelectActor( Actor, true, true );
@@ -5603,8 +5607,8 @@ void ListMapPackageDependencies(const TCHAR* InStr)
 				{
 					for (int32 ExportIdx = 0; ExportIdx < Linker->ExportMap.Num(); ExportIdx++)
 					{
-						FName CheckClassName = Linker->GetExportClassName(ExportIdx);
-						UClass* CheckClass = (UClass*)(StaticFindObject(UClass::StaticClass(), ANY_PACKAGE, *(CheckClassName.ToString()), true));
+						FTopLevelAssetPath CheckClassPathName(Linker->GetExportClassPackage(ExportIdx), Linker->GetExportClassName(ExportIdx));
+						UClass* CheckClass = FindObject<UClass>(CheckClassPathName, true);
 						if (
 							(CheckClass != NULL) &&
 							(CheckClass->IsChildOf(UTexture::StaticClass()) == true)
@@ -6021,7 +6025,7 @@ bool UEditorEngine::HandleTestPropsCommand( const TCHAR* Str, FOutputDevice& Ar
 	{
 		UObject* Object;
 		UClass* Class = NULL;
-		if( ParseObject<UClass>( Str, TEXT("CLASS="), Class, ANY_PACKAGE ) != false )
+		if( ParseObject<UClass>( Str, TEXT("CLASS="), Class, nullptr ) != false )
 		{ 
 			Object = NewObject<UObject>(GetTransientPackage(), Class);
 		}
@@ -6085,7 +6089,7 @@ bool UEditorEngine::HandleTestPropsCommand( const TCHAR* Str, FOutputDevice& Ar
 		StructArgs.bShowInterfaces = true;
 
 		UStruct* Struct = nullptr;
-		if (!ParseObject<UStruct>(Str, TEXT("STRUCT="), Struct, ANY_PACKAGE))
+		if (!ParseObject<UStruct>(Str, TEXT("STRUCT="), Struct, nullptr))
 		{
 			Struct = FPropertyEditorTestBasicStruct::StaticStruct();
 		}
@@ -6116,7 +6120,7 @@ bool UEditorEngine::HandleTestPropsCommand( const TCHAR* Str, FOutputDevice& Ar
 	{
 		UObject* Object;
 		UClass* Class = nullptr;
-		if( ParseObject<UClass>( Str, TEXT("CLASS="), Class, ANY_PACKAGE ) != false )
+		if( ParseObject<UClass>( Str, TEXT("CLASS="), Class, nullptr ) != false )
 		{ 
 			Object = NewObject<UObject>(GetTransientPackage(), Class);
 		}
diff --git a/Engine/Source/Editor/UnrealEd/Private/Factories/ActorFactory.cpp b/Engine/Source/Editor/UnrealEd/Private/Factories/ActorFactory.cpp
--- a/Engine/Source/Editor/UnrealEd/Private/Factories/ActorFactory.cpp
+++ b/Engine/Source/Editor/UnrealEd/Private/Factories/ActorFactory.cpp
@@ -1639,12 +1639,12 @@ bool UActorFactoryBlueprint::CanCreateActorFrom( const FAssetData& AssetData, FT
 		FAssetRegistryModule& AssetRegistryModule = FModuleManager::LoadModuleChecked<FAssetRegistryModule>("AssetRegistry");
 		IAssetRegistry& AssetRegistry = AssetRegistryModule.Get();
 
-		const FStringView ObjectPath = FPackageName::ExportTextPathToObjectPath(FStringView(ParentClassPath));
-		const FName ParentClassPathFName = FName( FPackageName::ObjectPathToObjectName(ObjectPath) );
-		TArray<FName> AncestorClassNames;
-		AssetRegistry.GetAncestorClassNames(ParentClassPathFName, AncestorClassNames);
+		const FTopLevelAssetPath ObjectPath(FPackageName::ExportTextPathToObjectPath(ParentClassPath));
+		const FName ParentClassPathFName = ObjectPath.GetAssetName();
+		TArray<FTopLevelAssetPath> AncestorClassNames;
+		AssetRegistry.GetAncestorClassNames(ObjectPath, AncestorClassNames);
 
-		bIsActorBased = AncestorClassNames.Contains(AActor::StaticClass()->GetFName());
+		bIsActorBased = AncestorClassNames.Contains(AActor::StaticClass()->GetClassPathName());
 	}
 
 	if ( !bIsActorBased )
diff --git a/Engine/Source/Editor/UnrealEd/Private/Factories/AnimBlueprintFactory.cpp b/Engine/Source/Editor/UnrealEd/Private/Factories/AnimBlueprintFactory.cpp
--- a/Engine/Source/Editor/UnrealEd/Private/Factories/AnimBlueprintFactory.cpp
+++ b/Engine/Source/Editor/UnrealEd/Private/Factories/AnimBlueprintFactory.cpp
@@ -304,7 +304,7 @@ private:
 
 		FAssetPickerConfig AssetPickerConfig;
 		AssetPickerConfig.RefreshAssetViewDelegates.Add(&RefreshSkeletonViewDelegate);
-		AssetPickerConfig.Filter.ClassNames.Add(USkeleton::StaticClass()->GetFName());
+		AssetPickerConfig.Filter.ClassPaths.Add(USkeleton::StaticClass()->GetClassPathName());
 		AssetPickerConfig.OnAssetSelected = FOnAssetSelected::CreateSP(this, &SAnimBlueprintCreateDialog::OnSkeletonSelected);
 		AssetPickerConfig.OnShouldFilterAsset = FOnShouldFilterAsset::CreateSP(this, &SAnimBlueprintCreateDialog::FilterSkeletonBasedOnParentClass);
 		AssetPickerConfig.bAllowNullSelection = false;
diff --git a/Engine/Source/Editor/UnrealEd/Private/Factories/AnimStreamableFactory.cpp b/Engine/Source/Editor/UnrealEd/Private/Factories/AnimStreamableFactory.cpp
--- a/Engine/Source/Editor/UnrealEd/Private/Factories/AnimStreamableFactory.cpp
+++ b/Engine/Source/Editor/UnrealEd/Private/Factories/AnimStreamableFactory.cpp
@@ -38,7 +38,7 @@ bool UAnimStreamableFactory::ConfigureProperties()
 
 	FAssetPickerConfig AssetPickerConfig;
 	
-	AssetPickerConfig.Filter.ClassNames.Add(USkeleton::StaticClass()->GetFName());
+	AssetPickerConfig.Filter.ClassPaths.Add(USkeleton::StaticClass()->GetFName());
 	AssetPickerConfig.Filter.bRecursiveClasses = true;
 
 	// The delegate that fires when an asset was selected
diff --git a/Engine/Source/Editor/UnrealEd/Private/Factories/EditorFactories.cpp b/Engine/Source/Editor/UnrealEd/Private/Factories/EditorFactories.cpp
--- a/Engine/Source/Editor/UnrealEd/Private/Factories/EditorFactories.cpp
+++ b/Engine/Source/Editor/UnrealEd/Private/Factories/EditorFactories.cpp
@@ -842,7 +842,7 @@ UObject* ULevelFactory::FactoryCreateText
 		else if( GetBEGIN(&Str,TEXT("ACTOR")) )
 		{
 			UClass* TempClass;
-			if( ParseObject<UClass>( Str, TEXT("CLASS="), TempClass, ANY_PACKAGE ) )
+			if( ParseObject<UClass>( Str, TEXT("CLASS="), TempClass, nullptr ) )
 			{
 				// Get actor name.
 				FName ActorUniqueName(NAME_None);
@@ -880,13 +880,13 @@ UObject* ULevelFactory::FactoryCreateText
 					if ( FPackageName::ParseExportTextPath(ArchetypeName, &ObjectClass, &ObjectPath) )
 					{
 						// find the class
-						UClass* ArchetypeClass = (UClass*)StaticFindObject(UClass::StaticClass(), ANY_PACKAGE, *ObjectClass);
+						UClass* ArchetypeClass = UClass::TryFindTypeSlow<UClass>(ObjectClass, EFindFirstObjectOptions::EnsureIfAmbiguous);
 						if ( ArchetypeClass )
 						{
 							if ( ArchetypeClass->IsChildOf(AActor::StaticClass()) )
 							{
 								// if we had the class, find the archetype
-								Archetype = Cast<AActor>(StaticFindObject(ArchetypeClass, ANY_PACKAGE, *ObjectPath));
+								Archetype = Cast<AActor>(StaticFindObject(ArchetypeClass, nullptr, *ObjectPath));
 							}
 							else
 							{
@@ -1423,13 +1423,13 @@ UObject* UPackageFactory::FactoryCreateText( UClass* Class, UObject* InParent, F
 		{
 			//Begin TopLevelPackage Class=Package Name=ExportTest_ORIG Archetype=Package'Core.Default__Package'
 			UClass* TempClass;
-			if( ParseObject<UClass>( Str, TEXT("CLASS="), TempClass, ANY_PACKAGE ) )
+			if( ParseObject<UClass>( Str, TEXT("CLASS="), TempClass, nullptr ) )
 			{
 				// Get actor name.
 				FName PackageName(NAME_None);
 				FParse::Value( Str, TEXT("NAME="), PackageName );
 
-				if (FindObject<UPackage>(ANY_PACKAGE, *(PackageName.ToString())))
+				if (FindObject<UPackage>(nullptr, *(PackageName.ToString())))
 				{
 					UE_LOG(LogEditorFactories, Warning, TEXT("Package factory can only handle the map package or new packages!"));
 					return nullptr;
@@ -1451,7 +1451,7 @@ UObject* UPackageFactory::FactoryCreateText( UClass* Class, UObject* InParent, F
 			FString ParentPackageName;
 			FParse::Value(Str, TEXT("PARENTPACKAGE="), ParentPackageName);
 			UClass* PkgClass;
-			if (ParseObject<UClass>(Str, TEXT("CLASS="), PkgClass, ANY_PACKAGE))
+			if (ParseObject<UClass>(Str, TEXT("CLASS="), PkgClass, nullptr))
 			{
 				// Get the name of the object.
 				FName NewPackageName(NAME_None);
@@ -1468,13 +1468,13 @@ UObject* UPackageFactory::FactoryCreateText( UClass* Class, UObject* InParent, F
 					if ( FPackageName::ParseExportTextPath(ArchetypeName, &ObjectClass, &ObjectPath) )
 					{
 						// find the class
-						UClass* ArchetypeClass = (UClass*)StaticFindObject(UClass::StaticClass(), ANY_PACKAGE, *ObjectClass);
+						UClass* ArchetypeClass = UClass::TryFindTypeSlow<UClass>(ObjectClass, EFindFirstObjectOptions::EnsureIfAmbiguous);
 						if (ArchetypeClass)
 						{
 							if (ArchetypeClass->IsChildOf(UPackage::StaticClass()))
 							{
 								// if we had the class, find the archetype
-								Archetype = Cast<UPackage>(StaticFindObject(ArchetypeClass, ANY_PACKAGE, *ObjectPath));
+								Archetype = Cast<UPackage>(StaticFindObject(ArchetypeClass, nullptr, *ObjectPath));
 							}
 							else
 							{
@@ -1760,7 +1760,7 @@ UObject* UPolysFactory::FactoryCreateText
 			// only load the texture if it was present
 			if (FParse::Value( Str, TEXT("TEXTURE="), TextureName ))
 			{
-				Poly.Material = Cast<UMaterialInterface>(StaticFindObject( UMaterialInterface::StaticClass(), ANY_PACKAGE, *TextureName ) );
+				Poly.Material = FindFirstObject<UMaterialInterface>( *TextureName, EFindFirstObjectOptions::EnsureIfAmbiguous );
 /***
 				if (Poly.Material == nullptr)
 				{
@@ -5707,7 +5707,7 @@ void FCustomizableTextObjectFactory::ProcessBuffer(UObject* InParent, EObjectFla
 			}
 
 			UClass* ObjClass;
-			if( ParseObject<UClass>( Str, TEXT("CLASS="), ObjClass, ANY_PACKAGE ) )
+			if( ParseObject<UClass>( Str, TEXT("CLASS="), ObjClass, nullptr ) )
 			{
 				bool bOmitSubObjects = false;
 				if (!CanCreateClass(ObjClass, bOmitSubObjects))
@@ -5833,7 +5833,7 @@ bool FCustomizableTextObjectFactory::CanCreateObjectsFromText( const FString& Te
 			}
 
 			UClass* ObjClass;
-			if (ParseObject<UClass>(Str, TEXT("CLASS="), ObjClass, ANY_PACKAGE))
+			if (ParseObject<UClass>(Str, TEXT("CLASS="), ObjClass, nullptr))
 			{
 				bool bOmitSubObjects = false;;
 				if (CanCreateClass(ObjClass, bOmitSubObjects))
@@ -7655,7 +7655,7 @@ bool UBlendSpaceFactoryNew::ConfigureProperties()
 	FAssetPickerConfig AssetPickerConfig;
 
 	/** The asset picker will only show skeletal meshes */
-	AssetPickerConfig.Filter.ClassNames.Add(USkeleton::StaticClass()->GetFName());
+	AssetPickerConfig.Filter.ClassPaths.Add(USkeleton::StaticClass()->GetClassPathName());
 	AssetPickerConfig.Filter.bRecursiveClasses = true;
 
 	/** The delegate that fires when an asset was selected */
@@ -7730,7 +7730,7 @@ bool UBlendSpaceFactory1D::ConfigureProperties()
 	FAssetPickerConfig AssetPickerConfig;
 
 	/** The asset picker will only show skeletal meshes */
-	AssetPickerConfig.Filter.ClassNames.Add(USkeleton::StaticClass()->GetFName());
+	AssetPickerConfig.Filter.ClassPaths.Add(USkeleton::StaticClass()->GetClassPathName());
 	AssetPickerConfig.Filter.bRecursiveClasses = true;
 
 	/** The delegate that fires when an asset was selected */
diff --git a/Engine/Source/Editor/UnrealEd/Private/Factories/MaterialImportHelpers.cpp b/Engine/Source/Editor/UnrealEd/Private/Factories/MaterialImportHelpers.cpp
--- a/Engine/Source/Editor/UnrealEd/Private/Factories/MaterialImportHelpers.cpp
+++ b/Engine/Source/Editor/UnrealEd/Private/Factories/MaterialImportHelpers.cpp
@@ -91,7 +91,7 @@ UMaterialInterface* UMaterialImportHelpers::FindExistingMaterial(const FString&
 	FARFilter Filter;
 	Filter.bRecursiveClasses = true;
 	Filter.bRecursivePaths = bRecursivePaths;
-	Filter.ClassNames.Add(UMaterialInterface::StaticClass()->GetFName());
+	Filter.ClassPaths.Add(UMaterialInterface::StaticClass()->GetClassPathName());
 	Filter.PackagePaths.Add(FName(*BasePath));
 
 	TArray<FAssetData> AssetData;
diff --git a/Engine/Source/Editor/UnrealEd/Private/Factories/PoseAssetFactory.cpp b/Engine/Source/Editor/UnrealEd/Private/Factories/PoseAssetFactory.cpp
--- a/Engine/Source/Editor/UnrealEd/Private/Factories/PoseAssetFactory.cpp
+++ b/Engine/Source/Editor/UnrealEd/Private/Factories/PoseAssetFactory.cpp
@@ -50,7 +50,7 @@ public:
 
 		FAssetPickerConfig AssetPickerConfig;
 		/** The asset picker will only show sequences */
-		AssetPickerConfig.Filter.ClassNames.Add(UAnimSequence::StaticClass()->GetFName());
+		AssetPickerConfig.Filter.ClassPaths.Add(UAnimSequence::StaticClass()->GetClassPathName());
 		/** The delegate that fires when an asset was selected */
 		AssetPickerConfig.OnAssetSelected = FOnAssetSelected::CreateRaw(this, &SPoseConfigureWindow::OnSourceAnimationSelected);
 		if (SourceSequence != nullptr)
diff --git a/Engine/Source/Editor/UnrealEd/Private/Fbx/FbxMainImport.cpp b/Engine/Source/Editor/UnrealEd/Private/Fbx/FbxMainImport.cpp
--- a/Engine/Source/Editor/UnrealEd/Private/Fbx/FbxMainImport.cpp
+++ b/Engine/Source/Editor/UnrealEd/Private/Fbx/FbxMainImport.cpp
@@ -105,7 +105,7 @@ FBXImportOptions* GetImportOptions( UnFbx::FFbxImporter* FbxImporter, UFbxImport
 			// Look in the current target directory to see if we have a skeleton
 			FARFilter Filter;
 			Filter.PackagePaths.Add(*FPaths::GetPath(FullPath));
-			Filter.ClassNames.Add(USkeleton::StaticClass()->GetFName());
+			Filter.ClassPaths.Add(USkeleton::StaticClass()->GetClassPathName());
 
 			IAssetRegistry& AssetRegistry = FModuleManager::LoadModuleChecked<FAssetRegistryModule>("AssetRegistry").Get();
 			TArray<FAssetData> SkeletonAssets;
diff --git a/Engine/Source/Editor/UnrealEd/Private/Fbx/FbxMatineeImport.cpp b/Engine/Source/Editor/UnrealEd/Private/Fbx/FbxMatineeImport.cpp
--- a/Engine/Source/Editor/UnrealEd/Private/Fbx/FbxMatineeImport.cpp
+++ b/Engine/Source/Editor/UnrealEd/Private/Fbx/FbxMatineeImport.cpp
@@ -46,7 +46,7 @@ inline bool _HasUnknownCameras( AMatineeActor* InMatineeActor, FbxNode* Node, co
 		}
 		
 		// Attempt to name-match the scene node for this camera with one of the actors.
-		AActor* Actor = FindObject<AActor>( ANY_PACKAGE, Name );
+		AActor* Actor = FindFirstObject<AActor>( Name, EFindFirstObjectOptions::None, ELogVerbosity::Warning, TEXT("_HasUnknownCameras"));
 		if ( !IsValid(Actor) )
 		{
 			return true;
@@ -211,7 +211,7 @@ bool FFbxImporter::ImportMatineeSequence(AMatineeActor* InMatineeActor)
 		// Attempt to name-match the scene node with one of the actors.
 		if( Actor == NULL )
 		{
-			Actor = FindObject<AActor>( ANY_PACKAGE, UTF8_TO_TCHAR(Node->GetName()) );
+			Actor = FindFirstObject<AActor>( UTF8_TO_TCHAR(Node->GetName()), EFindFirstObjectOptions::None, ELogVerbosity::Warning, TEXT("ImportMatineeSequence"));
 		}
 
 		FbxCamera* CameraNode = NULL;
diff --git a/Engine/Source/Editor/UnrealEd/Private/Fbx/FbxReimportDialogs.cpp b/Engine/Source/Editor/UnrealEd/Private/Fbx/FbxReimportDialogs.cpp
--- a/Engine/Source/Editor/UnrealEd/Private/Fbx/FbxReimportDialogs.cpp
+++ b/Engine/Source/Editor/UnrealEd/Private/Fbx/FbxReimportDialogs.cpp
@@ -152,7 +152,7 @@ void FFbxImporter::ShowFbxSkeletonConflictWindow(USkeletalMesh* SkeletalMesh, US
 					FAssetData AssetData = AssetRegistryModule.Get().GetAssetByObjectPath(FName(*FullAssetPathName));
 					if (AssetData.GetClass() != nullptr)
 					{
-						TSharedPtr<FString> AssetReferencing = MakeShareable(new FString(AssetData.AssetClass.ToString() + TEXT(" ") + FullAssetPathName));
+						TSharedPtr<FString> AssetReferencing = MakeShareable(new FString(AssetData.AssetClassPath.ToString() + TEXT(" ") + FullAssetPathName));
 						AssetReferencingSkeleton.Add(AssetReferencing);
 					}
 				}
diff --git a/Engine/Source/Editor/UnrealEd/Private/Fbx/FbxSceneImportFactory.cpp b/Engine/Source/Editor/UnrealEd/Private/Fbx/FbxSceneImportFactory.cpp
--- a/Engine/Source/Editor/UnrealEd/Private/Fbx/FbxSceneImportFactory.cpp
+++ b/Engine/Source/Editor/UnrealEd/Private/Fbx/FbxSceneImportFactory.cpp
@@ -759,7 +759,7 @@ UObject *FFbxAttributeInfo::GetContentObject()
 	{
 		ContentPackage->FullyLoad();
 	}
-	ContentObject = FindObjectSafe<UObject>(ANY_PACKAGE, *AssetName);
+	ContentObject = FindFirstObject<UObject>(*AssetName, EFindFirstObjectOptions::None, ELogVerbosity::Warning, TEXT("FFbxAttributeInfo::GetContentObject()"));
 	if (ContentObject != nullptr)
 	{
 		if (ContentObject->HasAnyFlags(RF_Transient) || !IsValid(ContentObject))
diff --git a/Engine/Source/Editor/UnrealEd/Private/Fbx/FbxSkeletalMeshImport.cpp b/Engine/Source/Editor/UnrealEd/Private/Fbx/FbxSkeletalMeshImport.cpp
--- a/Engine/Source/Editor/UnrealEd/Private/Fbx/FbxSkeletalMeshImport.cpp
+++ b/Engine/Source/Editor/UnrealEd/Private/Fbx/FbxSkeletalMeshImport.cpp
@@ -1672,7 +1672,7 @@ USkeletalMesh* UnFbx::FFbxImporter::ImportSkeletalMesh(FImportSkeletalMeshArgs &
 	USkeletalMesh* ExistingSkelMesh = nullptr;
 	if ( !ImportSkeletalMeshArgs.FbxShapeArray  )
 	{
-		UObject* ExistingObject = StaticFindObjectFast(UObject::StaticClass(), ImportSkeletalMeshArgs.InParent, ImportSkeletalMeshArgs.Name, false, false, RF_NoFlags, EInternalObjectFlags::Garbage);
+		UObject* ExistingObject = StaticFindObjectFast(UObject::StaticClass(), ImportSkeletalMeshArgs.InParent, ImportSkeletalMeshArgs.Name, false, RF_NoFlags, EInternalObjectFlags::Garbage);
 		ExistingSkelMesh = Cast<USkeletalMesh>(ExistingObject);
 
 		if (!ExistingSkelMesh && ExistingObject)
@@ -2109,7 +2109,7 @@ USkeletalMesh* UnFbx::FFbxImporter::ImportSkeletalMesh(FImportSkeletalMeshArgs &
 						TArray<FAssetData> SkeletalMeshAssetData;
 						
 						FARFilter ARFilter;
-						ARFilter.ClassNames.Add(*USkeletalMesh::StaticClass()->GetName());
+						ARFilter.ClassPaths.Add(USkeletalMesh::StaticClass()->GetClassPathName());
 						ARFilter.TagsAndValues.Add(TEXT("Skeleton"), FAssetData(Skeleton).GetExportTextName());
 
 						IAssetRegistry& AssetRegistry = AssetRegistryModule.Get();
diff --git a/Engine/Source/Editor/UnrealEd/Private/Fbx/ReimportFbxSceneFactory.cpp b/Engine/Source/Editor/UnrealEd/Private/Fbx/ReimportFbxSceneFactory.cpp
--- a/Engine/Source/Editor/UnrealEd/Private/Fbx/ReimportFbxSceneFactory.cpp
+++ b/Engine/Source/Editor/UnrealEd/Private/Fbx/ReimportFbxSceneFactory.cpp
@@ -953,7 +953,7 @@ UBlueprint *UReimportFbxSceneFactory::UpdateOriginalBluePrint(FString &BluePrint
 	}
 	//Load the package before searching the asset
 	PkgExist->FullyLoad();
-	UBlueprint* BluePrint = FindObjectSafe<UBlueprint>(ANY_PACKAGE, *BluePrintFullName);
+	UBlueprint* BluePrint = FindObjectSafe<UBlueprint>(nullptr, *BluePrintFullName);
 	if (BluePrint == nullptr)
 	{
 		return nullptr;
@@ -1257,7 +1257,7 @@ EReimportResult::Type UReimportFbxSceneFactory::ReimportSkeletalMesh(void* VoidF
 	UPackage* PkgExist = MeshInfo->GetContentPackage();
 
 	FString AssetName = MeshInfo->GetFullImportName();
-	USkeletalMesh* Mesh = FindObjectSafe<USkeletalMesh>(ANY_PACKAGE, *AssetName);
+	USkeletalMesh* Mesh = FindFirstObject<USkeletalMesh>(*AssetName, EFindFirstObjectOptions::None, ELogVerbosity::Warning, TEXT("ReimportSkeletalMesh"));
 	if (Mesh == nullptr)
 	{
 		//We reimport only skeletal mesh here
@@ -1398,7 +1398,8 @@ EReimportResult::Type UReimportFbxSceneFactory::ReimportSkeletalMesh(void* VoidF
 						{
 							ParentPackage->FullyLoad();
 						}
-						UObject* Object = FindObjectSafe<UObject>(ANY_PACKAGE, *SequenceName);
+
+						UObject* Object = FindFirstObject<UObject>(*SequenceName, EFindFirstObjectOptions::None, ELogVerbosity::Warning, TEXT("ReimportSkeletalMesh"));
 						if (Object != nullptr)
 						{
 							if (ParentPackage == nullptr)
@@ -1481,7 +1482,7 @@ EReimportResult::Type UReimportFbxSceneFactory::ReimportStaticMesh(void* VoidFbx
 	MeshInfo->GetContentPackage();
 
 	FString AssetName = MeshInfo->GetFullImportName();
-	UStaticMesh* Mesh = FindObjectSafe<UStaticMesh>(ANY_PACKAGE, *AssetName);
+	UStaticMesh* Mesh = FindFirstObject<UStaticMesh>(*AssetName, EFindFirstObjectOptions::None, ELogVerbosity::Warning, TEXT("ReimportStaticMesh"));
 	if (Mesh == nullptr)
 	{
 		//We reimport only static mesh here
diff --git a/Engine/Source/Editor/UnrealEd/Private/Fbx/SSceneMaterialsListView.cpp b/Engine/Source/Editor/UnrealEd/Private/Fbx/SSceneMaterialsListView.cpp
--- a/Engine/Source/Editor/UnrealEd/Private/Fbx/SSceneMaterialsListView.cpp
+++ b/Engine/Source/Editor/UnrealEd/Private/Fbx/SSceneMaterialsListView.cpp
@@ -335,7 +335,7 @@ TSharedPtr<SWidget> SFbxSceneMaterialsListView::OnOpenContextMenu()
 			// Configure filter for asset picker
 			FAssetPickerConfig Config;
 			Config.Filter.bRecursiveClasses = true;
-			Config.Filter.ClassNames.Add(UMaterialInterface::StaticClass()->GetFName());
+			Config.Filter.ClassPaths.Add(UMaterialInterface::StaticClass()->GetClassPathName());
 			Config.InitialAssetViewType = EAssetViewType::List;
 			Config.OnAssetSelected = FOnAssetSelected::CreateSP(this, &SFbxSceneMaterialsListView::AssignMaterialAssetData);
 			Config.bAllowNullSelection = false;
@@ -440,7 +440,7 @@ void SFbxSceneMaterialsListView::AssignMaterialToExisting()
 		SelectAssetConfig.DialogTitleOverride = LOCTEXT("FbxChooseMaterialAssetContentPath", "Choose a material asset");
 		//SelectAssetConfig.DefaultPath = ItemPtr->OriginalImportPath;
 		SelectAssetConfig.bAllowMultipleSelection = false;
-		SelectAssetConfig.AssetClassNames.Add(UMaterial::StaticClass()->GetFName());
+		SelectAssetConfig.AssetClassNames.Add(UMaterial::StaticClass()->GetClassPathName());
 		TArray<FAssetData> AssetData = ContentBrowserModule.Get().CreateModalOpenAssetDialog(SelectAssetConfig);
 		if (AssetData.Num() == 1)
 		{
diff --git a/Engine/Source/Editor/UnrealEd/Private/Fbx/SSceneReimportStaticMeshListView.cpp b/Engine/Source/Editor/UnrealEd/Private/Fbx/SSceneReimportStaticMeshListView.cpp
--- a/Engine/Source/Editor/UnrealEd/Private/Fbx/SSceneReimportStaticMeshListView.cpp
+++ b/Engine/Source/Editor/UnrealEd/Private/Fbx/SSceneReimportStaticMeshListView.cpp
@@ -511,7 +511,7 @@ void SFbxSceneStaticMeshReimportListView::AssignToStaticMesh()
 		SelectAssetConfig.DialogTitleOverride = LOCTEXT("FbxChooseReimportAssetContentPath", "Choose static mesh asset for reimporting the fbx scene content");
 		//SelectAssetConfig.DefaultPath = ItemPtr->OriginalImportPath;
 		SelectAssetConfig.bAllowMultipleSelection = false;
-		SelectAssetConfig.AssetClassNames.Add(UStaticMesh::StaticClass()->GetFName());
+		SelectAssetConfig.AssetClassNames.Add(UStaticMesh::StaticClass()->GetClassPathName());
 		TArray<FAssetData> AssetData = ContentBrowserModule.Get().CreateModalOpenAssetDialog(SelectAssetConfig);
 		if (AssetData.Num() == 1)
 		{
diff --git a/Engine/Source/Editor/UnrealEd/Private/FileHelpers.cpp b/Engine/Source/Editor/UnrealEd/Private/FileHelpers.cpp
--- a/Engine/Source/Editor/UnrealEd/Private/FileHelpers.cpp
+++ b/Engine/Source/Editor/UnrealEd/Private/FileHelpers.cpp
@@ -1018,7 +1018,7 @@
 	{
 		SaveAssetDialogConfig.DefaultPath = DefaultPath;
 		SaveAssetDialogConfig.DefaultAssetName = NewNameSuggestion;
-		SaveAssetDialogConfig.AssetClassNames.Add(SavedClass->GetFName());
+		SaveAssetDialogConfig.AssetClassNames.Add(SavedClass->GetClassPathName());
 		SaveAssetDialogConfig.ExistingAssetPolicy = ESaveAssetDialogExistingAssetPolicy::AllowButWarn;
 		SaveAssetDialogConfig.DialogTitleOverride = (SavedClass == UWorld::StaticClass())
 			? LOCTEXT("SaveLevelDialogTitle", "Save Level As")
@@ -2427,7 +2427,7 @@
 	FOpenAssetDialogConfig OpenAssetDialogConfig;
 	OpenAssetDialogConfig.DialogTitleOverride = LOCTEXT("OpenLevelDialogTitle", "Open Level");
 	OpenAssetDialogConfig.DefaultPath = DefaultPath;
-	OpenAssetDialogConfig.AssetClassNames.Add(UWorld::StaticClass()->GetFName());
+	OpenAssetDialogConfig.AssetClassNames.Add(UWorld::StaticClass()->GetClassPathName());
 	OpenAssetDialogConfig.bAllowMultipleSelection = bAllowMultipleSelection;
 
 	FContentBrowserModule& ContentBrowserModule = FModuleManager::LoadModuleChecked<FContentBrowserModule>("ContentBrowser");
diff --git a/Engine/Source/Editor/UnrealEd/Private/ImportUtils/SkeletalMeshImportUtils.cpp b/Engine/Source/Editor/UnrealEd/Private/ImportUtils/SkeletalMeshImportUtils.cpp
--- a/Engine/Source/Editor/UnrealEd/Private/ImportUtils/SkeletalMeshImportUtils.cpp
+++ b/Engine/Source/Editor/UnrealEd/Private/ImportUtils/SkeletalMeshImportUtils.cpp
@@ -122,7 +122,7 @@ void SkeletalMeshImportUtils::ProcessImportMeshMaterials(TArray<FSkeletalMateria
 		{
 			const FString& MaterialName = ImportedMaterial.MaterialImportName;
 			MaterialNameNoSkin = MaterialName;
-			Material = FindObject<UMaterialInterface>(ANY_PACKAGE, *MaterialName);
+			Material = FindFirstObject<UMaterialInterface>(*MaterialName, EFindFirstObjectOptions::None, ELogVerbosity::Warning, TEXT("ProcessImportMeshMaterials"));
 			if (Material == nullptr)
 			{
 				SkinOffset = MaterialName.Find(TEXT("_skin"), ESearchCase::IgnoreCase, ESearchDir::FromEnd);
@@ -132,7 +132,7 @@ void SkeletalMeshImportUtils::ProcessImportMeshMaterials(TArray<FSkeletalMateria
 					if (SkinXXNumber.IsNumeric())
 					{
 						MaterialNameNoSkin = MaterialName.LeftChop(MaterialName.Len() - SkinOffset);
-						Material = FindObject<UMaterialInterface>(ANY_PACKAGE, *MaterialNameNoSkin);
+						Material = FindFirstObject<UMaterialInterface>(*MaterialNameNoSkin, EFindFirstObjectOptions::None, ELogVerbosity::Warning, TEXT("ProcessImportMeshMaterials"));
 					}
 				}
 			}
diff --git a/Engine/Source/Editor/UnrealEd/Private/Kismet2/BlueprintEditorUtils.cpp b/Engine/Source/Editor/UnrealEd/Private/Kismet2/BlueprintEditorUtils.cpp
--- a/Engine/Source/Editor/UnrealEd/Private/Kismet2/BlueprintEditorUtils.cpp
+++ b/Engine/Source/Editor/UnrealEd/Private/Kismet2/BlueprintEditorUtils.cpp
@@ -6225,7 +6225,7 @@ void FBlueprintEditorUtils::FixupVariableDescription(UBlueprint* Blueprint, FBPV
 		FString BitmaskEnumTypeName = VarDesc.GetMetaData(FBlueprintMetadata::MD_BitmaskEnum);
 		if (!BitmaskEnumTypeName.IsEmpty())
 		{
-			UEnum* BitflagsEnum = FindObject<UEnum>(ANY_PACKAGE, *BitmaskEnumTypeName);
+			UEnum* BitflagsEnum = UClass::TryFindTypeSlow<UEnum>(BitmaskEnumTypeName);
 			if (BitflagsEnum == nullptr || !BitflagsEnum->HasMetaData(*FBlueprintMetadata::MD_Bitflags.ToString()))
 			{
 				VarDesc.RemoveMetaData(FBlueprintMetadata::MD_BitmaskEnum);
@@ -6437,12 +6437,12 @@ FGuid FBlueprintEditorUtils::FindInterfaceFunctionGuid(const UFunction* Function
 }
 
 // Add a new interface, and member function graphs to the blueprint
-bool FBlueprintEditorUtils::ImplementNewInterface(UBlueprint* Blueprint, const FName& InterfaceClassName)
+bool FBlueprintEditorUtils::ImplementNewInterface(UBlueprint* Blueprint, FTopLevelAssetPath InterfaceClassPathName)
 {
-	check(InterfaceClassName != NAME_None);
+	check(!InterfaceClassPathName.IsNull());
 
 	// Attempt to find the class we want to implement
-	UClass* InterfaceClass = (UClass*)StaticFindObject(UClass::StaticClass(), ANY_PACKAGE, *InterfaceClassName.ToString());
+	UClass* InterfaceClass = FindObject<UClass>(InterfaceClassPathName);
 
 	// Make sure the class is found, and isn't native (since Blueprints don't necessarily generate native classes.
 	check(InterfaceClass);
@@ -6456,7 +6456,7 @@ bool FBlueprintEditorUtils::ImplementNewInterface(UBlueprint* Blueprint, const F
 				FText::Format(
 					LOCTEXT("InterfaceAlreadyImplementedFmt", "Blueprint '{0}' already implements the interface '{1}'"),
 					FText::FromString(Blueprint->GetName()),
-					FText::FromString(InterfaceClassName.ToString())
+					FText::FromString(InterfaceClassPathName.ToString())
 				),
 				EMessageSeverity::Warning
 			);
@@ -6489,7 +6489,7 @@ bool FBlueprintEditorUtils::ImplementNewInterface(UBlueprint* Blueprint, const F
 						LOCTEXT("InterfaceFunctionConflictsFmt", "Blueprint '{0}' has a function or graph which conflicts with function '{1}' in interface '{2}'"),
 						FText::FromString(Blueprint->GetName()),
 						FText::FromName(FunctionName),
-						FText::FromName(InterfaceClassName)
+						FText::FromString(InterfaceClassPathName.ToString())
 					),
 					EMessageSeverity::Error
 				);
@@ -6522,12 +6522,18 @@ bool FBlueprintEditorUtils::ImplementNewInterface(UBlueprint* Blueprint, const F
 	return bAllFunctionsAdded;
 }
 
+bool FBlueprintEditorUtils::ImplementNewInterface(UBlueprint* Blueprint, const FName& InterfaceClassName)
+{
+	FTopLevelAssetPath InterfaceClassPathName = UClass::TryConvertShortTypeNameToPathName<UStruct>(InterfaceClassName.ToString(), ELogVerbosity::Warning, TEXT("FBlueprintEditorUtils::ImplementNewInterface"));
+	return ImplementNewInterface(Blueprint, InterfaceClassPathName);
+}
+
 // Gets the graphs currently in the blueprint associated with the specified interface
-void FBlueprintEditorUtils::GetInterfaceGraphs(UBlueprint* Blueprint, const FName& InterfaceClassName, TArray<UEdGraph*>& ChildGraphs)
+void FBlueprintEditorUtils::GetInterfaceGraphs(UBlueprint* Blueprint, FTopLevelAssetPath InterfaceClassPathName, TArray<UEdGraph*>& ChildGraphs)
 {
 	ChildGraphs.Empty();
 
-	if( InterfaceClassName == NAME_None )
+	if (InterfaceClassPathName.IsNull())
 	{
 		return;
 	}
@@ -6535,7 +6541,7 @@ void FBlueprintEditorUtils::GetInterfaceGraphs(UBlueprint* Blueprint, const FNam
 	// Find the implemented interface
 	for( int32 i = 0; i < Blueprint->ImplementedInterfaces.Num(); i++ )
 	{
-		if( Blueprint->ImplementedInterfaces[i].Interface->GetFName() == InterfaceClassName )
+		if( Blueprint->ImplementedInterfaces[i].Interface->GetClassPathName() == InterfaceClassPathName)
 		{
 			ChildGraphs = Blueprint->ImplementedInterfaces[i].Graphs;
 			return;			
@@ -6543,6 +6549,12 @@ void FBlueprintEditorUtils::GetInterfaceGraphs(UBlueprint* Blueprint, const FNam
 	}
 }
 
+void FBlueprintEditorUtils::GetInterfaceGraphs(UBlueprint* Blueprint, const FName& InterfaceClassName, TArray<UEdGraph*>& ChildGraphs)
+{
+	FTopLevelAssetPath InterfaceClassPathName = UClass::TryConvertShortTypeNameToPathName<UStruct>(InterfaceClassName.ToString(), ELogVerbosity::Warning, TEXT("FBlueprintEditorUtils::GetInterfaceGraphs"));
+	GetInterfaceGraphs(Blueprint, InterfaceClassPathName, ChildGraphs);
+}
+
 UFunction* FBlueprintEditorUtils::GetInterfaceFunction(UBlueprint* Blueprint, const FName FuncName)
 {
 	UFunction* Function = nullptr;
@@ -6592,9 +6604,9 @@ bool FBlueprintEditorUtils::IsInterfaceFunction(UBlueprint* Blueprint, UFunction
 }
 
 // Remove an implemented interface, and its associated member function graphs
-void FBlueprintEditorUtils::RemoveInterface(UBlueprint* Blueprint, const FName& InterfaceClassName, bool bPreserveFunctions /*= false*/)
+void FBlueprintEditorUtils::RemoveInterface(UBlueprint* Blueprint, FTopLevelAssetPath InterfaceClassPathName, bool bPreserveFunctions /*= false*/)
 {
-	if( InterfaceClassName == NAME_None )
+	if (InterfaceClassPathName.IsNull())
 	{
 		return;
 	}
@@ -6603,7 +6615,7 @@ void FBlueprintEditorUtils::RemoveInterface(UBlueprint* Blueprint, const FName&
 	int32 Idx = INDEX_NONE;
 	for( int32 i = 0; i < Blueprint->ImplementedInterfaces.Num(); i++ )
 	{
-		if( Blueprint->ImplementedInterfaces[i].Interface->GetFName() == InterfaceClassName )
+		if( Blueprint->ImplementedInterfaces[i].Interface->GetClassPathName() == InterfaceClassPathName)
 		{
 			Idx = i;
 			break;
@@ -6679,6 +6691,12 @@ void FBlueprintEditorUtils::RemoveInterface(UBlueprint* Blueprint, const FName&
 	}
 }
 
+void FBlueprintEditorUtils::RemoveInterface(UBlueprint* Blueprint, const FName& InterfaceClassName, bool bPreserveFunctions /*= false*/)
+{
+	FTopLevelAssetPath InterfaceClassPathName = UClass::TryConvertShortTypeNameToPathName<UStruct>(InterfaceClassName.ToString(), ELogVerbosity::Warning, TEXT("FBlueprintEditorUtils::RemoveInterface"));
+	return RemoveInterface(Blueprint, InterfaceClassPathName, bPreserveFunctions);
+}
+
 bool FBlueprintEditorUtils::RemoveInterfaceFunction(UBlueprint* Blueprint, FBPInterfaceDescription& Interface, UFunction* Function, bool bPreserveFunction)
 {
 	for (TArray<UEdGraph*>::TIterator it(Interface.Graphs); it; ++it)
@@ -8503,7 +8521,7 @@ bool FBlueprintEditorUtils::KismetDiagnosticExec(const TCHAR* InStream, FOutputD
 	}
 	else if (FParse::Command(&Str, TEXT("RepairBlueprint")))
 	{
-		if (UBlueprint* Blueprint = FindObject<UBlueprint>(ANY_PACKAGE, Str))
+		if (UBlueprint* Blueprint = FindFirstObject<UBlueprint>(Str, EFindFirstObjectOptions::None, ELogVerbosity::Warning, TEXT("RepairBlueprint")))
 		{
 			IKismetCompilerInterface& Compiler = FModuleManager::LoadModuleChecked<IKismetCompilerInterface>(KISMET_COMPILER_MODULENAME);
 			Compiler.RecoverCorruptedBlueprint(Blueprint);
@@ -8831,7 +8849,7 @@ TSharedRef<SWidget> FBlueprintEditorUtils::ConstructBlueprintInterfaceClassPicke
 			{
 				ProhibitedInterfaceNames[ExclusionIndex].TrimStartInline();
 				FString const& ProhibitedInterfaceName = ProhibitedInterfaceNames[ExclusionIndex].RightChop(1);
-				UClass* ProhibitedInterface = (UClass*)StaticFindObject(UClass::StaticClass(), ANY_PACKAGE, *ProhibitedInterfaceName);
+				UClass* ProhibitedInterface = UClass::TryFindTypeSlow<UClass>(ProhibitedInterfaceName);
 				if(ProhibitedInterface)
 				{
 					Filter->DisallowedClasses.Add(ProhibitedInterface);
@@ -9876,7 +9894,7 @@ const FSlateBrush* FBlueprintEditorUtils::GetIconFromPin( const FEdGraphPinType&
 	}
 	else if( PinSubObject )
 	{
-		UClass* VarClass = FindObject<UClass>(ANY_PACKAGE, *PinSubObject->GetName());
+		UClass* VarClass = FindObject<UClass>(nullptr, *PinSubObject->GetFullName());
 		if( VarClass )
 		{
 			IconBrush = FSlateIconFinder::FindIconBrushForClass( VarClass );
diff --git a/Engine/Source/Editor/UnrealEd/Private/LevelEditorViewport.cpp b/Engine/Source/Editor/UnrealEd/Private/LevelEditorViewport.cpp
--- a/Engine/Source/Editor/UnrealEd/Private/LevelEditorViewport.cpp
+++ b/Engine/Source/Editor/UnrealEd/Private/LevelEditorViewport.cpp
@@ -341,12 +341,13 @@ static UTexture* GetTextureWithNameVariations( const FString& BasePackageName, c
 	FAssetRegistryModule& AssetRegistryModule = FModuleManager::LoadModuleChecked<FAssetRegistryModule>( TEXT( "AssetRegistry" ) );
 
 	// Try all the variations of suffixes, if we find a package matching the suffix, return it.
+	const FTopLevelAssetPath Texture2DClassPath(TEXT("/Script/Engine"), TEXT("Texture2D"));
 	for ( int i = 0; i < Suffixes.Num(); i++ )
 	{
 		TArray<FAssetData> OutAssetData;
 		if ( AssetRegistryModule.Get().GetAssetsByPackageName( *( BasePackageName + Suffixes[i] ), OutAssetData ) && OutAssetData.Num() > 0 )
 		{
-			if ( OutAssetData[0].AssetClass == "Texture2D" )
+			if ( OutAssetData[0].AssetClassPath == Texture2DClassPath )
 			{
 				return Cast<UTexture>(OutAssetData[0].GetAsset());
 			}
diff --git a/Engine/Source/Editor/UnrealEd/Private/LocalizedAssetUtil.cpp b/Engine/Source/Editor/UnrealEd/Private/LocalizedAssetUtil.cpp
--- a/Engine/Source/Editor/UnrealEd/Private/LocalizedAssetUtil.cpp
+++ b/Engine/Source/Editor/UnrealEd/Private/LocalizedAssetUtil.cpp
@@ -138,7 +138,7 @@ bool FLocalizedAssetSCCUtil::SaveFileWithSCC(const TSharedPtr<FLocalizationSCC>&
 }
 
 
-bool FLocalizedAssetUtil::GetAssetsByPathAndClass(IAssetRegistry& InAssetRegistry, const FName InPackagePath, const FName InClassName, const bool bIncludeLocalizedAssets, TArray<FAssetData>& OutAssets)
+bool FLocalizedAssetUtil::GetAssetsByPathAndClass(IAssetRegistry& InAssetRegistry, const FName InPackagePath, const FTopLevelAssetPath& InClassName, const bool bIncludeLocalizedAssets, TArray<FAssetData>& OutAssets)
 {
 	TArray<FName> PackagePaths;
 	PackagePaths.Add(InPackagePath);
@@ -146,12 +146,12 @@ bool FLocalizedAssetUtil::GetAssetsByPathAndClass(IAssetRegistry& InAssetRegistr
 	return GetAssetsByPathAndClass(InAssetRegistry, PackagePaths, InClassName, bIncludeLocalizedAssets, OutAssets);
 }
 
-bool FLocalizedAssetUtil::GetAssetsByPathAndClass(IAssetRegistry& InAssetRegistry, const TArray<FName>& InPackagePaths, const FName InClassName, const bool bIncludeLocalizedAssets, TArray<FAssetData>& OutAssets)
+bool FLocalizedAssetUtil::GetAssetsByPathAndClass(IAssetRegistry& InAssetRegistry, const TArray<FName>& InPackagePaths, const FTopLevelAssetPath& InClassName, const bool bIncludeLocalizedAssets, TArray<FAssetData>& OutAssets)
 {
 	FARFilter AssetFilter;
 	AssetFilter.PackagePaths = InPackagePaths;
 	AssetFilter.bRecursivePaths = true;
-	AssetFilter.ClassNames.Add(InClassName);
+	AssetFilter.ClassPaths.Add(InClassName);
 	AssetFilter.bRecursiveClasses = true;
 
 	if (!InAssetRegistry.GetAssets(AssetFilter, OutAssets))
diff --git a/Engine/Source/Editor/UnrealEd/Private/MaterialGraphSchema.cpp b/Engine/Source/Editor/UnrealEd/Private/MaterialGraphSchema.cpp
--- a/Engine/Source/Editor/UnrealEd/Private/MaterialGraphSchema.cpp
+++ b/Engine/Source/Editor/UnrealEd/Private/MaterialGraphSchema.cpp
@@ -922,7 +922,7 @@ void UMaterialGraphSchema::GetMaterialFunctionActions(FGraphActionMenuBuilder& A
 
 	// Collect a full list of assets with the specified class
 	TArray<FAssetData> AssetDataList;
-	AssetRegistryModule.Get().GetAssetsByClass(UMaterialFunction::StaticClass()->GetFName(), AssetDataList);
+	AssetRegistryModule.Get().GetAssetsByClass(UMaterialFunction::StaticClass()->GetClassPathName(), AssetDataList);
 
 	for (const FAssetData& AssetData : AssetDataList)
 	{
diff --git a/Engine/Source/Editor/UnrealEd/Private/PackageTools.cpp b/Engine/Source/Editor/UnrealEd/Private/PackageTools.cpp
--- a/Engine/Source/Editor/UnrealEd/Private/PackageTools.cpp
+++ b/Engine/Source/Editor/UnrealEd/Private/PackageTools.cpp
@@ -1472,7 +1472,7 @@ UPackageTools::UPackageTools(const FObjectInitializer& ObjectInitializer)
 			if (NumberOfAssets == 1)
 			{
 				const FAssetData& AssetData = OutAssets[0];
-				if (AssetData.AssetClass != AssetClass->GetFName())
+				if (AssetData.AssetClassPath != AssetClass->GetClassPathName())
 				{
 					bShouldGenerateUniquePackageName = true;
 				}
diff --git a/Engine/Source/Editor/UnrealEd/Private/SSkeletonWidget.cpp b/Engine/Source/Editor/UnrealEd/Private/SSkeletonWidget.cpp
--- a/Engine/Source/Editor/UnrealEd/Private/SSkeletonWidget.cpp
+++ b/Engine/Source/Editor/UnrealEd/Private/SSkeletonWidget.cpp
@@ -35,7 +35,7 @@ void SSkeletonListWidget::Construct(const FArguments& InArgs)
 	FContentBrowserModule& ContentBrowserModule = FModuleManager::Get().LoadModuleChecked<FContentBrowserModule>(TEXT("ContentBrowser"));
 
 	FAssetPickerConfig AssetPickerConfig;
-	AssetPickerConfig.Filter.ClassNames.Add(USkeleton::StaticClass()->GetFName());
+	AssetPickerConfig.Filter.ClassPaths.Add(USkeleton::StaticClass()->GetClassPathName());
 	AssetPickerConfig.OnAssetSelected = FOnAssetSelected::CreateSP(this, &SSkeletonListWidget::SkeletonSelectionChanged);
 	AssetPickerConfig.InitialAssetViewType = EAssetViewType::Column;
 	AssetPickerConfig.SelectionMode = ESelectionMode::Single;
@@ -126,7 +126,7 @@ void SSkeletonCompareWidget::Construct(const FArguments& InArgs)
 	FContentBrowserModule& ContentBrowserModule = FModuleManager::Get().LoadModuleChecked<FContentBrowserModule>(TEXT("ContentBrowser"));
 
 	FAssetPickerConfig AssetPickerConfig;
-	AssetPickerConfig.Filter.ClassNames.Add(USkeleton::StaticClass()->GetFName());
+	AssetPickerConfig.Filter.ClassPaths.Add(USkeleton::StaticClass()->GetClassPathName());
 	AssetPickerConfig.OnAssetSelected = FOnAssetSelected::CreateSP(this, &SSkeletonCompareWidget::SkeletonSelectionChanged);
 	AssetPickerConfig.InitialAssetViewType = EAssetViewType::Column;
 	AssetPickerConfig.bShowPathInColumnView = true;
@@ -780,7 +780,7 @@ void SReplaceMissingSkeletonDialog::Construct(const FArguments& InArgs)
 	FAssetPickerConfig AssetPickerConfig;
 
 	/** The asset picker will only show skeleton assets */
-	AssetPickerConfig.Filter.ClassNames.Add(USkeleton::StaticClass()->GetFName());
+	AssetPickerConfig.Filter.ClassPaths.Add(USkeleton::StaticClass()->GetClassPathName());
 	AssetPickerConfig.Filter.bRecursiveClasses = true;
 
 	/** The delegate that fires when an asset was selected */
