Fixing crash in AssetThumbnail.cpp due to unchecked pointer to an unloaded asset.

diff --git a/Engine/Source/Editor/UnrealEd/Private/AssetThumbnail.cpp b/Engine/Source/Editor/UnrealEd/Private/AssetThumbnail.cpp
--- a/Engine/Source/Editor/UnrealEd/Private/AssetThumbnail.cpp
+++ b/Engine/Source/Editor/UnrealEd/Private/AssetThumbnail.cpp
@@ -1109,9 +1109,11 @@ void FAssetThumbnailPool::Tick( float DeltaTime )
 						// Only load the custom thumbnail if the original asset is also loaded
 						//if (InfoRef->AssetData.IsAssetLoaded())
 						{
-							UObject* CustomThumbnail = CustomThumbnailAsset.GetAsset();
-							bLoadedThumbnail = LoadThumbnail(InfoRef, bIsAssetStillCompiling, CustomThumbnailAsset);
-							CustomThumbnail->ClearFlags(RF_Standalone);
+							if (UObject* CustomThumbnail = CustomThumbnailAsset.GetAsset())
+							{
+								bLoadedThumbnail = LoadThumbnail(InfoRef, bIsAssetStillCompiling, CustomThumbnailAsset);
+								CustomThumbnail->ClearFlags(RF_Standalone);
+							}
 						}
 					}
 
