Added PreTransactionFinalized transaction event to allow handling of finalization which may create nested transactions

diff --git a/Engine/Source/Editor/UnrealEd/Private/EditorTransaction.cpp b/Engine/Source/Editor/UnrealEd/Private/EditorTransaction.cpp
--- a/Engine/Source/Editor/UnrealEd/Private/EditorTransaction.cpp
+++ b/Engine/Source/Editor/UnrealEd/Private/EditorTransaction.cpp
@@ -1267,6 +1267,11 @@ int32 UTransBuffer::End()
 	// So, the active count could be 0, but the code path may still call end.
 	if (ActiveCount >= 1)
 	{
+		if (ActiveCount == 1)
+		{
+			TransactionStateChangedDelegate.Broadcast(GUndo->GetContext(), ETransactionStateEventType::PreTransactionFinalized);
+		}
+
 		if( --ActiveCount==0 )
 		{
 			if (GUndo)
diff --git a/Engine/Source/Runtime/Core/Public/Misc/ITransaction.h b/Engine/Source/Runtime/Core/Public/Misc/ITransaction.h
--- a/Engine/Source/Runtime/Core/Public/Misc/ITransaction.h
+++ b/Engine/Source/Runtime/Core/Public/Misc/ITransaction.h
@@ -20,6 +20,8 @@ enum class ETransactionStateEventType : uint8
 	TransactionStarted,
 	/** A transaction was canceled. */
 	TransactionCanceled,
+	/** A transaction is about to be finalized. It is still safe to temporarily open new transactions within the scope of this event. */
+	PreTransactionFinalized,
 	/** A transaction was finalized. */
 	TransactionFinalized,
 
