Implement export of the source HDR texture for long-lat cubemaps.

diff --git a/Engine/Source/Editor/UnrealEd/Private/Factories/EditorFactories.cpp b/Engine/Source/Editor/UnrealEd/Private/Factories/EditorFactories.cpp
--- a/Engine/Source/Editor/UnrealEd/Private/Factories/EditorFactories.cpp
+++ b/Engine/Source/Editor/UnrealEd/Private/Factories/EditorFactories.cpp
@@ -4849,9 +4849,14 @@ UVirtualTextureBuilderExporterDDS::UVirtualTextureBuilderExporterDDS(const FObje
 	UTextureExporterGeneric base class for all standard 2d texture exporters
 ------------------------------------------------------------------------------*/
 
-static bool ExportTexture2DGeneric( UObject* Object,FArchive& Ar, const TCHAR * ImageFormat, FFeedbackContext* Warn, int FileIndex=0)
+static bool ExportTextureSourceGeneric( UObject* Object,FArchive& Ar, const TCHAR * ImageFormat, FFeedbackContext* Warn, int FileIndex=0)
 {
-	UTexture2D* Texture = CastChecked<UTexture2D>( Object );
+	// supports export of the source texture for Texture2D and TextureCube generated from long-lat image
+	if (!Object->IsA(UTexture2D::StaticClass()) && !(Object->IsA(UTextureCube::StaticClass()) && Cast<UTextureCube>(Object)->Source.IsLongLatCubemap()))
+	{
+		return false;
+	}
+	UTexture* Texture = CastChecked<UTexture>( Object );
 	if ( Texture == nullptr )
 	{
 		return false;
@@ -5010,7 +5015,7 @@ bool UTextureExporterGeneric::ExportBinary( UObject* Object, const TCHAR* Type,
 	
 	// FileIndex is for layers for VT
 	// Type == Image extension
-	return ExportTexture2DGeneric(Texture2D,Ar,Type,Warn,FileIndex);
+	return ExportTextureSourceGeneric(Texture2D,Ar,Type,Warn,FileIndex);
 }
 
 //===========================================================
@@ -5209,7 +5214,7 @@ UTextureCubeExporterHDR::UTextureCubeExporterHDR(const FObjectInitializer& Objec
 	SupportedClass = UTextureCube::StaticClass();
 	PreferredFormatIndex = 0;
 	FormatExtension.Add(TEXT("HDR"));
-	FormatDescription.Add(TEXT("HDR cube LongLatUnwrap"));
+	FormatDescription.Add(TEXT("HDR cube LongLat"));
 	
 	// this is the only exporter registered for UTextureCube
 	// this does GenerateLongLatUnwrap
@@ -5223,8 +5228,16 @@ bool UTextureCubeExporterHDR::ExportBinary(UObject* Object, const TCHAR* Type, F
 
 	if (TexCube != nullptr)
 	{
-		// this does GenerateLongLatUnwrap
-		return FImageUtils::ExportTextureCubeAsHDR(TexCube, Ar);
+		if (TexCube->Source.IsLongLatCubemap() && TexCube->Source.GetFormat() == TSF_BGRE8)
+		{
+			// export the source HDR long-lat texture if available
+			return ExportTextureSourceGeneric(TexCube, Ar, Type, Warn, FileIndex);
+		}
+		else
+		{
+			// this does GenerateLongLatUnwrap
+			return FImageUtils::ExportTextureCubeAsHDR(TexCube, Ar);
+		}
 	}
 	else if (TexRTCube != nullptr)
 	{
