Fix thumbnails not loading in editor until the package is loaded. ThumbnailTools::FindCachedThumbnailInPackage and LoadThumbnailsFromPackage are sometimes called with AssetData.GetFullName, and that function has changed to describe the class as a ClassPath instead of a ShortClassName. But package thumbnails are stored by FullName with class described as a ShortClassName. Change Thumbnail tools to convert input to ShortClassName.
#rnx

diff --git a/Engine/Source/Editor/UnrealEd/Private/ObjectTools.cpp b/Engine/Source/Editor/UnrealEd/Private/ObjectTools.cpp
--- a/Engine/Source/Editor/UnrealEd/Private/ObjectTools.cpp
+++ b/Engine/Source/Editor/UnrealEd/Private/ObjectTools.cpp
@@ -4960,56 +4960,92 @@ namespace ThumbnailTools
 		return true;
 	}
 
-
-
-	/** Searches for an object's thumbnail in memory and returns it if found */
-	FObjectThumbnail* FindCachedThumbnailInPackage( UPackage* InPackage, const FName InObjectFullName )
+	namespace Private
 	{
-		FObjectThumbnail* FoundThumbnail = NULL;
-
-		// We're expecting this to be an outermost package!
-		check( InPackage->GetOutermost() == InPackage );
-
-		// Does the package have any thumbnails?
-		if( InPackage->HasThumbnailMap() )
+		FObjectThumbnail* FindCachedThumbnailInPackage(UPackage* InPackage, const FName InObjectShortClassFullName)
 		{
-			// @todo thumbnails: Backwards compat
-			FThumbnailMap& PackageThumbnailMap = InPackage->AccessThumbnailMap();
-			FoundThumbnail = PackageThumbnailMap.Find( InObjectFullName );
-		}
-
-		return FoundThumbnail;
-	}
+			FObjectThumbnail* FoundThumbnail = NULL;
 
+			// We're expecting this to be an outermost package!
+			check(InPackage->GetOutermost() == InPackage);
 
+			// Does the package have any thumbnails?
+			if (InPackage->HasThumbnailMap())
+			{
+				// @todo thumbnails: Backwards compat
+				FThumbnailMap& PackageThumbnailMap = InPackage->AccessThumbnailMap();
+				FoundThumbnail = PackageThumbnailMap.Find(InObjectShortClassFullName);
+			}
 
+			return FoundThumbnail;
+		}
+	}
+	/** Searches for an object's thumbnail in memory and returns it if found */
+	FObjectThumbnail* FindCachedThumbnailInPackage(UPackage* InPackage, FName InObjectFullName)
+	{
+		return Private::FindCachedThumbnailInPackage(InPackage,
+			FName(*UClass::ConvertFullNameToShortTypeFullName(WriteToString<256>(InObjectFullName).ToView())));
+	}
 	/** Searches for an object's thumbnail in memory and returns it if found */
-	FObjectThumbnail* FindCachedThumbnailInPackage( const FString& InPackageFileName, const FName InObjectFullName )
+	FObjectThumbnail* FindCachedThumbnailInPackage(UPackage* InPackage, FStringView InObjectFullName)
 	{
-		FObjectThumbnail* FoundThumbnail = nullptr;
+		return Private::FindCachedThumbnailInPackage(InPackage,
+			FName(*UClass::ConvertFullNameToShortTypeFullName(InObjectFullName)));
+	}
+	/** Searches for an object's thumbnail in memory and returns it if found */
+	FObjectThumbnail* FindCachedThumbnailInPackage(UPackage* InPackage, const TCHAR* InObjectFullName)
+	{
+		return Private::FindCachedThumbnailInPackage(InPackage,
+			FName(*UClass::ConvertFullNameToShortTypeFullName(InObjectFullName)));
+	}
 
-		FString PackageName = InPackageFileName;
-		if (FPackageName::TryConvertFilenameToLongPackageName(PackageName, PackageName))
+	namespace Private
+	{
+		/** Searches for an object's thumbnail in memory and returns it if found */
+		FObjectThumbnail* FindCachedThumbnailInPackage(const FString& InPackageFileName, const FName InObjectShortClassFullName)
 		{
-			if (PackageName == TEXT("None"))
-			{
-				UE_LOG(LogUObjectGlobals, Warning, TEXT("Attempted to FindCachedThumbnailInPackage named 'None' - PackageName: %s InPackageFileName: %s"), *PackageName, *InPackageFileName);
-				return nullptr;
-			}
+			FObjectThumbnail* FoundThumbnail = nullptr;
 
-			// First check to see if the package is already in memory.  If it is, some or all of the thumbnails
-			// may already be loaded and ready.
-			UObject* PackageOuter = nullptr;
-			UPackage* Package = FindPackage(PackageOuter, *PackageName);
-			if (Package != nullptr)
+			FString PackageName = InPackageFileName;
+			if (FPackageName::TryConvertFilenameToLongPackageName(PackageName, PackageName))
 			{
-				FoundThumbnail = FindCachedThumbnailInPackage(Package, InObjectFullName);
+				if (PackageName == TEXT("None"))
+				{
+					UE_LOG(LogUObjectGlobals, Warning, TEXT("Attempted to FindCachedThumbnailInPackage named 'None' - PackageName: %s InPackageFileName: %s"), *PackageName, *InPackageFileName);
+					return nullptr;
+				}
+
+				// First check to see if the package is already in memory.  If it is, some or all of the thumbnails
+				// may already be loaded and ready.
+				UObject* PackageOuter = nullptr;
+				UPackage* Package = FindPackage(PackageOuter, *PackageName);
+				if (Package != nullptr)
+				{
+					FoundThumbnail = Private::FindCachedThumbnailInPackage(Package, InObjectShortClassFullName);
+				}
 			}
+			return FoundThumbnail;
 		}
-		return FoundThumbnail;
 	}
 
-
+	/** Searches for an object's thumbnail in memory and returns it if found */
+	FObjectThumbnail* FindCachedThumbnailInPackage(const FString& InPackageFileName, FName InObjectFullName)
+	{
+		return Private::FindCachedThumbnailInPackage(InPackageFileName,
+			FName(*UClass::ConvertFullNameToShortTypeFullName(WriteToString<256>(InObjectFullName).ToView())));
+	}
+	/** Searches for an object's thumbnail in memory and returns it if found */
+	FObjectThumbnail* FindCachedThumbnailInPackage(const FString& InPackageFileName, FStringView InObjectFullName)
+	{
+		return Private::FindCachedThumbnailInPackage(InPackageFileName,
+			FName(*UClass::ConvertFullNameToShortTypeFullName(InObjectFullName)));
+	}
+	/** Searches for an object's thumbnail in memory and returns it if found */
+	FObjectThumbnail* FindCachedThumbnailInPackage(const FString& InPackageFileName, const TCHAR* InObjectFullName)
+	{
+		return Private::FindCachedThumbnailInPackage(InPackageFileName,
+			FName(*UClass::ConvertFullNameToShortTypeFullName(InObjectFullName)));
+	}
 
 	/** Searches for an object's thumbnail in memory and returns it if found */
 	const FObjectThumbnail* FindCachedThumbnail( const FString& InFullName )
@@ -5023,7 +5059,7 @@ namespace ThumbnailTools
 			return nullptr;
 		}
 
-		return FindCachedThumbnailInPackage( PackageFileName, FName( *InFullName ) );
+		return FindCachedThumbnailInPackage( PackageFileName, *InFullName );
 	}
 
 
@@ -5032,7 +5068,7 @@ namespace ThumbnailTools
 	FObjectThumbnail* GetThumbnailForObject( UObject* InObject )
 	{
 		UPackage* ObjectPackage = InObject->GetOutermost();
-		return FindCachedThumbnailInPackage( ObjectPackage, FName( *InObject->GetFullName() ) );
+		return FindCachedThumbnailInPackage( ObjectPackage, *InObject->GetFullName() );
 	}
 
 
@@ -5083,8 +5119,8 @@ namespace ThumbnailTools
 			for( int32 CurThumbnailIndex = 0; CurThumbnailIndex < ThumbnailCount; ++CurThumbnailIndex )
 			{
 				bool bHaveValidClassName = false;
-				FString ObjectClassName;
-				*FileReader << ObjectClassName;
+				FString ObjectShortClassName;
+				*FileReader << ObjectShortClassName;
 
 				// Object path
 				FString ObjectPathWithoutPackageName;
@@ -5093,7 +5129,7 @@ namespace ThumbnailTools
 				FString ObjectPath;
 
 				// handle UPackage thumbnails differently from usual assets
-				if (ObjectClassName == UPackage::StaticClass()->GetName())
+				if (ObjectShortClassName == UPackage::StaticClass()->GetName())
 				{
 					ObjectPath = ObjectPathWithoutPackageName;
 				}
@@ -5103,7 +5139,7 @@ namespace ThumbnailTools
 				}
 
 				// If the thumbnail was stored with a missing class name ("???") when we'll catch that here
-				if( ObjectClassName.Len() > 0 && ObjectClassName != TEXT( "???" ) )
+				if(ObjectShortClassName.Len() > 0 && ObjectShortClassName != TEXT( "???" ) )
 				{
 					bHaveValidClassName = true;
 				}
@@ -5129,7 +5165,8 @@ namespace ThumbnailTools
 							// Great, we found a path that matches -- we just need to add that class name
 							const int32 FirstSpaceIndex = CurObjectFullName.Find( TEXT( " " ) );
 							check( FirstSpaceIndex != -1 );
-							ObjectClassName = CurObjectFullName.Left( FirstSpaceIndex );
+							ObjectShortClassName = CurObjectFullName.Left( FirstSpaceIndex );
+							ObjectShortClassName = UClass::ConvertPathNameToShortTypeName(ObjectShortClassName);
 
 							// We have a useful class name now!
 							bHaveValidClassName = true;
@@ -5152,7 +5189,7 @@ namespace ThumbnailTools
 				if( bHaveValidClassName )
 				{
 					// Create a full name string with the object's class and fully qualified path
-					const FString ObjectFullName( ObjectClassName + TEXT( " " ) + ObjectPath );
+					const FString ObjectFullName(ObjectShortClassName + TEXT( " " ) + ObjectPath);
 
 
 					// Add to our map
@@ -5170,11 +5207,12 @@ namespace ThumbnailTools
 		// @todo CB: Should sort the thumbnails to load by file offset to reduce seeks [reviewed; pre-qa release]
 		for ( TSet<FName>::TConstIterator It(InObjectFullNames); It; ++It )
 		{
-			const FName& CurObjectFullName = *It;
+			FName CurObjectFullName = *It;
+			FName CurObjectShortClassFullName = FName(*UClass::ConvertFullNameToShortTypeFullName(WriteToString<256>(CurObjectFullName).ToView()));
 
 			// Do we have this thumbnail in the file?
 			// @todo thumbnails: Backwards compat
-			const int32* pFileOffset = ObjectNameToFileOffsetMap.Find(CurObjectFullName);
+			const int32* pFileOffset = ObjectNameToFileOffsetMap.Find(CurObjectShortClassFullName);
 			if ( pFileOffset != NULL )
 			{
 				// Seek to the location in the file with the image data
diff --git a/Engine/Source/Runtime/CoreUObject/Private/UObject/Class.cpp b/Engine/Source/Runtime/CoreUObject/Private/UObject/Class.cpp
--- a/Engine/Source/Runtime/CoreUObject/Private/UObject/Class.cpp
+++ b/Engine/Source/Runtime/CoreUObject/Private/UObject/Class.cpp
@@ -55,6 +55,7 @@
 #include "UObject/TopLevelAssetPath.h"
 #include "AssetRegistry/AssetData.h"
 #include "HAL/PlatformStackWalk.h"
+#include "String/Find.h"
 
 // This flag enables some expensive class tree validation that is meant to catch mutations of 
 // the class tree outside of SetSuperStruct. It has been disabled because loading blueprints 
@@ -6095,6 +6096,62 @@ bool UClass::TryFixShortClassNameExportPath(FString& InOutExportPathToFix, ELogV
 	return false;
 }
 
+FString UClass::ConvertPathNameToShortTypeName(FStringView InPathName)
+{
+	return FString(FPackageName::ObjectPathToObjectName(InPathName));
+}
+
+FString UClass::ConvertFullNameToShortTypeFullName(FStringView InFullName)
+{
+	FStringView ClassPath, PackageName, ObjectName, SubObjectName;
+	FPackageName::SplitFullObjectPath(InFullName, ClassPath, PackageName, ObjectName, SubObjectName);
+	if (ClassPath.IsEmpty() || PackageName.IsEmpty())
+	{
+		// Not a valid FullName, return it untransformed
+		return FString(InFullName);
+	}
+	FStringView ClassObjectName = FPackageName::ObjectPathToObjectName(ClassPath);
+	if (ClassObjectName.Len() == ClassPath.Len())
+	{
+		// Already a ShortTypeFullName, return it untransformed
+		return FString(InFullName);
+	}
+
+	FString Result;
+	Result.Reserve(InFullName.Len());
+	Result += ClassObjectName;
+	Result += TEXT(" ");
+	Result += PackageName;
+	if (ObjectName.Len())
+	{
+		Result += TEXT(".");
+		Result += ObjectName;
+		if (SubObjectName.Len())
+		{
+			Result += SUBOBJECT_DELIMITER;
+			Result += SubObjectName;
+		}
+	}
+	return Result;
+}
+
+bool UClass::IsShortTypeName(FStringView InClassPath)
+{
+	if (InClassPath.Len() == 0)
+	{
+		// Arbitrary
+		return true;
+	}
+	// If it starts with '/' it is definitely not a ShortTypeName
+	// If it doesn't start with '/', most of the time it is a ShortTypeName because this function is usually called on either PathNames or ShortTypeNames
+	if (InClassPath[0] == '/')
+	{
+		return false;
+	}
+	return UE::String::FindFirstOfAnyChar(InClassPath, TEXT("./\\" SUBOBJECT_DELIMITER)) == INDEX_NONE;
+}
+
+
 #if WITH_EDITOR || HACK_HEADER_GENERATOR
 void UClass::GetHideFunctions(TArray<FString>& OutHideFunctions) const
 {
diff --git a/Engine/Source/Runtime/CoreUObject/Public/UObject/Class.h b/Engine/Source/Runtime/CoreUObject/Public/UObject/Class.h
--- a/Engine/Source/Runtime/CoreUObject/Public/UObject/Class.h
+++ b/Engine/Source/Runtime/CoreUObject/Public/UObject/Class.h
@@ -2822,6 +2822,26 @@ public:
 	 */
 	static bool TryFixShortClassNameExportPath(FString& InOutExportPathToFix, ELogVerbosity::Type AmbiguousMessageVerbosity = ELogVerbosity::NoLogging, const TCHAR* AmbiguousClassMessage = nullptr);
 
+	/**
+	 * Returns the ObjectName portion of a ClassPath name: "/Path/To.Object" is converted to "Object". 
+	 * Returns the full string if it is already a ShortTypeName.
+	 */
+	static FString ConvertPathNameToShortTypeName(FStringView InClassPathOrShortTypeName);
+
+	/**
+	 * Takes a FullName (from e.g. AssetData.GetFullName or UObject.GetFullName) in either ShortTypeFullName or PathFullName form
+	 * ShortTypeFullName: "ClassObjectName /PackagePath/PackageShortName.ObjectName:SubObjectName"
+	 * PathFullName: "/ClassPath/ClassPackage.ClassObjectName /PackagePath/PackageShortName.ObjectName:SubObjectName"
+	 * Converts it to ShortTypeFullName if not already in that format and returns it
+	 */
+	static FString ConvertFullNameToShortTypeFullName(FStringView InFullName);
+
+	/**
+	 * Returns whether the given stringview is in ShortTypeName form: No directory separators "/" or object separators ".", SUBOBJECT_DELIMITER
+	 * Returns true for empty string.
+	 */
+	static bool IsShortTypeName(FStringView ClassPathOrShortTypeName);
+
 #if WITH_EDITOR
 	/** Provides access to C++ type info. */
 	const ICppClassTypeInfo* GetCppTypeInfo() const
