Cooker: PakFile Manifests: When a file is too big to be added to a pak file, display the filename and size in the error message.
#rnx

diff --git a/Engine/Source/Editor/UnrealEd/Private/Commandlets/AssetRegistryGenerator.cpp b/Engine/Source/Editor/UnrealEd/Private/Commandlets/AssetRegistryGenerator.cpp
--- a/Engine/Source/Editor/UnrealEd/Private/Commandlets/AssetRegistryGenerator.cpp
+++ b/Engine/Source/Editor/UnrealEd/Private/Commandlets/AssetRegistryGenerator.cpp
@@ -711,6 +711,8 @@ bool FAssetRegistryGenerator::GenerateStreamingInstallManifest(int64 InOverrideC
 			}
 
 			int64 CurrentPakSize = 0;
+			int64 NextFileSize = 0;
+			FString NextFilename;
 			bFinishedAllFiles = true;
 			for (; FilenameIndex < ChunkFilenames.Num(); ++FilenameIndex)
 			{
@@ -718,12 +720,15 @@ bool FAssetRegistryGenerator::GenerateStreamingInstallManifest(int64 InOverrideC
 				FString PakListLine = FPaths::ConvertRelativePathToFull(Filename.Replace(TEXT("[Platform]"), *Platform));
 				if (MaxChunkSize > 0)
 				{
-					const int64* PackageFileSize = PackageFileSizes.Find(PakListLine);
-					CurrentPakSize += PackageFileSize ? *PackageFileSize : 0;
+					const int64* ExistingPackageFileSize = PackageFileSizes.Find(PakListLine);
+					const int64 PackageFileSize = ExistingPackageFileSize ? *ExistingPackageFileSize : 0;
+					CurrentPakSize += PackageFileSize;
 					if (MaxChunkSize < CurrentPakSize)
 					{
 						// early out if we are over memory limit
 						bFinishedAllFiles = false;
+						NextFileSize = PackageFileSize;
+						NextFilename = MoveTemp(PakListLine);
 						break;
 					}
 				}
@@ -738,7 +743,17 @@ bool FAssetRegistryGenerator::GenerateStreamingInstallManifest(int64 InOverrideC
 
 			if (!bFinishedAllFiles && !bAddedFilesToPakList)
 			{
-				UE_LOG(LogAssetRegistryGenerator, Error, TEXT("Failed to add file(s) to paklist '%s', max chunk size '%d' too small"), *PakListFilename, MaxChunkSize);
+				const TCHAR* UnitsText = TEXT("MB");
+				int32 Unit = 1000*1000;
+				if (MaxChunkSize < Unit * 10)
+				{
+					Unit = 1;
+					UnitsText = TEXT("bytes");
+				}
+				UE_LOG(LogAssetRegistryGenerator, Error, TEXT("Failed to add file %s to paklist '%s'. The maximum size for a Pakfile is %d %s, but the file to add is %d %s."),
+					*NextFilename, *PakListFilename,
+					MaxChunkSize/Unit, UnitsText, // Round the limit down and round the value up, so that the display always shows that the value is greater than the limit
+					(NextFileSize+Unit-1)/Unit, UnitsText);
 				bSucceeded = false;
 				break;
 			}
