Duplicated world PIE packages tagged with PKG_NewlyCreated to avoid warnings when resolving linker
- This was happening on worlds with OFPA

diff --git a/Engine/Source/Editor/UnrealEd/Private/PlayLevel.cpp b/Engine/Source/Editor/UnrealEd/Private/PlayLevel.cpp
--- a/Engine/Source/Editor/UnrealEd/Private/PlayLevel.cpp
+++ b/Engine/Source/Editor/UnrealEd/Private/PlayLevel.cpp
@@ -2320,7 +2320,8 @@ UWorld* UEditorEngine::CreatePIEWorldByDuplication(FWorldContext &WorldContext,
 	UE_LOG( LogPlayLevel, Log, TEXT("Creating play world package: %s"),  *PlayWorldMapName );	
 
 	UPackage* PlayWorldPackage = CreatePackage(*PlayWorldMapName);
-	PlayWorldPackage->SetPackageFlags(PKG_PlayInEditor);
+	// Add PKG_NewlyCreated flag to this package so we don't try to resolve its linker as it is unsaved duplicated world package
+	PlayWorldPackage->SetPackageFlags(PKG_PlayInEditor | PKG_NewlyCreated);
 	PlayWorldPackage->SetPIEInstanceID(WorldContext.PIEInstance);
 	PlayWorldPackage->SetLoadedPath(InPackage->GetLoadedPath());
 	PRAGMA_DISABLE_DEPRECATION_WARNINGS
diff --git a/Engine/Source/Runtime/Engine/Private/World.cpp b/Engine/Source/Runtime/Engine/Private/World.cpp
--- a/Engine/Source/Runtime/Engine/Private/World.cpp
+++ b/Engine/Source/Runtime/Engine/Private/World.cpp
@@ -3585,7 +3585,8 @@ UWorld* UWorld::DuplicateWorldForPIE(const FString& PackageName, UWorld* OwningW
 
 	UWorld::WorldTypePreLoadMap.FindOrAdd(PrefixedLevelFName) = EWorldType::PIE;
 	UPackage* PIELevelPackage = CreatePackage(*PrefixedLevelName);
-	PIELevelPackage->SetPackageFlags(PKG_PlayInEditor);
+	// Add PKG_NewlyCreated flag to this package so we don't try to resolve its linker as it is unsaved duplicated world package
+	PIELevelPackage->SetPackageFlags(PKG_PlayInEditor | PKG_NewlyCreated);
 	PIELevelPackage->SetPIEInstanceID(PIEInstanceID);
 	PIELevelPackage->SetLoadedPath(EditorLevelPackage->GetLoadedPath());
 	PRAGMA_DISABLE_DEPRECATION_WARNINGS
