Support multiple console commands in the editor by using FParse::Line to allow commands to be split with newlines or pipes '|'
Allow adding a newline in the console command entry box with shift + enter, same as python.

diff --git a/Engine/Source/Editor/UnrealEd/Private/Toolkits/FConsoleCommandExecutor.cpp b/Engine/Source/Editor/UnrealEd/Private/Toolkits/FConsoleCommandExecutor.cpp
--- a/Engine/Source/Editor/UnrealEd/Private/Toolkits/FConsoleCommandExecutor.cpp
+++ b/Engine/Source/Editor/UnrealEd/Private/Toolkits/FConsoleCommandExecutor.cpp
@@ -65,6 +65,23 @@ bool FConsoleCommandExecutor::Exec(const TCHAR* Input)
 {
 	IConsoleManager::Get().AddConsoleHistoryEntry(TEXT(""), Input);
 
+	int32 Len = FCString::Strlen(Input);
+	TArray<TCHAR> Buffer;
+	Buffer.AddZeroed(Len+1);
+
+	bool bHandled = false;
+	const TCHAR* ParseCursor = Input;
+	while (FParse::Line(&ParseCursor, Buffer.GetData(), Buffer.Num()))
+	{
+		bHandled = ExecInternal(Buffer.GetData()) || bHandled;
+	}
+
+	// return true if we successfully executed any of the commands 
+	return bHandled;
+}
+
+bool FConsoleCommandExecutor::ExecInternal(const TCHAR* Input) const
+{
 	bool bWasHandled = false;
 	UWorld* World = nullptr;
 	UWorld* OldWorld = nullptr;
@@ -136,7 +153,7 @@ bool FConsoleCommandExecutor::AllowHotKeyClose() const
 
 bool FConsoleCommandExecutor::AllowMultiLine() const
 {
-	return false;
+	return true;
 }
 
 #undef LOCTEXT_NAMESPACE
\ No newline at end of file
diff --git a/Engine/Source/Editor/UnrealEd/Public/Toolkits/FConsoleCommandExecutor.h b/Engine/Source/Editor/UnrealEd/Public/Toolkits/FConsoleCommandExecutor.h
--- a/Engine/Source/Editor/UnrealEd/Public/Toolkits/FConsoleCommandExecutor.h
+++ b/Engine/Source/Editor/UnrealEd/Public/Toolkits/FConsoleCommandExecutor.h
@@ -28,4 +28,6 @@ public:
 	{
 		return FGlobalEditorCommonCommands::Get().OpenConsoleCommandBox->GetActiveChord(EMultipleKeyBindingIndex::Primary).Get();
 	}
+private:
+	bool ExecInternal(const TCHAR*) const;
 };
