EditorPerf: Memory, Rate, and better Counter statistics

### Features

Instrumentation is added to the FGenericStatsUpdater to push the sampled memory stats to counters. Ideally we would add TraceStat support for EditorPerf, but that is significantly more work.

CountPerSecond statistic has been added to compute rate information from the first start and last end timestamps of each Scope or Counter.

This change adds full min, max, mean, deviation statistics to Counters. It's all made available similarly to Scopes.

### Notes

Welford's mean, variance, and deviation are moved out FSummarizeScope to static utility functions.

This change closes a hole in FSummarizeScope::Merge() where Mean and Variance were not being handled at all. The variance merge math is probably not accurate, but it should yield something proportional. A follow up change will address that.

This change tidies up some class names and code formatting, and grouping the Scope, Counter, and Bookmark classes to be adjacent in the cpp file to make reading easier.

#rnx

diff --git a/Engine/Source/Editor/UnrealEd/Private/Commandlets/SummarizeTraceCommandlet.cpp b/Engine/Source/Editor/UnrealEd/Private/Commandlets/SummarizeTraceCommandlet.cpp
--- a/Engine/Source/Editor/UnrealEd/Private/Commandlets/SummarizeTraceCommandlet.cpp
+++ b/Engine/Source/Editor/UnrealEd/Private/Commandlets/SummarizeTraceCommandlet.cpp
@@ -176,8 +176,6 @@ private:
 
 	// Scope name lookup function, cached for efficiency.
 	TFunction<const FString*(uint32 ScopeId)> LookupScopeNameFn;
-
-
 };
 
 enum
@@ -444,7 +442,9 @@ void FCpuScopeStreamProcessor::OnCpuScopeTree(uint32 ThreadId, const FScopeTreeI
 	}
 }
 
-
+/**
+ * Base class to extend for counters analysis.
+ */
 class FCountersAnalyzer
 	: public UE::Trace::IAnalyzer
 {
@@ -459,12 +459,14 @@ public:
 	struct FCounterIntValue
 	{
 		uint16				Id;
+		double				Timestamp;
 		int64				Value;
 	};
 
 	struct FCounterFloatValue
 	{
 		uint16				Id;
+		double				Timestamp;
 		double				Value;
 	};
 
@@ -530,7 +532,9 @@ void FCountersAnalyzer::OnCountersSetValueInt(const FOnEventContext& Context)
 	const FEventData& EventData = Context.EventData;
 	uint16 CounterId = EventData.GetValue<uint16>("CounterId");
 	int64 Value = EventData.GetValue<int64>("Value");
-	OnCounterIntValue({ uint16(CounterId - 1), Value });
+	uint64 Cycle = EventData.GetValue<uint64>("Cycle");
+	double Timestamp = Context.EventTime.AsSeconds(Cycle);
+	OnCounterIntValue({ uint16(CounterId - 1), Timestamp, Value });
 }
 
 void FCountersAnalyzer::OnCountersSetValueFloat(const FOnEventContext& Context)
@@ -538,9 +542,14 @@ void FCountersAnalyzer::OnCountersSetValueFloat(const FOnEventContext& Context)
 	const FEventData& EventData = Context.EventData;
 	uint16 CounterId = EventData.GetValue<uint16>("CounterId");
 	double Value = EventData.GetValue<double>("Value");
-	OnCounterFloatValue({ uint16(CounterId - 1), Value });
+	uint64 Cycle = EventData.GetValue<uint64>("Cycle");
+	double Timestamp = Context.EventTime.AsSeconds(Cycle);
+	OnCounterFloatValue({ uint16(CounterId - 1), Timestamp, Value });
 }
 
+/**
+ * Base class to extend for bookmarks analysis.
+ */
 class FBookmarksAnalyzer
 	: public UE::Trace::IAnalyzer
 {
@@ -668,8 +677,56 @@ public:
 };
 
 /*
- * Helper classes for the SummarizeTrace commandlet. Aggregates statistics about a trace.
- */
+* 
+* Helper classes for the SummarizeTrace commandlet. Aggregates statistics about a trace.
+* Code preceeding this comment should eventually be reloated to the Trace/Insights modules.
+* Everything succeeding this comment should stay in this commandlet.
+* 
+*/
+
+/**
+* Perform an increment of work for Welford's variance, from which we can compute the standard deviation on the last step
+* 
+* @param InSample					The new sample value to operate on
+* @param InCount					The number of samples, inclusive
+* @param InOutMean					The mean of all samples, which will be updated in this function
+* @param InOutVarianceAccumulator	The accumulator for Welford's method
+*/
+static void WelfordsIncrement(const double InSample, const uint64 InCount, double& InOutMean, double& InOutVarianceAccumulator)
+{
+	ensure(InCount);
+	const double OldMean = InOutMean;
+	InOutMean += ((InSample - InOutMean) / double(InCount));
+	InOutVarianceAccumulator += ((InSample - InOutMean) * (InSample - OldMean));
+}
+
+/**
+* Compute the standard deviation given Welford's accumulator and the overall count
+* 
+* @param InVarianceAccumulator	The accumulator for Welford's method
+* @param InCount				The number of samples, inclusive
+* 
+* @return The standard deviation in sample units
+*/
+static double WelfordsDeviation(const double InVarianceAccumulator, const uint64 InCount)
+{
+	if (InCount > 1)
+	{
+		// Welford's final step, dependent on sample count
+		double VarianceSecondsSquared = InVarianceAccumulator / double(InCount - 1);
+
+		// stddev is sqrt of variance, to restore to units of seconds (vs. seconds squared)
+		return sqrt(VarianceSecondsSquared);
+	}
+	else
+	{
+		return 0.0;
+	}
+}
+
+//
+// FSummarizeCpuAnalyzer - Helpers & Analyzers for Cpu channel events
+//
 
 struct FSummarizeScope
 {
@@ -688,7 +745,7 @@ struct FSummarizeScope
 	double MinDurationSeconds = 1e10;
 	double MaxDurationSeconds = -1e10;
 	double MeanDurationSeconds = 0.0;
-	double VarianceAcc = 0.0; // Accumulator for Welford's
+	double VarianceAccumulator = 0.0; // Accumulator for Welford's
 
 	void AddDuration(double StartSeconds, double FinishSeconds)
 	{
@@ -697,8 +754,8 @@ struct FSummarizeScope
 		// compute the duration
 		double DurationSeconds = FinishSeconds - StartSeconds;
 
-		// only set first for the first sample, compare exact zero
-		if (FirstStartSeconds == 0.0)
+		// only set first for the first sample
+		if (Count == 1)
 		{
 			FirstStartSeconds = StartSeconds;
 			FirstFinishSeconds = FinishSeconds;
@@ -713,33 +770,24 @@ struct FSummarizeScope
 		TotalDurationSeconds += DurationSeconds;
 		MinDurationSeconds = FMath::Min(MinDurationSeconds, DurationSeconds);
 		MaxDurationSeconds = FMath::Max(MaxDurationSeconds, DurationSeconds);
-		UpdateVariance(DurationSeconds);
-	}
-
-	void UpdateVariance(double DurationSeconds)
-	{
-		ensure(Count);
 
-		// Welford's increment
-		double OldMeanDurationSeconds = MeanDurationSeconds;
-		MeanDurationSeconds = MeanDurationSeconds + ((DurationSeconds - MeanDurationSeconds) / double(Count));
-		VarianceAcc = VarianceAcc + ((DurationSeconds - MeanDurationSeconds) * (DurationSeconds - OldMeanDurationSeconds));
+		// update variance for stddev later
+		WelfordsIncrement(DurationSeconds, Count, MeanDurationSeconds, VarianceAccumulator);
 	}
 
 	double GetDeviationDurationSeconds() const
 	{
-		if (Count > 1)
-		{
-			// Welford's final step, dependent on sample count
-			double VarianceSecondsSquared = VarianceAcc / double(Count - 1);
+		return WelfordsDeviation(VarianceAccumulator, Count);
+	}
 
-			// stddev is sqrt of variance, to restore to units of seconds (vs. seconds squared)
-			return sqrt(VarianceSecondsSquared);
-		}
-		else
+	double GetCountPerSecond() const
+	{
+		double CountPerSecond = 0.0;
+		if (Count)
 		{
-			return 0.0;
+			CountPerSecond = Count / (LastFinishSeconds - FirstStartSeconds);
 		}
+		return CountPerSecond;
 	}
 
 	void Merge(const FSummarizeScope& Scope)
@@ -748,7 +796,19 @@ struct FSummarizeScope
 		TotalDurationSeconds += Scope.TotalDurationSeconds;
 		MinDurationSeconds = FMath::Min(MinDurationSeconds, Scope.MinDurationSeconds);
 		MaxDurationSeconds = FMath::Max(MaxDurationSeconds, Scope.MaxDurationSeconds);
-		Count += Scope.Count;
+
+		const uint64 NewCount = Count + Scope.Count;
+		if (NewCount)
+		{
+			// perform weighted average, weighted on each sub-population's count
+			MeanDurationSeconds = ((MeanDurationSeconds * double(Count)) + (Scope.MeanDurationSeconds * double(Scope.Count))) / double(NewCount);
+
+			// perform weighted average, weighted on each sub-population's count
+			//  this is likely not the perfect union of two Welford variance accumuations due to them being squares of variance offsets
+			VarianceAccumulator = ((VarianceAccumulator * double(Count)) + (Scope.VarianceAccumulator * double(Scope.Count))) / double(NewCount);
+		}
+
+		Count = NewCount;
 	}
 
 	FString GetValue(const FStringView& Statistic) const
@@ -805,6 +865,10 @@ struct FSummarizeScope
 		{
 			return FString::Printf(TEXT("%f"), GetDeviationDurationSeconds());
 		}
+		else if (Statistic == TEXT("CountPerSecond"))
+		{
+			return FString::Printf(TEXT("%f"), GetCountPerSecond());
+		}
 		return FString();
 	}
 
@@ -826,60 +890,6 @@ static uint32 GetTypeHash(const FSummarizeScope& Scope)
 	return FCrc::StrCrc32(*Scope.Name);
 }
 
-struct FSummarizeBookmark
-{
-	FString Name;
-	uint64 Count = 0;
-
-	double FirstSeconds = 0.0;
-	double LastSeconds = 0.0;
-
-	void AddTimestamp(double Seconds)
-	{
-		Count += 1;
-
-		// only set first for the first sample, compare exact zero
-		if (FirstSeconds == 0.0)
-		{
-			FirstSeconds = Seconds;
-		}
-
-		LastSeconds = Seconds;
-	}
-
-	FString GetValue(const FStringView& Statistic) const
-	{
-		if (Statistic == TEXT("Name"))
-		{
-			return Name;
-		}
-		else if (Statistic == TEXT("Count"))
-		{
-			return FString::Printf(TEXT("%llu"), Count);
-		}
-		else if (Statistic == TEXT("FirstSeconds"))
-		{
-			return FString::Printf(TEXT("%f"), FirstSeconds);
-		}
-		else if (Statistic == TEXT("LastSeconds"))
-		{
-			return FString::Printf(TEXT("%f"), LastSeconds);
-		}
-		return FString();
-	}
-
-	// for deduplication
-	bool operator==(const FSummarizeBookmark& Bookmark) const
-	{
-		return Name == Bookmark.Name;
-	}
-};
-
-static uint32 GetTypeHash(const FSummarizeBookmark& Bookmark)
-{
-	return FCrc::StrCrc32(*Bookmark.Name);
-}
-
 /**
  * This analyzer aggregates CPU scopes having the same name together, counting the number of occurrences, total duration,
  * average occurrence duration, for each scope name.
@@ -946,7 +956,6 @@ void FSummarizeCpuAnalyzer::OnCpuScopeAnalysisEnd()
 	PublishFn(LocalScopes);
 }
 
-
 /**
  * Summarizes matched CPU scopes, excluding time consumed by immediate children if any. The analyzer uses pattern matching to
  * selects the scopes of interest and detects parent/child relationship. Once a parent/child relationship is established in a
@@ -975,7 +984,9 @@ void FSummarizeCpuAnalyzer::OnCpuScopeAnalysisEnd()
  * @note If the matching expression was to match all scopes, the analyser would summarize all scopes,
  *       accounting for their exclusive time.
  */
-class FCpuScopeHierarchyAnalyzer : public FCpuScopeAnalyzer
+
+class FSummarizeCpuScopeHierarchyAnalyzer
+	: public FCpuScopeAnalyzer
 {
 public:
 	/**
@@ -988,7 +999,7 @@ public:
 	 * @note This analyzer publishes summarized scope names with a suffix ".excl" as it computes exclusive time to  prevent name collisions with other analyzers. The summary of all scopes
 	 *       suffixed with ".excl.all" and gets its base name from the analyzer name. The scopes in the array gets their names from the scope name themselves, but suffixed with .excl.
 	 */
-	FCpuScopeHierarchyAnalyzer(const FString& InAnalyzerName, TFunction<bool(const FString&)> InMatchFn, TFunction<void(const FSummarizeScope&, TArray<FSummarizeScope>&&)> InPublishFn);
+	FSummarizeCpuScopeHierarchyAnalyzer(const FString& InAnalyzerName, TFunction<bool(const FString&)> InMatchFn, TFunction<void(const FSummarizeScope&, TArray<FSummarizeScope>&&)> InPublishFn);
 
 	/**
 	 * Runs analysis on a collection of scopes events. The scopes are expected to be from the same thread and form a 'tree', meaning
@@ -1021,14 +1032,14 @@ private:
 	TFunction<void(const FSummarizeScope&, TArray<FSummarizeScope>&&)> PublishFn;
 };
 
-FCpuScopeHierarchyAnalyzer::FCpuScopeHierarchyAnalyzer(const FString& InAnalyzerName, TFunction<bool(const FString&)> InMatchFn, TFunction<void(const FSummarizeScope&, TArray<FSummarizeScope>&&)> InPublishFn)
+FSummarizeCpuScopeHierarchyAnalyzer::FSummarizeCpuScopeHierarchyAnalyzer(const FString& InAnalyzerName, TFunction<bool(const FString&)> InMatchFn, TFunction<void(const FSummarizeScope&, TArray<FSummarizeScope>&&)> InPublishFn)
  : MatchesFn(MoveTemp(InMatchFn))
  , PublishFn(MoveTemp(InPublishFn))
 {
 	MatchedScopesSummary.Name = InAnalyzerName;
 }
 
-void FCpuScopeHierarchyAnalyzer::OnCpuScopeTree(uint32 ThreadId, const TArray<FCpuScopeAnalyzer::FScopeEvent>& ScopeEvents, const TFunction<const FString*(uint32 /*ScopeId*/)>& InScopeNameLookup)
+void FSummarizeCpuScopeHierarchyAnalyzer::OnCpuScopeTree(uint32 ThreadId, const TArray<FCpuScopeAnalyzer::FScopeEvent>& ScopeEvents, const TFunction<const FString*(uint32 /*ScopeId*/)>& InScopeNameLookup)
 {
 	// Scope matching the pattern.
 	struct FMatchScopeEnter
@@ -1083,7 +1094,7 @@ void FCpuScopeHierarchyAnalyzer::OnCpuScopeTree(uint32 ThreadId, const TArray<FC
 	}
 }
 
-void FCpuScopeHierarchyAnalyzer::OnCpuScopeAnalysisEnd()
+void FSummarizeCpuScopeHierarchyAnalyzer::OnCpuScopeAnalysisEnd()
 {
 	MatchedScopesSummary.Name += TEXT(".excl.all");
 
@@ -1098,110 +1109,221 @@ void FCpuScopeHierarchyAnalyzer::OnCpuScopeAnalysisEnd()
 	PublishFn(MatchedScopesSummary, MoveTemp(ScopeSummaries));
 }
 
-
 //
 // FSummarizeCountersAnalyzer - Tally Counters from counter set/increment events
 //
 
-class FSummarizeCountersAnalyzer
-	: public FCountersAnalyzer
+struct FSummarizeCounter
 {
-public:
-	virtual void OnCounterName(const FCounterName& CounterName) override;
-	virtual void OnCounterIntValue(const FCounterIntValue& NewValue) override;
-	virtual void OnCounterFloatValue(const FCounterFloatValue& NewValue) override;
+	FString Name;
+	ETraceCounterType Type;
 
-	struct FCounter
+	union Value
 	{
-		FString Name;
-		ETraceCounterType Type;
+		int64 Int;
+		double Float;
+	};
 
-		union
-		{
-			int64 IntValue;
-			double FloatValue;
-		};
+	Value Zero;
+	Value First;
+	Value Last;
+	Value Minimum;
+	Value Maximum;
+
+	uint64 Count = 0;
+	double Mean = 0.0;
+	double VarianceAccumulator = 0.0; // Accumulator for Welford's
 
-		FCounter(FString InName, ETraceCounterType InType)
+	double FirstSeconds = 0.0;
+	double LastSeconds = 0.0;
+
+	FSummarizeCounter(FString InName, ETraceCounterType InType)
+	{
+		Name = InName;
+		Type = InType;
+		switch (Type)
 		{
-			Name = InName;
-			Type = InType;
-			switch (Type)
-			{
 			case TraceCounterType_Int:
-				IntValue = 0;
+				Zero.Int = 0;
+				First.Int = 0;
+				Last.Int = 0;
+				Minimum.Int = TNumericLimits<int64>::Max();
+				Maximum.Int = TNumericLimits<int64>::Min();
 				break;
 
 			case TraceCounterType_Float:
-				FloatValue = 0.0;
+				Zero.Float = 0;
+				First.Float = 0.0;
+				Last.Float = 0.0;
+				Minimum.Float = TNumericLimits<double>::Max();
+				Maximum.Float = TNumericLimits<double>::Min();
 				break;
-			}
 		}
+	}
 
-		template<typename ValueType>
-		void SetValue(ValueType InValue) = delete;
-
-		template<>
-		void SetValue(int64 InValue)
+	void SetValue(int64 InValue, double InTimestamp)
+	{
+		ensure(Type == TraceCounterType_Int);
+		if (Type == TraceCounterType_Int)
 		{
-			ensure(Type == TraceCounterType_Int);
-			if (Type == TraceCounterType_Int)
+			Count += 1;
+
+			if (Count == 1)
 			{
-				IntValue = InValue;
+				First.Int = InValue;
+				FirstSeconds = InTimestamp;
 			}
+
+			Last.Int = InValue;
+			LastSeconds = InTimestamp;
+
+			Minimum.Int = FMath::Min(Minimum.Int, InValue);
+			Maximum.Int = FMath::Max(Maximum.Int, InValue);
+
+			// update variance for stddev later
+			WelfordsIncrement(double(InValue), Count, Mean, VarianceAccumulator);
 		}
+	}
 
-		template<>
-		void SetValue(double InValue)
+	void SetValue(double InValue, double InTimestamp)
+	{
+		ensure(Type == TraceCounterType_Float);
+		if (Type == TraceCounterType_Float)
 		{
-			ensure(Type == TraceCounterType_Float);
-			if (Type == TraceCounterType_Float)
+			Count += 1;
+
+			if (Count == 1)
 			{
-				FloatValue = InValue;
+				First.Float = InValue;
+				FirstSeconds = InTimestamp;
 			}
+
+			Last.Float = InValue;
+			LastSeconds = InTimestamp;
+
+			Minimum.Float = FMath::Min(Minimum.Float, InValue);
+			Maximum.Float = FMath::Max(Maximum.Float, InValue);
+
+			// update variance for stddev later
+			WelfordsIncrement(InValue, Count, Mean, VarianceAccumulator);
 		}
+	}
+
+	double GetDeviation() const
+	{
+		return WelfordsDeviation(VarianceAccumulator, Count);
+	}
 
-		FString GetValue() const
+	double GetCountPerSecond() const
+	{
+		double CountPerSecond = 0.0;
+		if (Count)
+		{
+			CountPerSecond = Count / (LastSeconds - FirstSeconds);
+		}
+		return CountPerSecond;
+	}
+
+	FString PrintValue(const Value& InValue) const
+	{
+		switch (Type)
 		{
-			switch (Type)
-			{
 			case TraceCounterType_Int:
-				return FString::Printf(TEXT("%lld"), IntValue);
+				return FString::Printf(TEXT("%lld"), InValue.Int);
 
 			case TraceCounterType_Float:
-				return FString::Printf(TEXT("%f"), FloatValue);
-			}
+				return FString::Printf(TEXT("%f"), InValue.Float);
+		}
+
+		ensure(false);
+		return TEXT("");
+	}
 
-			ensure(false);
-			return TEXT("");
+	FString GetValue(const FStringView& Statistic) const
+	{
+		if (Statistic == TEXT("Name"))
+		{
+			return Name;
 		}
-	};
+		else if (Statistic == TEXT("Count"))
+		{
+			return FString::Printf(TEXT("%llu"), Count);
+		}
+		else if (Statistic == TEXT("First"))
+		{
+			return PrintValue(First);
+		}
+		else if (Statistic == TEXT("FirstSeconds"))
+		{
+			return FString::Printf(TEXT("%f"), FirstSeconds);
+		}
+		else if (Statistic == TEXT("Last"))
+		{
+			return PrintValue(Last);
+		}
+		else if (Statistic == TEXT("LastSeconds"))
+		{
+			return FString::Printf(TEXT("%f"), LastSeconds);
+		}
+		else if (Statistic == TEXT("Minimum"))
+		{
+			return PrintValue(Count ? Minimum : Zero);
+		}
+		else if (Statistic == TEXT("Maximum"))
+		{
+			return PrintValue(Count ? Maximum : Zero);
+		}
+		else if (Statistic == TEXT("Mean"))
+		{
+			return FString::Printf(TEXT("%f"), Mean);
+		}
+		else if (Statistic == TEXT("Deviation"))
+		{
+			return FString::Printf(TEXT("%f"), GetDeviation());
+		}
+		else if (Statistic == TEXT("CountPerSecond"))
+		{
+			return FString::Printf(TEXT("%f"), GetCountPerSecond());
+		}
+
+		ensure(false);
+		return TEXT("");
+	}
+};
 
-	TMap<uint16, FCounter> Counters;
+class FSummarizeCountersAnalyzer
+	: public FCountersAnalyzer
+{
+public:
+	virtual void OnCounterName(const FCounterName& CounterName) override;
+	virtual void OnCounterIntValue(const FCounterIntValue& NewValue) override;
+	virtual void OnCounterFloatValue(const FCounterFloatValue& NewValue) override;
+
+	TMap<uint16, FSummarizeCounter> Counters;
 };
 
 void FSummarizeCountersAnalyzer::OnCounterName(const FCounterName& CounterName)
 {
-	Counters.Add(CounterName.Id, FCounter(CounterName.Name, CounterName.Type));
+	Counters.Add(CounterName.Id, FSummarizeCounter(CounterName.Name, CounterName.Type));
 }
 
 void FSummarizeCountersAnalyzer::OnCounterIntValue(const FCounterIntValue& NewValue)
 {
-	FCounter* FoundCounter = Counters.Find(NewValue.Id);
+	FSummarizeCounter* FoundCounter = Counters.Find(NewValue.Id);
 	ensure(FoundCounter);
 	if (FoundCounter)
 	{
-		FoundCounter->SetValue(NewValue.Value);
+		FoundCounter->SetValue(NewValue.Value, NewValue.Timestamp);
 	}
 }
 
 void FSummarizeCountersAnalyzer::OnCounterFloatValue(const FCounterFloatValue& NewValue)
 {
-	FCounter* FoundCounter = Counters.Find(NewValue.Id);
+	FSummarizeCounter* FoundCounter = Counters.Find(NewValue.Id);
 	ensure(FoundCounter);
 	if (FoundCounter)
 	{
-		FoundCounter->SetValue(NewValue.Value);
+		FoundCounter->SetValue(NewValue.Value, NewValue.Timestamp);
 	}
 }
 
@@ -1209,6 +1331,60 @@ void FSummarizeCountersAnalyzer::OnCounterFloatValue(const FCounterFloatValue& N
 // FSummarizeBookmarksAnalyzer - Tally Bookmarks from bookmark events
 //
 
+struct FSummarizeBookmark
+{
+	FString Name;
+	uint64 Count = 0;
+
+	double FirstSeconds = 0.0;
+	double LastSeconds = 0.0;
+
+	void AddTimestamp(double Seconds)
+	{
+		Count += 1;
+
+		// only set first for the first sample, compare exact zero
+		if (FirstSeconds == 0.0)
+		{
+			FirstSeconds = Seconds;
+		}
+
+		LastSeconds = Seconds;
+	}
+
+	FString GetValue(const FStringView& Statistic) const
+	{
+		if (Statistic == TEXT("Name"))
+		{
+			return Name;
+		}
+		else if (Statistic == TEXT("Count"))
+		{
+			return FString::Printf(TEXT("%llu"), Count);
+		}
+		else if (Statistic == TEXT("FirstSeconds"))
+		{
+			return FString::Printf(TEXT("%f"), FirstSeconds);
+		}
+		else if (Statistic == TEXT("LastSeconds"))
+		{
+			return FString::Printf(TEXT("%f"), LastSeconds);
+		}
+		return FString();
+	}
+
+	// for deduplication
+	bool operator==(const FSummarizeBookmark& Bookmark) const
+	{
+		return Name == Bookmark.Name;
+	}
+};
+
+static uint32 GetTypeHash(const FSummarizeBookmark& Bookmark)
+{
+	return FCrc::StrCrc32(*Bookmark.Name);
+}
+
 class FSummarizeBookmarksAnalyzer
 	: public FBookmarksAnalyzer
 {
@@ -1803,7 +1979,7 @@ int32 USummarizeTraceCommandlet::Main(const FString& CmdLineParams)
 		});
 
 	// Analyze 'LoadModule*' scope timer hierarchically to account individual load time only (substracting time consumed to load dependent module(s)).
-	TSharedPtr<FCpuScopeHierarchyAnalyzer> HierarchicalScopeAnalyzer = MakeShared<FCpuScopeHierarchyAnalyzer>(
+	TSharedPtr<FSummarizeCpuScopeHierarchyAnalyzer> HierarchicalScopeAnalyzer = MakeShared<FSummarizeCpuScopeHierarchyAnalyzer>(
 		TEXT("LoadModule"), // Analyzer Name.
 		[](const FString& ScopeName)
 		{
@@ -1906,7 +2082,7 @@ int32 USummarizeTraceCommandlet::Main(const FString& CmdLineParams)
 	else
 	{
 		// no newline, see row printfs
-		WriteUTF8(CsvHandle, FString::Printf(TEXT("Name,Count,TotalDurationSeconds,FirstStartSeconds,FirstFinishSeconds,FirstDurationSeconds,LastStartSeconds,LastFinishSeconds,LastDurationSeconds,MinDurationSeconds,MaxDurationSeconds,MeanDurationSeconds,DeviationDurationSeconds,")));
+		WriteUTF8(CsvHandle, FString::Printf(TEXT("Name,Count,CountPerSecond,TotalDurationSeconds,FirstStartSeconds,FirstFinishSeconds,FirstDurationSeconds,LastStartSeconds,LastFinishSeconds,LastDurationSeconds,MinDurationSeconds,MaxDurationSeconds,MeanDurationSeconds,DeviationDurationSeconds,")));
 		for (const FSummarizeScope& Scope : SortedScopes)
 		{
 			if (!IsCsvSafeString(Scope.Name))
@@ -1915,7 +2091,13 @@ int32 USummarizeTraceCommandlet::Main(const FString& CmdLineParams)
 			}
 
 			// note newline is at the front of every data line to prevent final extraneous newline, per customary for csv
-			WriteUTF8(CsvHandle, FString::Printf(TEXT("\n%s,%llu,%f,%f,%f,%f,%f,%f,%f,%f,%f,%f,%f,"), *Scope.Name, Scope.Count, Scope.TotalDurationSeconds, Scope.FirstStartSeconds, Scope.FirstFinishSeconds, Scope.FirstDurationSeconds, Scope.FirstStartSeconds, Scope.FirstFinishSeconds, Scope.FirstDurationSeconds, Scope.MinDurationSeconds, Scope.MaxDurationSeconds, Scope.MeanDurationSeconds, Scope.GetDeviationDurationSeconds()));
+			WriteUTF8(CsvHandle, FString::Printf(TEXT("\n%s,%llu,%f,%f,%f,%f,%f,%f,%f,%f,%f,%f,%f,%f,"),
+				*Scope.Name, Scope.Count, Scope.GetCountPerSecond(),
+				Scope.TotalDurationSeconds,
+				Scope.FirstStartSeconds, Scope.FirstFinishSeconds, Scope.FirstDurationSeconds,
+				Scope.LastStartSeconds, Scope.LastFinishSeconds, Scope.LastDurationSeconds,
+				Scope.MinDurationSeconds, Scope.MaxDurationSeconds,
+				Scope.MeanDurationSeconds, Scope.GetDeviationDurationSeconds()));
 		}
 		CsvHandle->Flush();
 		delete CsvHandle;
@@ -1934,8 +2116,8 @@ int32 USummarizeTraceCommandlet::Main(const FString& CmdLineParams)
 	else
 	{
 		// no newline, see row printfs
-		WriteUTF8(CsvHandle, FString::Printf(TEXT("Name,Value,")));
-		for (const TMap<uint16, FSummarizeCountersAnalyzer::FCounter>::ElementType& Counter : CountersAnalyzer.Counters)
+		WriteUTF8(CsvHandle, FString::Printf(TEXT("Name,Count,CountPerSecond,First,FirstSeconds,Last,LastSeconds,Minimum,Maximum,Mean,Deviation,")));
+		for (const TMap<uint16, FSummarizeCounter>::ElementType& Counter : CountersAnalyzer.Counters)
 		{
 			if (!IsCsvSafeString(Counter.Value.Name))
 			{
@@ -1943,7 +2125,12 @@ int32 USummarizeTraceCommandlet::Main(const FString& CmdLineParams)
 			}
 
 			// note newline is at the front of every data line to prevent final extraneous newline, per customary for csv
-			WriteUTF8(CsvHandle, FString::Printf(TEXT("\n%s,%s,"), *Counter.Value.Name, *Counter.Value.GetValue()));
+			WriteUTF8(CsvHandle, FString::Printf(TEXT("\n%s,%llu,%f,%s,%f,%s,%f,%s,%s,%f,%f,"),
+				*Counter.Value.Name, Counter.Value.Count, Counter.Value.GetCountPerSecond(),
+				*Counter.Value.GetValue(TEXT("First")), Counter.Value.FirstSeconds,
+				*Counter.Value.GetValue(TEXT("Last")), Counter.Value.LastSeconds,
+				*Counter.Value.GetValue(TEXT("Minimum")), *Counter.Value.GetValue(TEXT("Maximum")),
+				Counter.Value.Mean, Counter.Value.GetDeviation()));
 		}
 		CsvHandle->Flush();
 		delete CsvHandle;
@@ -1971,7 +2158,9 @@ int32 USummarizeTraceCommandlet::Main(const FString& CmdLineParams)
 			}
 
 			// note newline is at the front of every data line to prevent final extraneous newline, per customary for csv
-			WriteUTF8(CsvHandle, FString::Printf(TEXT("\n%s,%d,%f,%f,"), *Bookmark.Value.Name, Bookmark.Value.Count, Bookmark.Value.FirstSeconds, Bookmark.Value.LastSeconds));
+			WriteUTF8(CsvHandle, FString::Printf(TEXT("\n%s,%d,%f,%f,"),
+				*Bookmark.Value.Name, Bookmark.Value.Count,
+				Bookmark.Value.FirstSeconds, Bookmark.Value.LastSeconds));
 		}
 		CsvHandle->Flush();
 		delete CsvHandle;
@@ -2041,7 +2230,7 @@ int32 USummarizeTraceCommandlet::Main(const FString& CmdLineParams)
 		}
 
 		// resolve counters to telemetry
-		for (const TMap<uint16, FSummarizeCountersAnalyzer::FCounter>::ElementType& Counter : CountersAnalyzer.Counters)
+		for (const TMap<uint16, FSummarizeCounter>::ElementType& Counter : CountersAnalyzer.Counters)
 		{
 			if (!IsCsvSafeString(Counter.Value.Name))
 			{
@@ -2050,15 +2239,22 @@ int32 USummarizeTraceCommandlet::Main(const FString& CmdLineParams)
 
 			if (bAllTelemetry)
 			{
-				TelemetryData.Add(TelemetryDefinition(TestName, Counter.Value.Name, TEXT("Count"), Counter.Value.GetValue(), TEXT("Count")));
+				TelemetryData.Add(TelemetryDefinition(TestName, Counter.Value.Name, TEXT("Count"), Counter.Value.GetValue(TEXT("Count")), TEXT("Count")));
+				TelemetryData.Add(TelemetryDefinition(TestName, Counter.Value.Name, TEXT("First"), Counter.Value.GetValue(TEXT("First")), TEXT("Count")));
+				TelemetryData.Add(TelemetryDefinition(TestName, Counter.Value.Name, TEXT("FirstSeconds"), Counter.Value.GetValue(TEXT("FirstSeconds")), TEXT("Seconds")));
+				TelemetryData.Add(TelemetryDefinition(TestName, Counter.Value.Name, TEXT("Last"), Counter.Value.GetValue(TEXT("Last")), TEXT("Count")));
+				TelemetryData.Add(TelemetryDefinition(TestName, Counter.Value.Name, TEXT("LastSeconds"), Counter.Value.GetValue(TEXT("LastSeconds")), TEXT("Seconds")));
+				TelemetryData.Add(TelemetryDefinition(TestName, Counter.Value.Name, TEXT("Minimum"), Counter.Value.GetValue(TEXT("Minimum")), TEXT("Count")));
+				TelemetryData.Add(TelemetryDefinition(TestName, Counter.Value.Name, TEXT("Maximum"), Counter.Value.GetValue(TEXT("Maximum")), TEXT("Count")));
+				TelemetryData.Add(TelemetryDefinition(TestName, Counter.Value.Name, TEXT("Mean"), Counter.Value.GetValue(TEXT("Mean")), TEXT("Count")));
+				TelemetryData.Add(TelemetryDefinition(TestName, Counter.Value.Name, TEXT("Deviation"), Counter.Value.GetValue(TEXT("Deviation")), TEXT("Count")));
 			}
 			else
 			{
 				NameToDefinitionMap.MultiFind(Counter.Value.Name, Statistics, true);
-				ensure(Statistics.Num() <= 1); // there should only be one, the counter value
 				for (const StatisticDefinition& Statistic : Statistics)
 				{
-					TelemetryData.Add(TelemetryDefinition(TestName, Statistic.TelemetryContext, Statistic.TelemetryDataPoint, Statistic.TelemetryUnit, Counter.Value.GetValue()));
+					TelemetryData.Add(TelemetryDefinition(TestName, Statistic.TelemetryContext, Statistic.TelemetryDataPoint, Statistic.TelemetryUnit, Counter.Value.GetValue(Statistic.Statistic)));
 					ResolvedStatistics.Add(Statistic.Name);
 				}
 				Statistics.Reset();			
diff --git a/Engine/Source/Runtime/Core/Private/GenericPlatform/GenericPlatformMemory.cpp b/Engine/Source/Runtime/Core/Private/GenericPlatform/GenericPlatformMemory.cpp
--- a/Engine/Source/Runtime/Core/Private/GenericPlatform/GenericPlatformMemory.cpp
+++ b/Engine/Source/Runtime/Core/Private/GenericPlatform/GenericPlatformMemory.cpp
@@ -22,6 +22,7 @@
 #include "Templates/UnrealTemplate.h"
 #include "HAL/IConsoleManager.h"
 #include "ProfilingDebugging/MemoryTrace.h"
+#include "ProfilingDebugging/CountersTrace.h"
 
 DEFINE_STAT(MCR_Physical);
 DEFINE_STAT(MCR_PhysicalLLM);
@@ -71,6 +72,15 @@ struct TUnalignedTester
 	}
 };
 
+TRACE_DECLARE_MEMORY_COUNTER(PlatformMemoryTotalPhysical, TEXT("PlatformMemory/TotalPhysical"));
+TRACE_DECLARE_MEMORY_COUNTER(PlatformMemoryTotalVirtual, TEXT("PlatformMemory/TotalVirtual"));
+TRACE_DECLARE_MEMORY_COUNTER(PlatformMemoryPageSize, TEXT("PlatformMemory/PageSize"));
+TRACE_DECLARE_MEMORY_COUNTER(PlatformMemoryAvailablePhysical, TEXT("PlatformMemory/AvailablePhysical"));
+TRACE_DECLARE_MEMORY_COUNTER(PlatformMemoryAvailableVirtual, TEXT("PlatformMemory/AvailableVirtual"));
+TRACE_DECLARE_MEMORY_COUNTER(PlatformMemoryUsedPhysical, TEXT("PlatformMemory/UsedPhysical"));
+TRACE_DECLARE_MEMORY_COUNTER(PlatformMemoryPeakUsedPhysical, TEXT("PlatformMemory/PeakUsedPhysical"));
+TRACE_DECLARE_MEMORY_COUNTER(PlatformMemoryUsedVirtual, TEXT("PlatformMemory/UsedVirtual"));
+TRACE_DECLARE_MEMORY_COUNTER(PlatformMemoryPeakUsedVirtual, TEXT("PlatformMemory/PeakUsedVirtual"));
 
 /** Helper class used to update platform memory stats. */
 struct FGenericStatsUpdater
@@ -104,6 +114,16 @@ struct FGenericStatsUpdater
 		SET_MEMORY_STAT( STAT_UsedVirtual, MemoryStats.UsedVirtual );
 		SET_MEMORY_STAT( STAT_PeakUsedVirtual, MemoryStats.PeakUsedVirtual );
 
+		TRACE_COUNTER_SET(PlatformMemoryTotalPhysical, MemoryStats.TotalPhysical);
+		TRACE_COUNTER_SET(PlatformMemoryTotalVirtual, MemoryStats.TotalVirtual);
+		TRACE_COUNTER_SET(PlatformMemoryPageSize, MemoryStats.PageSize);
+		TRACE_COUNTER_SET(PlatformMemoryAvailablePhysical, MemoryStats.AvailablePhysical);
+		TRACE_COUNTER_SET(PlatformMemoryAvailableVirtual, MemoryStats.AvailableVirtual);
+		TRACE_COUNTER_SET(PlatformMemoryUsedPhysical, MemoryStats.UsedPhysical);
+		TRACE_COUNTER_SET(PlatformMemoryPeakUsedPhysical, MemoryStats.PeakUsedPhysical);
+		TRACE_COUNTER_SET(PlatformMemoryUsedVirtual, MemoryStats.UsedVirtual);
+		TRACE_COUNTER_SET(PlatformMemoryPeakUsedVirtual, MemoryStats.PeakUsedVirtual);
+
 		// Platform specific stats.
 		FPlatformMemory::InternalUpdateStats( MemoryStats );
 	}
