Cooker: Move CookOnTheFly pollables into the PumpPollables system. This reduces per-loop costs in TickMainCookLoop and allows more code sharing with CookByTheBook.
Remove no-longer needed COTFS API elements.
#rnx

diff --git a/Engine/Source/Editor/UnrealEd/Classes/Commandlets/CookCommandlet.h b/Engine/Source/Editor/UnrealEd/Classes/Commandlets/CookCommandlet.h
--- a/Engine/Source/Editor/UnrealEd/Classes/Commandlets/CookCommandlet.h
+++ b/Engine/Source/Editor/UnrealEd/Classes/Commandlets/CookCommandlet.h
@@ -18,6 +18,7 @@
 
 class FSandboxPlatformFile;
 class ITargetPlatform;
+class UCookOnTheFlyServer;
 
 UCLASS(config=Editor)
 class UCookCommandlet
@@ -83,8 +84,8 @@ class UCookCommandlet
 	/** Cooks for specified targets */
 	bool CookByTheBook(const TArray<ITargetPlatform*>& Platforms);
 
-	/**	Process deferred commands */
-	void ProcessDeferredCommands();
+	/** Collect garbage if the cooker's TickResults requested it */
+	void ConditionalCollectGarbage(uint32 TickResults, UCookOnTheFlyServer& COTFS);
 
 public:
 
diff --git a/Engine/Source/Editor/UnrealEd/Classes/CookOnTheSide/CookOnTheFlyServer.h b/Engine/Source/Editor/UnrealEd/Classes/CookOnTheSide/CookOnTheFlyServer.h
--- a/Engine/Source/Editor/UnrealEd/Classes/CookOnTheSide/CookOnTheFlyServer.h
+++ b/Engine/Source/Editor/UnrealEd/Classes/CookOnTheSide/CookOnTheFlyServer.h
@@ -175,9 +175,6 @@ private:
 		 */
 		void RunNow(UCookOnTheFlyServer& COTFS);
 
-		/** Run the pollable's callback function, from PumpPollables. Calculate the new time and store it in the output. */
-		void RunDuringPump(UE::Cook::FTickStackData& StackData, double& OutNewCurrentTime, double& OutNextTimeSeconds);
-
 	public:
 		const TCHAR* DebugName = TEXT("");
 		FPollFunction PollFunction;
@@ -185,6 +182,14 @@ private:
 		double NextTimeIdleSeconds = 0.;
 		float PeriodSeconds = 60.f;
 		float PeriodIdleSeconds = 5.f;
+
+	public: // To be called only by PumpPollables
+		/** Run the pollable's callback function, from PumpPollables. Calculate the new time and store it in the output. */
+		void RunDuringPump(UE::Cook::FTickStackData& StackData, double& OutNewCurrentTime, double& OutNextTimeSeconds);
+		/** Update the position of this in the Pollables queue, called from Trigger or for a deferred Trigger. */
+		void TriggerInternal(UCookOnTheFlyServer& COTFS);
+		/** Update the position of this in the Pollables queue, called from RunNow or for a deferred RunNow. */
+		void RunNowInternal(UCookOnTheFlyServer& COTFS, double LastTimeRun);
 	};
 	/** A key/value pair for storing Pollables in a priority queue, keyed by next call time. */
 	struct FPollableQueueKey
@@ -360,7 +365,14 @@ private:
 	void TickCookStatus(UE::Cook::FTickStackData& StackData);
 	void SetSaveBusy(bool bInSaveBusy);
 	void SetLoadBusy(bool bInLoadBusy);
-	void UpdateDisplay(ECookTickFlags CookFlags, bool bForceDisplay);
+	enum class EIdleStatus
+	{
+		Active,
+		Idle,
+		Done
+	};
+	void SetIdleStatus(UE::Cook::FTickStackData& StackData, EIdleStatus InStatus);
+	void UpdateDisplay(UE::Cook::FTickStackData& StackData, bool bForceDisplay);
 	enum class ECookAction
 	{
 		Done,			// The cook is complete; no requests remain in any non-idle state
@@ -373,7 +385,6 @@ private:
 		PollIdle,		// Execute pollables which have exceeded their idle period
 		WaitForAsync,	// Sleep for a time slice while we wait for async tasks to complete
 		YieldTick,		// Progress is blocked by an async result. Temporarily exit TickMainCookLoop.
-		Cancel,			// Cancel the current CookByTheBook
 	};
 	/** Inspect all tasks the scheduler could do and return which one it should do. */
 	ECookAction DecideNextCookAction(UE::Cook::FTickStackData& StackData);
@@ -430,8 +441,10 @@ private:
 	void PollFlushRenderingCommands();
 	TRefCountPtr<FPollable> CreatePollableLLM();
 	TRefCountPtr<FPollable> CreatePollableTriggerGC();
-	void PollPackagesPerGC(UE::Cook::FTickStackData& StackData);
+	void PollTimeForGC(UE::Cook::FTickStackData& StackData);
+	void PollQueuedCancel(UE::Cook::FTickStackData& StackData);
 	void WaitForAsync(UE::Cook::FTickStackData& StackData);
+	void TickRecompileShaderRequestsPrivate();
 
 public:
 
@@ -446,6 +459,8 @@ public:
 		COSR_MarkedUpKeepPackages	= 0x00000040,
 		COSR_RequiresGC_OOM			= 0x00000080,
 		COSR_RequiresGC_PackageCount= 0x00000100,
+		COSR_RequiresGC_IdleTimer	= 0x00000200,
+		COSR_YieldTick				= 0x00000400,
 	};
 
 	struct FCookByTheBookStartupOptions
@@ -562,9 +577,9 @@ public:
 	void PostLoadPackageFixup(UE::Cook::FPackageData& PackageData, UPackage* Package);
 
 	/** Tick CBTB until it finishes or needs to yield. Should only be called when in CookByTheBook Mode. */
-	uint32 TickCookByTheBook(const float TimeSlice, uint32& CookedPackageCount, ECookTickFlags TickFlags = ECookTickFlags::None, int32 InMaxNumPackagesToSave = 50);
+	uint32 TickCookByTheBook(const float TimeSlice, ECookTickFlags TickFlags = ECookTickFlags::None);
 	/** Tick COTF until it finishes or needs to yield. Should only be called when in CookOnTheFly Mode. */
-	uint32 TickCookOnTheFly(const float TimeSlice, uint32& CookedPackageCount, ECookTickFlags TickFlags = ECookTickFlags::None, int32 InMaxNumPackagesToSave = 50);
+	uint32 TickCookOnTheFly(const float TimeSlice, ECookTickFlags TickFlags = ECookTickFlags::None);
 
 	/**
 	 * Execute CookByTheBook as a load of all packages followed by a save of all packages. This mode is sometimes faster than normal CBTB.
@@ -636,8 +651,10 @@ public:
 	/** 
 	* Process any shader recompile requests
 	*/
-	void TickRecompileShaderRequests();
+	UE_DEPRECATED(5.1, "Ticking RecompileShaderRequests is now handled by the Tick function")
+	void TickRecompileShaderRequests() {}
 
+	UE_DEPRECATED(5.1, "Ticking RecompileShaderRequests is now handled by the Tick function")
 	bool HasRecompileShaderRequests() const;
 
 	/** Return whether the tick needs to take any action for the current session. If not, the session is done. Used for external managers of the cooker to know when to tick it. */
@@ -1210,12 +1227,12 @@ private:
 	bool bFirstCookInThisProcessInitialized = false;
 	bool bFirstCookInThisProcess = true;
 	bool bImportBehaviorCallbackInstalled = false;
-	/** Cancel has been queued and will be processed next tick */
-	bool bCancelSession = false;
 	/** True if a session is started for any mode. If the session is started, then we have at least one TargetPlatform specified. */
 	bool bSessionRunning = false;
 	/** Whether we're using ZenStore for storage of cook results. If false, we are using LooseCookedPackageWriter. */
 	bool bZenStore = false;
+	/** Multithreaded synchronization of Pollables, accessible only inside PollablesLock. */
+	bool bPollablesInTick = false;
 
 	/** Timers for tracking how long we have been busy, to manage retries and warnings of deadlock */
 	double SaveBusyRetryTimeSeconds = MAX_flt;
@@ -1238,12 +1255,28 @@ private:
 	/** Helper struct for running cooking in diagnostic modes */
 	TUniquePtr<FDiffModeCookServerUtils> DiffModeHelper;
 
+	/**
+	 * Heap of Pollables to tick, ordered by NextTimeSeconds.
+	 * Only accessible by PumpPollables when bPollablesInTick==true, can be accessed outside of PollablesLock.
+	 * or elsewhere when bPollablesInTick==false, must be accessed inside PollablesLock.
+	 */
 	TArray<FPollableQueueKey> Pollables;
+	/**
+	 * List of pollables that were triggered during PumpPollables and need to be updated when the Pump is done.
+	 * Accessible only inside PollablesLock.
+	 */
+	TArray<FPollableQueueKey> PollablesDeferredTriggers;
+	/** Together with bPollablesInTick, provides a lock around Pollables. */
 	FCriticalSection PollablesLock;
+	TRefCountPtr<FPollable> RecompileRequestsPollable;
+	TRefCountPtr<FPollable> QueuedCancelPollable;
 	double PollNextTimeSeconds = 0.;
 	double PollNextTimeIdleSeconds = 0.;
+	double IdleStatusStartTime = 0.0;
 	uint32 CookedPackageCountSinceLastGC = 0;
 	float WaitForAsyncSleepSeconds = 0.0f;
+	EIdleStatus IdleStatus = EIdleStatus::Done;
+
 	friend UE::Cook::FPackageData;
 	friend UE::Cook::FPendingCookedPlatformData;
 	friend UE::Cook::FPlatformManager;
diff --git a/Engine/Source/Editor/UnrealEd/Private/Commandlets/CookCommandlet.cpp b/Engine/Source/Editor/UnrealEd/Private/Commandlets/CookCommandlet.cpp
--- a/Engine/Source/Editor/UnrealEd/Private/Commandlets/CookCommandlet.cpp
+++ b/Engine/Source/Editor/UnrealEd/Private/Commandlets/CookCommandlet.cpp
@@ -463,125 +463,17 @@ bool UCookCommandlet::CookOnTheFly( FGuid InstanceId, int32 Timeout, bool bForce
 		GShaderCompilingManager->SkipShaderCompilation(true);
 	}
 
-	// Garbage collection should happen when either
-	//	1. We have cooked a map (configurable asset type)
-	//	2. We have cooked non-map packages and...
-	//		a. we have accumulated 50 (configurable) of these since the last GC.
-	//		b. we have been idle for 20 (configurable) seconds.
-	struct FCookOnTheFlyGCController
-	{
-	public:
-		FCookOnTheFlyGCController(const UCookOnTheFlyServer* COtFServer)
-			: PackagesPerGC(COtFServer->GetPackagesPerGC())
-			, IdleTimeToGC(COtFServer->GetIdleTimeToGC())
-			, bShouldGC(true)
-			, PackagesCookedSinceLastGC(0)
-			, LastCookActionTime(FPlatformTime::Seconds())
-		{}
-
-		/** Intended to be called with stats from the COTF tick loop. Determines if we should be calling GC after TickCookOnTheFly(). */
-		void Update(uint32 CookedCount, UCookOnTheFlyServer::ECookOnTheSideResult ResultFlags)
-		{
-			if (ResultFlags & (UCookOnTheFlyServer::COSR_CookedMap | UCookOnTheFlyServer::COSR_CookedPackage | UCookOnTheFlyServer::COSR_WaitingOnCache))
-			{
-				LastCookActionTime = FPlatformTime::Seconds();
-			}
-			
-			if (ResultFlags & UCookOnTheFlyServer::COSR_RequiresGC)
-			{
-				UE_LOG(LogCookCommandlet, Display, TEXT("Cooker cooked a map since last gc... collecting garbage"));
-				bShouldGC |= true;
-			}
-
-			PackagesCookedSinceLastGC += CookedCount;
-			if ((PackagesPerGC > 0) && (PackagesCookedSinceLastGC > PackagesPerGC))
-			{
-				UE_LOG(LogCookCommandlet, Display, TEXT("Cooker has exceeded max number of non map packages since last gc"));
-				bShouldGC |= true;
-			}
-
-			// we don't want to gc if we are waiting on cache of objects. this could clean up objects which we will need to reload next frame
-			bPostponeGC = (ResultFlags & UCookOnTheFlyServer::COSR_WaitingOnCache) != 0;
-		}
-
-		/** Runs GC if Update() determined it should happen. Also checks the idle time against the limit, and runs GC then if packages have been loaded. */
-		void ConditionallyCollectGarbage(const UCookOnTheFlyServer* COtFServer)
-		{
-			if (!bShouldGC)
-			{
-				if (PackagesCookedSinceLastGC > 0 && IdleTimeToGC > 0)
-				{
-					double IdleTime = FPlatformTime::Seconds() - LastCookActionTime;
-					if (IdleTime >= IdleTimeToGC)
-					{
-						UE_LOG(LogCookCommandlet, Display, TEXT("Cooker has been idle for long time gc"));
-						bShouldGC |= true;
-					}
-				}
-
-				if (!bShouldGC && COtFServer->HasExceededMaxMemory())
-				{
-					UE_LOG(LogCookCommandlet, Display, TEXT("Cooker has exceeded max memory usage collecting garbage"));
-					bShouldGC |= true;
-				}
-			}
-
-			if (bShouldGC && !bPostponeGC)
-			{	
-				Reset();
-				UE_LOG(LogCookCommandlet, Display, TEXT("GC..."));
-				CollectGarbage(RF_NoFlags);
-			}
-		}
-
-	private:
-		/** Resets counters and flags used to determine when we should GC. */
-		void Reset()
-		{
-			bShouldGC = false;
-			PackagesCookedSinceLastGC = 0;
-		}
-
-	private:
-		const uint32 PackagesPerGC;
-		const double IdleTimeToGC;
-
-		bool   bShouldGC;
-		uint32 PackagesCookedSinceLastGC;
-		double LastCookActionTime;
-		bool   bPostponeGC;
-
-	} CookOnTheFlyGCController(CookOnTheFlyServer);
-
 	FDateTime LastConnectionTime = FDateTime::UtcNow();
 	bool bHadConnection = false;
 
 	while (!IsEngineExitRequested())
 	{
-		uint32 CookedPkgCount = 0;
-		uint32 TickResults = CookOnTheFlyServer->TickCookOnTheFly(/*TimeSlice =*/10.f, CookedPkgCount, ShowProgress ? ECookTickFlags::None : ECookTickFlags::HideProgressDisplay);
-
-		// Flush the asset registry before GC
-		FAssetRegistryModule::TickAssetRegistry(-1.0f);
-
-		CookOnTheFlyGCController.Update(CookedPkgCount, (UCookOnTheFlyServer::ECookOnTheSideResult)TickResults);
-		CookOnTheFlyGCController.ConditionallyCollectGarbage(CookOnTheFlyServer);
+		uint32 TickResults = CookOnTheFlyServer->TickCookOnTheFly(/*TimeSlice =*/MAX_flt,
+			ShowProgress ? ECookTickFlags::None : ECookTickFlags::HideProgressDisplay);
+		ConditionalCollectGarbage(TickResults, *CookOnTheFlyServer);
 
-		// force at least a tick shader compilation even if we are requesting stuff
-		CookOnTheFlyServer->TickRecompileShaderRequests();
-		GShaderCompilingManager->ProcessAsyncResults(true, false);
-
-		while ( (CookOnTheFlyServer->HasRemainingWork() == false) && !IsEngineExitRequested())
+		if (!CookOnTheFlyServer->HasRemainingWork() && !IsEngineExitRequested())
 		{
-			CookOnTheFlyServer->TickRequestManager();
-
-			CookOnTheFlyServer->TickRecompileShaderRequests();
-
-			// Shaders need to be updated
-			GShaderCompilingManager->ProcessAsyncResults(true, false);
-
-			ProcessDeferredCommands();
-
 			// handle server timeout
 			if (InstanceId.IsValid() || bForceClose)
 			{
@@ -610,14 +502,6 @@ bool UCookCommandlet::CookOnTheFly( FGuid InstanceId, int32 Timeout, bool bForce
 				}
 			}
 
-			CookOnTheFlyGCController.ConditionallyCollectGarbage(CookOnTheFlyServer);
-
-			{
-				// Flush rendering commands to release any RHI resources (shaders and shader maps).
-				// Delete any FPendingCleanupObjects (shader maps).
-				FlushRenderingCommands();
-			}
-
 			CookOnTheFlyServer->WaitForRequests(100 /* timeoutMs */);
 		}
 	}
@@ -1004,86 +888,9 @@ bool UCookCommandlet::CookByTheBook( const TArray<ITargetPlatform*>& Platforms)
 				uint32 TickResults = 0;
 				uint32 UnusedVariable = 0;
 
-				TickResults = CookOnTheFlyServer->TickCookByTheBook(MAX_flt, UnusedVariable,
-					ShowProgress ? ECookTickFlags::None : ECookTickFlags::HideProgressDisplay, MAX_int32);
-
-				if ((TickResults & UCookOnTheFlyServer::COSR_RequiresGC) != 0)
-				{
-					FString GCReason;
-					if ((TickResults & UCookOnTheFlyServer::COSR_RequiresGC_PackageCount) != 0)
-					{
-						GCReason = TEXT("Exceeded packages per GC");
-					}
-					else if ((TickResults & UCookOnTheFlyServer::COSR_RequiresGC_OOM) != 0)
-					{
-						// this can cause thrashing if the cooker loads the same stuff into memory next tick
-						GCReason = TEXT("Exceeded Max Memory");
-
-						int32 JobsToLogAt = GShaderCompilingManager->GetNumRemainingJobs();
-
-						UE_SCOPED_COOKTIMER(CookByTheBook_ShaderJobFlush);
-						UE_LOG(LogCookCommandlet, Display, TEXT("Detected max mem exceeded - forcing shader compilation flush"));
-						while (true)
-						{
-							int32 NumRemainingJobs = GShaderCompilingManager->GetNumRemainingJobs();
-							if (NumRemainingJobs < 1000)
-							{
-								UE_LOG(LogCookCommandlet, Display, TEXT("Finished flushing shader jobs at %d"), NumRemainingJobs);
-								break;
-							}
-
-							if (NumRemainingJobs < JobsToLogAt)
-							{
-								UE_LOG(LogCookCommandlet, Display, TEXT("Flushing shader jobs, remaining jobs %d"), NumRemainingJobs);
-							}
-
-							GShaderCompilingManager->ProcessAsyncResults(false, false);
-
-							FPlatformProcess::Sleep(0.05);
-
-							// GShaderCompilingManager->FinishAllCompilation();
-						}
-					}
-					else
-					{
-						// cooker loaded some object which needs to be cleaned up before the cooker can proceed so force gc
-						GCReason = TEXT("COSR_RequiresGC");
-					}
-
-					// Flush the asset registry before GC
-					{
-						UE_SCOPED_COOKTIMER(CookByTheBook_TickAssetRegistry);
-						FAssetRegistryModule::TickAssetRegistry(-1.0f);
-					}
-					UE_SCOPED_COOKTIMER_AND_DURATION(CookByTheBook_GC, DetailedCookStats::TickLoopGCTimeSec);
-
-					int32 NumObjectsBeforeGC = GUObjectArray.GetObjectArrayNumMinusAvailable();
-					int32 NumObjectsAvailableBeforeGC = GUObjectArray.GetObjectArrayEstimatedAvailable();
-					UE_LOG(LogCookCommandlet, Display, TEXT("GarbageCollection...%s (%s)"), (bPartialGC ? TEXT(" partial gc") : TEXT("")), *GCReason);
-
-#if ENABLE_LOW_LEVEL_MEM_TRACKER
-					FLowLevelMemTracker::Get().UpdateStatsPerFrame();
-#endif
-					auto DumpMemStats = []()
-					{
-						FGenericMemoryStats MemStats;
-						GMalloc->GetAllocatorStats(MemStats);
-						for (const auto& Item : MemStats.Data)
-						{
-							UE_LOG(LogCookCommandlet, Display, TEXT("Item %s = %d"), *Item.Key, Item.Value);
-						}
-					};
-
-					DumpMemStats();
-
-					CollectGarbage(RF_NoFlags);
-
-					int32 NumObjectsAfterGC = GUObjectArray.GetObjectArrayNumMinusAvailable();
-					int32 NumObjectsAvailableAfterGC = GUObjectArray.GetObjectArrayEstimatedAvailable();
-					UE_LOG(LogCookCommandlet, Display, TEXT("%s GC before %d available %d after %d available %d"), (bPartialGC ? TEXT("Partial") : TEXT("Full")), NumObjectsBeforeGC, NumObjectsAvailableBeforeGC, NumObjectsAfterGC, NumObjectsAvailableAfterGC);
-
-					DumpMemStats();
-				}
+				TickResults = CookOnTheFlyServer->TickCookByTheBook(MAX_flt,
+					ShowProgress ? ECookTickFlags::None : ECookTickFlags::HideProgressDisplay);
+				ConditionalCollectGarbage(TickResults, *CookOnTheFlyServer);
 			}
 		}
 	} while (bTestCook);
@@ -1101,21 +908,102 @@ bool UCookCommandlet::CookByTheBook( const TArray<ITargetPlatform*>& Platforms)
 	return true;
 }
 
-void UCookCommandlet::ProcessDeferredCommands()
+void UCookCommandlet::ConditionalCollectGarbage(uint32 TickResults, UCookOnTheFlyServer& COTFS)
 {
-#if PLATFORM_MAC
-	// On Mac we need to process Cocoa events so that the console window for CookOnTheFlyServer is interactive
-	FPlatformApplicationMisc::PumpMessages(true);
+	if ((TickResults & UCookOnTheFlyServer::COSR_RequiresGC) == 0)
+	{
+		return;
+	}
+
+	FString GCReason;
+	if ((TickResults & UCookOnTheFlyServer::COSR_RequiresGC_PackageCount) != 0)
+	{
+		GCReason = TEXT("Exceeded packages per GC");
+	}
+	else if ((TickResults & UCookOnTheFlyServer::COSR_RequiresGC_OOM) != 0)
+	{
+		// this can cause thrashing if the cooker loads the same stuff into memory next tick
+		GCReason = TEXT("Exceeded Max Memory");
+
+		int32 JobsToLogAt = GShaderCompilingManager->GetNumRemainingJobs();
+
+		UE_SCOPED_COOKTIMER(CookByTheBook_ShaderJobFlush);
+		UE_LOG(LogCookCommandlet, Display, TEXT("Detected max mem exceeded - forcing shader compilation flush"));
+		while (true)
+		{
+			int32 NumRemainingJobs = GShaderCompilingManager->GetNumRemainingJobs();
+			if (NumRemainingJobs < 1000)
+			{
+				UE_LOG(LogCookCommandlet, Display, TEXT("Finished flushing shader jobs at %d"), NumRemainingJobs);
+				break;
+			}
+
+			if (NumRemainingJobs < JobsToLogAt)
+			{
+				UE_LOG(LogCookCommandlet, Display, TEXT("Flushing shader jobs, remaining jobs %d"), NumRemainingJobs);
+			}
+
+			GShaderCompilingManager->ProcessAsyncResults(false, false);
+
+			FPlatformProcess::Sleep(0.05);
+
+			// GShaderCompilingManager->FinishAllCompilation();
+		}
+	}
+	else if (TickResults & UCookOnTheFlyServer::COSR_RequiresGC_IdleTimer)
+	{
+		GCReason = TEXT("Cooker has been idle for long time gc");
+	}
+	else
+	{
+		// cooker loaded some object which needs to be cleaned up before the cooker can proceed so force gc
+		GCReason = TEXT("COSR_RequiresGC");
+	}
+
+	// Flush the asset registry before GC
+	{
+		UE_SCOPED_COOKTIMER(CookByTheBook_TickAssetRegistry);
+		FAssetRegistryModule::TickAssetRegistry(-1.0f);
+	}
+#if OUTPUT_COOKTIMING
+	TOptional<FScopedDurationTimer> CBTBScopedDurationTimer;
+	if (COTFS.IsCookByTheBookMode())
+	{
+		CBTBScopedDurationTimer.Emplace(DetailedCookStats::TickLoopGCTimeSec);
+	}
 #endif
+	UE_SCOPED_COOKTIMER(CookCommandlet_GC);
 
-	// update task graph
-	FTaskGraphInterface::Get().ProcessThreadUntilIdle(ENamedThreads::GameThread);
+	int32 NumObjectsBeforeGC = GUObjectArray.GetObjectArrayNumMinusAvailable();
+	int32 NumObjectsAvailableBeforeGC = GUObjectArray.GetObjectArrayEstimatedAvailable();
+	UE_LOG(LogCookCommandlet, Display, TEXT("GarbageCollection...%s (%s)"), (bPartialGC ? TEXT(" partial gc") : TEXT("")), *GCReason);
 
-	// execute deferred commands
-	for (int32 DeferredCommandsIndex = 0; DeferredCommandsIndex<GEngine->DeferredCommands.Num(); ++DeferredCommandsIndex)
+#if ENABLE_LOW_LEVEL_MEM_TRACKER
+	FLowLevelMemTracker::Get().UpdateStatsPerFrame();
+#endif
+	auto DumpMemStats = []()
 	{
-		GEngine->Exec( GWorld, *GEngine->DeferredCommands[DeferredCommandsIndex], *GLog);
+		FGenericMemoryStats MemStats;
+		GMalloc->GetAllocatorStats(MemStats);
+		for (const auto& Item : MemStats.Data)
+		{
+			UE_LOG(LogCookCommandlet, Display, TEXT("Item %s = %d"), *Item.Key, Item.Value);
+		}
+	};
+
+	if (COTFS.IsCookByTheBookMode())
+	{
+		DumpMemStats();
 	}
 
-	GEngine->DeferredCommands.Empty();
-}
+	CollectGarbage(RF_NoFlags);
+
+	int32 NumObjectsAfterGC = GUObjectArray.GetObjectArrayNumMinusAvailable();
+	int32 NumObjectsAvailableAfterGC = GUObjectArray.GetObjectArrayEstimatedAvailable();
+	if (COTFS.IsCookByTheBookMode())
+	{
+		UE_LOG(LogCookCommandlet, Display, TEXT("%s GC before %d available %d after %d available %d"),
+			(bPartialGC ? TEXT("Partial") : TEXT("Full")), NumObjectsBeforeGC, NumObjectsAvailableBeforeGC, NumObjectsAfterGC, NumObjectsAvailableAfterGC);
+		DumpMemStats();
+	}
+}
\ No newline at end of file
diff --git a/Engine/Source/Editor/UnrealEd/Private/CookOnTheFlyServer.cpp b/Engine/Source/Editor/UnrealEd/Private/CookOnTheFlyServer.cpp
--- a/Engine/Source/Editor/UnrealEd/Private/CookOnTheFlyServer.cpp
+++ b/Engine/Source/Editor/UnrealEd/Private/CookOnTheFlyServer.cpp
@@ -478,37 +478,46 @@ void UCookOnTheFlyServer::Tick(float DeltaTime)
 
 	check(IsCookingInEditor());
 
-	if (IsCookByTheBookMode() && !IsInSession() && !GIsSlowTask && IsCookFlagSet(ECookInitializationFlags::BuildDDCInBackground))
+	if (IsInSession())
 	{
-		// if we are in the editor then precache some stuff ;)
-		TArray<const ITargetPlatform*> CacheTargetPlatforms;
-		const ULevelEditorPlaySettings* PlaySettings = GetDefault<ULevelEditorPlaySettings>();
-		if (PlaySettings && (PlaySettings->LastExecutedLaunchModeType == LaunchMode_OnDevice))
+		// prevent autosave from happening until we are finished cooking
+		// causes really bad hitches
+		if (GUnrealEd)
 		{
-			FString DeviceName = PlaySettings->LastExecutedLaunchDevice.Left(PlaySettings->LastExecutedLaunchDevice.Find(TEXT("@")));
-			CacheTargetPlatforms.Add(GetTargetPlatformManager()->FindTargetPlatform(DeviceName));
+			constexpr float SecondsWarningTillAutosave = 10.0f;
+			GUnrealEd->GetPackageAutoSaver().ForceMinimumTimeTillAutoSave(SecondsWarningTillAutosave);
 		}
-		if (CacheTargetPlatforms.Num() > 0)
+	}
+	else
+	{
+		if (IsCookByTheBookMode() && !GIsSlowTask && IsCookFlagSet(ECookInitializationFlags::BuildDDCInBackground))
 		{
-			TickPrecacheObjectsForPlatforms(0.001, CacheTargetPlatforms);
+			// if we are in the editor then precache some stuff ;)
+			TArray<const ITargetPlatform*> CacheTargetPlatforms;
+			const ULevelEditorPlaySettings* PlaySettings = GetDefault<ULevelEditorPlaySettings>();
+			if (PlaySettings && (PlaySettings->LastExecutedLaunchModeType == LaunchMode_OnDevice))
+			{
+				FString DeviceName = PlaySettings->LastExecutedLaunchDevice.Left(PlaySettings->LastExecutedLaunchDevice.Find(TEXT("@")));
+				CacheTargetPlatforms.Add(GetTargetPlatformManager()->FindTargetPlatform(DeviceName));
+			}
+			if (CacheTargetPlatforms.Num() > 0)
+			{
+				TickPrecacheObjectsForPlatforms(0.001, CacheTargetPlatforms);
+			}
 		}
 	}
 
-	uint32 CookedPackagesCount = 0;
-	const static float TickTimeSliceSeconds = 0.1f;
+	const float TickTimeSliceSeconds = 0.1f;
 	TickCancels();
 	if (IsCookOnTheFlyMode())
 	{
-		TickCookOnTheFly(TickTimeSliceSeconds, CookedPackagesCount);
+		TickCookOnTheFly(TickTimeSliceSeconds);
 	}
 	else
 	{
 		check(IsCookByTheBookMode());
-		TickCookByTheBook(TickTimeSliceSeconds, CookedPackagesCount);
+		TickCookByTheBook(TickTimeSliceSeconds);
 	}
-
-	TickRequestManager();
-	TickRecompileShaderRequests();
 }
 
 bool UCookOnTheFlyServer::IsTickable() const 
@@ -616,6 +625,7 @@ bool UCookOnTheFlyServer::StartCookOnTheFly(FCookOnTheFlyStartupOptions InCookOn
 				};
 
 				PackageTracker->RecompileRequests.Enqueue({ RecompileData, MoveTemp(RecompileCompleted) });
+				RecompileRequestsPollable->Trigger(*this);
 
 				RecompileCompletedEvent->Wait();
 
@@ -650,6 +660,7 @@ bool UCookOnTheFlyServer::StartCookOnTheFly(FCookOnTheFlyStartupOptions InCookOn
 		UE_LOG(LogCook, Fatal, TEXT("Failed starting Unreal Network file server!"));
 	}
 	BeginCookEditorSystems();
+	InitializePollables();
 
 	const bool bInitialized = CookOnTheFlyRequestManager->Initialize();
 	return bInitialized;
@@ -1227,13 +1238,13 @@ bool UCookOnTheFlyServer::RequestPackage(const FName& StandardPackageFName, cons
 	return RequestPackage(StandardPackageFName, PlatformManager->GetSessionPlatforms(), bForceFrontOfQueue);
 }
 
-uint32 UCookOnTheFlyServer::TickCookByTheBook(const float TimeSlice, uint32& CookedPackageCount, ECookTickFlags TickFlags, int32 InMaxNumPackagesToSave)
+uint32 UCookOnTheFlyServer::TickCookByTheBook(const float TimeSlice, ECookTickFlags TickFlags)
 {
 	check(IsCookByTheBookMode());
 
 	LLM_SCOPE_BYTAG(Cooker);
 	COOK_STAT(FScopedDurationTimer TickTimer(DetailedCookStats::TickCookOnTheSideTimeSec));
-	UE::Cook::FTickStackData StackData(TimeSlice, IsRealtimeMode(), TickFlags, InMaxNumPackagesToSave);
+	UE::Cook::FTickStackData StackData(TimeSlice, TickFlags);
 
 	TickMainCookLoop(StackData);
 
@@ -1245,25 +1256,23 @@ uint32 UCookOnTheFlyServer::TickCookByTheBook(const float TimeSlice, uint32& Coo
 	}
 	else if (IsInSession() && StackData.bCookComplete)
 	{
-		UpdateDisplay(TickFlags, true /* bForceDisplay */);
+		UpdateDisplay(StackData, true /* bForceDisplay */);
 		CookByTheBookFinished();
 	}
-	CookedPackageCount += StackData.CookedPackageCount;
 	return StackData.ResultFlags;
 }
 
-uint32 UCookOnTheFlyServer::TickCookOnTheFly(const float TimeSlice, uint32& CookedPackageCount, ECookTickFlags TickFlags, int32 InMaxNumPackagesToSave)
+uint32 UCookOnTheFlyServer::TickCookOnTheFly(const float TimeSlice, ECookTickFlags TickFlags)
 {
 	check(IsCookOnTheFlyMode());
 
 	LLM_SCOPE_BYTAG(Cooker);
 	COOK_STAT(FScopedDurationTimer TickTimer(DetailedCookStats::TickCookOnTheSideTimeSec));
-	UE::Cook::FTickStackData StackData(TimeSlice, IsRealtimeMode(), TickFlags, InMaxNumPackagesToSave);
+	UE::Cook::FTickStackData StackData(TimeSlice, TickFlags);
 
 	TickNetwork();
 	TickMainCookLoop(StackData);
 
-	CookedPackageCount += StackData.CookedPackageCount;
 	return StackData.ResultFlags;
 }
 
@@ -1332,10 +1341,6 @@ void UCookOnTheFlyServer::TickMainCookLoop(UE::Cook::FTickStackData& StackData)
 			bContinueTick = false;
 			StackData.bCookComplete = true;
 			break;
-		case ECookAction::Cancel:
-			StackData.bCookCancelled = true;
-			bContinueTick = false;
-			break;
 		default:
 			check(false);
 			break;
@@ -1347,8 +1352,8 @@ void UCookOnTheFlyServer::TickCookStatus(UE::Cook::FTickStackData& StackData)
 {
 	UE_SCOPED_COOKTIMER(TickCookStatus);
 
-	// TODO: Calculate CurrentTime once per status update and share it with e.g. StackData.Timer
 	double CurrentTime = FPlatformTime::Seconds();
+	StackData.LoopStartTime = CurrentTime;
 	if (LastCookableObjectTickTime + TickCookableObjectsFrameTime <= CurrentTime)
 	{
 		UE_SCOPED_COOKTIMER(TickCookableObjects);
@@ -1356,15 +1361,7 @@ void UCookOnTheFlyServer::TickCookStatus(UE::Cook::FTickStackData& StackData)
 		LastCookableObjectTickTime = CurrentTime;
 	}
 
-	UpdateDisplay(StackData.TickFlags, false /* bForceDisplay */);
-	// prevent autosave from happening until we are finished cooking
-	// causes really bad hitches
-	if (GUnrealEd)
-	{
-		const static float SecondsWarningTillAutosave = 10.0f;
-		GUnrealEd->GetPackageAutoSaver().ForceMinimumTimeTillAutoSave(SecondsWarningTillAutosave);
-	}
-
+	UpdateDisplay(StackData, false /* bForceDisplay */);
 	ProcessUnsolicitedPackages();
 	UpdatePackageFilter();
 	PumpExternalRequests(StackData.Timer);
@@ -1559,11 +1556,22 @@ void UCookOnTheFlyServer::SetLoadBusy(bool bInLoadBusy)
 	}
 }
 
-void UCookOnTheFlyServer::UpdateDisplay(ECookTickFlags TickFlags, bool bForceDisplay)
+void UCookOnTheFlyServer::SetIdleStatus(UE::Cook::FTickStackData& StackData, EIdleStatus InStatus)
+{
+	if (InStatus == IdleStatus)
+	{
+		return;
+	}
+
+	IdleStatusStartTime = StackData.LoopStartTime;
+	IdleStatus = InStatus;
+}
+
+void UCookOnTheFlyServer::UpdateDisplay(UE::Cook::FTickStackData& StackData, bool bForceDisplay)
 {
 	using namespace UE::Cook;
 
-	const double CurrentTime = FPlatformTime::Seconds();
+	const double CurrentTime = StackData.LoopStartTime;
 	const float DeltaProgressDisplayTime = CurrentTime - LastProgressDisplayTime;
 	const int32 CookedPackagesCount = PackageDatas->GetNumCooked();
 	const int32 CookPendingCount = ExternalRequests->GetNumRequests() + PackageDatas->GetMonitor().GetNumInProgress();
@@ -1571,7 +1579,7 @@ void UCookOnTheFlyServer::UpdateDisplay(ECookTickFlags TickFlags, bool bForceDis
 		(DeltaProgressDisplayTime >= GCookProgressUpdateTime && CookPendingCount != 0 &&
 			(LastCookedPackagesCount != CookedPackagesCount || LastCookPendingCount != CookPendingCount || DeltaProgressDisplayTime > GCookProgressRepeatTime)))
 	{
-		UE_CLOG(!(TickFlags & ECookTickFlags::HideProgressDisplay) && (GCookProgressDisplay & (int32)ECookProgressDisplayMode::RemainingPackages),
+		UE_CLOG(!(StackData.TickFlags & ECookTickFlags::HideProgressDisplay) && (GCookProgressDisplay & (int32)ECookProgressDisplayMode::RemainingPackages),
 			LogCook,
 			Display,
 			TEXT("Cooked packages %d Packages Remain %d Total %d"),
@@ -1590,22 +1598,31 @@ void UCookOnTheFlyServer::UpdateDisplay(ECookTickFlags TickFlags, bool bForceDis
 #if PLATFORMFILETRACE_ENABLED
 		OpenFileHandles = FPlatformFileTrace::GetOpenFileHandleCount();
 #endif
-		UE_CLOG(!(TickFlags & ECookTickFlags::HideProgressDisplay) && (GCookProgressDisplay != (int32) ECookProgressDisplayMode::Nothing),
-			LogCook, Display,
-			TEXT("Cook Diagnostics: OpenFileHandles=%d, VirtualMemory=%dMiB"),
-			OpenFileHandles, FPlatformMemory::GetStats().UsedVirtual / 1024 / 1024);
-		LastDiagnosticsDisplayTime = CurrentTime;
-
+		bool bCookOnTheFlyShouldDisplay = false;
 		if (IsCookOnTheFlyMode() && (IsCookingInEditor() == false))
 		{
 			// Dump stats in CookOnTheFly, but only if there is new data
 			static uint64 LastNumLoadedAndSaved = 0;
 			if (StatLoadedPackageCount + StatSavedPackageCount != LastNumLoadedAndSaved)
 			{
+				bCookOnTheFlyShouldDisplay = true;
 				LastNumLoadedAndSaved = StatLoadedPackageCount + StatSavedPackageCount;
-				DumpStats();
 			}
 		}
+		if (!IsCookOnTheFlyMode() || bCookOnTheFlyShouldDisplay)
+		{
+			UE_CLOG(!(StackData.TickFlags & ECookTickFlags::HideProgressDisplay) && (GCookProgressDisplay != (int32)ECookProgressDisplayMode::Nothing),
+				LogCook, Display,
+				TEXT("Cook Diagnostics: OpenFileHandles=%d, VirtualMemory=%dMiB"),
+				OpenFileHandles, FPlatformMemory::GetStats().UsedVirtual / 1024 / 1024);
+		}
+		if (bCookOnTheFlyShouldDisplay)
+		{
+			DumpStats();
+		}
+
+		LastDiagnosticsDisplayTime = CurrentTime;
+
 	}
 }
 
@@ -1662,7 +1679,19 @@ UCookOnTheFlyServer::FPollableQueueKey::FPollableQueueKey(TRefCountPtr<FPollable
 void UCookOnTheFlyServer::FPollable::Trigger(UCookOnTheFlyServer& COTFS)
 {
 	FScopeLock PollablesScopeLock(&COTFS.PollablesLock);
+	if (COTFS.bPollablesInTick)
+	{
+		FPollableQueueKey DeferredTrigger(this);
+		DeferredTrigger.NextTimeSeconds = 0.;
+		COTFS.PollablesDeferredTriggers.Add(MoveTemp(DeferredTrigger));
+		return;
+	}
+
+	TriggerInternal(COTFS);
+}
 
+void UCookOnTheFlyServer::FPollable::TriggerInternal(UCookOnTheFlyServer& COTFS)
+{
 	FPollableQueueKey* KeyInQueue = COTFS.Pollables.FindByPredicate(
 		[this](const FPollableQueueKey& Existing) { return Existing.Pollable.GetReference() == this; });
 	if (ensure(KeyInQueue))
@@ -1690,24 +1719,37 @@ void UCookOnTheFlyServer::FPollable::RunNow(UCookOnTheFlyServer& COTFS)
 {
 	FScopeLock PollablesScopeLock(&COTFS.PollablesLock);
 
-	UE::Cook::FTickStackData StackData(MAX_flt, COTFS.IsRealtimeMode(), ECookTickFlags::None, MAX_int32);
+	UE::Cook::FTickStackData StackData(MAX_flt, ECookTickFlags::None);
 	PollFunction(StackData);
 
+	double CurrentTime = FPlatformTime::Seconds();
+	if (COTFS.bPollablesInTick)
+	{
+		FPollableQueueKey DeferredTrigger(this);
+		DeferredTrigger.NextTimeSeconds = CurrentTime;
+		COTFS.PollablesDeferredTriggers.Add(MoveTemp(DeferredTrigger));
+		return;
+	}
+
+	RunNowInternal(COTFS, CurrentTime);
+}
+
+void UCookOnTheFlyServer::FPollable::RunNowInternal(UCookOnTheFlyServer & COTFS, double TimeLastRun)
+{
 	FPollableQueueKey* KeyInQueue = COTFS.Pollables.FindByPredicate(
 		[this](const FPollableQueueKey& Existing) { return Existing.Pollable.GetReference() == this; });
 	if (ensure(KeyInQueue))
 	{
 		FPollableQueueKey LocalQueueKey;
 		LocalQueueKey.Pollable = MoveTemp(KeyInQueue->Pollable);
-		double CurrentTime = FPlatformTime::Seconds();
-		LocalQueueKey.NextTimeSeconds = CurrentTime + this->PeriodSeconds;
-		this->NextTimeIdleSeconds = CurrentTime + this->PeriodIdleSeconds;
+		LocalQueueKey.NextTimeSeconds = TimeLastRun + this->PeriodSeconds;
+		this->NextTimeIdleSeconds = TimeLastRun + this->PeriodIdleSeconds;
 
 		int32 Index = KeyInQueue - COTFS.Pollables.GetData();
+		COTFS.PollNextTimeSeconds = FMath::Min(LocalQueueKey.NextTimeSeconds, COTFS.PollNextTimeSeconds);
+		COTFS.PollNextTimeIdleSeconds = FMath::Min(this->NextTimeIdleSeconds, COTFS.PollNextTimeIdleSeconds);
 		COTFS.Pollables.HeapRemoveAt(Index, false /* bAllowShrinking */);
 		COTFS.Pollables.HeapPush(MoveTemp(LocalQueueKey));
-		COTFS.PollNextTimeSeconds = FMath::Min(LocalQueueKey.NextTimeSeconds, COTFS.PollNextTimeSeconds);
-		COTFS.PollNextTimeIdleSeconds = FMath::Min(LocalQueueKey.Pollable->NextTimeIdleSeconds, COTFS.PollNextTimeIdleSeconds);
 	}
 }
 
@@ -1722,7 +1764,10 @@ void UCookOnTheFlyServer::FPollable::RunDuringPump(UE::Cook::FTickStackData& Sta
 void UCookOnTheFlyServer::PumpPollables(UE::Cook::FTickStackData& StackData, bool bIsIdle)
 {
 	UE_SCOPED_HIERARCHICAL_COOKTIMER(PumpPollables);
-	FScopeLock PollablesScopeLock(&PollablesLock);
+	{
+		FScopeLock PollablesScopeLock(&PollablesLock);
+		bPollablesInTick = true;
+	}
 
 	int32 NumPollables = Pollables.Num();
 	if (NumPollables == 0)
@@ -1732,7 +1777,7 @@ void UCookOnTheFlyServer::PumpPollables(UE::Cook::FTickStackData& StackData, boo
 		return;
 	}
 
-	double CurrentTime = FPlatformTime::Seconds();
+	double CurrentTime = StackData.LoopStartTime;
 	if (!bIsIdle)
 	{
 		// To avoid an infinite loop, we keep the popped pollables in a separate list to readd afterwards
@@ -1792,6 +1837,23 @@ void UCookOnTheFlyServer::PumpPollables(UE::Cook::FTickStackData& StackData, boo
 			Pollables.Heapify();
 		}
 	}
+
+	{
+		FScopeLock PollablesScopeLock(&PollablesLock);
+		for (FPollableQueueKey& QueueKey : PollablesDeferredTriggers)
+		{
+			if (QueueKey.NextTimeSeconds == 0)
+			{
+				QueueKey.Pollable->TriggerInternal(*this);
+			}
+			else
+			{
+				QueueKey.Pollable->RunNowInternal(*this, QueueKey.NextTimeSeconds);
+			}
+		}
+		PollablesDeferredTriggers.Reset();
+		bPollablesInTick = false;
+	}
 }
 
 void UCookOnTheFlyServer::PollFlushRenderingCommands()
@@ -1820,39 +1882,61 @@ TRefCountPtr<UCookOnTheFlyServer::FPollable> UCookOnTheFlyServer::CreatePollable
 TRefCountPtr<UCookOnTheFlyServer::FPollable> UCookOnTheFlyServer::CreatePollableTriggerGC()
 {
 	bool bTestCook = IsCookFlagSet(ECookInitializationFlags::TestCook);
-	if (PackagesPerGC > 0 || bTestCook)
+	if (PackagesPerGC == 0 && !bTestCook && !IsCookOnTheFlyMode())
+	{
+		return TRefCountPtr<FPollable>();
+	}
+
+	float PeriodSeconds = MAX_flt;
+	float IdlePeriodSeconds = MAX_flt;
+	constexpr float SecondsPerPackage = .01f;
+	if (bTestCook)
+	{
+		PeriodSeconds = FMath::Min(PeriodSeconds, 50 * SecondsPerPackage);
+	}
+	if (PackagesPerGC > 0)
 	{
 		// PackagesPerGC is usually used only to debug; max memory counts are commonly used instead
 		// Since it's not commonly used, we make a concession to support it: we check on a timer rather than checking after every saved package.
 		// For large values, check less frequently.
-		uint32 NumPackagesForTimeEstimate = UINT32_MAX;
-		if (bTestCook)
-		{
-			NumPackagesForTimeEstimate = FMath::Min(NumPackagesForTimeEstimate, 50U);
-		}
-		if (PackagesPerGC > 0)
-		{
-			NumPackagesForTimeEstimate = FMath::Min(NumPackagesForTimeEstimate, PackagesPerGC);
-		}
-		float PeriodSeconds = FMath::Min(NumPackagesForTimeEstimate * .01f, 10.f);
-		return TRefCountPtr<FPollable>(new FPollable(TEXT("TimeForGC"), PeriodSeconds, PeriodSeconds,
-			[this](UE::Cook::FTickStackData& StackData) { PollPackagesPerGC(StackData); }));
+		PeriodSeconds = FMath::Min(PeriodSeconds, PackagesPerGC * SecondsPerPackage);
 	}
-	return TRefCountPtr<FPollable>();
+	if (IsCookOnTheFlyMode())
+	{
+		PeriodSeconds = FMath::Min(PeriodSeconds, 10.f);
+		IdlePeriodSeconds = FMath::Min(IdlePeriodSeconds, 0.1f);
+	}
+	IdlePeriodSeconds = FMath::Min(IdlePeriodSeconds, PeriodSeconds);
+
+	return TRefCountPtr<FPollable>(new FPollable(TEXT("TimeForGC"), PeriodSeconds, IdlePeriodSeconds,
+		[this](UE::Cook::FTickStackData& StackData) { PollTimeForGC(StackData); }));
 }
 
-void UCookOnTheFlyServer::PollPackagesPerGC(UE::Cook::FTickStackData& StackData)
+void UCookOnTheFlyServer::PollTimeForGC(UE::Cook::FTickStackData& StackData)
 {
 	if (IsCookFlagSet(ECookInitializationFlags::TestCook))
 	{
-		StackData.ResultFlags |= COSR_RequiresGC;
+		StackData.ResultFlags |= COSR_RequiresGC | COSR_YieldTick;
+		return;
 	}
-	else if (PackagesPerGC > 0 && CookedPackageCountSinceLastGC > PackagesPerGC)
+	if (PackagesPerGC > 0 && CookedPackageCountSinceLastGC > PackagesPerGC)
 	{
 		// if we are waiting on things to cache then ignore the PackagesPerGC
 		if (!bSaveBusy)
 		{
-			StackData.ResultFlags |= COSR_RequiresGC | COSR_RequiresGC_PackageCount;
+			StackData.ResultFlags |= COSR_RequiresGC | COSR_RequiresGC_PackageCount | COSR_YieldTick;
+			return;
+		}
+	}
+	if (IsCookOnTheFlyMode())
+	{
+		double CurrentTime = FPlatformTime::Seconds();
+		if (IdleStatus == EIdleStatus::Done &&
+			CurrentTime - IdleStatusStartTime > GetIdleTimeToGC() &&
+			IdleStatusStartTime > GetLastGCTime())
+		{
+			StackData.ResultFlags |= COSR_RequiresGC | COSR_RequiresGC_IdleTimer | COSR_YieldTick;
+			return;
 		}
 	}
 }
@@ -1860,9 +1944,17 @@ void UCookOnTheFlyServer::PollPackagesPerGC(UE::Cook::FTickStackData& StackData)
 namespace UE::Cook
 {
 
-static void CBTBProcessDeferredCommands()
+static void ProcessDeferredCommands(UCookOnTheFlyServer& COTFS)
 {
-	UE_SCOPED_COOKTIMER_AND_DURATION(CookByTheBook_ProcessDeferredCommands, DetailedCookStats::TickLoopProcessDeferredCommandsTimeSec);
+#if OUTPUT_COOKTIMING
+	TOptional<FScopedDurationTimer> CBTBScopedDurationTimer;
+	if (COTFS.IsCookByTheBookMode())
+	{
+		CBTBScopedDurationTimer.Emplace(DetailedCookStats::TickLoopProcessDeferredCommandsTimeSec);
+	}
+#endif
+	UE_SCOPED_COOKTIMER(ProcessDeferredCommands);
+
 #if PLATFORM_MAC
 	// On Mac we need to process Cocoa events so that the console window for CookOnTheFlyServer is interactive
 	FPlatformApplicationMisc::PumpMessages(true);
@@ -1886,13 +1978,13 @@ static void CBTBTickCommandletStats()
 	FStats::TickCommandletStats();
 }
 
-static void CBTBTickShaderCompilingManager()
+static void TickShaderCompilingManager()
 {
 	UE_SCOPED_COOKTIMER_AND_DURATION(CookByTheBook_ShaderProcessAsync, DetailedCookStats::TickLoopShaderProcessAsyncResultsTimeSec);
 	GShaderCompilingManager->ProcessAsyncResults(true, false);
 }
 
-static void CBTBTickAssetRegistry()
+static void TickAssetRegistry()
 {
 	UE_SCOPED_COOKTIMER(CookByTheBook_TickAssetRegistry);
 	FAssetRegistryModule::TickAssetRegistry(-1.0f);
@@ -1906,21 +1998,36 @@ void UCookOnTheFlyServer::InitializePollables()
 
 	Pollables.Reset();
 
-	if (IsCookByTheBookMode() && !IsCookingInEditor())
+	QueuedCancelPollable = new FPollable(TEXT("QueuedCancel"), FPollable::EManualTrigger(), [this](FTickStackData& StackData) { PollQueuedCancel(StackData); });
+	Pollables.Emplace(QueuedCancelPollable);
+	if (!IsCookingInEditor())
 	{
+		Pollables.Emplace(new FPollable(TEXT("AssetRegistry"), 60.f, 5.f, [](FTickStackData&) { TickAssetRegistry(); }));
+		if (TRefCountPtr<FPollable> Pollable = CreatePollableTriggerGC())
+		{
+			Pollables.Emplace(MoveTemp(Pollable));
+		}
+		Pollables.Emplace(new FPollable(TEXT("ProcessDeferredCommands"), 60.f, 5.f, [this](FTickStackData&) { ProcessDeferredCommands(*this); }));
+		Pollables.Emplace(new FPollable(TEXT("ShaderCompilingManager"), 60.f, 5.f, [this](FTickStackData&) { TickShaderCompilingManager(); }));
 		Pollables.Emplace(new FPollable(TEXT("FlushRenderingCommands"), 60.f, 5.f, [this](FTickStackData&) { PollFlushRenderingCommands(); }));
 		if (TRefCountPtr<FPollable> Pollable = CreatePollableLLM())
 		{
 			Pollables.Emplace(MoveTemp(Pollable));
 		}
-		if (TRefCountPtr<FPollable> Pollable = CreatePollableTriggerGC())
+	}
+	if (IsCookByTheBookMode())
+	{
+		if (!IsCookingInEditor())
 		{
-			Pollables.Emplace(MoveTemp(Pollable));
+			Pollables.Emplace(new FPollable(TEXT("CommandletStats"), 60.f, 5.f, [](FTickStackData&) { CBTBTickCommandletStats(); }));
 		}
-		Pollables.Emplace(new FPollable(TEXT("ProcessDeferredCommands"), 60.f, 5.f, [this](FTickStackData&) { CBTBProcessDeferredCommands(); }));
-		Pollables.Emplace(new FPollable(TEXT("CommandletStats"), 60.f, 5.f, [this](FTickStackData&) { CBTBTickCommandletStats(); }));
-		Pollables.Emplace(new FPollable(TEXT("ShaderCompilingManager"), 60.f, 5.f, [this](FTickStackData&) { CBTBTickShaderCompilingManager(); }));
-		Pollables.Emplace(new FPollable(TEXT("AssetRegistry"), 60.f, 5.f, [this](FTickStackData&) { CBTBTickAssetRegistry(); }));
+	}
+	else if (IsCookOnTheFlyMode())
+	{
+		RecompileRequestsPollable = new FPollable(TEXT("RecompileShaderRequests"), FPollable::EManualTrigger(), [this](FTickStackData&) { TickRecompileShaderRequestsPrivate(); });
+		Pollables.Add(FPollableQueueKey(RecompileRequestsPollable));
+		Pollables.Emplace(new FPollable(TEXT("RequestManager"), 0.5f, 0.5f, [this](FTickStackData&) { TickRequestManager(); }));
+
 	}
 	Pollables.Heapify();
 
@@ -1944,30 +2051,28 @@ void UCookOnTheFlyServer::WaitForAsync(UE::Cook::FTickStackData& StackData)
 
 UCookOnTheFlyServer::ECookAction UCookOnTheFlyServer::DecideNextCookAction(UE::Cook::FTickStackData& StackData)
 {
-	if (bCancelSession)
-	{
-		return ECookAction::Cancel;
-	}
-
-	if (StackData.ResultFlags & COSR_RequiresGC)
+	if (StackData.ResultFlags & COSR_YieldTick)
 	{
-		// if we just cooked a map then don't process anything the rest of this tick
+		// Yielding on demand does not impact idle status
 		return ECookAction::YieldTick;
 	}
 
-	double CurrentTime = FPlatformTime::Seconds();
+	double CurrentTime = StackData.LoopStartTime;
 	if (StackData.Timer.IsTimeUp(CurrentTime))
 	{
+		// Timeup does not impact idle status
 		return ECookAction::YieldTick;
 	}
 	else if (CurrentTime >= PollNextTimeSeconds)
 	{
+		// Polling does not impact idle status
 		return ECookAction::Poll;
 	}
 
 	UE::Cook::FRequestQueue& RequestQueue = PackageDatas->GetRequestQueue();
 	if (RequestQueue.GetRequestClusters().Num() != 0 || RequestQueue.GetUnclusteredRequests().Num() != 0)
 	{
+		SetIdleStatus(StackData, EIdleStatus::Active);
 		return ECookAction::Request;
 	}
 
@@ -1976,18 +2081,22 @@ UCookOnTheFlyServer::ECookAction UCookOnTheFlyServer::DecideNextCookAction(UE::C
 	{
 		if (Monitor.GetNumUrgent(UE::Cook::EPackageState::Save) > 0)
 		{
+			SetIdleStatus(StackData, EIdleStatus::Active);
 			return ECookAction::Save;
 		}
 		else if (Monitor.GetNumUrgent(UE::Cook::EPackageState::LoadPrepare) > 0)
 		{
+			SetIdleStatus(StackData, EIdleStatus::Active);
 			return ECookAction::Load;
 		}
 		else if (Monitor.GetNumUrgent(UE::Cook::EPackageState::LoadReady) > 0)
 		{
+			SetIdleStatus(StackData, EIdleStatus::Active);
 			return ECookAction::Load;
 		}
 		else if (Monitor.GetNumUrgent(UE::Cook::EPackageState::Request) > 0)
 		{
+			SetIdleStatus(StackData, EIdleStatus::Active);
 			return ECookAction::Request;
 		}
 		else
@@ -2000,6 +2109,7 @@ UCookOnTheFlyServer::ECookAction UCookOnTheFlyServer::DecideNextCookAction(UE::C
 	bool bSaveAvailable = ((!bSaveBusy) & (NumSaves > 0)) != 0;
 	if (bSaveAvailable & (NumSaves > static_cast<int32>(DesiredSaveQueueLength)))
 	{
+		SetIdleStatus(StackData, EIdleStatus::Active);
 		return ECookAction::SaveLimited;
 	}
 
@@ -2007,30 +2117,36 @@ UCookOnTheFlyServer::ECookAction UCookOnTheFlyServer::DecideNextCookAction(UE::C
 	bool bLoadAvailable = ((!bLoadBusy) & (NumLoads > 0)) != 0;
 	if (bLoadAvailable & (NumLoads > static_cast<int32>(DesiredLoadQueueLength)))
 	{
+		SetIdleStatus(StackData, EIdleStatus::Active);
 		return ECookAction::LoadLimited;
 	}
 
 	if (!RequestQueue.IsReadyRequestsEmpty())
 	{
+		SetIdleStatus(StackData, EIdleStatus::Active);
 		return ECookAction::Request;
 	}
 
 	if (bSaveAvailable)
 	{
+		SetIdleStatus(StackData, EIdleStatus::Active);
 		return ECookAction::Save;
 	}
 
 	if (bLoadAvailable)
 	{
+		SetIdleStatus(StackData, EIdleStatus::Active);
 		return ECookAction::Load;
 	}
 
 	if (NumSaves > 0 && CurrentTime >= SaveBusyRetryTimeSeconds)
 	{
+		SetIdleStatus(StackData, EIdleStatus::Active);
 		return ECookAction::Save;
 	}
 	if (NumLoads > 0 && CurrentTime >= LoadBusyRetryTimeSeconds)
 	{
+		SetIdleStatus(StackData, EIdleStatus::Active);
 		return ECookAction::Load;
 	}
 
@@ -2038,18 +2154,31 @@ UCookOnTheFlyServer::ECookAction UCookOnTheFlyServer::DecideNextCookAction(UE::C
 	{
 		if (CurrentTime >= PollNextTimeIdleSeconds)
 		{
+			// Polling does not impact idle status
 			return ECookAction::PollIdle;
 		}
 		else if (IsRealtimeMode() || IsCookOnTheFlyMode())
 		{
+			SetIdleStatus(StackData, EIdleStatus::Idle);
 			return ECookAction::YieldTick;
 		}
 		else
 		{
+			SetIdleStatus(StackData, EIdleStatus::Idle);
 			return ECookAction::WaitForAsync;
 		}
 	}
 
+	if (IsCookOnTheFlyMode())
+	{
+		if (CurrentTime >= PollNextTimeIdleSeconds)
+		{
+			// Polling does not impact idle status
+			return ECookAction::PollIdle;
+		}
+	}
+
+	SetIdleStatus(StackData, EIdleStatus::Done);
 	return ECookAction::Done;
 }
 
@@ -2234,7 +2363,7 @@ void UCookOnTheFlyServer::PumpLoads(UE::Cook::FTickStackData& StackData, uint32
 
 		if (HasExceededMaxMemory())
 		{
-			StackData.ResultFlags |= COSR_RequiresGC | COSR_RequiresGC_OOM;
+			StackData.ResultFlags |= COSR_RequiresGC | COSR_RequiresGC_OOM | COSR_YieldTick;
 			return;
 		}
 	}
@@ -3312,7 +3441,7 @@ void UCookOnTheFlyServer::PumpSaves(UE::Cook::FTickStackData& StackData, uint32
 	{
 		if (HasExceededMaxMemory())
 		{
-			StackData.ResultFlags |= COSR_RequiresGC | COSR_RequiresGC_OOM;
+			StackData.ResultFlags |= COSR_RequiresGC | COSR_RequiresGC_OOM | COSR_YieldTick;
 		}
 	};
 
@@ -3441,7 +3570,7 @@ void UCookOnTheFlyServer::PumpSaves(UE::Cook::FTickStackData& StackData, uint32
 			// GC is required
 			if (PackageData.GeneratorPackageRequiresGC())
 			{
-				StackData.ResultFlags |= COSR_RequiresGC;
+				StackData.ResultFlags |= COSR_RequiresGC | COSR_YieldTick;
 				SaveQueue.AddFront(&PackageData);
 				return;
 			}
@@ -3595,7 +3724,7 @@ void UCookOnTheFlyServer::TickPrecacheObjectsForPlatforms(const float TimeSlice,
 	SCOPE_CYCLE_COUNTER(STAT_TickPrecacheCooking);
 
 
-	UE::Cook::FCookerTimer Timer(TimeSlice, true);
+	UE::Cook::FCookerTimer Timer(TimeSlice);
 
 	if (LastUpdateTick > 50 ||
 		((CachedMaterialsToCacheArray.Num() == 0) && (CachedTexturesToCacheArray.Num() == 0)))
@@ -4264,7 +4393,7 @@ void UCookOnTheFlyServer::TickRequestManager()
 	}
 }
 
-void UCookOnTheFlyServer::TickRecompileShaderRequests()
+void UCookOnTheFlyServer::TickRecompileShaderRequestsPrivate()
 {
 	// try to pull off a request
 	UE::Cook::FRecompileShaderRequest RecompileShaderRequest;
@@ -4273,6 +4402,10 @@ void UCookOnTheFlyServer::TickRecompileShaderRequests()
 		RecompileShadersForRemote(RecompileShaderRequest.RecompileArguments, GetSandboxDirectory(RecompileShaderRequest.RecompileArguments.PlatformName));
 		RecompileShaderRequest.CompletionCallback();
 	}
+	if (PackageTracker->RecompileRequests.HasItems())
+	{
+		RecompileRequestsPollable->Trigger(*this);
+	}
 }
 
 bool UCookOnTheFlyServer::HasRecompileShaderRequests() const 
@@ -4417,7 +4550,6 @@ void UCookOnTheFlyServer::SaveCookedPackage(UE::Cook::FSaveCookedPackageContext&
 	}
 
 	Context.FinishPackage();
-	Context.StackData.Timer.SavedPackage();
 }
 
 namespace UE::Cook
@@ -4632,7 +4764,6 @@ void FSaveCookedPackageContext::FinishPlatform()
 		}
 		else
 		{
-			++StackData.CookedPackageCount;
 			++COTFS.CookedPackageCountSinceLastGC;
 			StackData.ResultFlags |= UCookOnTheFlyServer::COSR_CookedPackage;
 		}
@@ -4856,6 +4987,9 @@ void UCookOnTheFlyServer::Initialize( ECookMode::Type DesiredCookMode, ECookInit
 	PollNextTimeIdleSeconds = MAX_flt;
 	WaitForAsyncSleepSeconds = 1.0f;
 
+	IdleStatus = EIdleStatus::Done;
+	IdleStatusStartTime = FPlatformTime::Seconds();
+
 	{
 		const FConfigSection* CacheSettings = GConfig->GetSectionPrivate(TEXT("CookPlatformDataCacheSettings"), false, true, GEditorIni);
 		if ( CacheSettings )
@@ -7390,7 +7524,6 @@ void UCookOnTheFlyServer::ShutdownCookSession()
 
 	ConditionalUninstallImportBehaviorCallback();
 
-	bCancelSession = false;
 	if (IsCookByTheBookMode())
 	{
 		CookByTheBookOptions->ClearSessionData();
@@ -7476,10 +7609,16 @@ void UCookOnTheFlyServer::QueueCancelCookByTheBook()
 {
 	if (IsCookByTheBookMode() && IsInSession())
 	{
-		bCancelSession = true;
+		QueuedCancelPollable->Trigger(*this);
 	}
 }
 
+void UCookOnTheFlyServer::PollQueuedCancel(UE::Cook::FTickStackData& StackData)
+{
+	StackData.bCookCancelled = true;
+	StackData.ResultFlags |= COSR_YieldTick;
+}
+
 void UCookOnTheFlyServer::CancelCookByTheBook()
 {
 	check(IsCookByTheBookMode());
@@ -7488,6 +7627,8 @@ void UCookOnTheFlyServer::CancelCookByTheBook()
 	{
 		CancelAllQueues();
 		ShutdownCookSession();
+		UE::Cook::FTickStackData StackData(MAX_flt, ECookTickFlags::None);
+		SetIdleStatus(StackData, EIdleStatus::Done);
 	} 
 }
 
diff --git a/Engine/Source/Editor/UnrealEd/Private/Cooker/CookTypes.cpp b/Engine/Source/Editor/UnrealEd/Private/Cooker/CookTypes.cpp
--- a/Engine/Source/Editor/UnrealEd/Private/Cooker/CookTypes.cpp
+++ b/Engine/Source/Editor/UnrealEd/Private/Cooker/CookTypes.cpp
@@ -10,19 +10,18 @@
 namespace UE::Cook
 {
 
-FCookerTimer::FCookerTimer(const float& InTimeSlice, bool bInIsRealtimeMode, int InMaxNumPackagesToSave)
-	: bIsRealtimeMode(bInIsRealtimeMode), StartTime(FPlatformTime::Seconds()), TimeSlice(InTimeSlice)
-	, MaxNumPackagesToSave(InMaxNumPackagesToSave), NumPackagesSaved(0)
+FCookerTimer::FCookerTimer(float InTimeSlice)
+	: StartTime(FPlatformTime::Seconds()), TimeSlice(InTimeSlice)
 {
 }
 
 FCookerTimer::FCookerTimer(EForever)
-	:bIsRealtimeMode(false), StartTime(FPlatformTime::Seconds()), TimeSlice(ForeverTimeSlice), MaxNumPackagesToSave(TNumericLimits<int>::Max()), NumPackagesSaved(0)
+	: FCookerTimer(MAX_flt)
 {
 }
 
 FCookerTimer::FCookerTimer(ENoWait)
-	:bIsRealtimeMode(false), StartTime(FPlatformTime::Seconds()), TimeSlice(ZeroTimeSlice), MaxNumPackagesToSave(TNumericLimits<int>::Max()), NumPackagesSaved(0)
+	: FCookerTimer(0.0f)
 {
 }
 
@@ -33,7 +32,7 @@ double FCookerTimer::GetTimeTillNow() const
 
 double FCookerTimer::GetEndTimeSeconds() const
 {
-	return bIsRealtimeMode ? StartTime + TimeSlice : MAX_flt;
+	return FMath::Min(StartTime + TimeSlice,  MAX_flt);
 }
 
 bool FCookerTimer::IsTimeUp() const
@@ -43,23 +42,7 @@ bool FCookerTimer::IsTimeUp() const
 
 bool FCookerTimer::IsTimeUp(double CurrentTimeSeconds) const
 {
-	if (bIsRealtimeMode)
-	{
-		if ((CurrentTimeSeconds - StartTime) > TimeSlice)
-		{
-			return true;
-		}
-	}
-	if (NumPackagesSaved >= MaxNumPackagesToSave)
-	{
-		return true;
-	}
-	return false;
-}
-
-void FCookerTimer::SavedPackage()
-{
-	++NumPackagesSaved;
+	return CurrentTimeSeconds - StartTime > TimeSlice;
 }
 
 double FCookerTimer::GetTimeRemain() const
@@ -67,9 +50,6 @@ double FCookerTimer::GetTimeRemain() const
 	return TimeSlice - (FPlatformTime::Seconds() - StartTime);
 }
 
-float FCookerTimer::ZeroTimeSlice = 0.0f;
-float FCookerTimer::ForeverTimeSlice = TNumericLimits<float>::Max();
-
 static uint32 SchedulerThreadTlsSlot = 0;
 void InitializeTls()
 {
diff --git a/Engine/Source/Editor/UnrealEd/Private/Cooker/CookTypes.h b/Engine/Source/Editor/UnrealEd/Private/Cooker/CookTypes.h
--- a/Engine/Source/Editor/UnrealEd/Private/Cooker/CookTypes.h
+++ b/Engine/Source/Editor/UnrealEd/Private/Cooker/CookTypes.h
@@ -134,7 +134,7 @@ namespace UE::Cook
 			NoWait
 		};
 
-		FCookerTimer(const float& InTimeSlice, bool bInIsRealtimeMode, int InMaxNumPackagesToSave = 50);
+		FCookerTimer(float InTimeSlice);
 		FCookerTimer(EForever);
 		FCookerTimer(ENoWait);
 
@@ -142,28 +142,20 @@ namespace UE::Cook
 		double GetEndTimeSeconds() const;
 		bool IsTimeUp() const;
 		bool IsTimeUp(double CurrentTimeSeconds) const;
-		void SavedPackage();
 		double GetTimeRemain() const;
 
 	public:
-		const bool bIsRealtimeMode;
 		const double StartTime;
-		const float& TimeSlice;
-		const int MaxNumPackagesToSave; // maximum packages to save before exiting tick (this should never really hit unless we are not using realtime mode)
-		int NumPackagesSaved;
-
-	private:
-		static float ZeroTimeSlice;
-		static float ForeverTimeSlice;
+		const float TimeSlice;
 	};
 
 	/** Temporary-lifetime data about the current tick of the cooker. */
 	struct FTickStackData
 	{
+		/** Time at which the current iteration of the DecideCookAction loop started. */
+		double LoopStartTime = 0.;
 		/** A bitmask of flags of type enum ECookOnTheSideResult that were set during the tick. */
 		uint32 ResultFlags = 0;
-		/** The number of packages that have been cooked in the current tick. It is used for e.g. decisions about when to collect garbage. */
-		uint32 CookedPackageCount = 0;
 		/** The CookerTimer for the current tick. Used by slow reentrant operations that need to check whether they have timed out. */
 		FCookerTimer Timer;
 		/** CookFlags describing details of the caller's desired behavior for the current tick. */
@@ -172,8 +164,8 @@ namespace UE::Cook
 		bool bCookComplete = false;
 		bool bCookCancelled = false;
 
-		explicit FTickStackData(const float& TimeSlice, const bool bIsRealtimeMode, ECookTickFlags InTickFlags, int32 InMaxNumPackagesToSave)
-			:Timer(TimeSlice, bIsRealtimeMode, InMaxNumPackagesToSave), TickFlags(InTickFlags)
+		explicit FTickStackData(float TimeSlice, ECookTickFlags InTickFlags)
+			:Timer(TimeSlice), TickFlags(InTickFlags)
 		{
 		}
 	};
