EditorPerf: Clean up from async review for 20473282

#rnx

diff --git a/Engine/Source/Editor/UnrealEd/Private/Commandlets/SummarizeTraceCommandlet.cpp b/Engine/Source/Editor/UnrealEd/Private/Commandlets/SummarizeTraceCommandlet.cpp
--- a/Engine/Source/Editor/UnrealEd/Private/Commandlets/SummarizeTraceCommandlet.cpp
+++ b/Engine/Source/Editor/UnrealEd/Private/Commandlets/SummarizeTraceCommandlet.cpp
@@ -679,7 +679,7 @@ public:
 /*
 * 
 * Helper classes for the SummarizeTrace commandlet. Aggregates statistics about a trace.
-* Code preceeding this comment should eventually be reloated to the Trace/Insights modules.
+* Code preceding this comment should eventually be relocated to the Trace/Insights modules.
 * Everything succeeding this comment should stay in this commandlet.
 * 
 */
@@ -694,10 +694,13 @@ public:
 */
 static void WelfordsIncrement(const double InSample, const uint64 InCount, double& InOutMean, double& InOutVarianceAccumulator)
 {
-	ensure(InCount);
-	const double OldMean = InOutMean;
-	InOutMean += ((InSample - InOutMean) / double(InCount));
-	InOutVarianceAccumulator += ((InSample - InOutMean) * (InSample - OldMean));
+	check(InCount);
+	if (InCount)
+	{
+		const double OldMean = InOutMean;
+		InOutMean += ((InSample - InOutMean) / double(InCount));
+		InOutVarianceAccumulator += ((InSample - InOutMean) * (InSample - OldMean));
+	}
 }
 
 /**
@@ -1118,17 +1121,17 @@ struct FSummarizeCounter
 	FString Name;
 	ETraceCounterType Type;
 
-	union Value
+	union FValue
 	{
 		int64 Int;
 		double Float;
 	};
 
-	Value Zero;
-	Value First;
-	Value Last;
-	Value Minimum;
-	Value Maximum;
+	const FValue Zero = { 0 };
+	FValue First;
+	FValue Last;
+	FValue Minimum;
+	FValue Maximum;
 
 	uint64 Count = 0;
 	double Mean = 0.0;
@@ -1144,7 +1147,6 @@ struct FSummarizeCounter
 		switch (Type)
 		{
 			case TraceCounterType_Int:
-				Zero.Int = 0;
 				First.Int = 0;
 				Last.Int = 0;
 				Minimum.Int = TNumericLimits<int64>::Max();
@@ -1152,7 +1154,6 @@ struct FSummarizeCounter
 				break;
 
 			case TraceCounterType_Float:
-				Zero.Float = 0;
 				First.Float = 0.0;
 				Last.Float = 0.0;
 				Minimum.Float = TNumericLimits<double>::Max();
@@ -1219,12 +1220,19 @@ struct FSummarizeCounter
 		double CountPerSecond = 0.0;
 		if (Count)
 		{
-			CountPerSecond = Count / (LastSeconds - FirstSeconds);
+			if (Count == 1)
+			{
+				CountPerSecond = 1.0;
+			}
+			else
+			{
+				CountPerSecond = Count / (LastSeconds - FirstSeconds);
+			}
 		}
 		return CountPerSecond;
 	}
 
-	FString PrintValue(const Value& InValue) const
+	FString PrintValue(const FValue& InValue) const
 	{
 		switch (Type)
 		{
