Fix static analysis warning (not checking GUndo for null, when it is checked further down in the same function).

diff --git a/Engine/Source/Editor/UnrealEd/Private/EditorTransaction.cpp b/Engine/Source/Editor/UnrealEd/Private/EditorTransaction.cpp
--- a/Engine/Source/Editor/UnrealEd/Private/EditorTransaction.cpp
+++ b/Engine/Source/Editor/UnrealEd/Private/EditorTransaction.cpp
@@ -1267,7 +1267,7 @@ int32 UTransBuffer::End()
 	// So, the active count could be 0, but the code path may still call end.
 	if (ActiveCount >= 1)
 	{
-		if (ActiveCount == 1)
+		if (ActiveCount == 1 && GUndo != nullptr)
 		{
 			TransactionStateChangedDelegate.Broadcast(GUndo->GetContext(), ETransactionStateEventType::PreTransactionFinalized);
 		}
