Prevent maps without any HLODs to build from triggering errors during a distributed HLOD build

diff --git a/Engine/Source/Editor/UnrealEd/Private/WorldPartition/WorldPartitionHLODsBuilder.cpp b/Engine/Source/Editor/UnrealEd/Private/WorldPartition/WorldPartitionHLODsBuilder.cpp
--- a/Engine/Source/Editor/UnrealEd/Private/WorldPartition/WorldPartitionHLODsBuilder.cpp
+++ b/Engine/Source/Editor/UnrealEd/Private/WorldPartition/WorldPartitionHLODsBuilder.cpp
@@ -785,7 +785,7 @@ bool UWorldPartitionHLODsBuilder::GetHLODActorsToBuild(TArray<FGuid>& HLODActors
 		}
 		else
 		{
-			UE_LOG(LogWorldPartitionHLODsBuilder, Error, TEXT("Missing section [%s] in config file \"%s\""), *SectionName, *BuildManifest);
+			UE_LOG(LogWorldPartitionHLODsBuilder, Log, TEXT("No section [%s] found in config file \"%s\", assuming no HLOD needs to be built."), *SectionName, *BuildManifest);
 			bRet = false;
 		}
 	}
