LevelEditor: Added delegate to extract an asset from a drag and drop operation when SLevelViewport cannot get one.

#rnx

diff --git a/Engine/Source/Editor/LevelEditor/Private/SLevelViewport.cpp b/Engine/Source/Editor/LevelEditor/Private/SLevelViewport.cpp
--- a/Engine/Source/Editor/LevelEditor/Private/SLevelViewport.cpp
+++ b/Engine/Source/Editor/LevelEditor/Private/SLevelViewport.cpp
@@ -5,6 +5,7 @@
 #include "Engine/Selection.h"
 #include "Framework/Commands/UIAction.h"
 #include "Framework/Commands/UICommandList.h"
+#include "Subsystems/EditorAssetSubsystem.h"
 #include "Widgets/Text/STextBlock.h"
 #include "Framework/MultiBox/MultiBoxExtender.h"
 #include "Framework/MultiBox/MultiBoxBuilder.h"
@@ -707,7 +708,7 @@ bool SLevelViewport::HandleDragObjects(const FGeometry& MyGeometry, const FDragD
 	}
 	else
 	{
-		SelectedAssetDatas = AssetUtil::ExtractAssetDataFromDrag( DragDropEvent );
+		GetAssetsFromDrag(DragDropEvent, SelectedAssetDatas);
 
 		if ( SelectedAssetDatas.Num() > 0 )
 		{
@@ -862,7 +863,8 @@ bool SLevelViewport::HandlePlaceDraggedObjects(const FGeometry& MyGeometry, cons
 		bValidDrop = true;
 		DroppedObjects.Empty();
 
-		TArray<FAssetData> DroppedAssetDatas = AssetUtil::ExtractAssetDataFromDrag(DragDropEvent);
+		TArray<FAssetData> DroppedAssetDatas;
+		GetAssetsFromDrag(DragDropEvent, DroppedAssetDatas);
 
 		bAllAssetWereLoaded = true;
 		for (int32 AssetIdx = 0; AssetIdx < DroppedAssetDatas.Num(); ++AssetIdx)
@@ -959,6 +961,38 @@ bool SLevelViewport::HandlePlaceDraggedObjects(const FGeometry& MyGeometry, cons
 	return false;
 }
 
+void SLevelViewport::GetAssetsFromDrag(const FDragDropEvent& DragDropEvent, TArray<FAssetData>& AssetDataArray)
+{
+	AssetDataArray = AssetUtil::ExtractAssetDataFromDrag(DragDropEvent);
+
+	// Did we get anything?
+	if (AssetDataArray.Num() == 0)
+	{
+		// Get files from the drag/drop.
+		const TSharedPtr<FDragDropOperation>& Operation = DragDropEvent.GetOperation();
+		if (Operation.IsValid())
+		{
+			if (Operation->IsOfType<FExternalDragOperation>())
+			{
+				TSharedPtr<FExternalDragOperation> DragDropOp = StaticCastSharedPtr<FExternalDragOperation>(Operation);
+				if (DragDropOp->HasFiles())
+				{
+					const TArray<FString>& Files = DragDropOp->GetFiles();
+					if (Files.Num() > 0)
+					{
+						// See if anyone else can get us an asset.
+						UEditorAssetSubsystem* EditorAssetSubsystem = GEditor->GetEditorSubsystem<UEditorAssetSubsystem>();
+						if (EditorAssetSubsystem != nullptr)
+						{
+							EditorAssetSubsystem->GetOnExtractAssetFromFile().Broadcast(Files, AssetDataArray);
+						}
+					}
+				}
+			}
+		}
+	}
+}
+
 FReply SLevelViewport::OnDrop( const FGeometry& MyGeometry, const FDragDropEvent& DragDropEvent )
 {
 	if (DragDropEvent.GetOperation()->IsOfType<FInViewportUIDragOperation>())
diff --git a/Engine/Source/Editor/LevelEditor/Public/SLevelViewport.h b/Engine/Source/Editor/LevelEditor/Public/SLevelViewport.h
--- a/Engine/Source/Editor/LevelEditor/Public/SLevelViewport.h
+++ b/Engine/Source/Editor/LevelEditor/Public/SLevelViewport.h
@@ -612,6 +612,14 @@ private:
 	 */
 	bool HandlePlaceDraggedObjects(const FGeometry& MyGeometry, const FDragDropEvent& DragDropEvent, bool bCreateDropPreview);
 
+	/**
+	 * Tries to get assets from a drag and drop event.
+	 * 
+	 * @param DragDropEvent		Event to get assets from.
+	 * @param AssetDataArray	Asets will be added here.
+	 */
+	void GetAssetsFromDrag(const FDragDropEvent& DragDropEvent, TArray<FAssetData>& AssetDataArray);
+
 	/** SWidget Interface */
 	virtual FReply OnKeyDown( const FGeometry& MyGeometry, const FKeyEvent& InKeyEvent ) override;
 	virtual void OnDragEnter( const FGeometry& MyGeometry, const FDragDropEvent& DragDropEvent ) override;
diff --git a/Engine/Source/Editor/UnrealEd/Private/Subsystems/EditorAssetSubsystem.cpp b/Engine/Source/Editor/UnrealEd/Private/Subsystems/EditorAssetSubsystem.cpp
--- a/Engine/Source/Editor/UnrealEd/Private/Subsystems/EditorAssetSubsystem.cpp
+++ b/Engine/Source/Editor/UnrealEd/Private/Subsystems/EditorAssetSubsystem.cpp
@@ -340,6 +340,16 @@ UEditorAssetSubsystem::UEditorAssetSubsystem()
 {
 }
 
+void UEditorAssetSubsystem::Initialize(FSubsystemCollectionBase& Collection)
+{
+	OnExtractAssetFromFile.AddUObject(this, &UEditorAssetSubsystem::CallOnExtractAssetFromFileDynamicArray);
+}
+
+void UEditorAssetSubsystem::Deinitialize()
+{
+	OnExtractAssetFromFile.RemoveAll(this);
+}
+
 // Load operations
 
 // A wrapper around
@@ -1544,3 +1554,26 @@ void UEditorAssetSubsystem::RemoveMetadataTag(UObject* Object, FName Tag)
 	Object->GetPackage()->GetMetaData()->RemoveValue(Object, Tag);
 #endif // WITH_EDITORONLY_DATA
 }
+
+
+void UEditorAssetSubsystem::AddOnExtractAssetFromFile(FOnExtractAssetFromFileDynamic Delegate)
+{
+	OnExtractAssetFromFileDynamicArray.Add(Delegate);
+}
+
+void UEditorAssetSubsystem::RemoveOnExtractAssetFromFile(FOnExtractAssetFromFileDynamic Delegate)
+{
+	OnExtractAssetFromFileDynamicArray.Remove(Delegate);
+}
+
+void UEditorAssetSubsystem::CallOnExtractAssetFromFileDynamicArray(const TArray<FString>& Files,
+	TArray<FAssetData>& OutAssetDataArray)
+{
+	TArray<FAssetData> LocalArray;
+	for (FOnExtractAssetFromFileDynamic& Delegate : OnExtractAssetFromFileDynamicArray)
+	{
+		Delegate.ExecuteIfBound(Files, LocalArray);
+		OutAssetDataArray += LocalArray;
+		LocalArray.Reset();
+	}
+}
diff --git a/Engine/Source/Editor/UnrealEd/Public/Subsystems/EditorAssetSubsystem.h b/Engine/Source/Editor/UnrealEd/Public/Subsystems/EditorAssetSubsystem.h
--- a/Engine/Source/Editor/UnrealEd/Public/Subsystems/EditorAssetSubsystem.h
+++ b/Engine/Source/Editor/UnrealEd/Public/Subsystems/EditorAssetSubsystem.h
@@ -28,6 +28,8 @@ class UNREALED_API UEditorAssetSubsystem : public UEditorSubsystem
 public:
 
 	UEditorAssetSubsystem();
+	virtual void Initialize(FSubsystemCollectionBase& Collection) override;
+	virtual void Deinitialize() override;
 	
 	/**
 	 * Load an asset. It will verify if the object is already loaded and only load it if it's necessary.
@@ -365,4 +367,54 @@ public:
      */
     UFUNCTION(BlueprintCallable, Category = "Editor Scripting | Asset | Metadata")
     void RemoveMetadataTag(UObject* Object, FName Tag);
+
+	/**
+	 * Delegate for extracting an asset from a file, for example from a drag and drop operation.
+	 *
+	 * @param Files				List of files.
+	 * @param AssetDataArray	Add extracted assets here.
+	 */
+	DECLARE_MULTICAST_DELEGATE_TwoParams(FOnExtractAssetFromFile, const TArray<FString>& Files, TArray<FAssetData>& AssetDataArray);
+
+	/**
+	 * Dynamic version of FOnExtractAssetFromFile.
+	 */
+	DECLARE_DYNAMIC_DELEGATE_TwoParams(FOnExtractAssetFromFileDynamic, const TArray<FString>&, Files, TArray<FAssetData>&, AssetDataArray);
+
+	/**
+	 * Call this to add a callback to extract an asset from a file,
+	 * for example from a drag and drop operation.
+	 */
+	UFUNCTION(BlueprintCallable, Category = "Editor Scripting | Asset")
+		void AddOnExtractAssetFromFile(FOnExtractAssetFromFileDynamic Delegate);
+
+	/**
+	 * Call this to remove a callback added with AddOnExtractAssetFromFile.
+	 */
+	UFUNCTION(BlueprintCallable, Category = "Editor Scripting | Asset")
+		void RemoveOnExtractAssetFromFile(FOnExtractAssetFromFileDynamic Delegate);
+
+	/**
+	 * Get the delegate for extracting an asset from a file,
+	 * for example from a drag and drop operation.
+	 * Add to this so you can do your own extraction.
+	 *
+	 * Broadcasting this will also call anything added with AddOnExtractAssetFromFile.
+	 */
+	FOnExtractAssetFromFile& GetOnExtractAssetFromFile() { return OnExtractAssetFromFile; }
+
+private:
+
+	/**
+	 * Calls all the delegates in OnExtractAssetFromFileDynamicArray when OnExtractAssetFromFile
+	 * is broadcast.
+	 */
+	void CallOnExtractAssetFromFileDynamicArray(const TArray<FString>& Files,
+		TArray<FAssetData>& InAssetDataArray);
+
+	/** Delegate for extracting an asset from a file, for example from a drag and drop operation. */
+	FOnExtractAssetFromFile OnExtractAssetFromFile;
+	/** Array of dynamic delegates for extracting an asset from a file. */
+	TArray<FOnExtractAssetFromFileDynamic> OnExtractAssetFromFileDynamicArray;
+
 };
