EditorDomain: Fix deadlock in TryFindOrAddPackageSource during startup when turning on asynchronous DDC requests.
#rnx

diff --git a/Engine/Source/Editor/UnrealEd/Private/EditorDomain/EditorDomain.cpp b/Engine/Source/Editor/UnrealEd/Private/EditorDomain/EditorDomain.cpp
--- a/Engine/Source/Editor/UnrealEd/Private/EditorDomain/EditorDomain.cpp
+++ b/Engine/Source/Editor/UnrealEd/Private/EditorDomain/EditorDomain.cpp
@@ -166,14 +166,32 @@
 {
 }
 
-bool FEditorDomain::TryFindOrAddPackageSource(FName PackageName,
+bool FEditorDomain::TryFindOrAddPackageSource(FScopeLock& ScopeLock, bool& bOutReenteredLock, FName PackageName,
 	TRefCountPtr<FPackageSource>& OutSource, UE::EditorDomain::FPackageDigest* OutErrorDigest)
 {
-	// Called within Locks.Lock
+	// Called within &Locks->Lock, ScopeLock is locked on that Lock
 	using namespace UE::EditorDomain;
 
 	// EDITOR_DOMAIN_TODO: Need to delete entries from PackageSources when the assetregistry reports the package is
 	// resaved on disk.
+	TRefCountPtr<FPackageSource>* PackageSourcePtr = PackageSources.Find(PackageName);
+	if (PackageSourcePtr && *PackageSourcePtr)
+	{
+		bOutReenteredLock = false;
+		OutSource = *PackageSourcePtr;
+		return true;
+	}
+
+	// CalculatepackageDigest calls arbitrary code because of class loads and can reenter EditorDomain functions
+	// We have to drop the lock around our call to it
+	FPackageDigest PackageDigest;
+	{
+		Locks->Lock.Unlock();
+		ON_SCOPE_EXIT{ Locks->Lock.Lock(); };
+		PackageDigest = CalculatePackageDigest(*AssetRegistry, PackageName);
+	}
+
+	bOutReenteredLock = true;
 	TRefCountPtr<FPackageSource>& PackageSource = PackageSources.FindOrAdd(PackageName);
 	if (PackageSource)
 	{
@@ -181,7 +199,6 @@
 		return true;
 	}
 
-	FPackageDigest PackageDigest = CalculatePackageDigest(*AssetRegistry, PackageName);
 	switch (PackageDigest.Status)
 	{
 	case FPackageDigest::EStatus::Successful:
@@ -218,17 +235,18 @@
 UE::EditorDomain::FPackageDigest FEditorDomain::GetPackageDigest(FName PackageName)
 {
 	FScopeLock ScopeLock(&Locks->Lock);
-	return GetPackageDigest_WithinLock(PackageName);
+	bool bReenteredLock;
+	return GetPackageDigest_WithinLock(ScopeLock, bReenteredLock, PackageName);
 }
 
-UE::EditorDomain::FPackageDigest FEditorDomain::GetPackageDigest_WithinLock(FName PackageName)
+UE::EditorDomain::FPackageDigest FEditorDomain::GetPackageDigest_WithinLock(FScopeLock& ScopeLock, bool& bOutReenteredLock, FName PackageName)
 {
-	// Called within &Locks->Lock
+	// Called within &Locks->Lock, ScopeLock is locked on that Lock
 	using namespace UE::EditorDomain;
 
 	TRefCountPtr<FPackageSource> PackageSource;
 	FPackageDigest ErrorDigest;
-	if (!TryFindOrAddPackageSource(PackageName, PackageSource, &ErrorDigest))
+	if (!TryFindOrAddPackageSource(ScopeLock, bOutReenteredLock, PackageName, PackageSource, &ErrorDigest))
 	{
 		return ErrorDigest;
 	}
@@ -316,7 +334,8 @@
 			return Workspace->FileSize(PackagePath, PackageSegment, OutUpdatedPath);
 		}
 
-		if (!TryFindOrAddPackageSource(PackageName, PackageSource) || PackageSource->Source == EPackageSource::Workspace)
+		bool bReenteredLock;
+		if (!TryFindOrAddPackageSource(ScopeLock, bReenteredLock, PackageName, PackageSource) || PackageSource->Source == EPackageSource::Workspace)
 		{
 			return Workspace->FileSize(PackagePath, PackageSegment, OutUpdatedPath);
 		}
@@ -399,7 +418,8 @@
 		return Workspace->OpenReadPackage(PackagePath, PackageSegment, OutUpdatedPath);
 	}
 	TRefCountPtr<FPackageSource> PackageSource;
-	if (!TryFindOrAddPackageSource(PackageName, PackageSource) || (PackageSource->Source == EPackageSource::Workspace))
+	bool bReenteredLock;
+	if (!TryFindOrAddPackageSource(ScopeLock, bReenteredLock, PackageName, PackageSource) || (PackageSource->Source == EPackageSource::Workspace))
 	{
 		return Workspace->OpenReadPackage(PackagePath, PackageSegment, OutUpdatedPath);
 	}
@@ -456,7 +476,8 @@
 		return Workspace->OpenAsyncReadPackage(PackagePath, PackageSegment);
 	}
 	TRefCountPtr<FPackageSource> PackageSource;
-	if (!TryFindOrAddPackageSource(PackageName, PackageSource) ||
+	bool bReenteredLock;
+	if (!TryFindOrAddPackageSource(ScopeLock, bReenteredLock, PackageName, PackageSource) ||
 		(PackageSource->Source == EPackageSource::Workspace))
 	{
 		return Workspace->OpenAsyncReadPackage(PackagePath, PackageSegment);
@@ -672,7 +693,8 @@
 	ECachePolicy CachePolicy = ECachePolicy::Default | ECachePolicy::SkipData;
 	for (FName PackageName : PackageNames)
 	{
-		FPackageDigest PackageDigest = GetPackageDigest_WithinLock(PackageName);
+		bool bReenteredLock;
+		FPackageDigest PackageDigest = GetPackageDigest_WithinLock(ScopeLock, bReenteredLock, PackageName);
 		if (PackageDigest.IsSuccessful() && EnumHasAnyFlags(PackageDigest.DomainUse, EDomainUse::LoadEnabled))
 		{
 			CacheRequests.Add({ { WriteToString<256>(PackageName) }, GetEditorDomainPackageKey(PackageDigest.Hash),
@@ -688,7 +710,8 @@
 				FScopeLock ScopeLock(&Locks->Lock);
 				TRefCountPtr<FPackageSource> PackageSource;
 				FName PackageName = FName(*Response.Name);
-				if (TryFindOrAddPackageSource(PackageName, PackageSource))
+				bool bReenteredLock;
+				if (TryFindOrAddPackageSource(ScopeLock, bReenteredLock, PackageName, PackageSource))
 				{
 					PackageSource->bHasQueriedCatalog = true;
 				}
diff --git a/Engine/Source/Editor/UnrealEd/Private/EditorDomain/EditorDomain.h b/Engine/Source/Editor/UnrealEd/Private/EditorDomain/EditorDomain.h
--- a/Engine/Source/Editor/UnrealEd/Private/EditorDomain/EditorDomain.h
+++ b/Engine/Source/Editor/UnrealEd/Private/EditorDomain/EditorDomain.h
@@ -18,6 +18,7 @@
 
 class FAssetPackageData;
 class FEditorDomainSaveClient;
+class FScopeLock;
 class IAssetRegistry;
 class UObject;
 class UPackage;
@@ -215,8 +216,12 @@ private:
 	FEditorDomain(const FEditorDomain& Other) = delete;
 	FEditorDomain(FEditorDomain&& Other) = delete;
 
-	/** Read the PackageSource data from PackageSources, or from the asset registry if not in PackageSources. */
-	bool TryFindOrAddPackageSource(FName PackageName, TRefCountPtr<FPackageSource>& OutSource,
+	/**
+	 * Read the PackageSource data from PackageSources, or from the asset registry if not in PackageSources.
+	 * Note this function can exit and reenter the provided ScopeLock. If so it will set bOutReenteredLock=true and caller must handle
+	 * retesting variables that may have changed while outside of the lock.
+	*/
+	bool TryFindOrAddPackageSource(FScopeLock& ScopeLock, bool& bOutReenteredLock, FName PackageName, TRefCountPtr<FPackageSource>& OutSource,
 		UE::EditorDomain::FPackageDigest* OutErrorDigest=nullptr);
 	/** Return the PackageSource data in PackageSources, if it exists */
 	TRefCountPtr<FPackageSource> FindPackageSource(const FPackagePath& PackagePath);
@@ -236,7 +241,7 @@ private:
 	/** AssetUpdated event to invalidate our information about where it should be loaded from. */
 	void OnAssetUpdatedOnDisk(const FAssetData& AssetData);
 	/** Same As GetPackageDigest, but assumes lock is already held. */
-	UE::EditorDomain::FPackageDigest GetPackageDigest_WithinLock(FName PackageDigest);
+	UE::EditorDomain::FPackageDigest GetPackageDigest_WithinLock(FScopeLock& ScopeLock, bool& bOutReenteredLock, FName PackageDigest);
 
 
 	/** Subsystem used to request the save of missing packages into the EditorDomain from a separate process. */
