Add null check in FBluepirntEditorUtils::FindAndSetDebuggableBlueprintInstances

diff --git a/Engine/Source/Editor/UnrealEd/Private/Kismet2/BlueprintEditorUtils.cpp b/Engine/Source/Editor/UnrealEd/Private/Kismet2/BlueprintEditorUtils.cpp
--- a/Engine/Source/Editor/UnrealEd/Private/Kismet2/BlueprintEditorUtils.cpp
+++ b/Engine/Source/Editor/UnrealEd/Private/Kismet2/BlueprintEditorUtils.cpp
@@ -9016,7 +9016,7 @@ void FBlueprintEditorUtils::FindAndSetDebuggableBlueprintInstances()
 			UBlueprint* EachBlueprint = ObjIt.Key();
 			bool bFoundItemToDebug = false;
 
-			if( Selected->Num() != 0 )
+			if( Selected )
 			{
 				for (int32 iSelected = 0; iSelected < Selected->Num() ; iSelected++)
 				{
