AssetRegistry cook: Fix AssetDatas' PackageFlags being clobbered during CookByTheBookFinished when we update the AssetDatas for changes made on load of the packages. The PackageFlags set int he AssetRegistryGenerator are authoritative for the cooked AssetRegistry.
#rnx

diff --git a/Engine/Source/Editor/UnrealEd/Private/Commandlets/AssetRegistryGenerator.cpp b/Engine/Source/Editor/UnrealEd/Private/Commandlets/AssetRegistryGenerator.cpp
--- a/Engine/Source/Editor/UnrealEd/Private/Commandlets/AssetRegistryGenerator.cpp
+++ b/Engine/Source/Editor/UnrealEd/Private/Commandlets/AssetRegistryGenerator.cpp
@@ -1451,6 +1451,23 @@ void FAssetRegistryGenerator::AddAssetToFileOrderRecursive(const FName& InPackag
 
 void FAssetRegistryGenerator::UpdateAssetDatas(bool bForceNoFilter)
 {
+	// Keep the existing AssetData->PackageFlags. These flags were set by UpdateAssetDataPackageFlags after saving
+	// each Asset and the existing values are authoritative for the cooked AssetRegistry.
+	// TODO: Remove this end-of-cook updating of the AssetDatas. Update them as we go along,
+	// in UpdateAssetDataPackageFlags, and set the PackageFlags afterwards. Then we won't have to
+	// have the cost of the many map lookups we do here.
+	TSet<FName> SkipNone;
+	TArray<TPair<FName, uint32>> PackageNameAndFlags;
+	PackageNameAndFlags.Reserve(State.GetNumAssets());
+	State.EnumerateAllAssets(SkipNone, [&PackageNameAndFlags](const FAssetData& AssetData)
+		{
+			if (!AssetData.PackageName.IsNone())
+			{
+				PackageNameAndFlags.Emplace(AssetData.PackageName, AssetData.PackageFlags);
+			}
+			return true;
+		});
+
 	FAssetRegistrySerializationOptions SaveOptions;
 	AssetRegistry.InitializeSerializationOptions(SaveOptions, TargetPlatform->IniPlatformName());
 	if (bForceNoFilter)
@@ -1458,6 +1475,11 @@ void FAssetRegistryGenerator::UpdateAssetDatas(bool bForceNoFilter)
 		SaveOptions.DisableFilters();
 	}
 	AssetRegistry.InitializeTemporaryAssetRegistryState(State, SaveOptions, true);
+
+	for (TPair<FName, uint32>& Pair : PackageNameAndFlags)
+	{
+		State.UpdateAssetDataPackageFlags(Pair.Key, Pair.Value);
+	}
 }
 
 bool FAssetRegistryGenerator::SaveAssetRegistry(const FString& SandboxPath, bool bSerializeDevelopmentAssetRegistry, bool bForceNoFilter)
