WorldPartition: RVT World Builder
- Actors that output to RVT now have a ActorDesc property so they can be loaded before updating the RVT

- Fix a typo
- Fix minimap build calling hlod build in Build All

diff --git a/Engine/Source/Editor/LevelEditor/Private/LevelEditorActions.cpp b/Engine/Source/Editor/LevelEditor/Private/LevelEditorActions.cpp
--- a/Engine/Source/Editor/LevelEditor/Private/LevelEditorActions.cpp
+++ b/Engine/Source/Editor/LevelEditor/Private/LevelEditorActions.cpp
@@ -3387,7 +3387,7 @@ void FLevelEditorCommands::RegisterCommands()
 	UI_COMMAND( BuildGeometryOnly_OnlyCurrentLevel, "Build Geometry (Current Level)", "Builds geometry, only for the current level", EUserInterfaceActionType::Button, FInputChord() );
 	UI_COMMAND( BuildPathsOnly, "Build Paths", "Only builds paths (all levels.)", EUserInterfaceActionType::Button, FInputChord() );
 	UI_COMMAND( BuildHLODs, "Build HLODs", "Builds all HLODs for the current world", EUserInterfaceActionType::Button, FInputChord() );
-	UI_COMMAND( BuildMinimap, "Build Minimap", "Builds the minimap for the curent world", EUserInterfaceActionType::Button, FInputChord() );
+	UI_COMMAND( BuildMinimap, "Build Minimap", "Builds the minimap for the current world", EUserInterfaceActionType::Button, FInputChord() );
 	UI_COMMAND( BuildTextureStreamingOnly, "Build Texture Streaming", "Build texture streaming data", EUserInterfaceActionType::Button, FInputChord() );
 	UI_COMMAND( BuildVirtualTextureOnly, "Build Virtual Textures", "Build runtime virtual texture low mips streaming data", EUserInterfaceActionType::Button, FInputChord());
 	UI_COMMAND( BuildAllLandscape, "Build All Landscape", "Build All Landscape Data(Grass, Baked Textures)", EUserInterfaceActionType::Button, FInputChord());
diff --git a/Engine/Source/Editor/UnrealEd/Private/EditorBuildUtils.cpp b/Engine/Source/Editor/UnrealEd/Private/EditorBuildUtils.cpp
--- a/Engine/Source/Editor/UnrealEd/Private/EditorBuildUtils.cpp
+++ b/Engine/Source/Editor/UnrealEd/Private/EditorBuildUtils.cpp
@@ -8,6 +8,7 @@
 #include "Misc/MessageDialog.h"
 #include "HAL/FileManager.h"
 #include "Misc/ScopedSlowTask.h"
+#include "Misc/ScopeExit.h"
 #include "Modules/ModuleManager.h"
 #include "Misc/PackageName.h"
 #include "Engine/EngineTypes.h"
@@ -52,6 +53,8 @@
 #include "WorldPartition/IWorldPartitionEditorModule.h"
 #include "WorldPartition/SWorldPartitionBuildNavigationDialog.h"
 #include "WorldPartition/WorldPartitionBuildNavigationOptions.h"
+#include "WorldPartition/WorldPartitionHelpers.h"
+#include "WorldPartition/WorldPartitionRuntimeVirtualTextureBuilder.h"
 
 DEFINE_LOG_CATEGORY_STATIC(LogEditorBuildUtils, Log, All);
 
@@ -1192,8 +1195,8 @@ void FBuildAllHandler::ProcessBuild(const TWeakPtr<SBuildProgressWidget>& BuildP
 		}
 		else if (StepId == FBuildOptions::BuildMinimap)
 		{
-			BuildProgressWidget.Pin()->SetBuildType(SBuildProgressWidget::BUILDTYPE_HLODs);
-			FEditorBuildUtils::TriggerHierarchicalLODBuilder(CurrentWorld, CurrentBuildId);
+			BuildProgressWidget.Pin()->SetBuildType(SBuildProgressWidget::BUILDTYPE_Minimap);
+			FEditorBuildUtils::TriggerMinimapBuilder(CurrentWorld, CurrentBuildId);
 		}
 		else if (StepId == FBuildOptions::BuildTextureStreaming)
 		{
@@ -1596,37 +1599,49 @@ bool FEditorBuildUtils::EditorBuildVirtualTexture(UWorld* InWorld)
 	}
 
 	CollectGarbage(GARBAGE_COLLECTION_KEEPFLAGS);
+	ON_SCOPE_EXIT
+	{
+		CollectGarbage(GARBAGE_COLLECTION_KEEPFLAGS);
+	};
 
-	TArray<URuntimeVirtualTextureComponent*> Components;
-	for (TObjectIterator<URuntimeVirtualTextureComponent> It; It; ++It)
 	{
-		if (Module->HasStreamedMips(*It))
+		FWorldPartitionHelpers::FForEachActorWithLoadingResult ForEachActorWithLoadingResult;
+		if (UWorldPartition* WorldPartition = InWorld->GetWorldPartition())
 		{
-			Components.Add(*It);
+			FScopedSlowTask BuildTask(1.0f, LOCTEXT("VirtualTextureBuild", "Loading Actors"));
+			BuildTask.MakeDialog();
+			UWorldPartitionRuntimeVirtualTextureBuilder::LoadRuntimeVirtualTextureActors(WorldPartition, ForEachActorWithLoadingResult);
 		}
-	}
 
-	if (Components.Num() == 0)
-	{
-		return true;
-	}
+		TArray<URuntimeVirtualTextureComponent*> Components;
+		for (TObjectIterator<URuntimeVirtualTextureComponent> It; It; ++It)
+		{
+			if (Module->HasStreamedMips(*It))
+			{
+				Components.Add(*It);
+			}
+		}
 
-	FScopedSlowTask BuildTask(Components.Num(), LOCTEXT("VirtualTextureBuild", "Building Virtual Textures"));
-	BuildTask.MakeDialog(true);
+		if (Components.Num() == 0)
+		{
+			return true;
+		}
 
-	for (URuntimeVirtualTextureComponent* Component : Components)
-	{
-		BuildTask.EnterProgressFrame();
+		FScopedSlowTask BuildTask(Components.Num(), LOCTEXT("VirtualTextureBuild", "Building Virtual Textures"));
+		BuildTask.MakeDialog(true);
 
-		// Note that Build*() functions return true if the associated Has*() functions return false
-		if (BuildTask.ShouldCancel() || !Module->BuildStreamedMips(Component))
+		for (URuntimeVirtualTextureComponent* Component : Components)
 		{
-			return false;
+			BuildTask.EnterProgressFrame();
+
+			// Note that Build*() functions return true if the associated Has*() functions return false
+			if (BuildTask.ShouldCancel() || !Module->BuildStreamedMips(Component))
+			{
+				return false;
+			}
 		}
 	}
-
-	CollectGarbage(GARBAGE_COLLECTION_KEEPFLAGS);
-
+	
 	return true;
 }
 
diff --git a/Engine/Source/Editor/UnrealEd/Private/WorldPartition/WorldPartitionRuntimeVirtualTextureBuilder.cpp b/Engine/Source/Editor/UnrealEd/Private/WorldPartition/WorldPartitionRuntimeVirtualTextureBuilder.cpp
new file mode 100644
--- /dev/null
+++ b/Engine/Source/Editor/UnrealEd/Private/WorldPartition/WorldPartitionRuntimeVirtualTextureBuilder.cpp
@@ -0,0 +1,72 @@
+// Copyright Epic Games, Inc. All Rights Reserved.
+
+#include "WorldPartition/WorldPartitionRuntimeVirtualTextureBuilder.h"
+#include "WorldPartition/WorldPartition.h"
+#include "WorldPartition/WorldPartitionActorDesc.h"
+#include "WorldPartition/WorldPartitionHelpers.h"
+#include "WorldPartition/WorldPartitionHandle.h"
+#include "Components/RuntimeVirtualTextureComponent.h"
+#include "Components/PrimitiveComponent.h"
+#include "VT/VirtualTextureBuilder.h"
+#include "VirtualTexturingEditorModule.h"
+#include "AssetCompilingManager.h"
+#include "Modules/ModuleManager.h"
+#include "UObject/UObjectIterator.h"
+
+DEFINE_LOG_CATEGORY_STATIC(LogWorldPartitionRuntimeVirtualTextureBuilder, All, All);
+
+UWorldPartitionRuntimeVirtualTextureBuilder::UWorldPartitionRuntimeVirtualTextureBuilder(const FObjectInitializer& ObjectInitializer)
+	: Super(ObjectInitializer)
+{
+}
+
+void UWorldPartitionRuntimeVirtualTextureBuilder::LoadRuntimeVirtualTextureActors(UWorldPartition* WorldPartition, FWorldPartitionHelpers::FForEachActorWithLoadingResult& Result)
+{
+	check(WorldPartition);
+
+	FWorldPartitionHelpers::FForEachActorWithLoadingParams ForEachActorWithLoadingParams;
+	ForEachActorWithLoadingParams.bKeepReferences = true;
+	ForEachActorWithLoadingParams.FilterActorDesc = [](const FWorldPartitionActorDesc* ActorDesc) -> bool { return ActorDesc->HasProperty(UPrimitiveComponent::RVTActorDescProperty); };
+		
+	// @todo: in order to scale the RVTs should be generated with tiling so that we don't need to load all actors writing to RVTs at once.
+	FWorldPartitionHelpers::ForEachActorWithLoading(WorldPartition, [](const FWorldPartitionActorDesc*) { return true; }, ForEachActorWithLoadingParams, Result);
+
+	// Make sure all assets are finished compiling
+	FAssetCompilingManager::Get().FinishAllCompilation();
+}
+
+bool UWorldPartitionRuntimeVirtualTextureBuilder::RunInternal(UWorld* World, const FCellInfo& InCellInfo, FPackageSourceControlHelper& PackageHelper)
+{
+	IVirtualTexturingEditorModule& VTModule = FModuleManager::Get().LoadModuleChecked<IVirtualTexturingEditorModule>("VirtualTexturingEditor");
+	
+	UWorldPartition* WorldPartition = World->GetWorldPartition();
+	if (!WorldPartition)
+	{
+		UE_LOG(LogWorldPartitionRuntimeVirtualTextureBuilder, Error, TEXT("Failed to retrieve WorldPartition."));
+		return false;
+	}
+		
+	// Load required actors
+	FWorldPartitionHelpers::FForEachActorWithLoadingResult ForEachActorWithLoadingResult;
+	LoadRuntimeVirtualTextureActors(WorldPartition, ForEachActorWithLoadingResult);
+		
+	TSet<UPackage*> ModifiedPackages;
+	for (TObjectIterator<URuntimeVirtualTextureComponent> It(RF_ClassDefaultObject, false, EInternalObjectFlags::Garbage); It; ++It)
+	{
+		if (It->GetWorld() == World)
+		{
+			if (VTModule.HasStreamedMips(*It) && VTModule.BuildStreamedMips(*It))
+			{
+				if (UVirtualTextureBuilder* VTBuilder = It->GetStreamingTexture(); VTBuilder->GetPackage()->IsDirty())
+				{
+					ModifiedPackages.Add(VTBuilder->GetPackage());
+				}
+			}
+		}
+	}
+
+	// Wait for VT Textures to be ready before saving
+	FAssetCompilingManager::Get().FinishAllCompilation();
+
+	return SavePackages(ModifiedPackages.Array(), PackageHelper, false);
+}
diff --git a/Engine/Source/Editor/UnrealEd/Public/WorldPartition/WorldPartitionRuntimeVirtualTextureBuilder.h b/Engine/Source/Editor/UnrealEd/Public/WorldPartition/WorldPartitionRuntimeVirtualTextureBuilder.h
new file mode 100644
--- /dev/null
+++ b/Engine/Source/Editor/UnrealEd/Public/WorldPartition/WorldPartitionRuntimeVirtualTextureBuilder.h
@@ -0,0 +1,26 @@
+// Copyright Epic Games, Inc. All Rights Reserved.
+
+#pragma once
+
+#include "WorldPartition/WorldPartitionBuilder.h"
+#include "WorldPartition/WorldPartitionHelpers.h"
+#include "WorldPartitionRuntimeVirtualTextureBuilder.generated.h"
+
+class UWorldPartition;
+
+// Example Command Line: ProjectName MapName -run=WorldPartitionBuilderCommandlet -SCCProvider=Perforce -Builder=WorldPartitionRuntimeVirtualTextureBuilder
+UCLASS()
+class UNREALED_API UWorldPartitionRuntimeVirtualTextureBuilder : public UWorldPartitionBuilder
+{
+	GENERATED_UCLASS_BODY()
+
+public:
+	// UWorldPartitionBuilder interface begin
+	virtual bool RequiresCommandletRendering() const override { return true; }
+	virtual ELoadingMode GetLoadingMode() const override { return ELoadingMode::Custom; }
+	virtual bool RunInternal(UWorld* World, const FCellInfo& InCellInfo, FPackageSourceControlHelper& PackageHelper) override;
+	// UWorldPartitionBuilder interface end
+
+	// Helper method to load actors that contribute to the Runtime Virtual Texture
+	static void LoadRuntimeVirtualTextureActors(UWorldPartition* WorldPartition, FWorldPartitionHelpers::FForEachActorWithLoadingResult& Result);
+};
\ No newline at end of file
diff --git a/Engine/Source/Runtime/Engine/Classes/Components/PrimitiveComponent.h b/Engine/Source/Runtime/Engine/Classes/Components/PrimitiveComponent.h
--- a/Engine/Source/Runtime/Engine/Classes/Components/PrimitiveComponent.h
+++ b/Engine/Source/Runtime/Engine/Classes/Components/PrimitiveComponent.h
@@ -210,7 +210,8 @@ public:
 	~UPrimitiveComponent();
 
 	// Rendering
-	
+	static FName RVTActorDescProperty;
+
 	/**
 	 * The minimum distance at which the primitive should be rendered, 
 	 * measured in world space units from the center of the primitive's bounding sphere to the camera position.
@@ -2164,6 +2165,7 @@ public:
 	virtual void OnComponentDestroyed(bool bDestroyingHierarchy) override;
 #if WITH_EDITOR
 	virtual void CheckForErrors() override;
+	virtual void GetActorDescProperties(TMap<FName, FName>& Properties) const;
 #endif // WITH_EDITOR	
 	//~ End UActorComponent Interface
 
diff --git a/Engine/Source/Runtime/Engine/Private/Components/PrimitiveComponent.cpp b/Engine/Source/Runtime/Engine/Private/Components/PrimitiveComponent.cpp
--- a/Engine/Source/Runtime/Engine/Private/Components/PrimitiveComponent.cpp
+++ b/Engine/Source/Runtime/Engine/Private/Components/PrimitiveComponent.cpp
@@ -294,6 +294,8 @@ FORCEINLINE_DEBUGGABLE static void GetPointersToArrayDataByPredicate(TArray<cons
 
 uint32 UPrimitiveComponent::GlobalOverlapEventsCounter = 0;
 
+FName UPrimitiveComponent::RVTActorDescProperty(TEXT("RVT"));
+
 // 0 is reserved to mean invalid
 FThreadSafeCounter UPrimitiveComponent::NextRegistrationSerialNumber;
 
@@ -1206,6 +1208,14 @@ void UPrimitiveComponent::CheckForErrors()
 	}
 }
 
+void UPrimitiveComponent::GetActorDescProperties(TMap<FName, FName>& Properties) const
+{
+	if (RuntimeVirtualTextures.Num())
+	{
+		Properties.Add(UPrimitiveComponent::RVTActorDescProperty, NAME_None);
+	}
+}
+
 void UPrimitiveComponent::UpdateCollisionProfile()
 {
 	BodyInstance.LoadProfileData(false);
diff --git a/Engine/Source/Runtime/Engine/Public/WorldPartition/WorldPartitionActorDesc.h b/Engine/Source/Runtime/Engine/Public/WorldPartition/WorldPartitionActorDesc.h
--- a/Engine/Source/Runtime/Engine/Public/WorldPartition/WorldPartitionActorDesc.h
+++ b/Engine/Source/Runtime/Engine/Public/WorldPartition/WorldPartitionActorDesc.h
@@ -98,7 +98,7 @@ public:
 	inline const FGuid& GetParentActor() const { return ParentActor; }
 	inline bool IsUsingDataLayerAsset() const { return bIsUsingDataLayerAsset; }
 	inline FName GetProperty(FName PropertyName) const { return Properties.FindRef(PropertyName); }
-	inline bool HasProperty(FName PropertyName) const { return !GetProperty(PropertyName).IsNone(); }
+	inline bool HasProperty(FName PropertyName) const { return Properties.Contains(PropertyName); }
 
 	FName GetActorName() const;
 	FName GetActorLabelOrName() const;
