Switching rendering preview mode should use ODSC.

- AllMaterialsCacheResourceShadersForRendering now accepts a parameter to decide whether to fully cache and complete all shader maps.  By default this is true to maintain previous behavior.
- When switching preview materials this behavior defaults to false and we let ODSC in the editor to handle which shaders to compile.
- This results in far fewer shaders compiled.

diff --git a/Engine/Source/Editor/UnrealEd/Private/EditorEngine.cpp b/Engine/Source/Editor/UnrealEd/Private/EditorEngine.cpp
--- a/Engine/Source/Editor/UnrealEd/Private/EditorEngine.cpp
+++ b/Engine/Source/Editor/UnrealEd/Private/EditorEngine.cpp
@@ -7526,6 +7526,9 @@ void UEditorEngine::SetPreviewPlatform(const FPreviewPlatformInfo& NewPreviewPla
 		UStaticMesh::OnLodStrippingQualityLevelChanged(nullptr);
 	}
 
+	constexpr bool bUpdateProgressDialog = true;
+	constexpr bool bCacheAllRemainingShaders = false;
+
 	if (bChangedFeatureLevel)
 	{
 		FScopedSlowTask SlowTask(100.f, NSLOCTEXT("Engine", "ChangingPreviewPlatform", "Changing Preview Platform"), true);
@@ -7549,10 +7552,10 @@ void UEditorEngine::SetPreviewPlatform(const FPreviewPlatformInfo& NewPreviewPla
 		}
 
 		SlowTask.EnterProgressFrame(35.0f);
-		UMaterial::AllMaterialsCacheResourceShadersForRendering(true);
+		UMaterial::AllMaterialsCacheResourceShadersForRendering(bUpdateProgressDialog, bCacheAllRemainingShaders);
 
 		SlowTask.EnterProgressFrame(35.0f);
-		UMaterialInstance::AllMaterialsCacheResourceShadersForRendering(true);
+		UMaterialInstance::AllMaterialsCacheResourceShadersForRendering(bUpdateProgressDialog, bCacheAllRemainingShaders);
 
 		SlowTask.EnterProgressFrame(15.0f, NSLOCTEXT("Engine", "SlowTaskGlobalShaderMapMessage", "Compiling global shaders"));
 		CompileGlobalShaderMap(PreviewPlatform.PreviewFeatureLevel);
@@ -7568,8 +7571,8 @@ void UEditorEngine::SetPreviewPlatform(const FPreviewPlatformInfo& NewPreviewPla
 		// Rebuild materials if we have the same feature level but a different 'material quality platform'
 		FGlobalComponentRecreateRenderStateContext Recreate;
 		FlushRenderingCommands();
-		UMaterial::AllMaterialsCacheResourceShadersForRendering(true);
-		UMaterialInstance::AllMaterialsCacheResourceShadersForRendering(true);
+		UMaterial::AllMaterialsCacheResourceShadersForRendering(bUpdateProgressDialog, bCacheAllRemainingShaders);
+		UMaterialInstance::AllMaterialsCacheResourceShadersForRendering(bUpdateProgressDialog, bCacheAllRemainingShaders);
 		PreviewFeatureLevelChanged.Broadcast(EffectiveFeatureLevel);
 	}
 
diff --git a/Engine/Source/Runtime/Engine/Classes/Materials/Material.h b/Engine/Source/Runtime/Engine/Classes/Materials/Material.h
--- a/Engine/Source/Runtime/Engine/Classes/Materials/Material.h
+++ b/Engine/Source/Runtime/Engine/Classes/Materials/Material.h
@@ -1741,7 +1741,7 @@ public:
 #endif // WITH_EDITOR
 
 	// For all materials, UMaterial::CacheResourceShadersForRendering
-	ENGINE_API static void AllMaterialsCacheResourceShadersForRendering(bool bUpdateProgressDialog = false);
+	ENGINE_API static void AllMaterialsCacheResourceShadersForRendering(bool bUpdateProgressDialog = false, bool bCacheAllRemainingShaders = true);
 
 #if WITH_EDITORONLY_DATA
 	/**
diff --git a/Engine/Source/Runtime/Engine/Classes/Materials/MaterialInstance.h b/Engine/Source/Runtime/Engine/Classes/Materials/MaterialInstance.h
--- a/Engine/Source/Runtime/Engine/Classes/Materials/MaterialInstance.h
+++ b/Engine/Source/Runtime/Engine/Classes/Materials/MaterialInstance.h
@@ -845,7 +845,7 @@ public:
 	ENGINE_API virtual bool HasOverridenBaseProperties()const;
 
 	// For all materials instances, UMaterialInstance::CacheResourceShadersForRendering
-	ENGINE_API static void AllMaterialsCacheResourceShadersForRendering(bool bUpdateProgressDialog = false);
+	ENGINE_API static void AllMaterialsCacheResourceShadersForRendering(bool bUpdateProgressDialog = false, bool bCacheAllRemainingShaders = true);
 
 	/**
 	 * Determine whether this Material Instance is a child of another Material
diff --git a/Engine/Source/Runtime/Engine/Private/Materials/Material.cpp b/Engine/Source/Runtime/Engine/Private/Materials/Material.cpp
--- a/Engine/Source/Runtime/Engine/Private/Materials/Material.cpp
+++ b/Engine/Source/Runtime/Engine/Private/Materials/Material.cpp
@@ -5775,7 +5775,7 @@ UMaterial::FMaterialCompilationFinished& UMaterial::OnMaterialCompilationFinishe
 }
 #endif // WITH_EDITOR
 
-void UMaterial::AllMaterialsCacheResourceShadersForRendering(bool bUpdateProgressDialog)
+void UMaterial::AllMaterialsCacheResourceShadersForRendering(bool bUpdateProgressDialog, bool bCacheAllRemainingShaders)
 {
 #if STORE_ONLY_ACTIVE_SHADERMAPS
 	TArray<UMaterial*> Materials;
@@ -5813,7 +5813,7 @@ void UMaterial::AllMaterialsCacheResourceShadersForRendering(bool bUpdateProgres
 	{
 		UMaterial* Material = (UMaterial*)MaterialObj;
 
-		Material->CacheResourceShadersForRendering(false);
+		Material->CacheResourceShadersForRendering(false, bCacheAllRemainingShaders ? EMaterialShaderPrecompileMode::Default : EMaterialShaderPrecompileMode::None);
 
 #if WITH_EDITOR
 		if (bUpdateProgressDialog)
diff --git a/Engine/Source/Runtime/Engine/Private/Materials/MaterialInstance.cpp b/Engine/Source/Runtime/Engine/Private/Materials/MaterialInstance.cpp
--- a/Engine/Source/Runtime/Engine/Private/Materials/MaterialInstance.cpp
+++ b/Engine/Source/Runtime/Engine/Private/Materials/MaterialInstance.cpp
@@ -3795,7 +3795,7 @@ FPostProcessMaterialNode* IteratePostProcessMaterialNodes(const FFinalPostProces
 	}
 }
 
-void UMaterialInstance::AllMaterialsCacheResourceShadersForRendering(bool bUpdateProgressDialog)
+void UMaterialInstance::AllMaterialsCacheResourceShadersForRendering(bool bUpdateProgressDialog, bool bCacheAllRemainingShaders)
 {
 #if STORE_ONLY_ACTIVE_SHADERMAPS
 	TArray<UMaterialInstance*> MaterialInstances;
@@ -3827,7 +3827,7 @@ void UMaterialInstance::AllMaterialsCacheResourceShadersForRendering(bool bUpdat
 	{
 		UMaterialInstance* MaterialInstance = (UMaterialInstance*)MaterialInstanceObj;
 
-		MaterialInstance->CacheResourceShadersForRendering();
+		MaterialInstance->CacheResourceShadersForRendering(bCacheAllRemainingShaders ? EMaterialShaderPrecompileMode::Default : EMaterialShaderPrecompileMode::None);
 
 #if WITH_EDITOR
 		if (bUpdateProgressDialog)
