Cooking: Downgrade fatal assert to a warning if a package in memory has a non-mounted PackageRoot.
#rnx

diff --git a/Engine/Source/Editor/UnrealEd/Private/Cooker/CookPackageData.cpp b/Engine/Source/Editor/UnrealEd/Private/Cooker/CookPackageData.cpp
--- a/Engine/Source/Editor/UnrealEd/Private/Cooker/CookPackageData.cpp
+++ b/Engine/Source/Editor/UnrealEd/Private/Cooker/CookPackageData.cpp
@@ -1912,8 +1912,14 @@ bool FPackageDatas::TryLookupFileNameOnDisk(FName PackageName, FString& OutFileN
 	{
 		if (FindPackage(/*Outer =*/nullptr, *PackageNameStr))
 		{
-			return FPackageName::TryConvertLongPackageNameToFilename(PackageNameStr, OutFileName,
-				FPackageName::GetAssetPackageExtension());
+			if (!FPackageName::TryConvertLongPackageNameToFilename(PackageNameStr, OutFileName,
+				FPackageName::GetAssetPackageExtension()))
+			{
+				UE_LOG(LogCook, Warning, TEXT("Package %s exists in memory but its PackageRoot is not mounted. It will not be cooked."),
+					*PackageNameStr);
+				return false;
+			}
+			return true;
 		}
 		// else, the cooker could be responding to a NotifyUObjectCreated() event, and the object hasn't
 		// been fully constructed yet (missing from the FindObject() list) -- in this case, we've found 
