[Placement Mode] Fixed an issue where previewing a playset actor would cause odd flickering and jumping

# The spawned preview actor was not correctly flagged as a preview actor in the spawn parameters

#rnx

diff --git a/Engine/Source/Editor/UnrealEd/Private/Factories/ActorFactory.cpp b/Engine/Source/Editor/UnrealEd/Private/Factories/ActorFactory.cpp
--- a/Engine/Source/Editor/UnrealEd/Private/Factories/ActorFactory.cpp
+++ b/Engine/Source/Editor/UnrealEd/Private/Factories/ActorFactory.cpp
@@ -306,6 +306,8 @@ TArray<FTypedElementHandle> UActorFactory::PlaceAsset(const FAssetPlacementInfo&
 	FActorSpawnParameters SpawnParams;
 	SpawnParams.Name = InPlacementInfo.NameOverride;
 	SpawnParams.ObjectFlags = InPlacementOptions.bIsCreatingPreviewElements ? EObjectFlags::RF_Transient : EObjectFlags::RF_Transactional;
+	SpawnParams.bTemporaryEditorActor = InPlacementOptions.bIsCreatingPreviewElements;
+
 	AActor* NewActor = CreateActor(InPlacementInfo.AssetToPlace.GetAsset(), InPlacementInfo.PreferredLevel.Get(), InPlacementInfo.FinalizedTransform, SpawnParams);
 	if (NewActor)
 	{
