Strata - added support for material export to lightmass.
This works for legacy material behavior when using the conversion node, as well as what would in fact be expected from materials with emisive, basecolor, coverage and transmittance.

diff --git a/Engine/Shaders/Private/BasePassPixelShader.usf b/Engine/Shaders/Private/BasePassPixelShader.usf
--- a/Engine/Shaders/Private/BasePassPixelShader.usf
+++ b/Engine/Shaders/Private/BasePassPixelShader.usf
@@ -173,13 +173,13 @@ struct FShadingOcclusion
 
 #define MATERIAL_STRATA_OPAQUE_PRECOMPUTED_LIGHTING (MATERIAL_IS_STRATA && STRATA_ENABLED && STRATA_OPAQUE_DEFERRED)
 
-#if STRATA_TRANSLUCENT_FORWARD || STRATA_FORWARD_SHADING || MATERIAL_STRATA_OPAQUE_PRECOMPUTED_LIGHTING
+#if STRATA_TRANSLUCENT_FORWARD || STRATA_FORWARD_SHADING || MATERIAL_STRATA_OPAQUE_PRECOMPUTED_LIGHTING || STRATA_MATERIAL_EXPORT_EXECUTED
 #include "/Engine/Private/Strata/StrataEvaluation.ush"
 #endif
 #if STRATA_TRANSLUCENT_FORWARD || STRATA_FORWARD_SHADING
 #include "/Engine/Private/Strata/StrataForwardLighting.ush"
 #endif
-#if MATERIAL_STRATA_OPAQUE_PRECOMPUTED_LIGHTING
+#if MATERIAL_STRATA_OPAQUE_PRECOMPUTED_LIGHTING || STRATA_MATERIAL_EXPORT_EXECUTED
 #include "/Engine/Private/Strata/StrataExport.ush"
 #endif
 
@@ -1588,7 +1588,7 @@ void FPixelShaderInOut_MainPS(
 	#endif
 #endif
 
-#if MATERIAL_WORKS_WITH_DUAL_SOURCE_COLOR_BLENDING || STRATA_TRANSLUCENT_MATERIAL || STRATA_FORWARD_SHADING
+#if MATERIAL_WORKS_WITH_DUAL_SOURCE_COLOR_BLENDING || STRATA_TRANSLUCENT_MATERIAL || STRATA_FORWARD_SHADING || STRATA_MATERIAL_EXPORT_EXECUTED
 	float3 DualBlendSurfaceLuminancePostCoverage	= 0.0f;
 	float3 DualBlendSurfaceTransmittancePreCoverage = 1.0f;
 	float DualBlendSurfaceCoverage					= 1.0f;
@@ -1672,6 +1672,9 @@ void FPixelShaderInOut_MainPS(
 	}
 #endif // MATERIAL_SHADINGMODEL_THIN_TRANSLUCENT
 
+#if STRATA_MATERIAL_EXPORT_EXECUTED
+	float4 StrataMaterialExport = float4(0, 0, 0, 0);
+#endif
 
 	bool bStrataSubsurfaceEnable = false;
 #if STRATA_ENABLED && !STRATA_OPTIMIZED_UNLIT
@@ -1690,6 +1693,71 @@ void FPixelShaderInOut_MainPS(
 	}
 #endif
 
+	// Export non-unlit materials (opaque or translucent)
+#if STRATA_MATERIAL_EXPORT_EXECUTED
+
+	float  Coverage = 0;
+	float3 TransmittancePreCoverage = 0;
+	float3 BaseColorPostCoverage = 0;
+	float3 WorldNormal = 0;
+	float3 EmissiveLuminance = 0;
+
+	FStrataIntegrationSettings Settings = InitStrataIntegrationSettings(true /*bForceFullyRough*/, false /*StrataStruct.bRoughDiffuse*/, StrataStruct.PeelLayersAboveDepth);
+
+	const float3 SurfaceWorldNormal = MaterialParameters.TangentToWorld[2].xyz;
+	StrataMaterialExportOut(
+		Settings,
+		StrataPixelHeader,
+		StrataData,
+		SurfaceWorldNormal,
+		Coverage,
+		TransmittancePreCoverage,
+		BaseColorPostCoverage,
+		WorldNormal,
+		EmissiveLuminance);
+
+	#if STRATA_MATERIAL_EXPORT_FROM_OPAQUE
+		#if	  STRATA_MATERIAL_EXPORT_TYPE == STRATA_MATERIAL_EXPORT_DIFFUSE_COLOR
+			StrataMaterialExport = float4(BaseColorPostCoverage, 1.0);
+		#elif STRATA_MATERIAL_EXPORT_TYPE == STRATA_MATERIAL_EXPORT_NORMAL
+			StrataMaterialExport = float4(WorldNormal * 0.5 + 0.5, 1.0);
+		#elif STRATA_MATERIAL_EXPORT_TYPE == STRATA_MATERIAL_EXPORT_EMISSIVE
+			StrataMaterialExport = float4(EmissiveLuminance, 1.0);
+		#elif STRATA_MATERIAL_EXPORT_TYPE == STRATA_MATERIAL_EXPORT_TRANSMITTANCE
+			StrataMaterialExport = float4(PixelMaterialInputs.OpacityMask.xxx, 1.0f); // when not masked, OpacityMask default to 1.
+		#endif
+	#elif STRATA_MATERIAL_EXPORT_FROM_TRANSLUCENT
+		#if STRATA_USES_CONVERSION_FROM_LEGACY
+			#if	  STRATA_MATERIAL_EXPORT_TYPE == STRATA_MATERIAL_EXPORT_DIFFUSE_COLOR
+				StrataMaterialExport = float4(0.0, 0.0, 0.0, 1.0);
+			#elif STRATA_MATERIAL_EXPORT_TYPE == STRATA_MATERIAL_EXPORT_NORMAL
+				StrataMaterialExport = float4(0.0, 0.0, 0.0, 1.0);
+			#elif STRATA_MATERIAL_EXPORT_TYPE == STRATA_MATERIAL_EXPORT_EMISSIVE
+				StrataMaterialExport = float4(EmissiveLuminance, 1.0);
+			#elif STRATA_MATERIAL_EXPORT_TYPE == STRATA_MATERIAL_EXPORT_TRANSMITTANCE
+				#if STRATA_MATERIAL_EXPORT_LEGACY_BLEND_MODE == STRATA_MATERIAL_EXPORT_LEGACY_BLEND_MODULATE
+					StrataMaterialExport = float4(BaseColorPostCoverage, 1.0f);
+				#else
+					StrataMaterialExport = float4(BaseColorPostCoverage + saturate(1.0f - Coverage), 1.0f);
+				#endif
+			#endif
+		#else
+			#if	  STRATA_MATERIAL_EXPORT_TYPE == STRATA_MATERIAL_EXPORT_DIFFUSE_COLOR
+				StrataMaterialExport = float4(BaseColorPostCoverage, 1.0);
+			#elif STRATA_MATERIAL_EXPORT_TYPE == STRATA_MATERIAL_EXPORT_NORMAL
+				StrataMaterialExport = float4(WorldNormal * 0.5 + 0.5, 1.0);
+			#elif STRATA_MATERIAL_EXPORT_TYPE == STRATA_MATERIAL_EXPORT_EMISSIVE
+				StrataMaterialExport = float4(EmissiveLuminance, 1.0);
+			#elif STRATA_MATERIAL_EXPORT_TYPE == STRATA_MATERIAL_EXPORT_TRANSMITTANCE
+				StrataMaterialExport = float4(TransmittancePreCoverage * Coverage + (1.0f - Coverage), 1.0f);
+			#endif
+		#endif
+	#else
+	#error Unhandled STRATA_MATERIAL_EXPORT
+	#endif
+
+#else // STRATA_MATERIAL_EXPORT_EXECUTED
+
 #if STRATA_OPAQUE_DEFERRED
 
 #if STRATA_INLINE_SINGLELAYERWATER==0
@@ -1850,20 +1918,61 @@ void FPixelShaderInOut_MainPS(
 	}
 #endif // STRATA_TRANSLUCENT_FORWARD || STRATA_FORWARD_SHADING
 
+#endif // STRATA_MATERIAL_EXPORT_EXECUTED
+
 #elif STRATA_ENABLED && STRATA_OPTIMIZED_UNLIT
 
 	// Unlit forces a single BSDF
 	FStrataBSDF UnlitBSDF = PixelMaterialInputs.FrontMaterial.InlinedBSDF;
 
-#if STRATA_OPAQUE_DEFERRED || (STRATA_OPAQUE_MATERIAL && STRATA_FORWARD_SHADING)
+#if ((STRATA_OPAQUE_DEFERRED || (STRATA_OPAQUE_MATERIAL && STRATA_FORWARD_SHADING))) && !STRATA_MATERIAL_EXPORT_FROM_TRANSLUCENT
 	Color = BSDF_GETEMISSIVE(UnlitBSDF);
 	Opacity = 1.0f;
-#else // STRATA_TRANSLUCENT_FORWARD || STRATA_FORWARD_SHADING
+
+#if   STRATA_MATERIAL_EXPORT_TYPE == STRATA_MATERIAL_EXPORT_DIFFUSE_COLOR
+	StrataMaterialExport = float4(0.0, 0.0, 0.0, 1.0);
+#elif STRATA_MATERIAL_EXPORT_TYPE == STRATA_MATERIAL_EXPORT_NORMAL
+	StrataMaterialExport = float4(0.0, 0.0, 0.0, 1.0);
+#elif STRATA_MATERIAL_EXPORT_TYPE == STRATA_MATERIAL_EXPORT_EMISSIVE
+	StrataMaterialExport = float4(Color, 1.0);
+#elif STRATA_MATERIAL_EXPORT_TYPE == STRATA_MATERIAL_EXPORT_TRANSMITTANCE
+	StrataMaterialExport = float4(PixelMaterialInputs.OpacityMask.xxx, 1.0f); // when not masked, OpacityMask default to 1.
+#endif
+
+#else // STRATA_TRANSLUCENT_FORWARD || STRATA_FORWARD_SHADING || STRATA_MATERIAL_EXPORT_FROM_TRANSLUCENT
 	DualBlendSurfaceLuminancePostCoverage		= BSDF_GETEMISSIVE(UnlitBSDF);
 	DualBlendSurfaceTransmittancePreCoverage	= UNLIT_TRANSMITTANCE(UnlitBSDF);
 	DualBlendSurfaceCoverage					= UnlitBSDF.Coverage;
 	Color = 0.0f;
 	Opacity = 1.0f;
+
+#if STRATA_USES_CONVERSION_FROM_LEGACY
+	#if   STRATA_MATERIAL_EXPORT_TYPE == STRATA_MATERIAL_EXPORT_DIFFUSE_COLOR
+		StrataMaterialExport = float4(0.0, 0.0, 0.0, 1.0);
+	#elif STRATA_MATERIAL_EXPORT_TYPE == STRATA_MATERIAL_EXPORT_NORMAL
+		StrataMaterialExport = float4(0.0, 0.0, 0.0, 1.0);
+	#elif STRATA_MATERIAL_EXPORT_TYPE == STRATA_MATERIAL_EXPORT_EMISSIVE
+		StrataMaterialExport = float4(DualBlendSurfaceLuminancePostCoverage, 1.0);
+	#elif STRATA_MATERIAL_EXPORT_TYPE == STRATA_MATERIAL_EXPORT_TRANSMITTANCE
+		#if STRATA_MATERIAL_EXPORT_LEGACY_BLEND_MODE == STRATA_MATERIAL_EXPORT_LEGACY_BLEND_MODULATE
+		StrataMaterialExport = float4(DualBlendSurfaceLuminancePostCoverage, 1.0f);
+		#else
+		StrataMaterialExport = float4(DualBlendSurfaceLuminancePostCoverage + saturate(1.0 - DualBlendSurfaceCoverage.xxx), 1.0f);
+		#endif
+	#endif
+#else // STRATA_USES_CONVERSION_FROM_LEGACY
+	// The legacy unlit to lightmass conversion is weird and very specialized.
+	#if   STRATA_MATERIAL_EXPORT_TYPE == STRATA_MATERIAL_EXPORT_DIFFUSE_COLOR
+		StrataMaterialExport = float4(0.0, 0.0, 0.0, 1.0);
+	#elif STRATA_MATERIAL_EXPORT_TYPE == STRATA_MATERIAL_EXPORT_NORMAL
+		StrataMaterialExport = float4(0.0, 0.0, 0.0, 1.0);
+	#elif STRATA_MATERIAL_EXPORT_TYPE == STRATA_MATERIAL_EXPORT_EMISSIVE
+		StrataMaterialExport = float4(DualBlendSurfaceLuminancePostCoverage, 1.0);
+	#elif STRATA_MATERIAL_EXPORT_TYPE == STRATA_MATERIAL_EXPORT_TRANSMITTANCE
+		StrataMaterialExport = float4(DualBlendSurfaceTransmittancePreCoverage * DualBlendSurfaceCoverage + (1.0 - DualBlendSurfaceCoverage), 1.0f);
+	#endif
+#endif // STRATA_USES_CONVERSION_FROM_LEGACY
+
 #endif
 	
 #endif
@@ -2147,6 +2256,11 @@ void FPixelShaderInOut_MainPS(
 		View.VTFeedbackBuffer
 	);
 #endif
+
+#if STRATA_MATERIAL_EXPORT_EXECUTED
+	Out.MRT[0] = StrataMaterialExport;
+#endif
+
 }
 
 #if STRATA_OPAQUE_DEFERRED
diff --git a/Engine/Shaders/Private/Definitions.usf b/Engine/Shaders/Private/Definitions.usf
--- a/Engine/Shaders/Private/Definitions.usf
+++ b/Engine/Shaders/Private/Definitions.usf
@@ -214,4 +214,42 @@
 
 #ifndef FORCE_MATERIAL_FLOAT_FULL_PRECISION
 #define FORCE_MATERIAL_FLOAT_FULL_PRECISION             0
-#endif
\ No newline at end of file
+#endif
+
+///////////////////////////////////////////////////////////////////////////////
+// Global Strata material export related defines
+
+#ifndef STRATA_MATERIAL_EXPORT_TYPE
+#define STRATA_MATERIAL_EXPORT_TYPE						0
+#endif
+
+#ifndef STRATA_MATERIAL_EXPORT_CONTEXT
+#define STRATA_MATERIAL_EXPORT_CONTEXT					0
+#endif
+
+#ifndef STRATA_MATERIAL_EXPORT_LEGACY_BLEND_MODE
+#define STRATA_MATERIAL_EXPORT_LEGACY_BLEND_MODE		0
+#endif
+
+// Those define usage should match the meaning of EStrataMaterialExport enum from MaterialCompiler.h
+#define STRATA_MATERIAL_EXPORT_DIFFUSE_COLOR				1
+#define STRATA_MATERIAL_EXPORT_NORMAL						2
+#define STRATA_MATERIAL_EXPORT_EMISSIVE						3
+#define STRATA_MATERIAL_EXPORT_TRANSMITTANCE				4
+
+#define STRATA_MATERIAL_EXPORT_CONTEXT_OPAQUE				0
+#define STRATA_MATERIAL_EXPORT_CONTEXT_TRANSLUCENT			1
+
+#define STRATA_MATERIAL_EXPORT_LEGACY_BLEND_OPAQUE			0
+#define STRATA_MATERIAL_EXPORT_LEGACY_BLEND_MASKED			1
+#define STRATA_MATERIAL_EXPORT_LEGACY_BLEND_TRANSLUCENT		2
+#define STRATA_MATERIAL_EXPORT_LEGACY_BLEND_ADDITIVE		3
+#define STRATA_MATERIAL_EXPORT_LEGACY_BLEND_MODULATE		4
+#define STRATA_MATERIAL_EXPORT_LEGACY_BLEND_ALPHACOMPOSITE	5
+#define STRATA_MATERIAL_EXPORT_LEGACY_BLEND_ALPHAHOLDOUT	6
+
+#define STRATA_MATERIAL_EXPORT_EXECUTED						(STRATA_MATERIAL_EXPORT_TYPE > 0)
+#define STRATA_MATERIAL_EXPORT_FROM_OPAQUE					(STRATA_MATERIAL_EXPORT_EXECUTED && (STRATA_MATERIAL_EXPORT_CONTEXT == STRATA_MATERIAL_EXPORT_CONTEXT_OPAQUE))
+#define STRATA_MATERIAL_EXPORT_FROM_TRANSLUCENT				(STRATA_MATERIAL_EXPORT_EXECUTED && (STRATA_MATERIAL_EXPORT_CONTEXT == STRATA_MATERIAL_EXPORT_CONTEXT_TRANSLUCENT))
+
+
diff --git a/Engine/Shaders/Private/MaterialTemplate.ush b/Engine/Shaders/Private/MaterialTemplate.ush
--- a/Engine/Shaders/Private/MaterialTemplate.ush
+++ b/Engine/Shaders/Private/MaterialTemplate.ush
@@ -149,8 +149,8 @@
 #endif
 
 #define TEMPLATE_USES_STRATA			(STRATA_ENABLED && MATERIAL_IS_STRATA)
-#define STRATA_OPAQUE_MATERIAL			(TEMPLATE_USES_STRATA && MATERIALBLENDING_ANY_TRANSLUCENT<=0)
-#define STRATA_TRANSLUCENT_MATERIAL		(TEMPLATE_USES_STRATA && MATERIALBLENDING_ANY_TRANSLUCENT>0)
+#define STRATA_OPAQUE_MATERIAL			(TEMPLATE_USES_STRATA && ((!STRATA_MATERIAL_EXPORT_EXECUTED && MATERIALBLENDING_ANY_TRANSLUCENT<=0) || STRATA_MATERIAL_EXPORT_FROM_OPAQUE))
+#define STRATA_TRANSLUCENT_MATERIAL		(TEMPLATE_USES_STRATA && ((!STRATA_MATERIAL_EXPORT_EXECUTED && MATERIALBLENDING_ANY_TRANSLUCENT>0 ) || STRATA_MATERIAL_EXPORT_FROM_TRANSLUCENT))
 #define STRATA_OPAQUE_DEFERRED			(STRATA_OPAQUE_MATERIAL && !FORWARD_SHADING)
 #define STRATA_TRANSLUCENT_FORWARD		(STRATA_TRANSLUCENT_MATERIAL && !FORWARD_SHADING)
 #define STRATA_FORWARD_SHADING			(FORWARD_SHADING && (STRATA_OPAQUE_MATERIAL || STRATA_TRANSLUCENT_MATERIAL))
@@ -172,7 +172,7 @@ FStrataData GetInitialisedStrataData() { return (FStrataData)0; }
 
 #if TEMPLATE_USES_STRATA
 #include "/Engine/Private/Strata/StrataTree.ush"
-#include "/Engine/Private/Strata/StrataLegacyConversion.ush"
+#include "/Engine/Private/Strata/StrataLegacyConversion.ush"	// This can only included here due to interaction between STRATA_MATERIAL_EXPORT_EXECUTED and STRATA_LEGACY_MATERIAL_APPLIES_FINAL_WEIGHT
 #endif
 
 #include "/Engine/Private/DBufferDecalShared.ush"
diff --git a/Engine/Shaders/Private/Strata/Strata.ush b/Engine/Shaders/Private/Strata/Strata.ush
--- a/Engine/Shaders/Private/Strata/Strata.ush
+++ b/Engine/Shaders/Private/Strata/Strata.ush
@@ -87,7 +87,7 @@
 // When converting legacy materials to Strata, we need to override transmittance/coverage to be independent from the material and directly linked to opacity
 #define STRATA_LEGACY_PREMULT_ALPHA_OVERRIDE (STRATA_USES_CONVERSION_FROM_LEGACY && MATERIALBLENDING_ALPHACOMPOSITE)
 
-#define STRATA_LEGACY_MATERIAL_APPLIES_FINAL_WEIGHT (!MATERIALBLENDING_MASKED && !MATERIALBLENDING_SOLID && STRATA_LEGACY_PREMULT_ALPHA_OVERRIDE==0)
+#define STRATA_LEGACY_MATERIAL_APPLIES_FINAL_WEIGHT (!MATERIALBLENDING_MASKED && !MATERIALBLENDING_SOLID && STRATA_LEGACY_PREMULT_ALPHA_OVERRIDE==0 || STRATA_MATERIAL_EXPORT_FROM_TRANSLUCENT)
 
 #define STRATA_INLINE_SINGLELAYERWATER (STRATA_ENABLED && MATERIAL_IS_STRATA && STRATA_INLINE_SHADING && MATERIAL_SHADINGMODEL_SINGLELAYERWATER)
 
diff --git a/Engine/Shaders/Private/Strata/StrataExport.ush b/Engine/Shaders/Private/Strata/StrataExport.ush
--- a/Engine/Shaders/Private/Strata/StrataExport.ush
+++ b/Engine/Shaders/Private/Strata/StrataExport.ush
@@ -549,4 +549,68 @@ void PackStrataOut(
 
 }
 
+#if STRATA_MATERIAL_EXPORT_TYPE > 0
+
+void StrataMaterialExportOut(
+	in FStrataIntegrationSettings Settings,
+	in FStrataPixelHeader StrataPixelHeader,
+	in FStrataData StrataData,
+	in float3 SurfaceWorldNormal,
+	inout float  OutCoverage, 
+	inout float3 OutTransmittancePreCoverage, 
+	inout float3 OutBaseColorPostCoverage, 
+	inout float3 OutWorldNormal,
+	inout float3 OutEmissiveLuminance)
+{
+	FStrataAddressing NullStrataAddressing = (FStrataAddressing)0;
+
+	// Align V with N to not have any view dependent effect when exporting material data.
+	const float3 FakeV = SurfaceWorldNormal;
+
+	OutCoverage = 0.0f;
+	OutTransmittancePreCoverage = 0.0f;
+	OutBaseColorPostCoverage = 0.0f;
+	OutWorldNormal = 0.0f;
+	OutEmissiveLuminance = 0.0f;
+	
+	FStrataTree StrataTree = StrataPixelHeader.StrataTree;
+	if (StrataTree.BSDFCount > 0)
+	{
+		// Update tree (coverage/transmittance/luminace weights)
+		StrataUpdateTree(NullStrataAddressing, StrataPixelHeader, StrataTree, StrataData, FakeV, Settings, OutCoverage, OutTransmittancePreCoverage);
+
+		//
+		// Forward lighting
+		//
+		STRATA_UNROLL_N(STRATA_CLAMPED_BSDF_COUNT)
+		for (int BSDFIdx = 0; BSDFIdx < StrataTree.BSDFCount; ++BSDFIdx)
+		{
+		#define CurrentBSDF StrataTree.BSDFs[BSDFIdx]
+
+			if (StrataIsBSDFVisible(CurrentBSDF))
+			{
+				// Starting with environment lighting, the context does not specify a light
+				FStrataBSDFContext StrataBSDFContext = StrataCreateBSDFContext(StrataPixelHeader, CurrentBSDF, NullStrataAddressing, FakeV);
+
+				const float3x3 TangentBasis = StrataGetBSDFSharedBasis_InlineShading(StrataPixelHeader, CurrentBSDF, NullStrataAddressing);
+				const float3 N = TangentBasis[2];
+
+				const float GreyScaleWeight = dot(CurrentBSDF.LuminanceWeightV * CurrentBSDF.TransmittanceAboveAlongN, (1.0/3.0).xxx);
+
+				// Simply sum all BaseColors according to transmittance and coverage.
+				OutBaseColorPostCoverage += StrataGetBSDFBaseColor(CurrentBSDF) * CurrentBSDF.LuminanceWeightV * CurrentBSDF.TransmittanceAboveAlongN;
+
+				OutWorldNormal += N * GreyScaleWeight;
+
+				OutEmissiveLuminance += BSDF_GETEMISSIVE(CurrentBSDF) * CurrentBSDF.LuminanceWeightV;
+			}
+		#undef CurrentBSDF
+		}
+
+		OutWorldNormal = all(OutWorldNormal == 0) ? float3(0.0, 0.0, 0.0) : normalize(OutWorldNormal);
+	}
+}
+
+#endif // STRATA_MATERIAL_EXPORT_TYPE
+
 #endif // STRATA_ENABLED
diff --git a/Engine/Source/Editor/UnrealEd/Private/Lightmass/LightmassRender.cpp b/Engine/Source/Editor/UnrealEd/Private/Lightmass/LightmassRender.cpp
--- a/Engine/Source/Editor/UnrealEd/Private/Lightmass/LightmassRender.cpp
+++ b/Engine/Source/Editor/UnrealEd/Private/Lightmass/LightmassRender.cpp
@@ -28,6 +28,12 @@
 
 DEFINE_LOG_CATEGORY_STATIC(LogLightmassRender, Error, All);
 
+bool Lightmass_IsStrataEnabled()
+{
+	static const auto CVar = IConsoleManager::Get().FindTConsoleVariableDataInt(TEXT("r.Strata"));
+	return CVar && CVar->GetValueOnAnyThread() > 0;
+}
+
 // FLightmassMaterialCompiler - A proxy compiler that overrides various compiler functions for potential problem expressions.
 struct FLightmassMaterialCompiler : public FProxyMaterialCompiler
 {
@@ -240,6 +246,7 @@ public:
 	{
 		if (InMaterialInterface)
 		{
+			bStrataEnabled = Lightmass_IsStrataEnabled();
 			MaterialInterface = InMaterialInterface;
 			Material = MaterialInterface ? MaterialInterface->GetMaterial() : NULL;
 			PropertyToCompile = InPropertyToCompile;
@@ -347,8 +354,6 @@ public:
 	/** helper for CompilePropertyAndSetMaterialProperty() */
 	int32 CompilePropertyAndSetMaterialPropertyWithoutCast(EMaterialProperty Property, FMaterialCompiler* Compiler) const
 	{
-		EMaterialProperty DiffuseInput = MP_BaseColor;
-
 		// MAKE SURE THIS MATCHES THE CHART IN WillFillData
 		// 						  RETURNED VALUES (F16 'textures')
 		// 	BLEND MODE  | DIFFUSE     | SPECULAR     | EMISSIVE    | NORMAL    | TRANSMISSIVE              |
@@ -359,6 +364,34 @@ public:
 		// 	Additive    | 0 (EMPTY)   | 0 (EMPTY)    | Emissive    | 0 (EMPTY) | (Emsv | Diffuse)*Opacity  |
 		// 	Modulative  | 0 (EMPTY)   | 0 (EMPTY)    | Emissive    | 0 (EMPTY) | Emsv | Diffuse            |
 		// 	------------+-------------+--------------+-------------+-----------+---------------------------|
+
+		const uint32 ForceCast_Exact_Replicate = MFCF_ForceCast | MFCF_ExactMatch | MFCF_ReplicateValue;
+		const EMaterialProperty DiffuseInput = MP_BaseColor;
+
+		if (bStrataEnabled)
+		{
+			EStrataBlendMode StrataBlendMode = MaterialInterface->GetStrataBlendMode();
+			uint8 LegacyBlendMode = MaterialInterface->GetBlendMode();
+			EStrataMaterialExportContext StrataMaterialExportContext = (StrataBlendMode == SBM_Opaque || StrataBlendMode == SBM_Masked) ? EStrataMaterialExportContext::SMEC_Opaque : EStrataMaterialExportContext::SMEC_Translucent;
+
+			if (Usage == EMaterialShaderMapUsage::LightmassExportDiffuse)
+			{
+				Compiler->SetStrataMaterialExportType(SME_Diffuse, StrataMaterialExportContext, LegacyBlendMode);
+			}
+			else if (Usage == EMaterialShaderMapUsage::LightmassExportNormal)
+			{
+				Compiler->SetStrataMaterialExportType(SME_Normal, StrataMaterialExportContext, LegacyBlendMode);
+			}
+			else if (Usage == EMaterialShaderMapUsage::LightmassExportOpacity)
+			{
+				Compiler->SetStrataMaterialExportType(SME_Transmittance, StrataMaterialExportContext, LegacyBlendMode);
+			}
+			else if (Usage == EMaterialShaderMapUsage::LightmassExportEmissive)
+			{
+				Compiler->SetStrataMaterialExportType(SME_Emissive, StrataMaterialExportContext, LegacyBlendMode);
+			}
+		}
+
 		if( Property == MP_EmissiveColor )
 		{
 			UMaterial* ProxyMaterial = MaterialInterface->GetMaterial();
@@ -367,8 +400,6 @@ public:
 			check(ProxyMaterial);
 			FLightmassMaterialCompiler ProxyCompiler(Compiler);
 
-			const uint32 ForceCast_Exact_Replicate = MFCF_ForceCast | MFCF_ExactMatch | MFCF_ReplicateValue;
-
 			switch (PropertyToCompile)
 			{
 			case MP_EmissiveColor:
@@ -434,7 +465,17 @@ public:
 				break;
 			case MP_ShadingModel:
 				return MaterialInterface->CompileProperty(&ProxyCompiler, MP_ShadingModel);
-			// STRATA_TODO
+			case MP_FrontMaterial:
+				if (bStrataEnabled)
+				{
+					// When using strata, material property always compile from material. 
+					// We cannot use rediction so instead we instruct the compiler the type of data export we are looking for.
+					return MaterialInterface->CompileProperty(&ProxyCompiler, MP_FrontMaterial);
+				}
+				else
+				{
+					return ProxyCompiler.StrataCreateAndRegisterNullMaterial();
+				}
 			default:
 				return Compiler->Constant(1.0f);
 			}
@@ -453,11 +494,24 @@ public:
 		}
 		else if (Property == MP_ShadingModel)
 		{
-			return MaterialInterface->CompileProperty(Compiler, MP_ShadingModel);
+			return MaterialInterface->CompileProperty(Compiler, MP_ShadingModel); // useless with strata
+		}
+		// When using strata, we need to actually compile more root node inputs, 
+		// and then we handle what needs to actually be exported from the shader code (see STRATA_MATERIAL_EXPORT_TYPE).
+		else if (Property == MP_OpacityMask)
+		{
+			return MaterialInterface->CompileProperty(Compiler, MP_OpacityMask);
 		}
 		else if (Property == MP_FrontMaterial)
 		{
-			return MaterialInterface->CompileProperty(Compiler, MP_FrontMaterial);
+			if (bStrataEnabled)
+			{
+				return MaterialInterface->CompileProperty(Compiler, MP_FrontMaterial);
+			}
+			else
+			{
+				return Compiler->StrataCreateAndRegisterNullMaterial();
+			}
 		}
 		else
 		{
@@ -565,35 +619,44 @@ public:
 	bool IsMaterialInputConnected(UMaterial* InMaterial, EMaterialProperty MaterialInput)
 	{
 		bool bConnected = false;
-
 		UMaterialEditorOnlyData* MaterialEditorOnly = InMaterial->GetEditorOnlyData();
-		switch (MaterialInput)
+
+		if (bStrataEnabled)
 		{
-		case MP_EmissiveColor:
-			bConnected = MaterialEditorOnly->EmissiveColor.Expression != nullptr;
-			break;
-		case MP_DiffuseColor:
-			bConnected = MaterialEditorOnly->BaseColor.Expression != nullptr;
-			break;
-		case MP_SpecularColor:
-			bConnected = MaterialEditorOnly->Specular.Expression != nullptr;
-			break;
-		case MP_Normal:
-			bConnected = MaterialEditorOnly->Normal.Expression != nullptr;
-			break;
-		case MP_Opacity:
-			bConnected = MaterialEditorOnly->Opacity.Expression != nullptr;
-			break;
-		case MP_OpacityMask:
-			bConnected = MaterialEditorOnly->OpacityMask.Expression != nullptr;
-			break;
-		default:
-			break;
+			// Material attribute do not override the FrontMaterial input
+			bConnected = MaterialEditorOnly->FrontMaterial.Expression != nullptr;
+		}
+		else
+		{
+			switch (MaterialInput)
+			{
+			case MP_EmissiveColor:
+				bConnected = MaterialEditorOnly->EmissiveColor.Expression != nullptr;
+				break;
+			case MP_DiffuseColor:
+				bConnected = MaterialEditorOnly->BaseColor.Expression != nullptr;
+				break;
+			case MP_SpecularColor:
+				bConnected = MaterialEditorOnly->Specular.Expression != nullptr;
+				break;
+			case MP_Normal:
+				bConnected = MaterialEditorOnly->Normal.Expression != nullptr;
+				break;
+			case MP_Opacity:
+				bConnected = MaterialEditorOnly->Opacity.Expression != nullptr;
+				break;
+			case MP_OpacityMask:
+				bConnected = MaterialEditorOnly->OpacityMask.Expression != nullptr;
+				break;
+			default:
+				break;
+			}
+
+			// Note: only checking to see whether the entire material attributes connection exists.  
+			// This means materials using the material attributes input will export more attributes than is necessary.
+			bConnected = InMaterial->bUseMaterialAttributes ? MaterialEditorOnly->MaterialAttributes.Expression != NULL : bConnected;
 		}
 
-		// Note: only checking to see whether the entire material attributes connection exists.  
-		// This means materials using the material attributes input will export more attributes than is necessary.
-		bConnected = InMaterial->bUseMaterialAttributes ? MaterialEditorOnly->MaterialAttributes.Expression != NULL : bConnected;
 		return bConnected;
 	}
 
@@ -614,78 +677,116 @@ public:
 		OutUniformValue.B = 0.0f;
 		OutUniformValue.A = 0.0f;
 
-		EBlendMode BlendMode = MaterialInterface->GetBlendMode();
-		bool bIsMaterialUnlit = MaterialInterface->GetShadingModels().IsUnlit();
-		
 		check(Material);
 		bool bExpressionIsNULL = false;
-		switch (PropertyToCompile)
+
+		if (bStrataEnabled)
 		{
-		case MP_EmissiveColor:
-			// Emissive is ALWAYS returned...
-			bExpressionIsNULL = !IsMaterialInputConnected(Material, PropertyToCompile);
-			break;
-		case MP_DiffuseColor:
-			// Only return for Opaque and Masked...
-			if (BlendMode == BLEND_Opaque || BlendMode == BLEND_Masked)
-			{
-				bExpressionIsNULL = !IsMaterialInputConnected(Material, PropertyToCompile);
-			}
-			break;
-		case MP_SpecularColor: 
-			// Only return for Opaque and Masked...
-			if (BlendMode == BLEND_Opaque || BlendMode == BLEND_Masked)
+			EStrataBlendMode StrataBlendMode = MaterialInterface->GetStrataBlendMode();
+
+			switch (Usage)
 			{
-				bExpressionIsNULL = !IsMaterialInputConnected(Material, PropertyToCompile);
-				OutUniformValue.A = 15.0f;
+			case EMaterialShaderMapUsage::LightmassExportEmissive:
+				bExpressionIsNULL = !IsMaterialInputConnected(Material, MP_FrontMaterial);
+				break;
+			case EMaterialShaderMapUsage::LightmassExportDiffuse:
+				if (StrataBlendMode == SBM_Opaque || StrataBlendMode == SBM_Masked)
+				{
+					bExpressionIsNULL = !IsMaterialInputConnected(Material, MP_FrontMaterial);
+				}
+				break;
+			case EMaterialShaderMapUsage::LightmassExportOpacity:
+				if (StrataBlendMode != SBM_Opaque)
+				{
+					bExpressionIsNULL = !IsMaterialInputConnected(Material, MP_FrontMaterial);
+					OutUniformValue.A = 15.0f;
+				}
+				break;
+			case EMaterialShaderMapUsage::LightmassExportNormal:
+				if (StrataBlendMode == SBM_Opaque || StrataBlendMode == SBM_Masked)
+				{
+					bExpressionIsNULL = !IsMaterialInputConnected(Material, MP_FrontMaterial);
+					OutUniformValue.B = 1.0f;	// Default normal is (0,0,1)
+				}
+				break;
+			default:
+				UE_LOG(LogLightmassRender, Error, TEXT("WillGenerateUniformData - cannot export a requested property for %s"), *(Material->GetPathName()));
+				break;
 			}
-			break;
-		case MP_Normal:
-			// Only return for Opaque and Masked...
-			if (BlendMode == BLEND_Opaque || BlendMode == BLEND_Masked)
+		}
+		else
+		{
+			bool bIsMaterialUnlit = MaterialInterface->GetShadingModels().IsUnlit();
+			EBlendMode BlendMode = MaterialInterface->GetBlendMode();
+
+			switch (PropertyToCompile)
 			{
+			case MP_EmissiveColor:
+				// Emissive is ALWAYS returned...
 				bExpressionIsNULL = !IsMaterialInputConnected(Material, PropertyToCompile);
-				OutUniformValue.B = 1.0f;	// Default normal is (0,0,1)
-			}
-			break;
-		case MP_Opacity:
-			if (BlendMode == BLEND_Masked)
-			{
-				bExpressionIsNULL = !IsMaterialInputConnected(Material, MP_OpacityMask);
-				OutUniformValue.R = 1.0f;
-				OutUniformValue.G = 1.0f;
-				OutUniformValue.B = 1.0f;
-				OutUniformValue.A = 1.0f;
-			}
-			else
-			if ((BlendMode == BLEND_Modulate) ||
-				(BlendMode == BLEND_Translucent) || 
-				(BlendMode == BLEND_Additive) ||
-				(BlendMode == BLEND_AlphaComposite) ||
-				(BlendMode == BLEND_AlphaHoldout))
-			{
-				bool bColorInputIsNULL = false;
-				if (bIsMaterialUnlit)
+				break;
+			case MP_DiffuseColor:
+				// Only return for Opaque and Masked...
+				if (BlendMode == BLEND_Opaque || BlendMode == BLEND_Masked)
 				{
-					bColorInputIsNULL = !IsMaterialInputConnected(Material, MP_EmissiveColor);
+					bExpressionIsNULL = !IsMaterialInputConnected(Material, PropertyToCompile);
 				}
-				else
+				break;
+			case MP_SpecularColor: 
+				// Only return for Opaque and Masked...
+				if (BlendMode == BLEND_Opaque || BlendMode == BLEND_Masked)
+				{
+					bExpressionIsNULL = !IsMaterialInputConnected(Material, PropertyToCompile);
+					OutUniformValue.A = 15.0f;
+				}
+				break;
+			case MP_Normal:
+				// Only return for Opaque and Masked...
+				if (BlendMode == BLEND_Opaque || BlendMode == BLEND_Masked)
 				{
-					bColorInputIsNULL = !IsMaterialInputConnected(Material, MP_DiffuseColor);
+					bExpressionIsNULL = !IsMaterialInputConnected(Material, PropertyToCompile);
+					OutUniformValue.B = 1.0f;	// Default normal is (0,0,1)
 				}
-				if (BlendMode == BLEND_Translucent
-					|| BlendMode == BLEND_Additive
-					|| BlendMode == BLEND_AlphaComposite
-					|| BlendMode == BLEND_AlphaHoldout)
+				break;
+			case MP_Opacity:
+				if (BlendMode == BLEND_Masked)
 				{
-					bExpressionIsNULL = bColorInputIsNULL && !IsMaterialInputConnected(Material, PropertyToCompile);
+					bExpressionIsNULL = !IsMaterialInputConnected(Material, MP_OpacityMask);
+					OutUniformValue.R = 1.0f;
+					OutUniformValue.G = 1.0f;
+					OutUniformValue.B = 1.0f;
+					OutUniformValue.A = 1.0f;
 				}
 				else
+				if ((BlendMode == BLEND_Modulate) ||
+					(BlendMode == BLEND_Translucent) || 
+					(BlendMode == BLEND_Additive) ||
+					(BlendMode == BLEND_AlphaComposite) ||
+					(BlendMode == BLEND_AlphaHoldout))
 				{
-					bExpressionIsNULL = bColorInputIsNULL;
+					bool bColorInputIsNULL = false;
+					if (bIsMaterialUnlit)
+					{
+						bColorInputIsNULL = !IsMaterialInputConnected(Material, MP_EmissiveColor);
+					}
+					else
+					{
+						bColorInputIsNULL = !IsMaterialInputConnected(Material, MP_DiffuseColor);
+					}
+					if (BlendMode == BLEND_Translucent
+						|| BlendMode == BLEND_Additive
+						|| BlendMode == BLEND_AlphaComposite
+						|| BlendMode == BLEND_AlphaHoldout)
+					{
+						bExpressionIsNULL = bColorInputIsNULL && !IsMaterialInputConnected(Material, PropertyToCompile);
+					}
+					else
+					{
+						bExpressionIsNULL = bColorInputIsNULL;
+					}
 				}
+				break;
 			}
-			break;
 		}
 
 		return bExpressionIsNULL;
@@ -742,7 +843,7 @@ public:
 		return true;
 	}
 
-	static bool WillFillData(EBlendMode InBlendMode, EMaterialProperty InMaterialProperty)
+	static bool WillFillData(EBlendMode InBlendMode, EMaterialProperty InMaterialProperty, bool bStrataEnabled, EStrataBlendMode InStrataBlendMode)
 	{
 		// MAKE SURE THIS MATCHES THE CHART IN CompileProperty
 		// 						  RETURNED VALUES (F16 'textures')
@@ -755,11 +856,36 @@ public:
 		// 	Modulative  | 0 (EMPTY)   | 0 (EMPTY)    | Emissive    | 0 (EMPTY) | Emsv | Diffuse            |
 		// 	------------+-------------+--------------+-------------+-----------+---------------------------|
 
+		// Emissive will always fill data.
 		if (InMaterialProperty == MP_EmissiveColor)
 		{
 			return true;
 		}
 
+		if (bStrataEnabled)
+		{
+			switch (InMaterialProperty)
+			{
+			case MP_DiffuseColor:
+			{
+				return InStrataBlendMode == SBM_Opaque || InStrataBlendMode == SBM_Masked;
+			}
+			case MP_Normal:
+			{
+				return InStrataBlendMode == SBM_Opaque || InStrataBlendMode == SBM_Masked;
+			}
+			case MP_Opacity:
+			{
+				return InStrataBlendMode != SBM_Opaque;
+				break;
+			}
+			default:
+			{
+				UE_LOG(LogLightmassRender, Error, TEXT("FLightmassMaterialProxy::WillFillData - cannot export a requested property for"));
+			}
+			}
+		}
+
 		switch (InBlendMode)
 		{
 		case BLEND_Opaque:
@@ -838,6 +964,8 @@ private:
 	EMaterialProperty PropertyToCompile;
 	/** Stores which exported attribute this proxy is compiling for. */
 	EMaterialShaderMapUsage::Type Usage;
+	/** If Strata is enabled, we need to specify things differently since redirection cannot straiforwardly be used*/
+	bool bStrataEnabled;
 };
 
 FMaterialExportDataEntry::~FMaterialExportDataEntry()
@@ -886,22 +1014,24 @@ void FLightmassMaterialRenderer::BeginGenerateMaterialData(
 	if (BaseMaterial)
 	{
 		check(!MaterialExportData.Contains(InMaterial));
+		const bool bStrataEnabled = Lightmass_IsStrataEnabled();
+		EStrataBlendMode StrataBlendMode = InMaterial->GetStrataBlendMode();
 
 		FMaterialExportDataEntry& MaterialData = MaterialExportData.Add(InMaterial, FMaterialExportDataEntry(ChannelName));
 
-		if (FLightmassMaterialProxy::WillFillData(BlendMode, MP_DiffuseColor))
+		if (FLightmassMaterialProxy::WillFillData(BlendMode, MP_DiffuseColor, bStrataEnabled, StrataBlendMode))
 		{
 			MaterialData.DiffuseMaterialProxy = new FLightmassMaterialProxy();
 			MaterialData.DiffuseMaterialProxy->BeginCompiling(InMaterial, MP_DiffuseColor, EMaterialShaderMapUsage::LightmassExportDiffuse);
 		}
 
-		if (FLightmassMaterialProxy::WillFillData(BlendMode, MP_EmissiveColor))
+		if (FLightmassMaterialProxy::WillFillData(BlendMode, MP_EmissiveColor, bStrataEnabled, StrataBlendMode))
 		{
 			MaterialData.EmissiveMaterialProxy = new FLightmassMaterialProxy();
 			MaterialData.EmissiveMaterialProxy->BeginCompiling(InMaterial, MP_EmissiveColor, EMaterialShaderMapUsage::LightmassExportEmissive);
 		}
 
-		if (FLightmassMaterialProxy::WillFillData(BlendMode, MP_Opacity))
+		if (FLightmassMaterialProxy::WillFillData(BlendMode, MP_Opacity, bStrataEnabled, StrataBlendMode))
 		{
 			// Landscape opacity is generated from the hole mask, not the material
 			if (!bIsLandscapeMaterial)
@@ -911,7 +1041,7 @@ void FLightmassMaterialRenderer::BeginGenerateMaterialData(
 			}
 		}
 
-		if (bInWantNormals && FLightmassMaterialProxy::WillFillData(BlendMode, MP_Normal))
+		if (bInWantNormals && FLightmassMaterialProxy::WillFillData(BlendMode, MP_Normal, bStrataEnabled, StrataBlendMode))
 		{
 			MaterialData.NormalMaterialProxy = new FLightmassMaterialProxy();
 			MaterialData.NormalMaterialProxy->BeginCompiling(InMaterial, MP_Normal, EMaterialShaderMapUsage::LightmassExportNormal);
@@ -946,8 +1076,11 @@ bool FLightmassMaterialRenderer::GenerateMaterialData(
 	check(BaseMaterial);
 
 	EBlendMode BlendMode = InMaterial.GetBlendMode();
+	const bool bStrataEnabled = Lightmass_IsStrataEnabled();
+
 	FMaterialShadingModelField ShadingModels = InMaterial.GetShadingModels();
- 	if (!ShadingModels.HasShadingModel(MSM_DefaultLit) &&
+ 	if (!bStrataEnabled &&		// Shading models are irrelevant when using Strata
+		!ShadingModels.HasShadingModel(MSM_DefaultLit) &&
 		!ShadingModels.HasShadingModel(MSM_Unlit) &&
 		!ShadingModels.HasShadingModel(MSM_Subsurface) &&
 		!ShadingModels.HasShadingModel(MSM_PreintegratedSkin) &&
diff --git a/Engine/Source/Runtime/Engine/Private/Materials/HLSLMaterialTranslator.cpp b/Engine/Source/Runtime/Engine/Private/Materials/HLSLMaterialTranslator.cpp
--- a/Engine/Source/Runtime/Engine/Private/Materials/HLSLMaterialTranslator.cpp
+++ b/Engine/Source/Runtime/Engine/Private/Materials/HLSLMaterialTranslator.cpp
@@ -1891,6 +1891,10 @@ void FHLSLMaterialTranslator::GetMaterialEnvironment(EShaderPlatform InPlatform,
 
 	if (bMaterialIsStrata)
 	{
+		OutEnvironment.SetDefine(TEXT("STRATA_MATERIAL_EXPORT_TYPE"), (int32)GetStrataMaterialExportType());
+		OutEnvironment.SetDefine(TEXT("STRATA_MATERIAL_EXPORT_CONTEXT"), (int32)GetStrataMaterialExportContext());
+		OutEnvironment.SetDefine(TEXT("STRATA_MATERIAL_EXPORT_LEGACY_BLEND_MODE"), (int32)GetStrataMaterialExportLegacyBlendMode());
+
 		check(StrataMaterialRootOperator);
 		uint32 RootMaximumDistanceToLeaves = StrataMaterialRootOperator->MaxDistanceFromLeaves;
 		OutEnvironment.SetDefine(TEXT("MATERIAL_TOPOLOGY_TREE_MAX_DEPTH"), RootMaximumDistanceToLeaves);
diff --git a/Engine/Source/Runtime/Engine/Public/MaterialCompiler.h b/Engine/Source/Runtime/Engine/Public/MaterialCompiler.h
--- a/Engine/Source/Runtime/Engine/Public/MaterialCompiler.h
+++ b/Engine/Source/Runtime/Engine/Public/MaterialCompiler.h
@@ -58,6 +58,23 @@ enum class EMaterialCompilerType
 	MaterialProxy, /** Flat material proxy compiler */
 };
 
+
+/** Whether we need some data export from a Strata material from spatially varying properties, e.g. diffuse color for Lighmass to generate lightmaps. */
+enum EStrataMaterialExport : uint8
+{
+	SME_None			= 0,
+	SME_Diffuse			= 1,
+	SME_Normal			= 2,
+	SME_Emissive		= 3,
+	SME_Transmittance	= 4,
+};
+/** Exported materials are all opaque unlit. This is used to give some context to the export logic.*/
+enum EStrataMaterialExportContext : uint8
+{
+	SMEC_Opaque = 0,
+	SMEC_Translucent = 1
+};
+
 struct FStrataRegisteredSharedLocalBasis
 {
 	int32 NormalCodeChunk;
@@ -199,6 +216,16 @@ public:
 	inline bool IsMaterialProxyCompiler() const { return GetCompilerType() == EMaterialCompilerType::MaterialProxy; }
 	inline bool IsLightmassCompiler() const { return GetCompilerType() == EMaterialCompilerType::Lightmass; }
 
+	inline void SetStrataMaterialExportType(EStrataMaterialExport InStrataMaterialExport, EStrataMaterialExportContext InStrataMaterialExportContext, uint8 InStrataMaterialExportLegacyBlendMode)
+	{
+		StrataMaterialExport = InStrataMaterialExport; 
+		StrataMaterialExportContext = InStrataMaterialExportContext;
+		StrataMaterialExportLegacyBlendMode = InStrataMaterialExportLegacyBlendMode;
+	}
+	inline EStrataMaterialExport GetStrataMaterialExportType() const { return StrataMaterialExport; }
+	inline EStrataMaterialExportContext GetStrataMaterialExportContext() const { return StrataMaterialExportContext; }
+	inline uint8 GetStrataMaterialExportLegacyBlendMode() const { return StrataMaterialExportLegacyBlendMode; }
+
 	inline bool IsVertexInterpolatorBypass() const
 	{
 		const EMaterialCompilerType Type = GetCompilerType();
@@ -609,6 +636,11 @@ public:
 	// If possible we should re-factor this to avoid having to deal with compiler state
 	virtual bool IsCurrentlyCompilingForPreviousFrame() const { return false; }
 	virtual bool IsDevelopmentFeatureEnabled(const FName& FeatureName) const { return true; }
+
+private:
+	EStrataMaterialExport StrataMaterialExport = SME_None;
+	EStrataMaterialExportContext StrataMaterialExportContext = SMEC_Opaque;
+	uint8 StrataMaterialExportLegacyBlendMode = 0;
 };
 
 /** 
