Don't tick FComponentTypeRegistryData if GUnrealEd isn't valid

diff --git a/Engine/Source/Editor/UnrealEd/Private/ComponentTypeRegistry.cpp b/Engine/Source/Editor/UnrealEd/Private/ComponentTypeRegistry.cpp
--- a/Engine/Source/Editor/UnrealEd/Private/ComponentTypeRegistry.cpp
+++ b/Engine/Source/Editor/UnrealEd/Private/ComponentTypeRegistry.cpp
@@ -266,7 +266,8 @@ struct FComponentTypeRegistryData
 
 	/** Implementation of FTickableEditorObject */
 	virtual void Tick(float) override;
-	virtual ETickableTickType GetTickableTickType() const override { return ETickableTickType::Always; }
+	virtual ETickableTickType GetTickableTickType() const override { return ETickableTickType::Conditional; }
+	virtual bool IsTickable() const override { return GUnrealEd != nullptr; }
 	virtual TStatId GetStatId() const override { RETURN_QUICK_DECLARE_CYCLE_STAT(FTypeDatabaseUpdater, STATGROUP_Tickables); }
 	
 	/** Implementation of FGCObject */
