Various minimap improvements
* Reduced minimap texture waste, allowing for higher resolution minimap at a lower cost
* Improved minimap accuracy configuration, users must now specify Worlds Units / pixel instead of the previously unclear "MiniMapTileSize"
* Allow minimap to capture scene with different settings other than "unlit" base color

diff --git a/Engine/Source/Editor/UnrealEd/Private/WorldPartition/WorldPartitionMiniMapBuilder.cpp b/Engine/Source/Editor/UnrealEd/Private/WorldPartition/WorldPartitionMiniMapBuilder.cpp
--- a/Engine/Source/Editor/UnrealEd/Private/WorldPartition/WorldPartitionMiniMapBuilder.cpp
+++ b/Engine/Source/Editor/UnrealEd/Private/WorldPartition/WorldPartitionMiniMapBuilder.cpp
@@ -51,8 +51,38 @@ bool UWorldPartitionMiniMapBuilder::PreRun(UWorld* World, FPackageSourceControlH
 	}
 
 	// Reset minimap resources
-	WorldMiniMap->MiniMapTexture = nullptr;
-	MiniMapTiles.Empty();
+	{
+		IterativeCellSize = WorldMiniMap->BuilderCellSize;
+
+		WorldMiniMap->MiniMapTexture = nullptr;
+
+		UWorldPartition* WorldPartition = World->GetWorldPartition();
+		FBox WorldBounds = WorldPartition->GetEditorWorldBounds();
+
+		MinimapImageSizeX = WorldBounds.GetSize().X / WorldMiniMap->WorldUnitsPerPixel;
+		MinimapImageSizeY = WorldBounds.GetSize().Y / WorldMiniMap->WorldUnitsPerPixel;
+
+		// For now, let's clamp to the maximum supported texture size
+		MinimapImageSizeX = FMath::Min(MinimapImageSizeX, UTexture::GetMaximumDimensionOfNonVT());
+		MinimapImageSizeY = FMath::Min(MinimapImageSizeY, UTexture::GetMaximumDimensionOfNonVT());
+		WorldUnitsPerPixel = FMath::CeilToInt(FMath::Max(WorldBounds.GetSize().X / MinimapImageSizeX, WorldBounds.GetSize().Y / MinimapImageSizeY));
+		MinimapImageSizeX = WorldBounds.GetSize().X / WorldUnitsPerPixel;
+		MinimapImageSizeY = WorldBounds.GetSize().Y / WorldUnitsPerPixel;
+
+		TStrongObjectPtr<UTextureFactory> Factory(NewObject<UTextureFactory>());
+		WorldMiniMap->MiniMapTexture = Factory->CreateTexture2D(WorldMiniMap, TEXT("MinimapTexture"), RF_NoFlags);
+		WorldMiniMap->MiniMapTexture->Source.Init(MinimapImageSizeX, MinimapImageSizeY, 1, 1, TSF_BGRA8);
+		WorldMiniMap->MiniMapWorldBounds = WorldBounds;
+		MiniMapSourcePtr = WorldMiniMap->MiniMapTexture->Source.LockMip(0);
+
+		WorldToMinimap = FReversedZOrthoMatrix(WorldBounds.Min.X, WorldBounds.Max.X, WorldBounds.Min.Y, WorldBounds.Max.Y, 1.0f, 0.0f);
+
+		FVector3d Translation(WorldBounds.Max.X / WorldBounds.GetSize().X, WorldBounds.Max.Y / WorldBounds.GetSize().Y, 0);
+		FVector3d Scaling(MinimapImageSizeX, MinimapImageSizeY, 1);
+
+		WorldToMinimap *= FTranslationMatrix(Translation);
+		WorldToMinimap *= FScaleMatrix(Scaling);
+	}
 
 	UDataLayerSubsystem* DataLayerSubSystem = UWorld::GetSubsystem<UDataLayerSubsystem>(World);
 	for (const FActorDataLayer& ActorDataLayer : WorldMiniMap->ExcludedDataLayers)
@@ -70,24 +100,56 @@ bool UWorldPartitionMiniMapBuilder::PreRun(UWorld* World, FPackageSourceControlH
 
 bool UWorldPartitionMiniMapBuilder::RunInternal(UWorld* World, const FCellInfo& InCellInfo, FPackageSourceControlHelper& PackageHelper)
 {
-	EditorBounds = InCellInfo.EditorBounds;
-	IterativeCellSize = InCellInfo.IterativeCellSize;
-
-	FString TextureName = FString::Format(TEXT("MinimapTile_{0}_{1}_{2}"), { InCellInfo.Location.X, InCellInfo.Location.Y, InCellInfo.Location.Z });
-
-	UTexture2D* TileTexture = nullptr;
 	check(World != nullptr);
-	if (ULevelInstanceSubsystem* LevelInstanceSubsystem = World->GetSubsystem<ULevelInstanceSubsystem>())
+
+	// World X,Y to minimap X,Y
+	const FVector3d MinimapMin = WorldToMinimap.TransformPosition(InCellInfo.Bounds.Min);
+	const FVector3d MinimapMax = WorldToMinimap.TransformPosition(InCellInfo.Bounds.Max);
+	const FIntVector2 DstMin(FMath::Floor(MinimapMin.X), FMath::Floor(MinimapMin.Y));
+	const FIntVector2 DstMax(FMath::Floor(MinimapMax.X), FMath::Floor(MinimapMax.Y));
+	const FIntVector2 DstClampedMin = FIntVector2(FMath::Clamp(DstMin.X, 0, MinimapImageSizeX), FMath::Clamp(DstMin.Y, 0, MinimapImageSizeY));
+	const FIntVector2 DstClampedMax = FIntVector2(FMath::Clamp(DstMax.X, 0, MinimapImageSizeX), FMath::Clamp(DstMax.Y, 0, MinimapImageSizeY));
+	const uint32 CaptureWidthPixels = DstClampedMax.X - DstClampedMin.X;
+	const uint32 CaptureHeightPixels = DstClampedMax.Y - DstClampedMin.Y;
+
+	// Capture a tile if the region to capture is not empty
+	if (CaptureWidthPixels > 0 && CaptureHeightPixels > 0)
 	{
-		LevelInstanceSubsystem->UpdateStreamingState();
-	}
-	FWorldPartitionMiniMapHelper::CaptureBoundsMiniMapToTexture(World, WorldMiniMap, WorldMiniMap->MiniMapTileSize, TileTexture, TextureName, InCellInfo.Bounds);
+		FString TextureName = FString::Format(TEXT("MinimapTile_{0}_{1}_{2}"), { InCellInfo.Location.X, InCellInfo.Location.Y, InCellInfo.Location.Z });
 
-	FMinimapTile MinimapTile;
-	MinimapTile.Texture = TStrongObjectPtr<UTexture2D>(TileTexture);
-	MinimapTile.Coordinates.X = InCellInfo.Location.X;
-	MinimapTile.Coordinates.Y = InCellInfo.Location.Y;
-	MiniMapTiles.Add(MoveTemp(MinimapTile));
+		UTexture2D* TileTexture = NewObject<UTexture2D>(GetTransientPackage(), FName(TextureName), RF_Transient);
+		TileTexture->Source.Init(CaptureWidthPixels, CaptureHeightPixels, 1, 1, TSF_BGRA8);
+		TileTexture->PowerOfTwoMode = ETexturePowerOfTwoSetting::PadToPowerOfTwo;
+
+		FWorldPartitionMiniMapHelper::CaptureBoundsMiniMapToTexture(World, GetTransientPackage(), CaptureWidthPixels, CaptureHeightPixels, TileTexture, TextureName, InCellInfo.Bounds, WorldMiniMap->CaptureSource, WorldMiniMap->CaptureWarmupFrames);
+
+		// Copy captured image to VT minimap
+		const uint32 BPP = TileTexture->Source.GetBytesPerPixel();
+		const uint32 CopyWidthBytes = CaptureWidthPixels * BPP;
+
+		const uint8* SrcDataPtr = TileTexture->Source.LockMipReadOnly(0);
+		check(SrcDataPtr);
+
+		const uint32 DstDataStrideBytes = WorldMiniMap->MiniMapTexture->Source.GetSizeX() * BPP;
+		uint8* const DstDataPtr = MiniMapSourcePtr + (DstClampedMin.Y * DstDataStrideBytes) + (DstClampedMin.X * BPP);
+		check(DstDataPtr);
+
+		for (uint32 RowIdx = 0; RowIdx < CaptureHeightPixels; ++RowIdx)
+		{
+			uint8* DstCopy = DstDataPtr + DstDataStrideBytes * RowIdx;
+			const uint8* SrcCopy = SrcDataPtr + CopyWidthBytes * RowIdx;
+			
+			check(DstCopy >= DstDataPtr);
+			check(DstCopy + CopyWidthBytes <= DstDataPtr + WorldMiniMap->MiniMapTexture->Source.CalcMipSize(0));
+
+			check(SrcCopy >= SrcDataPtr);
+			check(SrcCopy + CopyWidthBytes <= SrcDataPtr + TileTexture->Source.CalcMipSize(0));
+
+			FMemory::Memcpy(DstCopy, SrcCopy, CopyWidthBytes);
+		}
+
+		TileTexture->Source.UnlockMip(0);
+	}
 
 	return true;
 }
@@ -102,89 +164,31 @@ bool UWorldPartitionMiniMapBuilder::PostRun(UWorld* World, FPackageSourceControl
 	// Make sure all assets and textures are ready
 	FAssetCompilingManager::Get().FinishAllCompilation();
 
-	WorldMiniMap->MiniMapWorldBounds = EditorBounds;
-
-	if (MiniMapTiles.IsEmpty())
+	// Finalize texture
 	{
-		UE_LOG(LogWorldPartitionMiniMapBuilder, Error, TEXT("No tiles were rendered, cannot compose Virtual Texture."));
-		return false;
-	}
-
-	FIntVector2 MinTileCoords(INT32_MAX);
-	FIntVector2 MaxTileCoords(INT32_MIN);
-	for (const FMinimapTile& Tile : MiniMapTiles)
-	{
-		MinTileCoords = FIntVector2(FMath::Min(MinTileCoords.X, Tile.Coordinates.X), FMath::Min(MinTileCoords.Y, Tile.Coordinates.Y));
-		MaxTileCoords = FIntVector2(FMath::Max(MaxTileCoords.X, Tile.Coordinates.X), FMath::Max(MaxTileCoords.Y, Tile.Coordinates.Y));
-	}
-
-	const FIntVector2 NumTilesCoords = FIntVector2(MaxTileCoords.X - MinTileCoords.X + 1, MaxTileCoords.Y - MinTileCoords.Y + 1);
-	const int32 MinimapTileSize = WorldMiniMap->MiniMapTileSize;
-	const int32 MinimapImageSizeX = NumTilesCoords.X * MinimapTileSize;
-	const int32 MinimapImageSizeY = NumTilesCoords.Y * MinimapTileSize;
-	
-	// Copy all tiles to the final minimap texture
-	{
-		TStrongObjectPtr<UTextureFactory> Factory(NewObject<UTextureFactory>());
-		UTexture2D* MiniMapTexture = Factory->CreateTexture2D(WorldMiniMap, TEXT("MinimapTexture"), RF_NoFlags);
-		MiniMapTexture->Source.Init(MinimapImageSizeX, MinimapImageSizeY, 1, 1, TSF_BGRA8);
-
-	    const uint32 DstDataStride = MiniMapTexture->Source.GetSizeX() * MiniMapTexture->Source.GetBytesPerPixel();
-	    uint8* DstDataBasePtr = MiniMapTexture->Source.LockMip(0);
-		check(DstDataBasePtr);
-    
-	    for (const FMinimapTile& Tile : MiniMapTiles)
-	    {
-		    const uint32 SrcDataStride = MinimapTileSize * Tile.Texture->Source.GetBytesPerPixel();
-		    const uint8* SrcDataPtr = Tile.Texture->Source.LockMipReadOnly(0);
-    
-		    const int32 TileOffsetX = Tile.Coordinates.X - MinTileCoords.X;
-		    const int32 TileOffsetY = Tile.Coordinates.Y - MinTileCoords.Y;
-    
-		    uint8* DstDataPtr = DstDataBasePtr + (TileOffsetY * MinimapTileSize * DstDataStride) + (TileOffsetX * SrcDataStride);
-    
-		    for (int RowIdx = 0; RowIdx < MinimapTileSize; ++RowIdx)
-		    {
-			    FMemory::Memcpy(DstDataPtr + DstDataStride * RowIdx,
-							    SrcDataPtr + SrcDataStride * RowIdx,
-							    SrcDataStride);
-		    }
-    
-		    Tile.Texture->Source.UnlockMip(0);
-	    }
-
-		MiniMapTexture->Source.UnlockMip(0);
-		MiniMapTexture->PowerOfTwoMode = ETexturePowerOfTwoSetting::PadToPowerOfTwo;	// Required for VTs
-		MiniMapTexture->AdjustMinAlpha = 1.f;
-		MiniMapTexture->LODGroup = TEXTUREGROUP_UI;
-		MiniMapTexture->VirtualTextureStreaming = true;
-		MiniMapTexture->UpdateResource();
-
-		WorldMiniMap->MiniMapTexture = MiniMapTexture;
+		WorldMiniMap->MiniMapTexture->Source.UnlockMip(0);
+		WorldMiniMap->MiniMapTexture->PowerOfTwoMode = ETexturePowerOfTwoSetting::PadToPowerOfTwo;	// Required for VTs
+		WorldMiniMap->MiniMapTexture->AdjustMinAlpha = 1.f;
+		WorldMiniMap->MiniMapTexture->LODGroup = TEXTUREGROUP_UI;
+		WorldMiniMap->MiniMapTexture->VirtualTextureStreaming = true;
+		WorldMiniMap->MiniMapTexture->UpdateResource();
 	}
 
 	// Compute relevant UV space for the minimap
 	{
-		FBox2D MinimapWorldBounds = FBox2D({ EditorBounds.Min.X, EditorBounds.Min.Y }, { EditorBounds.Max.X, EditorBounds.Max.Y });
-		FBox2D ImageWorldBounds = FBox2D(FVector2D(MinTileCoords.X, MinTileCoords.Y) * IterativeCellSize, FVector2D(MaxTileCoords.X + 1, MaxTileCoords.Y + 1) * IterativeCellSize);
-
 		FVector2D TexturePOW2ScaleFactor = FVector2D((float)MinimapImageSizeX / FMath::RoundUpToPowerOfTwo(MinimapImageSizeX),
 			(float)MinimapImageSizeY / FMath::RoundUpToPowerOfTwo(MinimapImageSizeY));
 
-		FBox2D UVOffset(FVector2D((MinimapWorldBounds.Min - ImageWorldBounds.Min) / ImageWorldBounds.GetSize()),
-			FVector2D((MinimapWorldBounds.Max - ImageWorldBounds.Min) / ImageWorldBounds.GetSize()));
-
-		UVOffset.Min *= TexturePOW2ScaleFactor;
-		UVOffset.Max *= TexturePOW2ScaleFactor;
-
-		WorldMiniMap->UVOffset = UVOffset;
+		WorldMiniMap->UVOffset.Min = FVector2d(0, 0);
+		WorldMiniMap->UVOffset.Max = TexturePOW2ScaleFactor;
+		WorldMiniMap->UVOffset.bIsValid = true;
 	}
 
 	// Make sure the minimap texture is ready before saving
 	FAssetCompilingManager::Get().FinishAllCompilation();
 
 	// Save MiniMap Package
-	auto WorldMiniMapExternalPackage = WorldMiniMap->GetExternalPackage();
+	UPackage* WorldMiniMapExternalPackage = WorldMiniMap->GetExternalPackage();
 	FString PackageFileName = SourceControlHelpers::PackageFilename(WorldMiniMapExternalPackage);
 
 	if (!PackageHelper.Checkout(WorldMiniMapExternalPackage))
@@ -195,7 +199,6 @@ bool UWorldPartitionMiniMapBuilder::PostRun(UWorld* World, FPackageSourceControl
 
 	FSavePackageArgs SaveArgs;
 	SaveArgs.TopLevelFlags = RF_Standalone;
-	SaveArgs.SaveFlags = SAVE_Async;
 	if (!UPackage::SavePackage(WorldMiniMapExternalPackage, nullptr, *PackageFileName, SaveArgs))
 	{
 		UE_LOG(LogWorldPartitionMiniMapBuilder, Error, TEXT("Error saving package %s."), *WorldMiniMapExternalPackage->GetName());
@@ -208,8 +211,6 @@ bool UWorldPartitionMiniMapBuilder::PostRun(UWorld* World, FPackageSourceControl
 		return false;
 	}
 
-	UPackage::WaitForAsyncFileWrites();
-
 	if (bAutoSubmit)
 	{
 		FText ChangelistDescription = FText::FromString(FString::Printf(TEXT("Rebuilt minimap for \"%s\" at %s"), *World->GetName(), *FEngineVersion::Current().ToString()));
diff --git a/Engine/Source/Editor/UnrealEd/Public/WorldPartition/WorldPartitionMiniMapBuilder.h b/Engine/Source/Editor/UnrealEd/Public/WorldPartition/WorldPartitionMiniMapBuilder.h
--- a/Engine/Source/Editor/UnrealEd/Public/WorldPartition/WorldPartitionMiniMapBuilder.h
+++ b/Engine/Source/Editor/UnrealEd/Public/WorldPartition/WorldPartitionMiniMapBuilder.h
@@ -10,16 +10,6 @@
 #include "WorldPartitionMiniMapBuilder.generated.h"
 
 
-/** Structure holding Minimap's tile resources */
-struct FMinimapTile
-{
-	/** Tile's texture */
-	TStrongObjectPtr<UTexture2D> Texture;
-
-	/** Coordinates of the tile in the Minimap's Virtual Texture */
-	FIntVector2 Coordinates;
-};
-
 UCLASS(config = Engine, defaultconfig)
 class UNREALED_API UWorldPartitionMiniMapBuilder : public UWorldPartitionBuilder
 {
@@ -34,13 +24,16 @@ protected:
 	virtual bool PreRun(UWorld* World, FPackageSourceControlHelper& PackageHelper) override;
 	virtual bool RunInternal(UWorld* World, const FCellInfo& InCellInfo, FPackageSourceControlHelper& PackageHelper) override;
 	virtual bool PostRun(UWorld* World, FPackageSourceControlHelper& PackageHelper, const bool bInRunSuccess) override;
-
 	// UWorldPartitionBuilder interface end
 
 private:
-	/*MiniMap Texture Tiles for displaying on world partition window*/
-	TArray<FMinimapTile> MiniMapTiles;
 	TObjectPtr<AWorldPartitionMiniMap> WorldMiniMap;
 
-	FBox EditorBounds;
+	FMatrix WorldToMinimap;
+
+	int32 WorldUnitsPerPixel;
+	int32 MinimapImageSizeX;
+	int32 MinimapImageSizeY;
+
+	uint8* MiniMapSourcePtr;
 };
\ No newline at end of file
diff --git a/Engine/Source/Runtime/Engine/Private/WorldPartition/WorldPartitionMiniMap.cpp b/Engine/Source/Runtime/Engine/Private/WorldPartition/WorldPartitionMiniMap.cpp
--- a/Engine/Source/Runtime/Engine/Private/WorldPartition/WorldPartitionMiniMap.cpp
+++ b/Engine/Source/Runtime/Engine/Private/WorldPartition/WorldPartitionMiniMap.cpp
@@ -14,9 +14,26 @@ AWorldPartitionMiniMap::AWorldPartitionMiniMap(const FObjectInitializer& ObjectI
 	: Super(ObjectInitializer.DoNotCreateDefaultSubobject(TEXT("Sprite")))
 	, MiniMapWorldBounds(ForceInit)
 	, MiniMapTexture(nullptr)
+	, WorldUnitsPerPixel(50)
+	, BuilderCellSize(102400)
+	, CaptureSource(ESceneCaptureSource::SCS_BaseColor)
+	, CaptureWarmupFrames(5)
 {
 }
 
+void AWorldPartitionMiniMap::PostLoad()
+{
+	Super::PostLoad();
+
+	if (MiniMapTileSize_DEPRECATED != 0)
+	{
+		const int32 MinimapBuilderIterativeCellSize = 102400;
+		WorldUnitsPerPixel = MinimapBuilderIterativeCellSize / MiniMapTileSize_DEPRECATED;
+		WorldUnitsPerPixel = FMath::Clamp(WorldUnitsPerPixel, 10, 100000);
+		MiniMapTileSize_DEPRECATED = 0;
+	}
+}
+
 #if WITH_EDITOR
 void AWorldPartitionMiniMap::CheckForErrors()
 {
@@ -38,10 +55,6 @@ void AWorldPartitionMiniMap::CheckForErrors()
 	}
 }
 
-void AWorldPartitionMiniMap::PostEditChangeProperty(FPropertyChangedEvent& PropertyChangedEvent)
-{
-	MiniMapTileSize = FMath::Clamp<uint32>(FMath::RoundUpToPowerOfTwo(MiniMapTileSize), 256, 8192);
-}
 #endif
 
 #undef LOCTEXT_NAMESPACE
diff --git a/Engine/Source/Runtime/Engine/Private/WorldPartition/WorldPartitionMiniMapHelper.cpp b/Engine/Source/Runtime/Engine/Private/WorldPartition/WorldPartitionMiniMapHelper.cpp
--- a/Engine/Source/Runtime/Engine/Private/WorldPartition/WorldPartitionMiniMapHelper.cpp
+++ b/Engine/Source/Runtime/Engine/Private/WorldPartition/WorldPartitionMiniMapHelper.cpp
@@ -58,9 +58,19 @@ void FWorldPartitionMiniMapHelper::CaptureWorldMiniMapToTexture(UWorld* InWorld,
 	OutWorldBounds = InWorld->GetWorldPartition()->GetEditorWorldBounds();
 
 	CaptureBoundsMiniMapToTexture(InWorld, InOuterForTexture, InMiniMapSize, InOutMiniMapTexture, InTextureName, OutWorldBounds);
+
+	InOutMiniMapTexture->AdjustMinAlpha = 1.f;
+	InOutMiniMapTexture->LODGroup = TEXTUREGROUP_UI;
+	InOutMiniMapTexture->UpdateResource();
 }
 
 void FWorldPartitionMiniMapHelper::CaptureBoundsMiniMapToTexture(UWorld* InWorld, UObject* InOuterForTexture, uint32 InMiniMapSize, UTexture2D*& InOutMiniMapTexture, const FString& InTextureName, const FBox& InBounds)
+{
+	const uint32 DefaultNumWarmupFrames = 5;
+	CaptureBoundsMiniMapToTexture(InWorld, InOuterForTexture, InMiniMapSize, InMiniMapSize, InOutMiniMapTexture, InTextureName, InBounds, SCS_BaseColor, DefaultNumWarmupFrames);
+}
+
+void FWorldPartitionMiniMapHelper::CaptureBoundsMiniMapToTexture(UWorld* InWorld, UObject* InOuterForTexture, uint32 InMiniMapSizeX, uint32 InMiniMapSizeY, UTexture2D*& InOutMiniMapTexture, const FString& InTextureName, const FBox& InBounds, ESceneCaptureSource InCaptureSource, uint32 InNumWarmupFrames)
 {
 	// Before capturing the scene, make sure all assets are finished compiling
 	FAssetCompilingManager::Get().FinishAllCompilation();
@@ -69,8 +79,8 @@ void FWorldPartitionMiniMapHelper::CaptureBoundsMiniMapToTexture(UWorld* InWorld
 	FBox2D WorldBounds2D(FVector2D(InBounds.Min), FVector2D(InBounds.Max));
 	FVector2D ViewSize = WorldBounds2D.Max - WorldBounds2D.Min;
 	float AspectRatio = FMath::Abs(ViewSize.X) / FMath::Abs(ViewSize.Y);
-	uint32 ViewportWidth = InMiniMapSize * AspectRatio;
-	uint32 ViewportHeight = InMiniMapSize;
+	uint32 ViewportWidth = InMiniMapSizeX * AspectRatio;
+	uint32 ViewportHeight = InMiniMapSizeY;
 
 	//Calculate Projection matrix based on world bounds.
 	FMatrix ProjectionMatrix;
@@ -80,7 +90,7 @@ void FWorldPartitionMiniMapHelper::CaptureBoundsMiniMapToTexture(UWorld* InWorld
 	UTextureRenderTarget2D* RenderTargetTexture = NewObject<UTextureRenderTarget2D>();
 	RenderTargetTexture->ClearColor = FLinearColor::Transparent;
 	RenderTargetTexture->TargetGamma = 2.2f;
-	RenderTargetTexture->InitCustomFormat(InMiniMapSize, InMiniMapSize, PF_B8G8R8A8, false);
+	RenderTargetTexture->InitCustomFormat(InMiniMapSizeX, InMiniMapSizeY, PF_B8G8R8A8, false);
 	RenderTargetTexture->UpdateResourceImmediate(true);
 
 	FActorSpawnParameters SpawnInfo;
@@ -91,12 +101,22 @@ void FWorldPartitionMiniMapHelper::CaptureBoundsMiniMapToTexture(UWorld* InWorld
 	auto CaptureComponent = CaptureActor->GetCaptureComponent2D();
 	CaptureComponent->TextureTarget = RenderTargetTexture;
 	CaptureComponent->ProjectionType = ECameraProjectionMode::Orthographic;
-	CaptureComponent->CaptureSource = SCS_BaseColor;
+	CaptureComponent->CaptureSource = InCaptureSource;
 	CaptureComponent->OrthoWidth = ViewportWidth;
 	CaptureComponent->bUseCustomProjectionMatrix = true;
 	CaptureComponent->CustomProjectionMatrix = ProjectionMatrix;
 	CaptureComponent->bCaptureEveryFrame = false;
 	CaptureComponent->bCaptureOnMovement = false;
+
+	// Disable vignetting, otherwise we'll see it as a pattern between each captured tiles
+	CaptureComponent->PostProcessSettings.bOverride_VignetteIntensity = true;
+	CaptureComponent->PostProcessSettings.VignetteIntensity = 0.0f;
+
+	for (uint32 i = 0; i < InNumWarmupFrames; i++)
+	{
+		CaptureComponent->CaptureScene();
+	}
+
 	CaptureComponent->CaptureScene();
 
 	InWorld->DestroyActor(CaptureActor);
@@ -111,10 +131,6 @@ void FWorldPartitionMiniMapHelper::CaptureBoundsMiniMapToTexture(UWorld* InWorld
 	{
 		RenderTargetTexture->UpdateTexture2D(InOutMiniMapTexture, TSF_BGRA8, CTF_Default);
 	}
-
-	InOutMiniMapTexture->AdjustMinAlpha = 1.f;
-	InOutMiniMapTexture->LODGroup = TEXTUREGROUP_UI;
-	InOutMiniMapTexture->UpdateResource();
 }
 
 void FWorldPartitionMiniMapHelper::CalTopViewOfWorld(FMatrix& OutProjectionMatrix, const FBox& WorldBox, uint32 ViewportWidth, uint32 ViewportHeight)
diff --git a/Engine/Source/Runtime/Engine/Public/WorldPartition/WorldPartitionMiniMap.h b/Engine/Source/Runtime/Engine/Public/WorldPartition/WorldPartitionMiniMap.h
--- a/Engine/Source/Runtime/Engine/Public/WorldPartition/WorldPartitionMiniMap.h
+++ b/Engine/Source/Runtime/Engine/Public/WorldPartition/WorldPartitionMiniMap.h
@@ -9,7 +9,7 @@
 /**
  * A mini map to preview the world in world partition window. (editor-only)
  */
-UCLASS(hidecategories = (Actor, Advanced, Display, Events, Object, Attachment, Info, Input, Blueprint, Layers, Tags, Replication), notplaceable)
+UCLASS(hidecategories = (Actor, Advanced, Display, Events, Object, Attachment, Info, Input, Blueprint, Layers, Tags, Replication, Physics, Cooking), notplaceable)
 class ENGINE_API AWorldPartitionMiniMap : public AInfo
 {
 	GENERATED_BODY()
@@ -23,31 +23,53 @@ public:
 	AWorldPartitionMiniMap(const FObjectInitializer& ObjectInitializer = FObjectInitializer::Get());
 
 	virtual bool IsEditorOnly() const final { return true; }
+	virtual void PostLoad() override;
 
 #if WITH_EDITOR
-	virtual void PostEditChangeProperty(FPropertyChangedEvent& PropertyChangedEvent) override;
 	virtual bool IsLockLocation() const { return true; }
 	virtual bool IsUserManaged() const final { return false; }
 	virtual void CheckForErrors() override;
 #endif
 
-	/*WorldBounds for MinMapTexture*/
-	UPROPERTY(VisibleAnywhere, Category = WorldPartitionMiniMap)
+	/* WorldBounds for MinMapTexture */
+	UPROPERTY(VisibleAnywhere, Category = WorldPartitionMiniMap, AdvancedDisplay)
 	FBox MiniMapWorldBounds;
 
 	/* UVOffset used to setup Virtual Texture */
-	UPROPERTY(EditAnywhere, Category = WorldPartitionMiniMap)
+	UPROPERTY(VisibleAnywhere, Category = WorldPartitionMiniMap, AdvancedDisplay)
 	FBox2D UVOffset;
 
-	/*MiniMap Texture for displaying on world partition window*/
+	/* MiniMap Texture for displaying on world partition window */
 	UPROPERTY(VisibleAnywhere, Category = WorldPartitionMiniMap)
 	TObjectPtr<UTexture2D> MiniMapTexture;
 
-	/*Datalayers excluded from MiniMap rendering*/
+	/* Datalayers excluded from MiniMap rendering */
 	UPROPERTY(EditAnywhere, Category = WorldPartitionMiniMap)
 	TSet<FActorDataLayer> ExcludedDataLayers;
 
-	/*MiniMap Tile Size*/
-	UPROPERTY(EditAnywhere, Category = WorldPartitionMiniMap, meta = (UIMin = "256", UIMax = "8192"))
-	int32 MiniMapTileSize = 1024;
+	/**
+	 * Target world units per pixel for the minimap texture. 
+	 * May not end up being the final minimap accuracy if the resulting texture resolution is unsupported.
+	 */
+	UPROPERTY(EditAnywhere, Category = WorldPartitionMiniMap, meta = (UIMin = "10", UIMax = "100000"), AdvancedDisplay)
+	int32 WorldUnitsPerPixel;
+
+	/**
+	 * Size of the loading region that will be used when iterating over the whole map during the minimap build process.
+	 * A smaller size may help reduce blurriness as it will put less pressure on various graphics pools, at the expanse of an increase in processing time. 
+	 */
+	UPROPERTY(EditAnywhere, Category = WorldPartitionMiniMap, meta = (UIMin = "3200", UIMax = "204800"), AdvancedDisplay)
+	int32 BuilderCellSize;
+
+	/* Specifies which component of the scene rendering should be output to the minimap texture. */
+	UPROPERTY(EditAnywhere, Category = WorldPartitionMiniMap, AdvancedDisplay)
+	TEnumAsByte<ESceneCaptureSource> CaptureSource;
+
+	/* Number of frames to render before each capture in order to warmup various rendering systems (VT/Nanite/etc). */
+	UPROPERTY(EditAnywhere, Category = WorldPartitionMiniMap, meta = (UIMin = "0", UIMax = "10"), AdvancedDisplay)
+	uint32 CaptureWarmupFrames;
+
+private:
+	UPROPERTY()
+	int32 MiniMapTileSize_DEPRECATED = 0;
 };
\ No newline at end of file
diff --git a/Engine/Source/Runtime/Engine/Public/WorldPartition/WorldPartitionMiniMapHelper.h b/Engine/Source/Runtime/Engine/Public/WorldPartition/WorldPartitionMiniMapHelper.h
--- a/Engine/Source/Runtime/Engine/Public/WorldPartition/WorldPartitionMiniMapHelper.h
+++ b/Engine/Source/Runtime/Engine/Public/WorldPartition/WorldPartitionMiniMapHelper.h
@@ -23,6 +23,7 @@ public:
 	static AWorldPartitionMiniMap* GetWorldPartitionMiniMap(UWorld* World, bool bCreateNewMiniMap=false);
 	static void CaptureWorldMiniMapToTexture(UWorld* InWorld, UObject* InOuterForTexture, uint32 InMiniMapSize, UTexture2D*& InOutMiniMapTexture, const FString& InTextureName, FBox& OutWorldBounds);
 	static void CaptureBoundsMiniMapToTexture(UWorld* InWorld, UObject* InOuterForTexture, uint32 InMiniMapSize, UTexture2D*& InOutMiniMapTexture, const FString& InTextureName, const FBox& InBounds);
+	static void CaptureBoundsMiniMapToTexture(UWorld* InWorld, UObject* InOuterForTexture, uint32 InMiniMapSizeX, uint32 InMiniMapSizeY, UTexture2D*& InOutMiniMapTexture, const FString& InTextureName, const FBox& InBounds, ESceneCaptureSource InCaptureSource, uint32 InNumWarmupFrames);
 
 private:
 	static void CalTopViewOfWorld(FMatrix& OutProjectionMatrix, const FBox& WorldBox, uint32 ViewportWidth, uint32 ViewportHeight);
