Fix Duplicate LOC key "VirtualTextureBuild"

#rnx

diff --git a/Engine/Source/Editor/UnrealEd/Private/EditorBuildUtils.cpp b/Engine/Source/Editor/UnrealEd/Private/EditorBuildUtils.cpp
--- a/Engine/Source/Editor/UnrealEd/Private/EditorBuildUtils.cpp
+++ b/Engine/Source/Editor/UnrealEd/Private/EditorBuildUtils.cpp
@@ -1608,7 +1608,7 @@ bool FEditorBuildUtils::EditorBuildVirtualTexture(UWorld* InWorld)
 		FWorldPartitionHelpers::FForEachActorWithLoadingResult ForEachActorWithLoadingResult;
 		if (UWorldPartition* WorldPartition = InWorld->GetWorldPartition())
 		{
-			FScopedSlowTask BuildTask(1.0f, LOCTEXT("VirtualTextureBuild", "Loading Actors"));
+			FScopedSlowTask BuildTask(1.0f, LOCTEXT("VirtualTextureLoadActors", "Loading Actors"));
 			BuildTask.MakeDialog();
 			UWorldPartitionRuntimeVirtualTextureBuilder::LoadRuntimeVirtualTextureActors(WorldPartition, ForEachActorWithLoadingResult);
 		}
