World Partition: added WorldPartitionMiniMapVolume that can be used to override the bounds of the world when generating the minimap.
This can be useful when the world's Landscape is non-spatially loaded so the world bounds don't include these, for example.

#rnx

diff --git a/Engine/Source/Editor/UnrealEd/Private/WorldPartition/WorldPartitionMiniMapBuilder.cpp b/Engine/Source/Editor/UnrealEd/Private/WorldPartition/WorldPartitionMiniMapBuilder.cpp
--- a/Engine/Source/Editor/UnrealEd/Private/WorldPartition/WorldPartitionMiniMapBuilder.cpp
+++ b/Engine/Source/Editor/UnrealEd/Private/WorldPartition/WorldPartitionMiniMapBuilder.cpp
@@ -9,6 +9,7 @@
 #include "AssetCompilingManager.h"
 #include "Engine/Texture2D.h"
 #include "Engine/World.h"
+#include "EngineUtils.h"
 #include "Factories/TextureFactory.h"
 #include "SourceControlHelpers.h"
 #include "UObject/SavePackage.h"
@@ -28,6 +29,7 @@
 #include "WorldPartition/WorldPartitionHandle.h"
 #include "WorldPartition/WorldPartitionMiniMap.h"
 #include "WorldPartition/WorldPartitionMiniMapHelper.h"
+#include "WorldPartition/WorldPartitionMiniMapVolume.h"
 #include "LevelInstance/LevelInstanceSubsystem.h"
 
 DEFINE_LOG_CATEGORY_STATIC(LogWorldPartitionMiniMapBuilder, All, All);
@@ -56,8 +58,21 @@ bool UWorldPartitionMiniMapBuilder::PreRun(UWorld* World, FPackageSourceControlH
 
 		WorldMiniMap->MiniMapTexture = nullptr;
 
-		UWorldPartition* WorldPartition = World->GetWorldPartition();
-		FBox WorldBounds = WorldPartition->GetEditorWorldBounds();
+		FBox WorldBounds(ForceInit);
+
+		// Override the minimap bounds it a world partiion minimap volume exists
+		for (TActorIterator<AWorldPartitionMiniMapVolume> It(World); It; ++It)
+		{
+			if (AWorldPartitionMiniMapVolume* WorldPartitionMiniMapVolume = *It)
+			{
+				WorldBounds += WorldPartitionMiniMapVolume->GetBounds().GetBox();
+			}
+		}
+
+		if (!WorldBounds.IsValid)
+		{
+			WorldBounds = World->GetWorldPartition()->GetEditorWorldBounds();
+		}
 
 		MinimapImageSizeX = WorldBounds.GetSize().X / WorldMiniMap->WorldUnitsPerPixel;
 		MinimapImageSizeY = WorldBounds.GetSize().Y / WorldMiniMap->WorldUnitsPerPixel;
diff --git a/Engine/Source/Runtime/Engine/Private/WorldPartition/WorldPartitionMiniMapHelper.cpp b/Engine/Source/Runtime/Engine/Private/WorldPartition/WorldPartitionMiniMapHelper.cpp
--- a/Engine/Source/Runtime/Engine/Private/WorldPartition/WorldPartitionMiniMapHelper.cpp
+++ b/Engine/Source/Runtime/Engine/Private/WorldPartition/WorldPartitionMiniMapHelper.cpp
@@ -52,18 +52,6 @@ AWorldPartitionMiniMap* FWorldPartitionMiniMapHelper::GetWorldPartitionMiniMap(U
 	return nullptr;
 }
 
-void FWorldPartitionMiniMapHelper::CaptureWorldMiniMapToTexture(UWorld* InWorld, UObject* InOuterForTexture, uint32 InMiniMapSize, UTexture2D*& InOutMiniMapTexture, const FString& InTextureName, FBox& OutWorldBounds)
-{
-	//Calculate bounds of the World
-	OutWorldBounds = InWorld->GetWorldPartition()->GetEditorWorldBounds();
-
-	CaptureBoundsMiniMapToTexture(InWorld, InOuterForTexture, InMiniMapSize, InOutMiniMapTexture, InTextureName, OutWorldBounds);
-
-	InOutMiniMapTexture->AdjustMinAlpha = 1.f;
-	InOutMiniMapTexture->LODGroup = TEXTUREGROUP_UI;
-	InOutMiniMapTexture->UpdateResource();
-}
-
 void FWorldPartitionMiniMapHelper::CaptureBoundsMiniMapToTexture(UWorld* InWorld, UObject* InOuterForTexture, uint32 InMiniMapSize, UTexture2D*& InOutMiniMapTexture, const FString& InTextureName, const FBox& InBounds)
 {
 	const uint32 DefaultNumWarmupFrames = 5;
diff --git a/Engine/Source/Runtime/Engine/Private/WorldPartition/WorldPartitionMiniMapVolume.cpp b/Engine/Source/Runtime/Engine/Private/WorldPartition/WorldPartitionMiniMapVolume.cpp
new file mode 100644
--- /dev/null
+++ b/Engine/Source/Runtime/Engine/Private/WorldPartition/WorldPartitionMiniMapVolume.cpp
@@ -0,0 +1,11 @@
+// Copyright Epic Games, Inc. All Rights Reserved.
+
+#include "WorldPartition/WorldPartitionMiniMapVolume.h"
+
+AWorldPartitionMiniMapVolume::AWorldPartitionMiniMapVolume(const FObjectInitializer& ObjectInitializer)
+	: Super(ObjectInitializer.DoNotCreateDefaultSubobject(TEXT("Sprite")))
+{
+#if WITH_EDITORONLY_DATA
+	bIsSpatiallyLoaded = false;
+#endif
+}
\ No newline at end of file
diff --git a/Engine/Source/Runtime/Engine/Public/WorldPartition/WorldPartitionMiniMapHelper.h b/Engine/Source/Runtime/Engine/Public/WorldPartition/WorldPartitionMiniMapHelper.h
--- a/Engine/Source/Runtime/Engine/Public/WorldPartition/WorldPartitionMiniMapHelper.h
+++ b/Engine/Source/Runtime/Engine/Public/WorldPartition/WorldPartitionMiniMapHelper.h
@@ -21,7 +21,6 @@ class ENGINE_API FWorldPartitionMiniMapHelper
 {
 public:
 	static AWorldPartitionMiniMap* GetWorldPartitionMiniMap(UWorld* World, bool bCreateNewMiniMap=false);
-	static void CaptureWorldMiniMapToTexture(UWorld* InWorld, UObject* InOuterForTexture, uint32 InMiniMapSize, UTexture2D*& InOutMiniMapTexture, const FString& InTextureName, FBox& OutWorldBounds);
 	static void CaptureBoundsMiniMapToTexture(UWorld* InWorld, UObject* InOuterForTexture, uint32 InMiniMapSize, UTexture2D*& InOutMiniMapTexture, const FString& InTextureName, const FBox& InBounds);
 	static void CaptureBoundsMiniMapToTexture(UWorld* InWorld, UObject* InOuterForTexture, uint32 InMiniMapSizeX, uint32 InMiniMapSizeY, UTexture2D*& InOutMiniMapTexture, const FString& InTextureName, const FBox& InBounds, ESceneCaptureSource InCaptureSource, uint32 InNumWarmupFrames);
 
diff --git a/Engine/Source/Runtime/Engine/Public/WorldPartition/WorldPartitionMiniMapVolume.h b/Engine/Source/Runtime/Engine/Public/WorldPartition/WorldPartitionMiniMapVolume.h
new file mode 100644
--- /dev/null
+++ b/Engine/Source/Runtime/Engine/Public/WorldPartition/WorldPartitionMiniMapVolume.h
@@ -0,0 +1,22 @@
+// Copyright Epic Games, Inc. All Rights Reserved.
+#pragma once
+
+#include "CoreMinimal.h"
+#include "GameFramework/Volume.h"
+#include "WorldPartitionMiniMapVolume.generated.h"
+
+UCLASS(hidecategories = (Actor, Advanced, Display, Events, Object, Attachment, Info, Input, Blueprint, Layers, Tags, Replication))
+class ENGINE_API AWorldPartitionMiniMapVolume : public AVolume
+{
+	GENERATED_BODY()
+
+public:
+	AWorldPartitionMiniMapVolume(const FObjectInitializer& ObjectInitializer = FObjectInitializer::Get());
+
+	virtual bool IsEditorOnly() const final { return true; }
+
+#if WITH_EDITOR
+	virtual bool ActorTypeSupportsDataLayer() const final { return false; }
+	virtual bool CanChangeIsSpatiallyLoadedFlag() const { return false; }
+#endif
+};
\ No newline at end of file
