Fbx Importer: downgrade the bones digested log from "warning" to "display" level.
#rnx

diff --git a/Engine/Source/Editor/UnrealEd/Private/Fbx/FbxSkeletalMeshImport.cpp b/Engine/Source/Editor/UnrealEd/Private/Fbx/FbxSkeletalMeshImport.cpp
--- a/Engine/Source/Editor/UnrealEd/Private/Fbx/FbxSkeletalMeshImport.cpp
+++ b/Engine/Source/Editor/UnrealEd/Private/Fbx/FbxSkeletalMeshImport.cpp
@@ -1866,7 +1866,9 @@ USkeletalMesh* UnFbx::FFbxImporter::ImportSkeletalMesh(FImportSkeletalMeshArgs &
 	EARLY_RETURN_ON_CANCEL(false, CancelCleanup);
 
 	if (!GIsAutomationTesting)
-		UE_LOG(LogFbx, Warning, TEXT("Bones digested - %i  Depth of hierarchy - %i"), SkeletalMesh->GetRefSkeleton().GetNum(), SkeletalDepth);
+	{
+		UE_LOG(LogFbx, Log, TEXT("Bones digested - %i  Depth of hierarchy - %i"), SkeletalMesh->GetRefSkeleton().GetNum(), SkeletalDepth);
+	}
 
 	// process bone influences from import data
 	SkeletalMeshImportUtils::ProcessImportMeshInfluences(*SkelMeshImportDataPtr, SkeletalMesh->GetPathName());
