CookPackageSplitter: Change the split to occur after BeginCacheForCookedPlatformData has been called on all packages, and also add a function for the splitter to declare the objects that will be moved into the package, so we can wait on BeginCache for those as well before modifying the package. If we modify the package and then have to wait for some long running BeginCache calls before we save and undo the object moves, it is possible that garbage collection will run and will break the objects moved from external packages because their packages are garbage collected.
#rnx

diff --git a/Engine/Source/Editor/UnrealEd/Classes/CookOnTheSide/CookOnTheFlyServer.h b/Engine/Source/Editor/UnrealEd/Classes/CookOnTheSide/CookOnTheFlyServer.h
--- a/Engine/Source/Editor/UnrealEd/Classes/CookOnTheSide/CookOnTheFlyServer.h
+++ b/Engine/Source/Editor/UnrealEd/Classes/CookOnTheSide/CookOnTheFlyServer.h
@@ -7,6 +7,7 @@
 #include "Containers/ArrayView.h"
 #include "Containers/RingBuffer.h"
 #include "CookOnTheSide/CookLog.h"
+#include "CookPackageSplitter.h"
 #include "Engine/ICookInfo.h"
 #include "INetworkFileSystemModule.h"
 #include "IPlatformFileSandboxWrapper.h"
@@ -126,6 +127,8 @@ namespace UE::Cook
 	class FSaveCookedPackageContext;
 	class ICookOnTheFlyRequestManager;
 	class ICookOnTheFlyNetworkServer;
+	enum class EPollStatus : uint8;
+	enum class EReleaseSaveReason : uint8;
 	struct FConstructPackageData;
 	struct FCookByTheBookOptions;
 	struct FCookOnTheFlyOptions;
@@ -997,27 +1000,23 @@ private:
 	 * for this UObject.
 	 *
 	 * @param PackageData the PackageData used to gather all uobjects from
-	 * @param bIsPreCaching true if called for precaching
+	 * @param bIsPrecaching true if called for precaching
 	 * @return false if time slice was reached, true if all objects have had BeginCacheForCookedPlatformData called
 	 */
-	bool BeginPrepareSave(UE::Cook::FPackageData& PackageData, UE::Cook::FCookerTimer& Timer, bool bIsPreCaching);
-	
-	/**
-	 * Returns true when all objects in package have all their cooked platform data loaded
-	 *	confirms that BeginCacheForCookedPlatformData is called and will return true after all objects return IsCachedCookedPlatformDataLoaded true
-	 *
-	 * @param PackageData the PackageData used to gather all uobjects from
-	 * @return false if time slice was reached, true if all return true for IsCachedCookedPlatformDataLoaded 
-	 */
-	bool FinishPrepareSave(UE::Cook::FPackageData& PackageData, UE::Cook::FCookerTimer& Timer);
+	UE::Cook::EPollStatus PrepareSave(UE::Cook::FPackageData& PackageData, UE::Cook::FCookerTimer& Timer, bool bPrecaching);
+	UE::Cook::EPollStatus PrepareSaveInternal(UE::Cook::FPackageData& PackageData, UE::Cook::FCookerTimer& Timer, bool bPrecaching);
+	/** Call BeginCacheForCookedPlatformData on all of the objects in a PackageData or a GeneratorPackage's current round. */
+	UE::Cook::EPollStatus CallBeginCacheOnObjects(UE::Cook::FPackageData& PackageData, UPackage* Package,
+		TArray<FWeakObjectPtr>& Objects, int32& NextIndex, UE::Cook::FCookerTimer& Timer);
 
 	/**
 	 * Frees all the memory used to call BeginCacheForCookedPlatformData on all the objects in PackageData.
 	 * If the calls were incomplete because the PackageData's save was cancelled, handles canceling them and leaving any required CancelManagers in GetPendingCookedPlatformDatas
 	 * 
 	 * @param bCompletedSave If false, data that can not be efficiently recomputed will be preserved to try the save again. If true, all data will be wiped.
+	 * @param ReleaseSaveReason Why the save data is being released, allows specifying how much to tear down
 	 */
-	void ReleaseCookedPlatformData(UE::Cook::FPackageData& PackageData, bool bCompletedSave);
+	void ReleaseCookedPlatformData(UE::Cook::FPackageData& PackageData, UE::Cook::EReleaseSaveReason ReleaseSaveReason);
 
 	/**
 	 * Poll the GetPendingCookedPlatformDatas and release their resources when they are complete.
@@ -1167,33 +1166,27 @@ private:
 	/** Generates long package names for all files to be cooked */
 	void GenerateLongPackageNames(TArray<FName>& FilesInPath, TMap<FName, UE::Cook::FInstigator>& Instigators);
 
-	/**
-	* Try to find a registered CookPackageSplitter to split the given package.
-	*
-	* @param PackageData			Package to be potentially split
-	* @param OutSplitDataObject		Return package's object to handle splitting
-	* @param bOutHasError			Return True if function fails
-	*/
-	UE::Cook::Private::FRegisteredCookPackageSplitter* TryGetRegisteredCookPackageSplitter(UE::Cook::FPackageData& PackageData, UObject*& OutSplitDataObject, bool& bOutError);
-
-	/**
-	* Generate a creator package helper object that will be used to split package.
-	*
-	* @param PackageData			Package to split
-	* @param SplitDataObject		Package's object to handle splitting
-	* @param Splitter				Registered CookPackageSplitter to be used for splitting the package
-	*/
-	UE::Cook::FGeneratorPackage* CreateGeneratorPackage(UE::Cook::FPackageData& PackageData, UObject* SplitDataObject, UE::Cook::Private::FRegisteredCookPackageSplitter* Splitter);
-
-	/**
-	* Split the package and generate other packages.
-	*
-	* @param PackageData			Package to be split
-	* @param SplitDataObject		Package's object to handle splitting
-	* @param bOutCompleted			Return true is function is done splitting
-	* @param bOutError				Return true if function fails
-	*/
-	void SplitPackage(UE::Cook::FPackageData& PackageData, UE::Cook::FGeneratorPackage* Generator, bool& bCompleted, bool& bOutError);
+	UE::Cook::EPollStatus ConditionalCreateGeneratorPackage(UE::Cook::FPackageData& PackageData);
+
+	/** Generate the list of cook-time-created packages created by the Generator. */
+	UE::Cook::EPollStatus QueueGeneratedPackages(UE::Cook::FGeneratorPackage& Generator,
+		UE::Cook::FPackageData& PackageData);
+	/** Run additional steps in PrepareSave that are required when the package has a Generator. */
+	UE::Cook::EPollStatus PrepareSaveGeneratorPackage(UE::Cook::FGeneratorPackage& Generator,
+		UE::Cook::FPackageData& PackageData, UE::Cook::FCookerTimer& Timer);
+	/** Call BeginCacheForCookedPlatformData on the objects the Generator plans to move into its main UPackage. */
+	UE::Cook::EPollStatus BeginCacheObjectsToMove(UE::Cook::FGeneratorPackage& Generator,
+		UE::Cook::FPackageData& PackageData, UE::Cook::FCookerTimer& Timer,
+		TArray<ICookPackageSplitter::FGeneratedPackageForPreSave>& GeneratedPackagesForPresave);
+	/** Call the Generator's PreSaveGeneratorPackage to create/move objects into its main UPackage. */
+	UE::Cook::EPollStatus PreSaveGeneratorPackage(UE::Cook::FPackageData& PackageData,
+		UE::Cook::FGeneratorPackage& Generator, TArray<ICookPackageSplitter::FGeneratedPackageForPreSave>& GeneratedPackagesForPresave);
+	/** Construct the list of generated packages that is required for some of the CookPackageSplitter interface calls. */
+	void ConstructGeneratedPackagesForPresave(UE::Cook::FPackageData& PackageData, UE::Cook::FGeneratorPackage& Generator,
+		TArray<ICookPackageSplitter::FGeneratedPackageForPreSave>& GeneratedPackagesForPresave);
+	/** Call BeginCacheForCookedPlatformData on any undeclared objects in the Generator's main UPackage after the move. */
+	UE::Cook::EPollStatus BeginCachePostMove(UE::Cook::FGeneratorPackage& Generator,
+		UE::Cook::FPackageData& PackageData, UE::Cook::FCookerTimer& Timer);
 
 	/** Try calling the splitter's populate to create the package */
 	UPackage* TryPopulateGeneratedPackage(UE::Cook::FPopulatePackageContext& Context);
@@ -1281,4 +1274,5 @@ private:
 	friend UE::Cook::FPendingCookedPlatformData;
 	friend UE::Cook::FPlatformManager;
 	friend UE::Cook::FRequestCluster;
+	friend UE::Cook::FGeneratorPackage;
 };
diff --git a/Engine/Source/Editor/UnrealEd/Private/Commandlets/CookCommandlet.cpp b/Engine/Source/Editor/UnrealEd/Private/Commandlets/CookCommandlet.cpp
--- a/Engine/Source/Editor/UnrealEd/Private/Commandlets/CookCommandlet.cpp
+++ b/Engine/Source/Editor/UnrealEd/Private/Commandlets/CookCommandlet.cpp
@@ -70,8 +70,7 @@ namespace DetailedCookStats
 	extern double TickCookOnTheSideLoadPackagesTimeSec;
 	extern double TickCookOnTheSideResolveRedirectorsTimeSec;
 	extern double TickCookOnTheSideSaveCookedPackageTimeSec;
-	extern double TickCookOnTheSideBeginPrepareSaveTimeSec;
-	extern double TickCookOnTheSideFinishPrepareSaveTimeSec;
+	extern double TickCookOnTheSidePrepareSaveTimeSec;
 	extern double GameCookModificationDelegateTimeSec;
 	double TickLoopGCTimeSec = 0.0;
 	double TickLoopRecompileShaderRequestsTimeSec = 0.0;
@@ -98,8 +97,7 @@ namespace DetailedCookStats
 		ADD_COOK_STAT_FLT(" 0. 1. 1. 0", TickCookOnTheSideLoadPackagesTimeSec);
 		ADD_COOK_STAT_FLT(" 0. 1. 1. 1", TickCookOnTheSideSaveCookedPackageTimeSec);
 		ADD_COOK_STAT_FLT(" 0. 1. 1. 1. 0", TickCookOnTheSideResolveRedirectorsTimeSec);
-		ADD_COOK_STAT_FLT(" 0. 1. 1. 2", TickCookOnTheSideBeginPrepareSaveTimeSec);
-		ADD_COOK_STAT_FLT(" 0. 1. 1. 3", TickCookOnTheSideFinishPrepareSaveTimeSec);
+		ADD_COOK_STAT_FLT(" 0. 1. 1. 2", TickCookOnTheSidePrepareSaveTimeSec);
 		ADD_COOK_STAT_FLT(" 0. 1. 2", TickLoopGCTimeSec);
 		ADD_COOK_STAT_FLT(" 0. 1. 3", TickLoopRecompileShaderRequestsTimeSec);
 		ADD_COOK_STAT_FLT(" 0. 1. 4", TickLoopShaderProcessAsyncResultsTimeSec);
diff --git a/Engine/Source/Editor/UnrealEd/Private/CookOnTheFlyServer.cpp b/Engine/Source/Editor/UnrealEd/Private/CookOnTheFlyServer.cpp
--- a/Engine/Source/Editor/UnrealEd/Private/CookOnTheFlyServer.cpp
+++ b/Engine/Source/Editor/UnrealEd/Private/CookOnTheFlyServer.cpp
@@ -2667,15 +2667,23 @@ void UCookOnTheFlyServer::TickNetwork()
 	}
 }
 
-UE::Cook::Private::FRegisteredCookPackageSplitter* UCookOnTheFlyServer::TryGetRegisteredCookPackageSplitter(UE::Cook::FPackageData& PackageData, UObject*& OutSplitDataObject, bool& bOutError)
+UE::Cook::EPollStatus UCookOnTheFlyServer::ConditionalCreateGeneratorPackage(UE::Cook::FPackageData& PackageData)
 {
-	bOutError = false;
-	OutSplitDataObject = nullptr;
+	using namespace UE::Cook;
 
-	UE::Cook::Private::FRegisteredCookPackageSplitter* FoundSplitter = nullptr;
-	UObject* FoundSplitDataObject = nullptr;
+	Private::FRegisteredCookPackageSplitter* Splitter = nullptr;
+	UObject* SplitDataObject = nullptr;
+	bool bGeneratorExists = false;
+	ON_SCOPE_EXIT
+	{
+		if (!bGeneratorExists)
+		{
+			// Destroy any old GeneratorPackage if we no longer find we need one
+			PackageData.DestroyGeneratorPackage();
+		}
+	};
 
-	TArray<UE::Cook::Private::FRegisteredCookPackageSplitter*> FoundRegisteredSplitters;
+	TArray<Private::FRegisteredCookPackageSplitter*> FoundRegisteredSplitters;
 
 	for (FWeakObjectPtr& WeakObj : PackageData.GetCachedObjectsInOuter())
 	{
@@ -2688,38 +2696,31 @@ UE::Cook::Private::FRegisteredCookPackageSplitter* UCookOnTheFlyServer::TryGetRe
 		FoundRegisteredSplitters.Reset();
 		RegisteredSplitDataClasses.MultiFind(Obj->GetClass(), FoundRegisteredSplitters);
 
-		for (UE::Cook::Private::FRegisteredCookPackageSplitter* Splitter : FoundRegisteredSplitters)
+		for (Private::FRegisteredCookPackageSplitter* SplitterForObject: FoundRegisteredSplitters)
 		{
-			if (Splitter && Splitter->ShouldSplitPackage(Obj))
+			if (SplitterForObject && SplitterForObject->ShouldSplitPackage(Obj))
 			{
 				if (!Obj->HasAnyFlags(RF_Public))
 				{
 					UE_LOG(LogCook, Error, TEXT("SplitterData object %s must be publicly referenceable so we can keep them from being garbage collected"), *Obj->GetFullName());
-					bOutError = true;
-					return nullptr;
+					return EPollStatus::Error;
 				}
 
-				if (FoundSplitter)
+				if (Splitter)
 				{
 					UE_LOG(LogCook, Error, TEXT("Found more than one registered Cook Package Splitter for package %s."), *PackageData.GetPackageName().ToString());
-					bOutError = true;
-					return nullptr;
+					return EPollStatus::Error;
 				}
 
-				FoundSplitter = Splitter;
-				FoundSplitDataObject = Obj;
+				Splitter = SplitterForObject;
+				SplitDataObject = Obj;
 			}
 		}
 	}
-
-	OutSplitDataObject = FoundSplitDataObject;
-	return FoundSplitter;
-}
-
-UE::Cook::FGeneratorPackage* UCookOnTheFlyServer::CreateGeneratorPackage(UE::Cook::FPackageData& PackageData, UObject* SplitDataObject, UE::Cook::Private::FRegisteredCookPackageSplitter* Splitter)
-{
-	check(Splitter);
-	check(!PackageData.GetGeneratorPackage());
+	if (!Splitter)
+	{
+		return EPollStatus::Success;
+	}
 
 	// TODO: Add support for cooking in the editor. Possibly moot since we plan to deprecate cooking in the editor.
 	if (IsCookingInEditor())
@@ -2728,96 +2729,235 @@ UE::Cook::FGeneratorPackage* UCookOnTheFlyServer::CreateGeneratorPackage(UE::Coo
 		// of it into the streaming packages. To allow its use in the editor, we will need to make it non-destructive
 		// (by e.g. copying to new packages), or restore the package after the changes have been made.
 		UE_LOG(LogCook, Error, TEXT("Cooking in editor doesn't support Cook Package Splitters."));
-		return nullptr;
+		return EPollStatus::Error;
 	}
 
-	UE_LOG(LogCook, Display, TEXT("Splitting Package %s with class %s acting on object %s."), *PackageData.GetPackageName().ToString(), *Splitter->GetSplitDataClass()->GetName(), *SplitDataObject->GetFullName());
+	UE_LOG(LogCook, Display, TEXT("Splitting Package %s with class %s acting on object %s."),
+		*PackageData.GetPackageName().ToString(), *Splitter->GetSplitDataClass()->GetName(), *SplitDataObject->GetFullName());
 
 	// Create instance of CookPackageSplitter class
 	ICookPackageSplitter* SplitterInstance = Splitter->CreateInstance(SplitDataObject);
 	if (!SplitterInstance)
 	{
 		UE_LOG(LogCook, Error, TEXT("Error instantiating Cook Package Splitter for object %s."), *SplitDataObject->GetFullName());
-		return nullptr;
+		return EPollStatus::Error;
 	}
 
-	// Create a FGeneratorPackage helper object using this CookPackageSplitter instance and return it
+	// Create a FGeneratorPackage helper object using this CookPackageSplitter instance
+	bGeneratorExists = true;
 	PackageData.CreateGeneratorPackage(SplitDataObject, SplitterInstance);
-	UE::Cook::FGeneratorPackage* GeneratorPackage = PackageData.GetGeneratorPackage();
-	return GeneratorPackage;
+	return EPollStatus::Success;
 }
 
-void UCookOnTheFlyServer::SplitPackage(UE::Cook::FPackageData& PackageData, UE::Cook::FGeneratorPackage* GeneratorStruct, bool& bOutCompleted, bool& bOutError)
+UE::Cook::EPollStatus UCookOnTheFlyServer::QueueGeneratedPackages(UE::Cook::FGeneratorPackage& Generator,
+	UE::Cook::FPackageData& PackageData)
 {
 	using namespace UE::Cook;
 
-	bOutCompleted = false;
-	bOutError = false;
-	check(GeneratorStruct);
-	ICookPackageSplitter* Splitter = GeneratorStruct->GetCookPackageSplitterInstance();
-	UObject* SplitObject = GeneratorStruct->FindSplitDataObject();
+	ICookPackageSplitter* Splitter = Generator.GetCookPackageSplitterInstance();
+	UObject* SplitObject = Generator.FindSplitDataObject();
 	if (!SplitObject)
 	{
-		UE_LOG(LogCook, Error, TEXT("Could not find SplitDataObject %s"), *GeneratorStruct->GetSplitDataObjectName().ToString());
-		bOutError = true;
-		return;
+		UE_LOG(LogCook, Error, TEXT("Could not find SplitDataObject %s"), *Generator.GetSplitDataObjectName().ToString());
+		return EPollStatus::Error;
 	}
 
-	if (!GeneratorStruct->HasGeneratedList())
+	if (Generator.GetGeneratorSaveState() <= FGeneratorPackage::ESaveState::GenerateList)
 	{
 		// Call the splitter to generate the list
-		if (!GeneratorStruct->TryGenerateList(SplitObject, *PackageDatas))
+		if (!Generator.TryGenerateList(SplitObject, *PackageDatas))
 		{
-			bOutError = true;
-			return;
+			return EPollStatus::Error;
 		}
-		GeneratorStruct->SetGeneratedList();
+		Generator.SetGeneratorSaveStateComplete(FGeneratorPackage::ESaveState::GenerateList);
 	}
 
-	if (!GeneratorStruct->HasClearedOldPackages())
+	if (Generator.GetGeneratorSaveState() <= FGeneratorPackage::ESaveState::ClearOldPackagesLastAttempt)
 	{
-		for (const FGeneratorPackage::FGeneratedStruct& GeneratedStruct : GeneratorStruct->GetPackagesToGenerate())
+		for (const FGeneratorPackage::FGeneratedStruct& GeneratedStruct : Generator.GetPackagesToGenerate())
 		{
 			const FString GeneratedPackageName = GeneratedStruct.PackageData->GetPackageName().ToString();
 			if (FindObject<UPackage>(nullptr, *GeneratedPackageName))
 			{
-				if (GeneratorStruct->HasClearedOldPackagesWithGC())
+				if (Generator.GetGeneratorSaveState() < FGeneratorPackage::ESaveState::ClearOldPackagesLastAttempt)
 				{
-					UE_LOG(LogCook, Error, TEXT("PackageSplitter was unable to construct new generated packages because an old version of the package is already in memory and GC did not remove it. Splitter=%s, Generated=%s."),
-						*GeneratorStruct->GetSplitDataObjectName().ToString(), *GeneratedStruct.RelativePath);
-					bOutError = true;
-					return;
+					PackageData.SetIsPrepareSaveRequiresGC(true);
+					Generator.SetGeneratorSaveState(FGeneratorPackage::ESaveState::ClearOldPackagesLastAttempt);
+					return EPollStatus::Incomplete;
 				}
 				else
 				{
-					GeneratorStruct->SetClearedOldPackagesWithGC();
-					return;
+					UE_LOG(LogCook, Error, TEXT("PackageSplitter was unable to construct new generated packages because an old version of the package is already in memory and GC did not remove it. Splitter=%s, Generated=%s."),
+						*Generator.GetSplitDataObjectName().ToString(), *GeneratedStruct.RelativePath);
+					return EPollStatus::Error;
 				}
 			}
 		}
 
-		GeneratorStruct->SetClearedOldPackages();
+		Generator.SetGeneratorSaveStateComplete(FGeneratorPackage::ESaveState::ClearOldPackagesLastAttempt);
 	}
 
-	UPackage* Owner = GeneratorStruct->GetOwner().GetPackage();
+	UPackage* Owner = PackageData.GetPackage();
 	FName OwnerName = Owner->GetFName();
-	if (!GeneratorStruct->HasQueuedGeneratedPackages())
+	if (Generator.GetGeneratorSaveState() <= FGeneratorPackage::ESaveState::QueueGeneratedPackages)
 	{
-		for (const FGeneratorPackage::FGeneratedStruct& GeneratedStruct : GeneratorStruct->GetPackagesToGenerate())
+		for (const FGeneratorPackage::FGeneratedStruct& GeneratedStruct : Generator.GetPackagesToGenerate())
 		{
 			FPackageData* GeneratedPackageData = GeneratedStruct.PackageData;
 			GeneratedPackageData->ClearCookedPlatformData();
 			QueueDiscoveredPackageData(*GeneratedPackageData, FInstigator(EInstigator::GeneratedPackage, OwnerName));
 		}
-		GeneratorStruct->SetQueuedGeneratedPackages();
+		Generator.SetGeneratorSaveStateComplete(FGeneratorPackage::ESaveState::QueueGeneratedPackages);
+	}
+	return EPollStatus::Success;
+}
+
+UE::Cook::EPollStatus UCookOnTheFlyServer::PrepareSaveGeneratorPackage(UE::Cook::FGeneratorPackage& Generator,
+	UE::Cook::FPackageData& PackageData, UE::Cook::FCookerTimer& Timer)
+{
+	using namespace UE::Cook;
+
+	if (Generator.GetGeneratorSaveState() <= FGeneratorPackage::ESaveState::FinishCachePreObjectsToMove)
+	{
+		if (PackageData.GetNumPendingCookedPlatformData() > 0)
+		{
+			return EPollStatus::Incomplete;
+		}
+		Generator.SetGeneratorSaveStateComplete(FGeneratorPackage::ESaveState::FinishCachePreObjectsToMove);
+	}
+
+	// GeneratedPackagesForPresave is used by multiple steps, recreate it when needed each time we come in to this function
+	TArray<ICookPackageSplitter::FGeneratedPackageForPreSave> GeneratedPackagesForPresave;
+	if (Generator.GetGeneratorSaveState() <= FGeneratorPackage::ESaveState::FinishCacheObjectsToMove)
+	{
+		if (Generator.GetGeneratorSaveState() <= FGeneratorPackage::ESaveState::BeginCacheObjectsToMove)
+		{
+			EPollStatus Result = BeginCacheObjectsToMove(Generator, PackageData, Timer, GeneratedPackagesForPresave);
+			if (Result != EPollStatus::Success)
+			{
+				return Result;
+			}
+			Generator.SetGeneratorSaveStateComplete(FGeneratorPackage::ESaveState::BeginCacheObjectsToMove);
+		}
+		check(Generator.GetGeneratorSaveState() <= FGeneratorPackage::ESaveState::FinishCacheObjectsToMove);
+		if (PackageData.GetNumPendingCookedPlatformData() > 0)
+		{
+			return EPollStatus::Incomplete;
+		}
+		Generator.SetGeneratorSaveStateComplete(FGeneratorPackage::ESaveState::FinishCacheObjectsToMove);
+	}
+
+	if (Generator.GetGeneratorSaveState() <= FGeneratorPackage::ESaveState::CallPreSaveGeneratorPackage)
+	{
+		EPollStatus Result = PreSaveGeneratorPackage(PackageData, Generator, GeneratedPackagesForPresave);
+		if (Result != EPollStatus::Success)
+		{
+			return Result;
+		}
+		Generator.SetGeneratorSaveStateComplete(FGeneratorPackage::ESaveState::CallPreSaveGeneratorPackage);
+	}
+
+	if (Generator.GetGeneratorSaveState() <= FGeneratorPackage::ESaveState::FinishCachePostMove)
+	{
+		if (Generator.GetGeneratorSaveState() <= FGeneratorPackage::ESaveState::BeginCachePostMove)
+		{
+			EPollStatus Result = BeginCachePostMove(Generator, PackageData, Timer);
+			if (Result != EPollStatus::Success)
+			{
+				return Result;
+			}
+			Generator.SetGeneratorSaveStateComplete(FGeneratorPackage::ESaveState::BeginCachePostMove);
+		}
+		check(Generator.GetGeneratorSaveState() <= FGeneratorPackage::ESaveState::FinishCachePostMove);
+		if (PackageData.GetNumPendingCookedPlatformData() > 0)
+		{
+			return EPollStatus::Incomplete;
+		}
+		Generator.SetGeneratorSaveStateComplete(FGeneratorPackage::ESaveState::FinishCachePostMove);
+	}
+	check(Generator.GetGeneratorSaveState() == FGeneratorPackage::ESaveState::ReadyForSave);
+
+	return EPollStatus::Success;
+}
+
+UE::Cook::EPollStatus UCookOnTheFlyServer::BeginCacheObjectsToMove(UE::Cook::FGeneratorPackage& Generator,
+	UE::Cook::FPackageData& PackageData, UE::Cook::FCookerTimer& Timer,
+	TArray<ICookPackageSplitter::FGeneratedPackageForPreSave>& GeneratedPackagesForPresave)
+{
+	using namespace UE::Cook;
+
+	UPackage* Package = PackageData.GetPackage();
+	ICookPackageSplitter* Splitter = Generator.GetCookPackageSplitterInstance();
+	UObject* SplitDataObject = Generator.FindSplitDataObject();
+	if (!Package || !Splitter || !SplitDataObject)
+	{
+		UE_LOG(LogCook, Error, TEXT("CookPackageSplitter is missing %s during BeginCacheObjectsToMove. PackageName: %s."),
+			(!Package ? TEXT("Package") : (!Splitter ? TEXT("Splitter") : TEXT("SplitDataObject"))),
+			*PackageData.GetPackageName().ToString());
+		return EPollStatus::Error;
+	}
+
+	if (Generator.GetGeneratorSaveState() <= FGeneratorPackage::ESaveState::CallObjectsToMoveIntoGenerator)
+	{
+		ConstructGeneratedPackagesForPresave(PackageData, Generator, GeneratedPackagesForPresave);
+		Generator.SetObjectsToMoveIntoGenerator(Package, SplitDataObject, Splitter, PackageData.GetCachedObjectsInOuter(),
+			GeneratedPackagesForPresave);
+		Generator.SetGeneratorSaveStateComplete(FGeneratorPackage::ESaveState::CallObjectsToMoveIntoGenerator);
+	}
+
+	FBeginCacheObjects& BeginCacheObjects = Generator.GetGeneratorBeginCacheObjects();
+	EPollStatus Result = CallBeginCacheOnObjects(PackageData, Package, BeginCacheObjects.ObjectsInRound,
+		BeginCacheObjects.NextIndexInRound, Timer);
+	if (Result != EPollStatus::Success)
+	{
+		return Result;
 	}
+	BeginCacheObjects.EndRound();
+	return EPollStatus::Success;
+}
+
+UE::Cook::EPollStatus UCookOnTheFlyServer::PreSaveGeneratorPackage(UE::Cook::FPackageData& PackageData, UE::Cook::FGeneratorPackage& Generator,
+	TArray<ICookPackageSplitter::FGeneratedPackageForPreSave>& GeneratedPackagesForPresave)
+{
+	using namespace UE::Cook;
 
-	TArrayView<FGeneratorPackage::FGeneratedStruct> PackagesToGenerate = GeneratorStruct->GetPackagesToGenerate();
+	UPackage* Package = PackageData.GetPackage();
+	ICookPackageSplitter* Splitter = Generator.GetCookPackageSplitterInstance();
+	UObject* SplitDataObject = Generator.FindSplitDataObject();
+	if (!Package || !Splitter || !SplitDataObject)
+	{
+		UE_LOG(LogCook, Error, TEXT("CookPackageSplitter is missing %s during PreSaveGeneratorPackage. PackageName: %s."),
+			(!Package ? TEXT("Package") : (!Splitter ? TEXT("Splitter") : TEXT("SplitDataObject"))),
+			*PackageData.GetPackageName().ToString());
+		return EPollStatus::Error;
+	}
+
+	ConstructGeneratedPackagesForPresave(PackageData, Generator, GeneratedPackagesForPresave);
+	Splitter->PreSaveGeneratorPackage(Package, SplitDataObject, GeneratedPackagesForPresave);
+	return EPollStatus::Success;
+}
+
+void UCookOnTheFlyServer::ConstructGeneratedPackagesForPresave(UE::Cook::FPackageData& PackageData, UE::Cook::FGeneratorPackage& Generator,
+	TArray<ICookPackageSplitter::FGeneratedPackageForPreSave>& GeneratedPackagesForPresave)
+{
+	using namespace UE::Cook;
+
+	if (GeneratedPackagesForPresave.Num() > 0)
+	{
+		// Already constructed, save time by early exiting
+		return;
+	}
+	UPackage* Package = PackageData.GetPackage();
+	check(Package);
+
+	// We need to find or (create empty stub packages for) each of the PackagesToGenerate so that PreSaveGeneratorPackage
+	// can refer to them to create hardlinks in the cooked Generator package
+	TArrayView<FGeneratorPackage::FGeneratedStruct> PackagesToGenerate = Generator.GetPackagesToGenerate();
 	TArray<ICookPackageSplitter::FGeneratedPackageForPreSave> SplitterDatas;
 	SplitterDatas.Reserve(PackagesToGenerate.Num());
 	for (FGeneratorPackage::FGeneratedStruct& GeneratedStruct : PackagesToGenerate)
 	{
-		ICookPackageSplitter::FGeneratedPackageForPreSave& SplitterData = SplitterDatas.Emplace_GetRef();
+		ICookPackageSplitter::FGeneratedPackageForPreSave& SplitterData = GeneratedPackagesForPresave.Emplace_GetRef();
 		SplitterData.RelativePath = GeneratedStruct.RelativePath;
 		SplitterData.bCreatedAsMap = GeneratedStruct.bCreateAsMap;
 
@@ -2825,18 +2965,67 @@ void UCookOnTheFlyServer::SplitPackage(UE::Cook::FPackageData& PackageData, UE::
 		SplitterData.Package = FindObject<UPackage>(nullptr, *GeneratedPackageName);
 		if (!SplitterData.Package)
 		{
-			SplitterData.Package = GeneratorStruct->CreateGeneratedUPackage(GeneratedStruct, Owner, *GeneratedPackageName);
+			SplitterData.Package = Generator.CreateGeneratedUPackage(GeneratedStruct, Package, *GeneratedPackageName);
 		}
 	}
+}
 
-	UPackage* OwnerPackage = PackageData.GetPackage();
-	check(OwnerPackage);
-	GeneratorStruct->GetCookPackageSplitterInstance()->PreSaveGeneratorPackage(OwnerPackage, SplitObject, SplitterDatas);
-	// PreSaveGenerator package can add new objects to the package due to moving them in from OneFilePerActor packages
-	// We want to support calling BeginCacheForCookedPlatformData on those new objects, so recreate the ObjectCache now.
-	PackageData.RecreateObjectCache();
 
-	bOutCompleted = true;
+UE::Cook::EPollStatus UCookOnTheFlyServer::BeginCachePostMove(UE::Cook::FGeneratorPackage& Generator,
+	UE::Cook::FPackageData& PackageData, UE::Cook::FCookerTimer& Timer)
+{
+	using namespace UE::Cook;
+
+	UPackage* Package = PackageData.GetPackage();
+	ICookPackageSplitter* Splitter = Generator.GetCookPackageSplitterInstance();
+	UObject* SplitDataObject = Generator.FindSplitDataObject();
+	if (!Package || !Splitter || !SplitDataObject)
+	{
+		UE_LOG(LogCook, Error, TEXT("CookPackageSplitter is missing %s during BeginCachePostMove. PackageName: %s."),
+			(!Package ? TEXT("Package") : (!Splitter ? TEXT("Splitter") : TEXT("SplitDataObject"))),
+			*PackageData.GetPackageName().ToString());
+		return EPollStatus::Error;
+	}
+
+	if (Generator.GetGeneratorSaveState() <= FGeneratorPackage::ESaveState::CallGetPostMoveObjects)
+	{
+		Generator.SetPostMoveObjects(Package);
+		Generator.SetGeneratorSaveStateComplete(FGeneratorPackage::ESaveState::CallGetPostMoveObjects);
+	}
+
+	FBeginCacheObjects& BeginCacheObjects = Generator.GetGeneratorBeginCacheObjects();
+	EPollStatus Result = CallBeginCacheOnObjects(PackageData, Package, BeginCacheObjects.ObjectsInRound,
+		BeginCacheObjects.NextIndexInRound, Timer);
+	if (PackageData.GetNumPendingCookedPlatformData() > 0 && !Generator.HasIssuedUndeclaredMovedObjectsWarning())
+	{
+		UObject* FirstPendingObject = nullptr;
+		FString FirstPendingObjectName;
+		for (const FPendingCookedPlatformData& Pending : PackageDatas->GetPendingCookedPlatformDatas())
+		{
+			if (&Pending.PackageData == &PackageData)
+			{
+				FString ObjectName = Pending.Object.IsValid() ? Pending.Object.Get()->GetPathName() : TEXT("");
+				if (ObjectName.Len() && (!FirstPendingObject || ObjectName < FirstPendingObjectName))
+				{
+					FirstPendingObject = Pending.Object.Get();
+					FirstPendingObjectName = MoveTemp(ObjectName);
+				}
+			}
+		}
+		UE_LOG(LogCook, Warning, TEXT("CookPackageSplitter created or moved objects during PreSaveGeneratorPackage that are not yet ready to save. This will cause an error if garbage collection runs before the package is saved.\n")
+			TEXT("Change the splitter's GetObjectsToMoveIntoGenerator to construct new objects and declare existing objects that will be moved from other packages.\n")
+			TEXT("SplitterObject: %s. NumPendingObjects: %d. FirstPendingObject: %s"),
+			*SplitDataObject->GetFullName(), PackageData.GetNumPendingCookedPlatformData(),
+			FirstPendingObject ? *FirstPendingObject->GetFullName() : TEXT("<unknown>"));
+		Generator.SetHasIssuedUndeclaredMovedObjectsWarning();
+	}
+	if (Result != EPollStatus::Success)
+	{
+		return Result;
+	}
+	BeginCacheObjects.EndRound();
+
+	return EPollStatus::Success;
 }
 
 UPackage* UCookOnTheFlyServer::TryPopulateGeneratedPackage(UE::Cook::FPopulatePackageContext& Context)
@@ -2926,129 +3115,164 @@ UPackage* UCookOnTheFlyServer::TryPopulateGeneratedPackage(UE::Cook::FPopulatePa
 	return GeneratedPackage;
 }
 
-
-bool UCookOnTheFlyServer::BeginPrepareSave(UE::Cook::FPackageData& PackageData, UE::Cook::FCookerTimer& Timer, bool bIsPreCaching)
+UE::Cook::EPollStatus UCookOnTheFlyServer::PrepareSave(UE::Cook::FPackageData& PackageData,
+	UE::Cook::FCookerTimer& Timer, bool bPrecaching)
 {
-	if (PackageData.GetCookedPlatformDataCalled())
+	using namespace UE::Cook;
+
+	if (PackageData.GetCookedPlatformDataComplete())
 	{
-		return true;
+		return EPollStatus::Success;
 	}
-
-	if (PackageData.GetHasBeginPrepareSaveFailed())
+	if (PackageData.HasPrepareSaveFailed())
 	{
-		return false;
+		return EPollStatus::Error;
 	}
-
-	if (!PackageData.GetCookedPlatformDataStarted())
+	EPollStatus Result = PrepareSaveInternal(PackageData, Timer, bPrecaching);
+	if (Result == EPollStatus::Error)
 	{
-		if (PackageData.GetNumPendingCookedPlatformData() > 0)
-		{
-			// A previous Save was started and deleted after some calls to BeginCacheForCookedPlatformData occurred, and some of those objects have still not returned true for IsCachedCookedPlatformDataLoaded
-			// We need to wait for all of pending async calls from the cancelled save to finish before we start the new ones
-			return false;
-		}
-		PackageData.SetCookedPlatformDataStarted(true);
+		PackageData.SetHasPrepareSaveFailed(true);
 	}
+	return Result;
+}
 
-	UE_SCOPED_HIERARCHICAL_COOKTIMER_AND_DURATION(BeginPrepareSave, DetailedCookStats::TickCookOnTheSideBeginPrepareSaveTimeSec);
+UE::Cook::EPollStatus UCookOnTheFlyServer::PrepareSaveInternal(UE::Cook::FPackageData& PackageData,
+	UE::Cook::FCookerTimer& Timer, bool bPrecaching)
+{
+	using namespace UE::Cook;
 
+	UE_SCOPED_HIERARCHICAL_COOKTIMER_AND_DURATION(PrepareSave, DetailedCookStats::TickCookOnTheSidePrepareSaveTimeSec);
 #if DEBUG_COOKONTHEFLY 
 	UE_LOG(LogCook, Display, TEXT("Caching objects for package %s"), *PackageData.GetPackageName().ToString());
 #endif
 	UPackage* Package = PackageData.GetPackage();
 	check(Package && Package->IsFullyLoaded());
-	check(PackageData.GetState() == UE::Cook::EPackageState::Save);
-	PackageData.CreateObjectCache();
+	check(PackageData.GetState() == EPackageState::Save);
+	FGeneratorPackage* Generator = nullptr;
 
-	if (!PackageData.HasCompletedGeneration())
+	if (!PackageData.GetCookedPlatformDataCalled())
 	{
-		// Check for Splitting the package; this needs to happen before we make any of the BeginCacheForCookedPlatformData calls in the package
-		bool bError = false;
-
-		UE::Cook::FGeneratorPackage* Generator = nullptr;
-		if (!PackageData.HasInitializedGeneratorSave())
+		if (!PackageData.GetCookedPlatformDataStarted())
 		{
-			PackageData.SetInitializedGeneratorSave(true);
-			PackageData.DestroyGeneratorPackage();
-			UObject* SplitDataObject = nullptr;
-			UE::Cook::Private::FRegisteredCookPackageSplitter* Splitter = TryGetRegisteredCookPackageSplitter(PackageData, SplitDataObject, bError);
-			if (bError)
+			if (PackageData.GetNumPendingCookedPlatformData() > 0)
 			{
-				PackageData.SetHasBeginPrepareSaveFailed(true);
-				return false;
+				// A previous Save was started and deleted after some calls to BeginCacheForCookedPlatformData occurred, and some of those objects have still not returned true for IsCachedCookedPlatformDataLoaded
+				// We need to wait for all of pending async calls from the cancelled save to finish before we start the new ones
+				return EPollStatus::Incomplete;
 			}
+			PackageData.SetCookedPlatformDataStarted(true);
+		}
+
+		PackageData.CreateObjectCache();
+
+		// Note that we cache cooked data for all requested platforms, rather than only for the requested platforms that have not cooked yet.  This allows
+		// us to avoid the complexity of needing to cancel the Save and keep track of the old list of uncooked platforms whenever the cooked platforms change
+		// while PrepareSave is active.
+		// Currently this does not cause significant cost since saving new platforms with some platforms already saved is a rare operation.
 
-			if (Splitter)
+		int32& CookedPlatformDataNextIndex = PackageData.GetCookedPlatformDataNextIndex();
+		if (CookedPlatformDataNextIndex == 0)
+		{
+			if (!BuildDefinitions->TryRemovePendingBuilds(PackageData.GetPackageName()))
 			{
-				// Found a splitter, create generator package
-				Generator = CreateGeneratorPackage(PackageData, SplitDataObject, Splitter);
-				if (!Generator)
-				{
-					PackageData.SetHasBeginPrepareSaveFailed(true);
-					return false;
-				}
+				// Builds are in progress; wait for them to complete
+				return EPollStatus::Incomplete;
 			}
 		}
-		else
+
+		TArray<FWeakObjectPtr>& CachedObjectsInOuter = PackageData.GetCachedObjectsInOuter();
+		EPollStatus Result = CallBeginCacheOnObjects(PackageData, Package, CachedObjectsInOuter,
+			CookedPlatformDataNextIndex, Timer);
+		if (Result != EPollStatus::Success)
 		{
-			Generator = PackageData.GetGeneratorPackage();
+			return Result;
 		}
-		if (Generator)
+
+		// Check for whether the Package has a Splitter and initialize its list if so
+		if (!PackageData.HasInitializedGeneratorSave())
 		{
-			// Don't split a generator package when precaching
-			if (bIsPreCaching)
+			Result = ConditionalCreateGeneratorPackage(PackageData);
+			if (Result != EPollStatus::Success)
 			{
-				return false;
+				return Result;
 			}
-			bool bCompleted = false;
-			SplitPackage(PackageData, Generator, bCompleted, bError);
-			if (bError)
+			PackageData.SetInitializedGeneratorSave(true);
+		}
+		Generator = PackageData.GetGeneratorPackage();
+		if (Generator)
+		{
+			if (bPrecaching)
 			{
-				// SplitPackage failure marks the package data for PumpSaves to handle this as an error and to put it in idle state
-				PackageData.SetHasBeginPrepareSaveFailed(true);
-				return false;
+				// We're not allowed to split the package while precaching, so don't progress any farther
+				// with PrepareSave while precaching
+				return EPollStatus::Incomplete;
 			}
-			if (!bCompleted)
+			Result = QueueGeneratedPackages(*Generator, PackageData);
+			if (Result != EPollStatus::Success)
 			{
-				return false;
+				return Result;
 			}
 		}
-		PackageData.SetCompletedGeneration(true);
-	}
 
-	// Note that we cache cooked data for all requested platforms, rather than only for the requested platforms that have not cooked yet.  This allows
-	// us to avoid the complexity of needing to cancel the Save and keep track of the old list of uncooked platforms whenever the cooked platforms change
-	// while BeginPrepareSave is active.
-	// Currently this does not cause significant cost since saving new platforms with some platforms already saved is a rare operation.
+		PackageData.SetCookedPlatformDataCalled(true);
+	}
+	else
+	{
+		Generator = PackageData.GetGeneratorPackage();
+	}
 
-	int32& CookedPlatformDataNextIndex = PackageData.GetCookedPlatformDataNextIndex();
-	if (CookedPlatformDataNextIndex == 0)
+	if (Generator)
 	{
-		if (!BuildDefinitions->TryRemovePendingBuilds(PackageData.GetPackageName()))
+		if (bPrecaching)
 		{
-			// Builds are in progress; wait for them to complete
-			return false;
+			// We're not allowed to split the package while precaching, so don't progress any farther
+			// with PrepareSave while precaching
+			return EPollStatus::Incomplete;
+		}
+		EPollStatus Result = PrepareSaveGeneratorPackage(*Generator, PackageData, Timer);
+		if (Result != EPollStatus::Success)
+		{
+			return Result;
 		}
 	}
+	else
+	{
+		if (PackageData.GetNumPendingCookedPlatformData() > 0)
+		{
+			return EPollStatus::Incomplete;
+		}
+	}
+
+	check(PackageData.GetNumPendingCookedPlatformData() == 0);
+	PackageData.SetCookedPlatformDataComplete(true);
+	return EPollStatus::Success;
+}
+
+UE::Cook::EPollStatus UCookOnTheFlyServer::CallBeginCacheOnObjects(UE::Cook::FPackageData& PackageData,
+	UPackage* Package, TArray<FWeakObjectPtr>& Objects, int32& NextIndex, UE::Cook::FCookerTimer& Timer)
+{
+	using namespace UE::Cook;
+
+	check(Package);
 
-	TArray<FWeakObjectPtr>& CachedObjectsInOuter = PackageData.GetCachedObjectsInOuter();
-	FWeakObjectPtr* CachedObjectsInOuterData = CachedObjectsInOuter.GetData();
 	TArray<const ITargetPlatform*, TInlineAllocator<ExpectedMaxNumPlatforms>> TargetPlatforms;
 	PackageData.GetRequestedPlatforms(TargetPlatforms);
+
+	FWeakObjectPtr* ObjectsData = Objects.GetData();
 	int NumPlatforms = TargetPlatforms.Num();
-	int NumIndexes = CachedObjectsInOuter.Num() * NumPlatforms;
+	int NumIndexes = Objects.Num() * NumPlatforms;
 	UE_TRACK_REFERENCING_PACKAGE_SCOPED(Package, PackageAccessTrackingOps::NAME_CookerBuildObject);
-	while (CookedPlatformDataNextIndex < NumIndexes)
+	while (NextIndex < NumIndexes)
 	{
-		int ObjectIndex = CookedPlatformDataNextIndex / NumPlatforms;
-		int PlatformIndex = CookedPlatformDataNextIndex - ObjectIndex * NumPlatforms;
-		UObject* Obj = CachedObjectsInOuterData[ObjectIndex].Get();
+		int ObjectIndex = NextIndex / NumPlatforms;
+		int PlatformIndex = NextIndex - ObjectIndex * NumPlatforms;
+		UObject* Obj = ObjectsData[ObjectIndex].Get();
 		if (!Obj)
 		{
 			// Objects can be marked as pending kill even without a garbage collect, and our weakptr.get will return null for them, so we have to always check the WeakPtr before using it
 			// Treat objects that have been marked as pending kill or deleted as no-longer-required for BeginCacheForCookedPlatformData and ClearAllCachedCookedPlatformData
-			CachedObjectsInOuterData[ObjectIndex] = nullptr; // If the weakptr is merely pendingkill, set it to null explicitly so we don't think that we've called BeginCacheForCookedPlatformData on it if it gets unmarked pendingkill later
-			++CookedPlatformDataNextIndex;
+			ObjectsData[ObjectIndex] = nullptr; // If the weakptr is merely pendingkill, set it to null explicitly so we don't think that we've called BeginCacheForCookedPlatformData on it if it gets unmarked pendingkill later
+			++NextIndex;
 			continue;
 		}
 		const ITargetPlatform* TargetPlatform = TargetPlatforms[PlatformIndex];
@@ -3060,7 +3284,7 @@ bool UCookOnTheFlyServer::BeginPrepareSave(UE::Cook::FPackageData& PackageData,
 #if DEBUG_COOKONTHEFLY
 				UE_LOG(LogCook, Display, TEXT("Delaying shader compilation of material %s"), *Obj->GetFullName());
 #endif
-				return false;
+				return EPollStatus::Incomplete;
 			}
 		}
 
@@ -3070,13 +3294,13 @@ bool UCookOnTheFlyServer::BeginPrepareSave(UE::Cook::FPackageData& PackageData,
 		{
 			if (*CurrentAsyncCache < 1)
 			{
-				return false;
+				return EPollStatus::Incomplete;
 			}
 			*CurrentAsyncCache -= 1;
 		}
 
 		RouteBeginCacheForCookedPlatformData(Obj, TargetPlatform);
-		++CookedPlatformDataNextIndex;
+		++NextIndex;
 		if (RouteIsCachedCookedPlatformDataLoaded(Obj, TargetPlatform))
 		{
 			if (CurrentAsyncCache)
@@ -3095,40 +3319,14 @@ bool UCookOnTheFlyServer::BeginPrepareSave(UE::Cook::FPackageData& PackageData,
 #if DEBUG_COOKONTHEFLY
 			UE_LOG(LogCook, Display, TEXT("Object %s took too long to cache"), *Obj->GetFullName());
 #endif
-			return false;
-		}
-	}
-
-	PackageData.SetCookedPlatformDataCalled(true);
-	return true;
-}
-
-bool UCookOnTheFlyServer::FinishPrepareSave(UE::Cook::FPackageData& PackageData, UE::Cook::FCookerTimer& Timer)
-{
-	if (PackageData.GetCookedPlatformDataComplete())
-	{
-		return true;
-	}
-
-	if (!PackageData.GetCookedPlatformDataCalled())
-	{
-		if (!BeginPrepareSave(PackageData, Timer, /*bIsPreCaching*/ false))
-		{
-			return false;
+			return EPollStatus::Incomplete;
 		}
-		check(PackageData.GetCookedPlatformDataCalled());
 	}
 
-	if (PackageData.GetNumPendingCookedPlatformData() > 0)
-	{
-		return false;
-	}
-
-	PackageData.SetCookedPlatformDataComplete(true);
-	return true;
+	return EPollStatus::Success;
 }
 
-void UCookOnTheFlyServer::ReleaseCookedPlatformData(UE::Cook::FPackageData& PackageData, bool bCompletedSave)
+void UCookOnTheFlyServer::ReleaseCookedPlatformData(UE::Cook::FPackageData& PackageData, UE::Cook::EReleaseSaveReason ReleaseSaveReason)
 {
 	using namespace UE::Cook;
 
@@ -3138,22 +3336,47 @@ void UCookOnTheFlyServer::ReleaseCookedPlatformData(UE::Cook::FPackageData& Pack
 		return;
 	}
 
+	FGeneratorPackage* Generator = PackageData.GetGeneratorPackage();
+
 	// For every Object on which we called BeginCacheForCookedPlatformData, we need to call ClearAllCachedCookedPlatformData
-	if (PackageData.GetCookedPlatformDataComplete() && bCompletedSave)
+	if (ReleaseSaveReason == EReleaseSaveReason::Completed)
 	{
 		// Since we have completed CookedPlatformData, we know we called BeginCacheForCookedPlatformData on all objects in the package, and none are pending
 		if (!IsCookingInEditor()) // ClearAllCachedCookedPlatformData and WillNeverCacheCookedPlatformDataAgain calls are only used when not in editor
 		{
 			UE_SCOPED_HIERARCHICAL_COOKTIMER(ClearAllCachedCookedPlatformData);
-			for (FWeakObjectPtr& WeakPtr: PackageData.GetCachedObjectsInOuter())
+			if (Generator)
 			{
-				UObject* Object = WeakPtr.Get();
-				if (Object)
+				check(Generator->HasTakenOverCachedCookedPlatformData());
+				check(PackageData.GetCachedObjectsInOuter().Num() == 0);
+
+				FBeginCacheObjects& BeginCacheObjects = Generator->GetGeneratorBeginCacheObjects();
+				for (FBeginCacheObject& BeginCacheObject : BeginCacheObjects.Objects)
 				{
-					Object->ClearAllCachedCookedPlatformData();
-					if (CurrentCookMode == ECookMode::CookByTheBook)
+					UObject* Object = BeginCacheObject.Object.Get();
+					if (Object)
 					{
-						Object->WillNeverCacheCookedPlatformDataAgain();
+						check(BeginCacheObject.bHasFinishedRound);
+						Object->ClearAllCachedCookedPlatformData();
+						if (IsCookByTheBookMode())
+						{
+							Object->WillNeverCacheCookedPlatformDataAgain();
+						}
+					}
+				}
+			}
+			else
+			{
+				for (FWeakObjectPtr& WeakPtr : PackageData.GetCachedObjectsInOuter())
+				{
+					UObject* Object = WeakPtr.Get();
+					if (Object)
+					{
+						Object->ClearAllCachedCookedPlatformData();
+						if (CurrentCookMode == ECookMode::CookByTheBook)
+						{
+							Object->WillNeverCacheCookedPlatformDataAgain();
+						}
 					}
 				}
 			}
@@ -3161,7 +3384,7 @@ void UCookOnTheFlyServer::ReleaseCookedPlatformData(UE::Cook::FPackageData& Pack
 	}
 	else 
 	{
-		// This is an exceptional flow handling case; we are releasing the CookedPlatformData before we called SavePackage
+		// This is a slower but more general flow that can handle releasing whether or not we called SavePackage
 		// Note that even after we return from this function, some objects with pending IsCachedCookedPlatformDataLoaded calls may still exist for this Package in PackageDatas->GetPendingCookedPlatformDatas(),
 		// and this PackageData may therefore still have GetNumPendingCookedPlatformData > 0
 		if (!IsCookingInEditor()) // ClearAllCachedCookedPlatformData calls are only used when not in editor.
@@ -3185,20 +3408,51 @@ void UCookOnTheFlyServer::ReleaseCookedPlatformData(UE::Cook::FPackageData& Pack
 					}
 				}
 
-				// Iterate over all objects in the FPackageData up to GetCookedPlatformDataNextIndex
-				TArray<FWeakObjectPtr>& CachedObjects = PackageData.GetCachedObjectsInOuter();
+				TArray<UObject*> ObjectsToClear;
+				TArray<FWeakObjectPtr>* CachedObjects = &PackageData.GetCachedObjectsInOuter();
 				int32 NumIndexes = PackageData.GetCookedPlatformDataNextIndex();
-				check(NumIndexes <= NumPlatforms * CachedObjects.Num());
+				if (Generator)
+				{
+					FBeginCacheObjects& BeginCacheObjects = Generator->GetGeneratorBeginCacheObjects();
+					if (Generator->HasTakenOverCachedCookedPlatformData())
+					{
+						check(PackageData.GetCachedObjectsInOuter().Num() == 0);
+
+						// Add all objects that were cached in previous rounds
+						for (FBeginCacheObject& BeginCacheObject : BeginCacheObjects.Objects)
+						{
+							UObject* Object = BeginCacheObject.Object.Get();
+							if (Object && BeginCacheObject.bHasFinishedRound)
+							{
+								ObjectsToClear.Add(Object);
+							}
+						}
+						// Add all objects cached in the latest round
+						CachedObjects = &BeginCacheObjects.ObjectsInRound;
+						NumIndexes = BeginCacheObjects.NextIndexInRound;
+					}
+					else
+					{
+						check(BeginCacheObjects.Objects.Num() == 0);
+					}
+				}
+				check(NumIndexes <= NumPlatforms * CachedObjects->Num());
+
+				// Iterate over all objects in the FPackageData up to GetCookedPlatformDataNextIndex
 				// GetCookedPlatformDataNextIndex is a value in an inline iteration over the two-dimensional array of Objects x Platforms, in Object-major order.
 				// We take the ceiling of NextIndex/NumPlatforms to get the number of objects.
 				int32 NumObjects = (NumIndexes + NumPlatforms - 1) / NumPlatforms;
 				for (int32 ObjectIndex = 0; ObjectIndex < NumObjects; ++ObjectIndex)
 				{
-					UObject* Object = CachedObjects[ObjectIndex].Get();
+					UObject* Object = (*CachedObjects)[ObjectIndex].Get();
 					if (!Object)
 					{
 						continue;
 					}
+					ObjectsToClear.Add(Object);
+				}
+				for (UObject* Object : ObjectsToClear)
+				{
 					TArray<FPendingCookedPlatformData*>* PendingDatas = PendingObjects.Find(Object);
 					if (!PendingDatas || PendingDatas->Num() == 0)
 					{
@@ -3224,35 +3478,38 @@ void UCookOnTheFlyServer::ReleaseCookedPlatformData(UE::Cook::FPackageData& Pack
 		}
 	}
 
-	if (bCompletedSave)
+	if (Generator)
 	{
-		FGeneratorPackage* GeneratorPackage = PackageData.GetGeneratorPackage();
-		if (GeneratorPackage)
+		Generator->ResetGeneratorSaveState(PackageData.GetPackage(), ReleaseSaveReason);
+		PackageData.SetInitializedGeneratorSave(false);
+
+		if (ReleaseSaveReason == EReleaseSaveReason::Completed)
 		{
-			if (PackageData.HasCompletedGeneration()) // bCompletedSave can be true if !HasCompletedGeneration in the case of a GC request from BeginPrepareSave
-			{
-				GeneratorPackage->SetGeneratorSaved(PackageData.GetPackage());
-			}
-			if (GeneratorPackage->IsComplete())
+			Generator->SetGeneratorSaved(PackageData.GetPackage());
+			if (Generator->IsComplete())
 			{
 				PackageData.DestroyGeneratorPackage();
+				Generator = nullptr; // Destroyed
 			}
 		}
-		PackageData.ResetGenerationProgress();
+	}
 
-		FGeneratorPackage* OwnerGeneratorPackage = PackageData.GetGeneratedOwner();
-		if (OwnerGeneratorPackage)
+	FGeneratorPackage* OwnerGeneratorPackage = PackageData.GetGeneratedOwner();
+	if (OwnerGeneratorPackage)
+	{
+		if (ReleaseSaveReason == EReleaseSaveReason::Completed)
 		{
 			OwnerGeneratorPackage->SetGeneratedSaved(PackageData);
 			if (OwnerGeneratorPackage->IsComplete())
 			{
-				PackageData.DestroyGeneratorPackage();
+				OwnerGeneratorPackage->GetOwner().DestroyGeneratorPackage();
+				OwnerGeneratorPackage = nullptr; // Destroyed
 			}
 		}
 	}
 	PackageData.ClearCookedPlatformData();
 
-	if (bCompletedSave)
+	if (ReleaseSaveReason != EReleaseSaveReason::RecreateObjectCache)
 	{
 		if (CurrentCookMode == ECookMode::CookByTheBook)
 		{
@@ -3552,15 +3809,14 @@ void UCookOnTheFlyServer::PumpSaves(UE::Cook::FTickStackData& StackData, uint32
 		bool bForce = IsCookOnTheFlyMode() && !IsRealtimeMode();
 		PackageDatas->PollPendingCookedPlatformDatas(bForce);
 
-		// Always wait for FinishPrepareSave before attempting to save the package
-		bool AllObjectsCookedDataCached = FinishPrepareSave(PackageData, StackData.Timer);
-
-		// If the CookPlatformData is not ready then postpone the package, exit, or wait for it as appropriate
-		if (!AllObjectsCookedDataCached)
+		// If BeginCacheCookPlatformData is not ready then postpone the package, exit, or wait for it as appropriate
+		EPollStatus PrepareSaveStatus = PrepareSave(PackageData, StackData.Timer, false /* bPrecaching */);
+		if (PrepareSaveStatus != EPollStatus::Success)
 		{
-			if (PackageData.GetHasBeginPrepareSaveFailed())
+			if (PrepareSaveStatus == EPollStatus::Error)
 			{
-				ReleaseCookedPlatformData(PackageData, true /* bCompletedSave */);
+				check(PackageData.HasPrepareSaveFailed()); // Should have been set by PrepareSave; we rely on this for cleanup
+				ReleaseCookedPlatformData(PackageData, EReleaseSaveReason::AbortSave);
 				PackageData.SetPlatformsCooked(PlatformsForPackage, false /* bSucceeded */);
 				PackageData.SendToState(EPackageState::Idle, ESendFlags::QueueAdd);
 				++OutNumPushed;
@@ -3568,8 +3824,10 @@ void UCookOnTheFlyServer::PumpSaves(UE::Cook::FTickStackData& StackData, uint32
 			}
 
 			// GC is required
-			if (PackageData.GeneratorPackageRequiresGC())
+			if (PackageData.IsPrepareSaveRequiresGC())
 			{
+				// We consume the requiresGC; it will not trigger GC again unless set again
+				PackageData.SetIsPrepareSaveRequiresGC(false);
 				StackData.ResultFlags |= COSR_RequiresGC | COSR_YieldTick;
 				SaveQueue.AddFront(&PackageData);
 				return;
@@ -3591,18 +3849,22 @@ void UCookOnTheFlyServer::PumpSaves(UE::Cook::FTickStackData& StackData, uint32
 				UE_SCOPED_HIERARCHICAL_COOKTIMER(WaitingForCachedCookedPlatformData);
 				do
 				{
-					// FinishPrepareSave might block on pending CookedPlatformDatas, and it might block on BeginPrepareSave, which can
-					// block on resources held by other CookedPlatformDatas. Calling PollPendingCookedPlatformDatas should handle pumping all of those.
-					check(PackageDatas->GetPendingCookedPlatformDatas().Num() || !PackageData.GetCookedPlatformDataCalled()); // FinishPrepareSave can only return false in one of these cases
+					// PrepareSave might block on pending CookedPlatformDatas, and it might block on resources held by other
+					// CookedPlatformDatas. Calling PollPendingCookedPlatformDatas should handle pumping all of those.
+					if (!PackageDatas->GetPendingCookedPlatformDatas().Num())
+					{
+						// We're waiting on something other than pendingcookedplatformdatas; this loop does not yet handle
+						// updating anything else, so break out
+					}
 					// sleep for a bit
 					FPlatformProcess::Sleep(0.0f);
 					// Poll the results again and check whether we are now done
 					PackageDatas->PollPendingCookedPlatformDatas(true);
-					AllObjectsCookedDataCached = FinishPrepareSave(PackageData, StackData.Timer);
-				} while (!StackData.Timer.IsTimeUp() && !AllObjectsCookedDataCached);
+					PrepareSaveStatus = PrepareSave(PackageData, StackData.Timer, false /* bPrecaching */);
+				} while (!StackData.Timer.IsTimeUp() && PrepareSaveStatus == EPollStatus::Incomplete);
 			}
 			// If we couldn't postpone or wait, then we need to exit and try again later
-			if (!AllObjectsCookedDataCached)
+			if (PrepareSaveStatus != EPollStatus::Success)
 			{
 				StackData.ResultFlags |= COSR_WaitingOnCache;
 				bOutBusy = true;
@@ -3610,7 +3872,7 @@ void UCookOnTheFlyServer::PumpSaves(UE::Cook::FTickStackData& StackData, uint32
 				return;
 			}
 		}
-		check(AllObjectsCookedDataCached == true); // We are not allowed to save until FinishPrepareSave returns true.  We should have early exited above if it didn't
+		check(PrepareSaveStatus == EPollStatus::Success); // We are not allowed to save until PrepareSave succeeds.  We should have early exited above if it didn't
 
 		// precache the next few packages
 		if (!IsCookOnTheFlyMode() && SaveQueue.Num() != 0)
@@ -3625,7 +3887,7 @@ void UCookOnTheFlyServer::PumpSaves(UE::Cook::FTickStackData& StackData, uint32
 					break;
 				}
 				--LeftToPrecache;
-				BeginPrepareSave(*NextData, StackData.Timer, /*bIsPreCaching*/ true);
+				PrepareSave(*NextData, StackData.Timer, /*bPrecaching*/ true);
 			}
 
 			// If we're in RealTimeMode, check whether the precaching overflowed our timer and if so exit before we do the potentially expensive SavePackage
@@ -3643,7 +3905,7 @@ void UCookOnTheFlyServer::PumpSaves(UE::Cook::FTickStackData& StackData, uint32
 		{
 			// Timeouts can occur because of new objects created during the save, so we need to update our object cache,
 			// so we call ReleaseCookedPlatformData and ClearObjectCache to clear it and recache on next attempt.
-			ReleaseCookedPlatformData(PackageData, false/* bCompletedSave */);
+			ReleaseCookedPlatformData(PackageData, EReleaseSaveReason::RecreateObjectCache);
 			PackageData.ClearObjectCache();
 			if (PackageData.GetIsUrgent())
 			{
@@ -3656,7 +3918,7 @@ void UCookOnTheFlyServer::PumpSaves(UE::Cook::FTickStackData& StackData, uint32
 			continue;
 		}
 
-		ReleaseCookedPlatformData(PackageData, true /* bCompletedSave */);
+		ReleaseCookedPlatformData(PackageData, !Context.bHasRetryErrorCode ? EReleaseSaveReason::Completed : EReleaseSaveReason::DoneForNow);
 		PackageData.SendToState(EPackageState::Idle, ESendFlags::QueueAdd);
 		++OutNumPushed;
 	}
@@ -4206,17 +4468,6 @@ void UCookOnTheFlyServer::PreGarbageCollect()
 
 	TArray<UPackage*> GCKeepPackages;
 	TArray<FPackageData*> GCKeepPackageDatas;
-	for (FPackageData* PackageData : PackageDatas->GetSaveQueue())
-	{
-		// Generator packages that have started generation and have not yet finished saving should not be garbage
-		// collected because we do not want to kick them out of save and then have to reexecute their generation 
-		// on the next save, so keep them loaded.
-		if (PackageData->GetGeneratorPackage())
-		{
-			GCKeepPackages.Add(PackageData->GetPackage());
-			GCKeepPackageDatas.Add(PackageData);
-		}
-	}
 
 	// Find the packages that are waiting on async jobs to finish cooking data
 	// and make sure that they are not garbage collected until the jobs have
diff --git a/Engine/Source/Editor/UnrealEd/Private/Cooker/CookPackageData.cpp b/Engine/Source/Editor/UnrealEd/Private/Cooker/CookPackageData.cpp
--- a/Engine/Source/Editor/UnrealEd/Private/Cooker/CookPackageData.cpp
+++ b/Engine/Source/Editor/UnrealEd/Private/Cooker/CookPackageData.cpp
@@ -54,8 +54,8 @@ FPackageData::FPackageData(FPackageDatas& PackageDatas, const FName& InPackageNa
 	: GeneratedOwner(nullptr), PackageName(InPackageName), FileName(InFileName), PackageDatas(PackageDatas)
 	, Instigator(EInstigator::NotYetRequested), bIsUrgent(0)
 	, bIsVisited(0), bIsPreloadAttempted(0)
-	, bIsPreloaded(0), bHasSaveCache(0), bHasBeginPrepareSaveFailed(0), bCookedPlatformDataStarted(0)
-	, bCookedPlatformDataCalled(0), bCookedPlatformDataComplete(0), bMonitorIsCooked(0)
+	, bIsPreloaded(0), bHasSaveCache(0), bPrepareSaveFailed(0), bPrepareSaveRequiresGC(0)
+	, bCookedPlatformDataStarted(0), bCookedPlatformDataCalled(0), bCookedPlatformDataComplete(0), bMonitorIsCooked(0)
 	, bInitializedGeneratorSave(0), bCompletedGeneration(0), bGenerated(0), bKeepReferencedDuringGC(0)
 {
 	SetState(EPackageState::Idle);
@@ -726,16 +726,17 @@ void FPackageData::OnEnterSave()
 {
 	check(GetPackage() != nullptr && GetPackage()->IsFullyLoaded());
 
-	check(!GetHasBeginPrepareSaveFailed());
+	check(!HasPrepareSaveFailed());
 	CheckObjectCacheEmpty();
 	CheckCookedPlatformDataEmpty();
 }
 
 void FPackageData::OnExitSave()
 {
-	PackageDatas.GetCookOnTheFlyServer().ReleaseCookedPlatformData(*this, false /* bCompletedSave */);
+	PackageDatas.GetCookOnTheFlyServer().ReleaseCookedPlatformData(*this, EReleaseSaveReason::Demoted);
 	ClearObjectCache();
-	SetHasBeginPrepareSaveFailed(false);
+	SetHasPrepareSaveFailed(false);
+	SetIsPrepareSaveRequiresGC(false);
 }
 
 void FPackageData::OnEnterInProgress()
@@ -989,19 +990,6 @@ void FPackageData::CreateObjectCache()
 	}
 }
 
-void FPackageData::RecreateObjectCache()
-{
-	check(GetPackage());
-	if (GetHasSaveCache())
-	{
-		// It is not valid to recreate the ObjectCache when we have already started calling BeginCacheForCookedPlatformData on the objects in it.
-		// ReleaseCookedPlatformData must be called first to tear down all of the calls and reset GetCookedPlatformDataNextIndex to 0
-		check(GetCookedPlatformDataNextIndex() == 0);
-		ClearObjectCache();
-	}
-	CreateObjectCache();
-}
-
 void FPackageData::ClearObjectCache()
 {
 	CachedObjectsInOuter.Empty();
@@ -1034,6 +1022,8 @@ void FPackageData::CheckCookedPlatformDataEmpty() const
 	check(!GetCookedPlatformDataStarted());
 	check(!GetCookedPlatformDataCalled());
 	check(!GetCookedPlatformDataComplete());
+	check(!GetGeneratorPackage() ||
+		GetGeneratorPackage()->GetGeneratorSaveState() <= FGeneratorPackage::ESaveState::StartGeneratorSave);
 }
 
 void FPackageData::ClearCookedPlatformData()
@@ -1045,12 +1035,6 @@ void FPackageData::ClearCookedPlatformData()
 	SetCookedPlatformDataComplete(false);
 }
 
-void FPackageData::ResetGenerationProgress()
-{
-	SetInitializedGeneratorSave(false);
-	SetCompletedGeneration(false);
-}
-
 void FPackageData::OnRemoveSessionPlatform(const ITargetPlatform* Platform)
 {
 	PlatformDatas.Remove(Platform);
@@ -1120,20 +1104,19 @@ void FPackageData::SetGeneratedOwner(FGeneratorPackage* InGeneratedOwner)
 	GeneratedOwner = InGeneratedOwner;
 }
 
-bool FPackageData::GeneratorPackageRequiresGC() const
-{
-	// We consider that if a FPackageData has valid GeneratorPackage helper object,
-	// this means that COTFS's process of generating packages was not completed 
-	// either due to an error or because it has exceeded a maximum memory threshold. 
-	return IsGenerating() && !GetHasBeginPrepareSaveFailed();
-}
-
 UE::Cook::FGeneratorPackage* FPackageData::CreateGeneratorPackage(const UObject* InSplitDataObject,
 	ICookPackageSplitter* InCookPackageSplitterInstance)
 {
-	check(!GetGeneratorPackage());
-	GeneratorPackage.Reset(new UE::Cook::FGeneratorPackage(*this, InSplitDataObject,
-		InCookPackageSplitterInstance));
+	if (!GetGeneratorPackage())
+	{
+		GeneratorPackage.Reset(new UE::Cook::FGeneratorPackage(*this, InSplitDataObject,
+			InCookPackageSplitterInstance));
+	}
+	else
+	{
+		// The earlier exit from SaveState should have reset the progress back to StartGeneratorSave or earlier
+		check(GetGeneratorPackage()->GetGeneratorSaveState() <= FGeneratorPackage::ESaveState::StartGeneratorSave);
+	}
 	return GetGeneratorPackage();
 }
 	
@@ -1255,10 +1238,6 @@ UObject* FGeneratorPackage::FindSplitDataObject() const
 
 void FGeneratorPackage::PostGarbageCollect()
 {
-	if (!bGeneratedList)
-	{
-		return;
-	}
 	if (Owner.GetState() == EPackageState::Save)
 	{
 		// UCookOnTheFlyServer::PreCollectGarbage adds references for the Generator package and all its public
@@ -1333,16 +1312,6 @@ UPackage* FGeneratorPackage::CreateGeneratedUPackage(FGeneratorPackage::FGenerat
 
 void FGeneratorPackage::SetGeneratorSaved(UPackage* GeneratorUPackage)
 {
-	UObject* SplitObject = FindSplitDataObject();
-	if (!SplitObject || !GeneratorUPackage)
-	{
-		UE_LOG(LogCook, Error, TEXT("PackageSplitter: %s was GarbageCollected before we finished saving it. This will cause a failure later when we attempt to save it again and generate a second time. Splitter=%s."),
-			(!GeneratorUPackage ? TEXT("UPackage") : TEXT("SplitDataObject")), *GetSplitDataObjectName().ToString());
-	}
-	else
-	{
-		GetCookPackageSplitterInstance()->PostSaveGeneratorPackage(GeneratorUPackage, SplitObject);
-	}
 	bOwnerHasSaved = true;
 	if (IsComplete())
 	{
@@ -1375,7 +1344,7 @@ void FGeneratorPackage::SetGeneratedSaved(FPackageData& PackageData)
 
 bool FGeneratorPackage::IsComplete() const
 {
-	return bGeneratedList && RemainingToPopulate == 0;
+	return bOwnerHasSaved && RemainingToPopulate == 0;
 }
 
 void FGeneratorPackage::GetIntermediateMountPoint(FString& OutPackagePath, FString& OutLocalFilePath) const
@@ -1398,6 +1367,164 @@ FString FGeneratorPackage::GetIntermediateLocalPath(const FGeneratorPackage::FGe
 		*MountLocalFilePath, *GeneratedStruct.RelativePath, *Extension));
 }
 
+void FGeneratorPackage::SetGeneratorSaveStateComplete(ESaveState CompletedState)
+{
+	GeneratorSaveState = CompletedState;
+	if (GeneratorSaveState < ESaveState::Last)
+	{
+		GeneratorSaveState = static_cast<ESaveState>(static_cast<uint8>(GeneratorSaveState) + 1);
+	}
+}
+
+void FGeneratorPackage::SetObjectsToMoveIntoGenerator(UPackage* Package, UObject* SplitDataObject,
+	ICookPackageSplitter* Splitter, TArray<FWeakObjectPtr>& CachedObjectsInOuter,
+	TArray<ICookPackageSplitter::FGeneratedPackageForPreSave>& GeneratedPackagesForPresave)
+{
+	TArray<UObject*> ObjectsToMove;
+	Splitter->GetObjectsToMoveIntoGenerator(Package, SplitDataObject, GeneratedPackagesForPresave, ObjectsToMove);
+
+	TMap<UObject*, TOptional<bool>> ObjectSet;
+	for (FWeakObjectPtr& ObjectInOuter : CachedObjectsInOuter)
+	{
+		UObject* Object = ObjectInOuter.Get();
+		if (Object)
+		{
+			TOptional<bool>& HasFinishedRound = ObjectSet.FindOrAdd(Object);
+			HasFinishedRound = true;
+		}
+	}
+	for (UObject* Object : ObjectsToMove)
+	{
+		TOptional<bool>& HasFinishedRound = ObjectSet.FindOrAdd(Object);
+		if (!HasFinishedRound.IsSet())
+		{
+			HasFinishedRound = false;
+		}
+	}
+
+	GeneratorBeginCacheObjects.SetObjects(ObjectSet);
+	GeneratorBeginCacheObjects.StartRound();
+
+	CachedObjectsInOuter.Reset();
+	SetHasTakenOverCachedCookedPlatformData();
+}
+
+void FGeneratorPackage::SetPostMoveObjects(UPackage* Package)
+{
+	TArray<UObject*> ObjectsInOuter;
+	GetObjectsWithOuter(Package, ObjectsInOuter);
+
+	TMap<UObject*, TOptional<bool>> ObjectSet;
+	ObjectSet.Reserve(GeneratorBeginCacheObjects.Objects.Num());
+	for (FBeginCacheObject& ExistingObject : GeneratorBeginCacheObjects.Objects)
+	{
+		UObject* Object = ExistingObject.Object.Get();
+		if (Object)
+		{
+			TOptional<bool>& HasFinishedRound = ObjectSet.FindOrAdd(Object);
+			if (!HasFinishedRound.IsSet() || ExistingObject.bHasFinishedRound)
+			{
+				HasFinishedRound = ExistingObject.bHasFinishedRound;
+			}
+		}
+	}
+	for (UObject* Object : ObjectsInOuter)
+	{
+		TOptional<bool>& HasFinishedRound = ObjectSet.FindOrAdd(Object);
+		if (!HasFinishedRound.IsSet())
+		{
+			HasFinishedRound = false;
+		}
+	}
+
+	GeneratorBeginCacheObjects.SetObjects(ObjectSet);
+	GeneratorBeginCacheObjects.StartRound();
+}
+
+void FGeneratorPackage::ResetGeneratorSaveState(UPackage* Package, UE::Cook::EReleaseSaveReason ReleaseSaveReason)
+{
+	if (GeneratorSaveState > ESaveState::CallPreSaveGeneratorPackage)
+	{
+		UObject* SplitObject = FindSplitDataObject();
+		if (!SplitObject || !Package)
+		{
+			UE_LOG(LogCook, Error, TEXT("PackageSplitter: %s was GarbageCollected before we finished saving it. This will cause a failure later when we attempt to save it again and generate a second time. Splitter=%s."),
+				(!Package ? TEXT("UPackage") : TEXT("SplitDataObject")), *GetSplitDataObjectName().ToString());
+		}
+		else
+		{
+			check(GeneratorSaveState > ESaveState::CallPreSaveGeneratorPackage); // SetGeneratorSaved should only be called on a completed save, which means we passed the state
+			GetCookPackageSplitterInstance()->PostSaveGeneratorPackage(Package, SplitObject);
+		}
+	}
+
+	if (ReleaseSaveReason == EReleaseSaveReason::RecreateObjectCache ||
+		ReleaseSaveReason == EReleaseSaveReason::Demoted ||
+		ReleaseSaveReason == EReleaseSaveReason::DoneForNow)
+	{
+		if (GeneratorSaveState >= ESaveState::FinishCachePreObjectsToMove)
+		{
+			GeneratorSaveState = ESaveState::FinishCachePreObjectsToMove;
+		}
+		else
+		{
+			// Redo all the steps since we didn't make it to the FinishCachePreObjectsToMove.
+			// Restarting in the middle of that flow after a GarbageCollect is not robust
+			GeneratorSaveState = ESaveState::Start;
+		}
+	}
+	else
+	{
+		GeneratorSaveState = ESaveState::Start;
+	}
+	GeneratorBeginCacheObjects.Reset();
+	bTakenOverCachedCookedPlatformData = false;
+	bIssuedUndeclaredMovedObjectsWarning = false;
+}
+
+void FBeginCacheObjects::Reset()
+{
+	Objects.Reset();
+	ObjectsInRound.Reset();
+	NextIndexInRound = 0;
+}
+
+void FBeginCacheObjects::SetObjects(TMap<UObject*, TOptional<bool>>& InObjectSet)
+{
+	Objects.Reset(InObjectSet.Num());
+	for (TPair<UObject*, TOptional<bool>>& Pair : InObjectSet)
+	{
+		FBeginCacheObject& BeginCacheObject = Objects.Emplace_GetRef();
+		BeginCacheObject.bHasFinishedRound = Pair.Value.IsSet() && *Pair.Value;
+		BeginCacheObject.Object = Pair.Key;
+	}
+}
+
+
+void FBeginCacheObjects::StartRound()
+{
+	ObjectsInRound.Reset(Objects.Num());
+	for (FBeginCacheObject& Object : Objects)
+	{
+		if (!Object.bHasFinishedRound)
+		{
+			ObjectsInRound.Add(Object.Object);
+		}
+	}
+	NextIndexInRound = 0;
+}
+
+void FBeginCacheObjects::EndRound()
+{
+	check(NextIndexInRound == ObjectsInRound.Num());
+	ObjectsInRound.Reset();
+	NextIndexInRound = 0;
+	for (FBeginCacheObject& Object : Objects)
+	{
+		Object.bHasFinishedRound = true;
+	}
+}
+
 //////////////////////////////////////////////////////////////////////////
 // FPendingCookedPlatformData
 
diff --git a/Engine/Source/Editor/UnrealEd/Private/Cooker/CookPackageData.h b/Engine/Source/Editor/UnrealEd/Private/Cooker/CookPackageData.h
--- a/Engine/Source/Editor/UnrealEd/Private/Cooker/CookPackageData.h
+++ b/Engine/Source/Editor/UnrealEd/Private/Cooker/CookPackageData.h
@@ -16,6 +16,7 @@
 #include "HAL/CriticalSection.h"
 #include "HAL/Platform.h"
 #include "Misc/EnumClassFlags.h"
+#include "Misc/Optional.h"
 #include "Templates/SharedPointer.h"
 #include "UObject/GCObject.h"
 #include "UObject/NameTypes.h"
@@ -357,8 +358,6 @@ public:
 	void CheckObjectCacheEmpty() const;
 	/** Populate CachedObjectsInOuter if not already populated. Invalid to call except when in the save state. */
 	void CreateObjectCache();
-	/** Clear CachedObjectsInOuter and recompute it. */
-	void RecreateObjectCache();
 	/** Clear the CachedObjectsInOuter list, when e.g. leaving the save state. */
 	void ClearObjectCache();
 
@@ -391,9 +390,12 @@ public:
 	/** Check whether savestate contracts on the PackageData were invalidated by by e.g. garbage collection. */
 	bool IsSaveInvalidated() const;
 
-	/** Get/Set the flag for whether BeginPrepareSave has been called and returned an error. */
-	bool GetHasBeginPrepareSaveFailed() const { return static_cast<bool>(bHasBeginPrepareSaveFailed); }
-	void SetHasBeginPrepareSaveFailed(bool bValue) { bHasBeginPrepareSaveFailed = bValue != 0; }
+	/** Get/Set the flag for whether PrepareSave has been called and returned an error. */
+	bool HasPrepareSaveFailed() const { return static_cast<bool>(bPrepareSaveFailed); }
+	void SetHasPrepareSaveFailed(bool bValue) { bPrepareSaveFailed = bValue != 0; }
+
+	bool IsPrepareSaveRequiresGC() const { return bPrepareSaveRequiresGC; }
+	void SetIsPrepareSaveRequiresGC(bool bValue) { bPrepareSaveRequiresGC = bValue != 0; }
 
 	/** Validate that the BeginCacheForCookedPlatformData-dependent fields are empty, when entering save. */
 	void CheckCookedPlatformDataEmpty() const;
@@ -427,11 +429,6 @@ public:
 	bool HasCompletedGeneration() const { return static_cast<bool>(bCompletedGeneration); }
 	/** Set whether the PackageData has done any necessary Generator steps and is ready for BeginCache calls. */
 	void SetCompletedGeneration(bool Value) { bCompletedGeneration = Value != 0; }
-	/** Return whether the package is a Generator package and has started but not completed generation */
-	bool IsGenerating() const
-	{
-		return GetGeneratorPackage() != nullptr && HasInitializedGeneratorSave() && !HasCompletedGeneration();
-	}
 	/** Get whether the PackageData has completed the check for whether it is a Generator. */
 	bool HasInitializedGeneratorSave() const { return static_cast<bool>(bInitializedGeneratorSave); }
 	/** Set whether the PackageData has completed the check for whether it is a Generator. */
@@ -444,11 +441,6 @@ public:
 	void SetGeneratedOwner(FGeneratorPackage* InGeneratedOwner);
 	/** Return the owning generator PackageData. Will be null if not a generated package or if orphaned. */
 	FGeneratorPackage* GetGeneratedOwner() const { return GeneratedOwner; }
-	/** Mark that the package has finished generation+saving, and should generate again if it needs resaving. */
-	void ResetGenerationProgress();
-
-	/** Return whether a GC is required by a generator package as soon as possible. */
-	bool GeneratorPackageRequiresGC() const;
 
 	/**
 	 * Return the instigator for this package. The Instigator is the first code location or
@@ -564,7 +556,8 @@ private:
 	uint32 bIsPreloadAttempted : 1;
 	uint32 bIsPreloaded : 1;
 	uint32 bHasSaveCache : 1;
-	uint32 bHasBeginPrepareSaveFailed : 1;
+	uint32 bPrepareSaveFailed : 1;
+	uint32 bPrepareSaveRequiresGC : 1;
 	uint32 bCookedPlatformDataStarted : 1;
 	uint32 bCookedPlatformDataCalled : 1;
 	uint32 bCookedPlatformDataComplete : 1;
@@ -575,6 +568,30 @@ private:
 	uint32 bKeepReferencedDuringGC : 1;
 };
 
+/** A single object in athe save of a package that might have had BeginCacheForCookedPlatformData called already */
+struct FBeginCacheObject
+{
+	FWeakObjectPtr Object;
+	bool bHasFinishedRound = false;
+};
+
+/**
+ * Collection of the known objects in the save of a package that might have had BeginCacheForCookedPlatformData called already
+ * and will have it called on them in the next round of calls if not.
+ */
+struct FBeginCacheObjects
+{
+	TArray<FBeginCacheObject> Objects;
+	TArray<FWeakObjectPtr> ObjectsInRound;
+	int32 NextIndexInRound = 0;
+
+	void Reset();
+	void SetObjects(TMap<UObject*, TOptional<bool>>& InObjectSet);
+
+	void StartRound();
+	void EndRound();
+};
+
 /**
  * Helper that wraps a ICookPackageSplitter, gets/caches packages to generate and provide
  * the necessary to iterate over all of them iteratively.
@@ -608,7 +625,7 @@ public:
 	/** Return CookPackageSplitter. */
 	ICookPackageSplitter* GetCookPackageSplitterInstance() const { return CookPackageSplitterInstance.Get(); }
 	/** Return owner FPackageData. */
-	const UE::Cook::FPackageData& GetOwner() const { return Owner; }
+	UE::Cook::FPackageData& GetOwner() { return Owner; }
 	/** Return the SplitDataObject's FullObjectPath. */
 	const FName GetSplitDataObjectName() const { return SplitDataObjectName; }
 	/**
@@ -630,15 +647,37 @@ public:
 	 */
 	UObject* FindSplitDataObject() const;
 
-	/** State variables for reentrant SplitPackage calls */
-	bool HasGeneratedList() const { return bGeneratedList; }
-	void SetGeneratedList() { bGeneratedList = true; }
-	bool HasClearedOldPackages() { return bClearedOldPackages; }
-	void SetClearedOldPackages() { bClearedOldPackages = true; }
-	bool HasClearedOldPackagesWithGC() { return bClearedOldPackagesWithGC; }
-	void SetClearedOldPackagesWithGC() { bClearedOldPackagesWithGC = true; }
-	bool HasQueuedGeneratedPackages() const { return bQueuedGeneratedPackages; }
-	void SetQueuedGeneratedPackages() { bQueuedGeneratedPackages = true; }
+	/** State variable for reentrant SplitPackage calls */
+	enum class ESaveState : uint8
+	{
+		Start = 0,
+
+		GenerateList = Start,
+		ClearOldPackagesFirstAttempt,
+		ClearOldPackagesLastAttempt,
+		QueueGeneratedPackages,
+
+		StartGeneratorSave,
+		FinishCachePreObjectsToMove = StartGeneratorSave,
+		CallObjectsToMoveIntoGenerator,
+		BeginCacheObjectsToMove,
+		FinishCacheObjectsToMove,
+		CallPreSaveGeneratorPackage,
+		CallGetPostMoveObjects,
+		BeginCachePostMove,
+		FinishCachePostMove,
+
+		ReadyForSave,
+		Last = ReadyForSave,
+	};
+	ESaveState GetGeneratorSaveState() const { return GeneratorSaveState; }
+	void SetGeneratorSaveState(ESaveState InValue) { GeneratorSaveState = InValue; }
+	void SetGeneratorSaveStateComplete(ESaveState CompletedState);
+	bool HasTakenOverCachedCookedPlatformData() const { return bTakenOverCachedCookedPlatformData; }
+	void SetHasTakenOverCachedCookedPlatformData() { bTakenOverCachedCookedPlatformData = true; }
+	bool HasIssuedUndeclaredMovedObjectsWarning() const { return bIssuedUndeclaredMovedObjectsWarning; }
+	void SetHasIssuedUndeclaredMovedObjectsWarning() { bIssuedUndeclaredMovedObjectsWarning = true; }
+
 	int32& GetNextPopulateIndex() { return NextPopulateIndex; }
 
 	/** Callback during garbage collection */
@@ -656,25 +695,33 @@ public:
 	void SetGeneratedSaved(FPackageData& PackageData);
 	/** Return whether list has been generated and all generated packages have been populated */
 	bool IsComplete() const;
+	FBeginCacheObjects& GetGeneratorBeginCacheObjects() { return GeneratorBeginCacheObjects; }
+
+	void SetObjectsToMoveIntoGenerator(UPackage* Package, UObject* SplitDataObject,
+		ICookPackageSplitter* Splitter, TArray<FWeakObjectPtr>& CachedObjectsInOuter,
+		TArray<ICookPackageSplitter::FGeneratedPackageForPreSave>& GeneratedPackagesForPresave);
+	void SetPostMoveObjects(UPackage* Package);
+	void ResetGeneratorSaveState(UPackage* Package, UE::Cook::EReleaseSaveReason ReleaseSaveReason);
 
 private:
 	void ConditionalNotifyCompletion(ICookPackageSplitter::ETeardown Status);
 
 	/** PackageData for the package that is being split */
-	const UE::Cook::FPackageData& Owner;
+	UE::Cook::FPackageData& Owner;
 	/** Name of the object that prompted the splitter creation */
 	FName SplitDataObjectName;
 	/** Cached CookPackageSplitter */
 	TUniquePtr<ICookPackageSplitter> CookPackageSplitterInstance;
 	/** Recorded list of packages to generate from the splitter, and data we need about them */
 	TArray<FGeneratedStruct> PackagesToGenerate;
+	FBeginCacheObjects GeneratorBeginCacheObjects;
 
 	int32 NextPopulateIndex = 0;
 	int32 RemainingToPopulate = 0;
-	bool bGeneratedList = false;
-	bool bClearedOldPackages = false;
-	bool bClearedOldPackagesWithGC = false;
-	bool bQueuedGeneratedPackages = false;
+	ESaveState GeneratorSaveState = ESaveState::Start;
+	bool bTakenOverCachedCookedPlatformData = false;
+	bool bIssuedUndeclaredMovedObjectsWarning = false;
+
 	bool bWasOwnerReloaded = false;
 	bool bOwnerHasSaved = false;
 	bool bNotifiedCompletion = false;
diff --git a/Engine/Source/Editor/UnrealEd/Private/Cooker/CookProfiling.cpp b/Engine/Source/Editor/UnrealEd/Private/Cooker/CookProfiling.cpp
--- a/Engine/Source/Editor/UnrealEd/Private/Cooker/CookProfiling.cpp
+++ b/Engine/Source/Editor/UnrealEd/Private/Cooker/CookProfiling.cpp
@@ -177,8 +177,7 @@ namespace DetailedCookStats
 	double TickCookOnTheSideLoadPackagesTimeSec = 0.0;
 	double TickCookOnTheSideResolveRedirectorsTimeSec = 0.0;
 	double TickCookOnTheSideSaveCookedPackageTimeSec = 0.0;
-	double TickCookOnTheSideBeginPrepareSaveTimeSec = 0.0;
-	double TickCookOnTheSideFinishPrepareSaveTimeSec = 0.0;
+	double TickCookOnTheSidePrepareSaveTimeSec = 0.0;
 	double GameCookModificationDelegateTimeSec = 0.0;
 
 	// Stats tracked through FAutoRegisterCallback
diff --git a/Engine/Source/Editor/UnrealEd/Private/Cooker/CookProfiling.h b/Engine/Source/Editor/UnrealEd/Private/Cooker/CookProfiling.h
--- a/Engine/Source/Editor/UnrealEd/Private/Cooker/CookProfiling.h
+++ b/Engine/Source/Editor/UnrealEd/Private/Cooker/CookProfiling.h
@@ -121,8 +121,7 @@ namespace DetailedCookStats
 	extern double TickCookOnTheSideLoadPackagesTimeSec;
 	extern double TickCookOnTheSideResolveRedirectorsTimeSec;
 	extern double TickCookOnTheSideSaveCookedPackageTimeSec;
-	extern double TickCookOnTheSideBeginPrepareSaveTimeSec;
-	extern double TickCookOnTheSideFinishPrepareSaveTimeSec;
+	extern double TickCookOnTheSidePrepareSaveTimeSec;
 	extern double GameCookModificationDelegateTimeSec;
 	extern double TickLoopShaderProcessAsyncResultsTimeSec;
 	extern double TickLoopProcessDeferredCommandsTimeSec;
diff --git a/Engine/Source/Editor/UnrealEd/Private/Cooker/CookTypes.h b/Engine/Source/Editor/UnrealEd/Private/Cooker/CookTypes.h
--- a/Engine/Source/Editor/UnrealEd/Private/Cooker/CookTypes.h
+++ b/Engine/Source/Editor/UnrealEd/Private/Cooker/CookTypes.h
@@ -91,6 +91,27 @@ namespace UE::Cook
 		Skipped     /* For reporting the ECookResults specific to a request: the package was skipped due to e.g. already being cooked or being in NeverCook packages. */
 	};
 
+	/** Return type for functions called reentrantly that can succeed,fail,or be incomplete */
+	enum class EPollStatus : uint8
+	{
+		Success,
+		Error,
+		Incomplete,
+	};
+
+	/**
+	 * The possible reasons that the save-state data on an FPackageData might be released.
+	 * Different levels of teardown will happen based on the reason.
+	 */
+	enum class EReleaseSaveReason : uint8
+	{
+		Completed,
+		DoneForNow,
+		Demoted,
+		AbortSave,
+		RecreateObjectCache,
+	};
+
 	/** The type of callback for External Requests that needs to be executed within the Scheduler's lock. */
 	typedef TUniqueFunction<void()> FSchedulerCallback;
 
diff --git a/Engine/Source/Editor/UnrealEd/Public/CookPackageSplitter.h b/Engine/Source/Editor/UnrealEd/Public/CookPackageSplitter.h
--- a/Engine/Source/Editor/UnrealEd/Public/CookPackageSplitter.h
+++ b/Engine/Source/Editor/UnrealEd/Public/CookPackageSplitter.h
@@ -99,6 +99,13 @@ public:
 		/** *GetCreateAsMap returned from GetGenerateList. The package filename extension has already been set based on this. */
 		bool bCreatedAsMap = false;
 	};
+	/**
+	 * Return a list of all the objects that will be moved into the Generator package during its save, so the cooker
+	 * can call BeginCacheForCookedPlatformData on them before the move
+	 */
+	virtual void GetObjectsToMoveIntoGenerator(UPackage* OwnerPackage, UObject* OwnerObject,
+		const TArray<ICookPackageSplitter::FGeneratedPackageForPreSave>& GeneratedPackages, TArray<UObject*>& OutObjectsToMove) {}
+
 	/**
 	 * Called before saving the parent generator package, which itself occurs before TryPopulatePackage is called on the generated packages.
 	 * Make any required adjustments to the parent package before it is saved into the target domain.
