Fix for crash in thumbnail renderer when missing a payload for a virtual asset

diff --git a/Engine/Source/Editor/UnrealEd/Private/ThumbnailRendering/SoundWaveThumbnailRenderer.cpp b/Engine/Source/Editor/UnrealEd/Private/ThumbnailRendering/SoundWaveThumbnailRenderer.cpp
--- a/Engine/Source/Editor/UnrealEd/Private/ThumbnailRendering/SoundWaveThumbnailRenderer.cpp
+++ b/Engine/Source/Editor/UnrealEd/Private/ThumbnailRendering/SoundWaveThumbnailRenderer.cpp
@@ -43,6 +43,11 @@ void USoundWaveThumbnailRenderer::Draw(UObject* Object, int32 X, int32 Y, uint32
 	TFuture<FSharedBuffer> BufferFuture = SoundWave->RawData.GetPayload();
 	const uint8* RawWaveData = (const uint8 *)BufferFuture.Get().GetData();  // Bulk
 	int32 RawDataSize = BufferFuture.Get().GetSize();
+
+	if (!RawWaveData || RawDataSize <= 0)
+	{
+		return;
+	}
 	 
 	// Compute the scaled y-value used to render the channel data
 	const float SampleYScale = Height / (2.f * 32767 * SoundWave->NumChannels);
