Minimap - Change iteration bounds to match MinimapVolume bounds
* World iteration was still performed on EditorWorldBounds, causing some parts of the minimap to never be written to

diff --git a/Engine/Source/Editor/UnrealEd/Private/WorldPartition/WorldPartitionBuilder.cpp b/Engine/Source/Editor/UnrealEd/Private/WorldPartition/WorldPartitionBuilder.cpp
--- a/Engine/Source/Editor/UnrealEd/Private/WorldPartition/WorldPartitionBuilder.cpp
+++ b/Engine/Source/Editor/UnrealEd/Private/WorldPartition/WorldPartitionBuilder.cpp
@@ -204,7 +204,7 @@ bool UWorldPartitionBuilder::Run(UWorld* World, FPackageSourceControlHelper& Pac
 	const ELoadingMode LoadingMode = GetLoadingMode();
 	FCellInfo CellInfo;
 
-	CellInfo.EditorBounds = WorldPartition->GetEditorWorldBounds();
+	CellInfo.EditorBounds = IterativeWorldBounds.IsValid ? IterativeWorldBounds : WorldPartition->GetEditorWorldBounds();
 	CellInfo.IterativeCellSize = IterativeCellSize;
 
 	if ((LoadingMode == ELoadingMode::IterativeCells) || (LoadingMode == ELoadingMode::IterativeCells2D))
diff --git a/Engine/Source/Editor/UnrealEd/Private/WorldPartition/WorldPartitionMiniMapBuilder.cpp b/Engine/Source/Editor/UnrealEd/Private/WorldPartition/WorldPartitionMiniMapBuilder.cpp
--- a/Engine/Source/Editor/UnrealEd/Private/WorldPartition/WorldPartitionMiniMapBuilder.cpp
+++ b/Engine/Source/Editor/UnrealEd/Private/WorldPartition/WorldPartitionMiniMapBuilder.cpp
@@ -41,9 +41,11 @@ UWorldPartitionMiniMapBuilder::UWorldPartitionMiniMapBuilder(const FObjectInitia
 
 bool UWorldPartitionMiniMapBuilder::PreRun(UWorld* World, FPackageSourceControlHelper& PackageHelper)
 {
+	// Find or create the minimap actor
 	if (WorldMiniMap == nullptr)
 	{
 		WorldMiniMap = FWorldPartitionMiniMapHelper::GetWorldPartitionMiniMap(World, true);
+		IterativeCellSize = WorldMiniMap->BuilderCellSize;
 	}
 
 	if (!WorldMiniMap)
@@ -52,14 +54,9 @@ bool UWorldPartitionMiniMapBuilder::PreRun(UWorld* World, FPackageSourceControlH
 		return false;
 	}
 
-	// Reset minimap resources
+	// World bounds to process
+	FBox WorldBounds(ForceInit);
 	{
-		IterativeCellSize = WorldMiniMap->BuilderCellSize;
-
-		WorldMiniMap->MiniMapTexture = nullptr;
-
-		FBox WorldBounds(ForceInit);
-
 		// Override the minimap bounds it a world partiion minimap volume exists
 		for (TActorIterator<AWorldPartitionMiniMapVolume> It(World); It; ++It)
 		{
@@ -74,6 +71,11 @@ bool UWorldPartitionMiniMapBuilder::PreRun(UWorld* World, FPackageSourceControlH
 			WorldBounds = World->GetWorldPartition()->GetEditorWorldBounds();
 		}
 
+		IterativeWorldBounds = WorldBounds;
+	}
+	
+	// Compute minimap image size
+	{
 		MinimapImageSizeX = WorldBounds.GetSize().X / WorldMiniMap->WorldUnitsPerPixel;
 		MinimapImageSizeY = WorldBounds.GetSize().Y / WorldMiniMap->WorldUnitsPerPixel;
 
@@ -83,13 +85,19 @@ bool UWorldPartitionMiniMapBuilder::PreRun(UWorld* World, FPackageSourceControlH
 		WorldUnitsPerPixel = FMath::CeilToInt(FMath::Max(WorldBounds.GetSize().X / MinimapImageSizeX, WorldBounds.GetSize().Y / MinimapImageSizeY));
 		MinimapImageSizeX = WorldBounds.GetSize().X / WorldUnitsPerPixel;
 		MinimapImageSizeY = WorldBounds.GetSize().Y / WorldUnitsPerPixel;
+	}
 
+	// Create minimap texture
+	{
 		TStrongObjectPtr<UTextureFactory> Factory(NewObject<UTextureFactory>());
 		WorldMiniMap->MiniMapTexture = Factory->CreateTexture2D(WorldMiniMap, TEXT("MinimapTexture"), RF_NoFlags);
 		WorldMiniMap->MiniMapTexture->Source.Init(MinimapImageSizeX, MinimapImageSizeY, 1, 1, TSF_BGRA8);
 		WorldMiniMap->MiniMapWorldBounds = WorldBounds;
 		MiniMapSourcePtr = WorldMiniMap->MiniMapTexture->Source.LockMip(0);
+	}
 
+	// Compute world to minimap transform
+	{
 		WorldToMinimap = FReversedZOrthoMatrix(WorldBounds.Min.X, WorldBounds.Max.X, WorldBounds.Min.Y, WorldBounds.Max.Y, 1.0f, 0.0f);
 
 		FVector3d Translation(WorldBounds.Max.X / WorldBounds.GetSize().X, WorldBounds.Max.Y / WorldBounds.GetSize().Y, 0);
@@ -99,14 +107,17 @@ bool UWorldPartitionMiniMapBuilder::PreRun(UWorld* World, FPackageSourceControlH
 		WorldToMinimap *= FScaleMatrix(Scaling);
 	}
 
-	UDataLayerSubsystem* DataLayerSubSystem = UWorld::GetSubsystem<UDataLayerSubsystem>(World);
-	for (const FActorDataLayer& ActorDataLayer : WorldMiniMap->ExcludedDataLayers)
+	// Gather excluded data layers
 	{
-		const UDataLayerInstance* DataLayerInstance = DataLayerSubSystem->GetDataLayerInstance(ActorDataLayer.Name);
-
-		if (DataLayerInstance != nullptr)
+		UDataLayerSubsystem* DataLayerSubSystem = UWorld::GetSubsystem<UDataLayerSubsystem>(World);
+		for (const FActorDataLayer& ActorDataLayer : WorldMiniMap->ExcludedDataLayers)
 		{
-			ExcludedDataLayerShortNames.Add(FName(DataLayerInstance->GetDataLayerShortName()));
+			const UDataLayerInstance* DataLayerInstance = DataLayerSubSystem->GetDataLayerInstance(ActorDataLayer.Name);
+
+			if (DataLayerInstance != nullptr)
+			{
+				ExcludedDataLayerShortNames.Add(FName(DataLayerInstance->GetDataLayerShortName()));
+			}
 		}
 	}
 
diff --git a/Engine/Source/Editor/UnrealEd/Public/WorldPartition/WorldPartitionBuilder.h b/Engine/Source/Editor/UnrealEd/Public/WorldPartition/WorldPartitionBuilder.h
--- a/Engine/Source/Editor/UnrealEd/Public/WorldPartition/WorldPartitionBuilder.h
+++ b/Engine/Source/Editor/UnrealEd/Public/WorldPartition/WorldPartitionBuilder.h
@@ -83,6 +83,8 @@ protected:
 
 	int32 IterativeCellSize = 102400;
 	int32 IterativeCellOverlapSize = 0;
+	FBox  IterativeWorldBounds;
+
 	TSet<FName> DataLayerShortNames;
 	TSet<FName> ExcludedDataLayerShortNames;
 	bool bLoadNonDynamicDataLayers = true;
