The colour of DragToolMeasure is now defined by [REDACTED_PERSON] instead of a viewport setting as per the github code.

PR #9181: Option to change the measuring tool color (Contributed by [REDACTED_PERSON])

diff --git a/Engine/Source/Editor/UnrealEd/Private/DragTool_Measure.cpp b/Engine/Source/Editor/UnrealEd/Private/DragTool_Measure.cpp
--- a/Engine/Source/Editor/UnrealEd/Private/DragTool_Measure.cpp
+++ b/Engine/Source/Editor/UnrealEd/Private/DragTool_Measure.cpp
@@ -10,6 +10,7 @@
 #include "Editor.h"
 #include "SnappingUtils.h"
 #include "CanvasTypes.h"
+#include "Settings/EditorStyleSettings.h"
 
 ///////////////////////////////////////////////////////////////////////////////////////////////////////////////////////
 //
@@ -77,6 +78,7 @@ void FDragTool_Measure::Render(const FSceneView* View, FCanvas* Canvas)
 	if (View != nullptr && Canvas != nullptr && Length >= 1.f)
 	{
 		FCanvasLineItem LineItem( PixelStart, PixelEnd );
+		LineItem.SetColor(GetDefault<UEditorStyleSettings>()->SelectionColor);
 		Canvas->DrawItem( LineItem );
 
 		const FVector2D PixelMid = FVector2D(PixelStart + ((PixelEnd - PixelStart) / 2));
