Interchange: Create a project option to allow smooth edge when importing a static mesh that has no smoothing group information.
By default the option is true, existing projects that use the legacy hard edges behavior in their workflows can turn it to false.

diff --git a/Engine/Plugins/Experimental/Interchange/Runtime/Source/Parsers/Fbx/InterchangeFbxParser.Build.cs b/Engine/Plugins/Experimental/Interchange/Runtime/Source/Parsers/Fbx/InterchangeFbxParser.Build.cs
--- a/Engine/Plugins/Experimental/Interchange/Runtime/Source/Parsers/Fbx/InterchangeFbxParser.Build.cs
+++ b/Engine/Plugins/Experimental/Interchange/Runtime/Source/Parsers/Fbx/InterchangeFbxParser.Build.cs
@@ -48,6 +48,7 @@ namespace UnrealBuildTool.Rules
 				PrivateDependencyModuleNames.AddRange(
 					new string[] {
 						"Engine",
+						"InterchangeEngine"
 					}
 				);
 			}
diff --git a/Engine/Plugins/Experimental/Interchange/Runtime/Source/Parsers/Fbx/Private/FbxMesh.cpp b/Engine/Plugins/Experimental/Interchange/Runtime/Source/Parsers/Fbx/Private/FbxMesh.cpp
--- a/Engine/Plugins/Experimental/Interchange/Runtime/Source/Parsers/Fbx/Private/FbxMesh.cpp
+++ b/Engine/Plugins/Experimental/Interchange/Runtime/Source/Parsers/Fbx/Private/FbxMesh.cpp
@@ -18,6 +18,9 @@
 #include "SkeletalMeshAttributes.h"
 #include "StaticMeshAttributes.h"
 
+#if WITH_ENGINE
+#include "InterchangeProjectSettings.h"
+#endif
 #define LOCTEXT_NAMESPACE "InterchangeFbxMesh"
 
 namespace UE::Interchange::Private {
@@ -294,6 +297,14 @@ bool FMeshDescriptionImporter::FillMeshDescriptionFromFbxMesh(FbxMesh* Mesh, TAr
 	FStaticMeshAttributes Attributes(*MeshDescription);
 	Attributes.Register();
 
+	bool bStaticMeshUseSmoothEdgesIfSmoothingInformationIsMissing = true;
+
+#if WITH_ENGINE
+	//Get the default hard/smooth edges from the project settings
+	const UInterchangeProjectSettings* InterchangeProjectSettings = GetDefault<UInterchangeProjectSettings>();
+	bStaticMeshUseSmoothEdgesIfSmoothingInformationIsMissing = InterchangeProjectSettings->bStaticMeshUseSmoothEdgesIfSmoothingInformationIsMissing;
+#endif
+
 	//Get the base layer of the mesh
 	FbxLayer* BaseLayer = Mesh->GetLayer(0);
 	if (BaseLayer == NULL)
@@ -916,7 +927,7 @@ bool FMeshDescriptionImporter::FillMeshDescriptionFromFbxMesh(FbxMesh* Mesh, TAr
 							else
 							{
 								//When there is no smoothing group we set all edge to: hard (faceted mesh) for static mesh and smooth for skinned and rigid
-								EdgeHardnesses[MatchEdgeId] = MeshType == EMeshType::Static ? true : false;
+								EdgeHardnesses[MatchEdgeId] = MeshType == EMeshType::Static ? !bStaticMeshUseSmoothEdgesIfSmoothingInformationIsMissing : false;
 							}
 						}
 					}
diff --git a/Engine/Source/Editor/UnrealEd/Private/Fbx/FbxStaticMeshImport.cpp b/Engine/Source/Editor/UnrealEd/Private/Fbx/FbxStaticMeshImport.cpp
--- a/Engine/Source/Editor/UnrealEd/Private/Fbx/FbxStaticMeshImport.cpp
+++ b/Engine/Source/Editor/UnrealEd/Private/Fbx/FbxStaticMeshImport.cpp
@@ -31,6 +31,7 @@
 #include "FbxImporter.h"
 #include "GeomFitUtils.h"
 #include "ImportUtils/StaticMeshImportUtils.h"
+#include "InterchangeProjectSettings.h"
 #include "Interfaces/ITargetPlatform.h"
 #include "Interfaces/ITargetPlatformManagerModule.h"
 #include "Misc/FbxErrors.h"
@@ -370,7 +371,6 @@ float UnFbx::FFbxImporter::GetTriangleAreaThreshold() const
 	return (ImportOptions->bRemoveDegenerates && !ImportOptions->bBuildNanite) ? SMALL_NUMBER : 0.0f;
 }
 
-
 bool UnFbx::FFbxImporter::BuildStaticMeshFromGeometry(FbxNode* Node, UStaticMesh* StaticMesh, TArray<FFbxMaterial>& MeshMaterials, int32 LODIndex,
 	EVertexColorImportOption::Type VertexColorImportOption, const TMap<FVector3f, FColor>& ExistingVertexColorData, const FColor& VertexOverrideColor)
 {
@@ -402,6 +402,10 @@ bool UnFbx::FFbxImporter::BuildStaticMeshFromGeometry(FbxNode* Node, UStaticMesh
 		StaticMesh->SetLightMapCoordinateIndex(FBXNamedLightMapCoordinateIndex);
 	}
 	
+	//Get the default hard/smooth edges from the project settings
+	const UInterchangeProjectSettings* InterchangeProjectSettings = GetDefault<UInterchangeProjectSettings>();
+	const bool bStaticMeshUseSmoothEdgesIfSmoothingInformationIsMissing = InterchangeProjectSettings->bStaticMeshUseSmoothEdgesIfSmoothingInformationIsMissing;
+
 	//
 	// create materials
 	//
@@ -1097,8 +1101,8 @@ bool UnFbx::FFbxImporter::BuildStaticMeshFromGeometry(FbxNode* Node, UStaticMesh
 							}
 							else
 							{
-								//When there is no smoothing group we set all edge to hard (faceted mesh)
-								EdgeHardnesses[MatchEdgeId] = true;
+								//When there is no smoothing group we set all edge to hard (faceted mesh) or false depending on the project settings
+								EdgeHardnesses[MatchEdgeId] = !bStaticMeshUseSmoothEdgesIfSmoothingInformationIsMissing;
 							}
 						}
 					}
diff --git a/Engine/Source/Runtime/Experimental/Interchange/Engine/Public/InterchangeProjectSettings.h b/Engine/Source/Runtime/Experimental/Interchange/Engine/Public/InterchangeProjectSettings.h
--- a/Engine/Source/Runtime/Experimental/Interchange/Engine/Public/InterchangeProjectSettings.h
+++ b/Engine/Source/Runtime/Experimental/Interchange/Engine/Public/InterchangeProjectSettings.h
@@ -56,4 +56,12 @@ public:
 	/** This tell interchange which pipeline to select when importing scenes.*/
 	UPROPERTY(EditAnywhere, config, Category = "Interchange (Experimental)")
 	FName DefaultScenePipelineStack;
+
+	/**
+	 * If checked, interchange translators and legacy importer will default static mesh geometry to smooth edge when the smoothing information is missing.
+	 * This option exist to allows old project to import the same way as before if their workflows need static mesh edges to be hard when the smoothing
+	 * info is missing.
+	 */
+	UPROPERTY(EditAnywhere, config, Category = "Import")
+	bool bStaticMeshUseSmoothEdgesIfSmoothingInformationIsMissing = true;
 };
\ No newline at end of file
