Cooker Logging: Change the busy notification message to show time since the cooker has been blocked rather than time since the previous message.
#rnx

diff --git a/Engine/Source/Editor/UnrealEd/Classes/CookOnTheSide/CookOnTheFlyServer.h b/Engine/Source/Editor/UnrealEd/Classes/CookOnTheSide/CookOnTheFlyServer.h
--- a/Engine/Source/Editor/UnrealEd/Classes/CookOnTheSide/CookOnTheFlyServer.h
+++ b/Engine/Source/Editor/UnrealEd/Classes/CookOnTheSide/CookOnTheFlyServer.h
@@ -1228,8 +1228,10 @@ private:
 	bool bPollablesInTick = false;
 
 	/** Timers for tracking how long we have been busy, to manage retries and warnings of deadlock */
+	double SaveBusyStartTimeSeconds = MAX_flt;
 	double SaveBusyRetryTimeSeconds = MAX_flt;
 	double SaveBusyWarnTimeSeconds = MAX_flt;
+	double LoadBusyStartTimeSeconds = MAX_flt;
 	double LoadBusyRetryTimeSeconds = MAX_flt;
 	double LoadBusyWarnTimeSeconds = MAX_flt;
 	/** Tracking for the ticking of tickable cook objects */
diff --git a/Engine/Source/Editor/UnrealEd/Private/CookOnTheFlyServer.cpp b/Engine/Source/Editor/UnrealEd/Private/CookOnTheFlyServer.cpp
--- a/Engine/Source/Editor/UnrealEd/Private/CookOnTheFlyServer.cpp
+++ b/Engine/Source/Editor/UnrealEd/Private/CookOnTheFlyServer.cpp
@@ -1377,11 +1377,13 @@ void UCookOnTheFlyServer::SetSaveBusy(bool bInBusy)
 		if (bSaveBusy)
 		{
 			const double CurrentTime = FPlatformTime::Seconds();
+			SaveBusyStartTimeSeconds = CurrentTime;
 			SaveBusyRetryTimeSeconds = CurrentTime + CookProgressRetryBusyPeriodSeconds;
 			SaveBusyWarnTimeSeconds = CurrentTime + GCookProgressWarnBusyTime;
 		}
 		else
 		{
+			SaveBusyStartTimeSeconds = MAX_flt;
 			SaveBusyRetryTimeSeconds = MAX_flt;
 			SaveBusyWarnTimeSeconds = MAX_flt;
 		}
@@ -1442,7 +1444,8 @@ void UCookOnTheFlyServer::SetSaveBusy(bool bInBusy)
 				ExpectedPackages.Remove(Package);
 			}
 			
-			FString Message = FString::Printf(TEXT("Cooker has been blocked from saving the current packages for %f seconds."), GCookProgressWarnBusyTime);
+			FString Message = FString::Printf(TEXT("Cooker has been blocked from saving the current packages for %.0f seconds."),
+				(float)CurrentTime - SaveBusyStartTimeSeconds);
 			if (NonExpectedObjects.IsEmpty())
 			{
 				UE_LOG(LogCook, Display, TEXT("%s"), *Message);
@@ -1508,11 +1511,13 @@ void UCookOnTheFlyServer::SetLoadBusy(bool bInLoadBusy)
 		if (bLoadBusy)
 		{
 			const double CurrentTime = FPlatformTime::Seconds();
+			LoadBusyStartTimeSeconds = CurrentTime;
 			LoadBusyRetryTimeSeconds = CurrentTime + CookProgressRetryBusyPeriodSeconds;
 			LoadBusyWarnTimeSeconds = CurrentTime + GCookProgressWarnBusyTime;
 		}
 		else
 		{
+			LoadBusyStartTimeSeconds = MAX_flt;
 			LoadBusyRetryTimeSeconds = MAX_flt;
 			LoadBusyWarnTimeSeconds = MAX_flt;
 		}
@@ -1526,7 +1531,8 @@ void UCookOnTheFlyServer::SetLoadBusy(bool bInLoadBusy)
 			int DisplayCount = 0;
 			const int DisplayMax = 10;
 			FLoadPrepareQueue& LoadPrepareQueue = PackageDatas->GetLoadPrepareQueue();
-			UE_LOG(LogCook, Warning, TEXT("Cooker has been blocked from loading the current packages for %f seconds. %d packages in the loadqueue:"), GCookProgressWarnBusyTime, LoadPrepareQueue.PreloadingQueue.Num() + LoadPrepareQueue.EntryQueue.Num());
+			UE_LOG(LogCook, Warning, TEXT("Cooker has been blocked from loading the current packages for %.0f seconds. %d packages in the loadqueue:"),
+				(float)(CurrentTime - LoadBusyStartTimeSeconds), LoadPrepareQueue.PreloadingQueue.Num() + LoadPrepareQueue.EntryQueue.Num());
 			for (FPackageData* PackageData : LoadPrepareQueue.PreloadingQueue)
 			{
 				if (DisplayCount == DisplayMax)
