CookPackageSplitter: Change Population of the package to occur after BeginCacheForCookedPlatformData has been called on all objects that will move into the package.
#rnx

diff --git a/Engine/Source/Editor/UnrealEd/Classes/CookOnTheSide/CookOnTheFlyServer.h b/Engine/Source/Editor/UnrealEd/Classes/CookOnTheSide/CookOnTheFlyServer.h
--- a/Engine/Source/Editor/UnrealEd/Classes/CookOnTheSide/CookOnTheFlyServer.h
+++ b/Engine/Source/Editor/UnrealEd/Classes/CookOnTheSide/CookOnTheFlyServer.h
@@ -133,6 +133,7 @@ namespace UE::Cook
 	struct FCookByTheBookOptions;
 	struct FCookOnTheFlyOptions;
 	struct FCookerTimer;
+	struct FCookGenerationInfo;
 	struct FCookSavePackageContext;
 	struct FGeneratorPackage;
 	struct FPackageData;
@@ -140,7 +141,6 @@ namespace UE::Cook
 	struct FPackageTracker;
 	struct FPendingCookedPlatformData;
 	struct FPlatformManager;
-	struct FPopulatePackageContext;
 	struct FTickStackData;
 }
 namespace UE::Cook::Private
@@ -1172,11 +1172,11 @@ private:
 	UE::Cook::EPollStatus QueueGeneratedPackages(UE::Cook::FGeneratorPackage& Generator,
 		UE::Cook::FPackageData& PackageData);
 	/** Run additional steps in PrepareSave that are required when the package has a Generator. */
-	UE::Cook::EPollStatus PrepareSaveGeneratorPackage(UE::Cook::FGeneratorPackage& Generator,
-		UE::Cook::FPackageData& PackageData, UE::Cook::FCookerTimer& Timer);
+	UE::Cook::EPollStatus PrepareSaveGeneratedPackage(UE::Cook::FGeneratorPackage& Generator,
+		UE::Cook::FPackageData& PackageData, UE::Cook::FCookerTimer& Timer, bool bPrecaching);
 	/** Call BeginCacheForCookedPlatformData on the objects the Generator plans to move into its main UPackage. */
 	UE::Cook::EPollStatus BeginCacheObjectsToMove(UE::Cook::FGeneratorPackage& Generator,
-		UE::Cook::FPackageData& PackageData, UE::Cook::FCookerTimer& Timer,
+		UE::Cook::FCookGenerationInfo& Info, UE::Cook::FCookerTimer& Timer,
 		TArray<ICookPackageSplitter::FGeneratedPackageForPreSave>& GeneratedPackagesForPresave);
 	/** Call the Generator's PreSaveGeneratorPackage to create/move objects into its main UPackage. */
 	UE::Cook::EPollStatus PreSaveGeneratorPackage(UE::Cook::FPackageData& PackageData,
@@ -1186,10 +1186,12 @@ private:
 		TArray<ICookPackageSplitter::FGeneratedPackageForPreSave>& GeneratedPackagesForPresave);
 	/** Call BeginCacheForCookedPlatformData on any undeclared objects in the Generator's main UPackage after the move. */
 	UE::Cook::EPollStatus BeginCachePostMove(UE::Cook::FGeneratorPackage& Generator,
-		UE::Cook::FPackageData& PackageData, UE::Cook::FCookerTimer& Timer);
+		UE::Cook::FCookGenerationInfo& Info, UE::Cook::FCookerTimer& Timer);
 
+	/** Try creating (or finding from earlier creation) the generated package for later population */
+	UPackage* TryCreateGeneratedPackage(UE::Cook::FGeneratorPackage& Generator, UE::Cook::FCookGenerationInfo& GeneratedInfo);
 	/** Try calling the splitter's populate to create the package */
-	UPackage* TryPopulateGeneratedPackage(UE::Cook::FPopulatePackageContext& Context);
+	UE::Cook::EPollStatus TryPopulateGeneratedPackage(UE::Cook::FGeneratorPackage& Generator, UE::Cook::FCookGenerationInfo& GeneratedInfo);
 
 	ICookedPackageWriter& FindOrCreatePackageWriter(const ITargetPlatform* TargetPlatform);
 	void FindOrCreateSaveContexts(TConstArrayView<const ITargetPlatform*> TargetPlatforms);
diff --git a/Engine/Source/Editor/UnrealEd/Private/CookOnTheFlyServer.cpp b/Engine/Source/Editor/UnrealEd/Private/CookOnTheFlyServer.cpp
--- a/Engine/Source/Editor/UnrealEd/Private/CookOnTheFlyServer.cpp
+++ b/Engine/Source/Editor/UnrealEd/Private/CookOnTheFlyServer.cpp
@@ -2420,29 +2420,6 @@ void UCookOnTheFlyServer::PumpPreloadStarts()
 	}
 }
 
-namespace UE::Cook
-{
-struct FPopulatePackageContext
-{
-	explicit FPopulatePackageContext(FGeneratorPackage* InGeneratorStruct, UPackage* InOwnerPackage)
-		: GeneratorStruct(InGeneratorStruct)
-		, OwnerPackage(InOwnerPackage)
-	{
-	}
-
-	// Input variables
-	FGeneratorPackage* GeneratorStruct = nullptr;
-	const UPackage* OwnerPackage = nullptr;
-	FGeneratorPackage::FGeneratedStruct* GeneratedStruct = nullptr;
-
-	// Cache of initialized-on-demand variables
-	UObject* OwnerObject = nullptr;
-	FString GeneratorPackagePath;
-	FString GeneratorPackageShortName;
-	FString GeneratedCookedRootPath;
-};
-
-}
 void UCookOnTheFlyServer::LoadPackageInQueue(UE::Cook::FPackageData& PackageData, uint32& ResultFlags, int32& OutNumPushed)
 {
 	using namespace UE::Cook;
@@ -2484,35 +2461,48 @@ void UCookOnTheFlyServer::LoadPackageInQueue(UE::Cook::FPackageData& PackageData
 	}
 	else
 	{
-		FGeneratorPackage* GeneratorStruct = PackageData.GetGeneratedOwner();
-		if (!GeneratorStruct)
+		FGeneratorPackage* Generator = PackageData.GetGeneratedOwner();
+		if (!Generator)
 		{
 			UE_LOG(LogCook, Error, TEXT("Package %s is an out-of-date generated package with a no-longer-available generator. It can not be loaded."), *PackageFileName.ToString());
 			RejectPackageToLoad(PackageData, TEXT("is an orphaned generated package"));
 			return;
 		}
-		FGeneratorPackage::FGeneratedStruct* GeneratedStruct = GeneratorStruct->FindGeneratedStruct(&PackageData);
-		if (!GeneratedStruct)
+		FCookGenerationInfo* Info = Generator->FindInfo(PackageData);
+		if (!Info)
 		{
 			UE_LOG(LogCook, Error, TEXT("Package %s is a generated package but its generator no longer has a record of it. It can not be loaded."), *PackageFileName.ToString());
 			RejectPackageToLoad(PackageData, TEXT("is an orphaned generated package"));
 			return;
 		}
 
-		FPackageData& OwnerPackageData = const_cast<FPackageData&>(GeneratorStruct->GetOwner());
-		UPackage* OwnerPackage;
-		bool bLoadFullySuccessful = LoadPackageForCooking(OwnerPackageData, OwnerPackage, &PackageData);
-		if (!bLoadFullySuccessful)
+		FPackageData& OwnerPackageData = Generator->GetOwner();
+		UPackage* OwnerPackage = Generator->GetOwnerPackage();
+		if (!OwnerPackage)
 		{
-			ResultFlags |= COSR_ErrorLoadingPackage;
-			UE_LOG(LogCook, Error, TEXT("Package %s is a generated package and we could not load its generator package %s. It can not be loaded."),
-				*PackageFileName.ToString(), *OwnerPackageData.GetFileName().ToString());
-			RejectPackageToLoad(PackageData, TEXT("is a generated package which could not load its generator"));
-			return;
+			OwnerPackage = FindObject<UPackage>(nullptr, *OwnerPackageData.GetPackageName().ToString());
+		}
+		if (!OwnerPackage || !OwnerPackage->IsFullyLoaded())
+		{
+			bool bLoadFullySuccessful = LoadPackageForCooking(OwnerPackageData, OwnerPackage, &PackageData);
+			UObject* SplitterDataObject = nullptr;
+			if (bLoadFullySuccessful)
+			{
+				SplitterDataObject = Generator->FindSplitDataObject();
+			}
+			if (!SplitterDataObject)
+			{
+				ResultFlags |= COSR_ErrorLoadingPackage;
+				UE_LOG(LogCook, Error, TEXT("Package %s is a generated package and we could not load its generator package %s. It can not be loaded."),
+					*PackageFileName.ToString(), *OwnerPackageData.GetFileName().ToString());
+				RejectPackageToLoad(PackageData, TEXT("is a generated package which could not load its generator"));
+				return;
+			}
+			Generator->GetCookPackageSplitterInstance()->OnOwnerReloaded(OwnerPackage, SplitterDataObject);
+			Generator->SetOwnerPackage(OwnerPackage);
 		}
-		FPopulatePackageContext Context(GeneratorStruct, OwnerPackage);
-		Context.GeneratedStruct = GeneratedStruct;
-		LoadedPackage = TryPopulateGeneratedPackage(Context);
+
+		LoadedPackage = TryCreateGeneratedPackage(*Generator, *Info);
 		if (!LoadedPackage)
 		{
 			RejectPackageToLoad(PackageData, TEXT("is a generated package which could not be populated"));
@@ -2762,136 +2752,178 @@ UE::Cook::EPollStatus UCookOnTheFlyServer::QueueGeneratedPackages(UE::Cook::FGen
 
 	ICookPackageSplitter* Splitter = Generator.GetCookPackageSplitterInstance();
 	UObject* SplitObject = Generator.FindSplitDataObject();
+	FCookGenerationInfo& Info = Generator.GetOwnerInfo();
 	if (!SplitObject)
 	{
 		UE_LOG(LogCook, Error, TEXT("Could not find SplitDataObject %s"), *Generator.GetSplitDataObjectName().ToString());
 		return EPollStatus::Error;
 	}
 
-	if (Generator.GetGeneratorSaveState() <= FGeneratorPackage::ESaveState::GenerateList)
+	if (Info.GetSaveState() <= FCookGenerationInfo::ESaveState::GenerateList)
 	{
 		// Call the splitter to generate the list
 		if (!Generator.TryGenerateList(SplitObject, *PackageDatas))
 		{
 			return EPollStatus::Error;
 		}
-		Generator.SetGeneratorSaveStateComplete(FGeneratorPackage::ESaveState::GenerateList);
+		Generator.SetOwnerPackage(PackageData.GetPackage());
+		Info.SetSaveStateComplete(FCookGenerationInfo::ESaveState::GenerateList);
 	}
 
-	if (Generator.GetGeneratorSaveState() <= FGeneratorPackage::ESaveState::ClearOldPackagesLastAttempt)
+	if (Info.GetSaveState() <= FCookGenerationInfo::ESaveState::ClearOldPackagesLastAttempt)
 	{
-		for (const FGeneratorPackage::FGeneratedStruct& GeneratedStruct : Generator.GetPackagesToGenerate())
+		for (const FCookGenerationInfo& ChildInfo: Generator.GetPackagesToGenerate())
 		{
-			const FString GeneratedPackageName = GeneratedStruct.PackageData->GetPackageName().ToString();
+			const FString GeneratedPackageName = ChildInfo.PackageData->GetPackageName().ToString();
 			if (FindObject<UPackage>(nullptr, *GeneratedPackageName))
 			{
-				if (Generator.GetGeneratorSaveState() < FGeneratorPackage::ESaveState::ClearOldPackagesLastAttempt)
+				if (Info.GetSaveState() < FCookGenerationInfo::ESaveState::ClearOldPackagesLastAttempt)
 				{
 					PackageData.SetIsPrepareSaveRequiresGC(true);
-					Generator.SetGeneratorSaveState(FGeneratorPackage::ESaveState::ClearOldPackagesLastAttempt);
+					Info.SetSaveState(FCookGenerationInfo::ESaveState::ClearOldPackagesLastAttempt);
 					return EPollStatus::Incomplete;
 				}
 				else
 				{
 					UE_LOG(LogCook, Error, TEXT("PackageSplitter was unable to construct new generated packages because an old version of the package is already in memory and GC did not remove it. Splitter=%s, Generated=%s."),
-						*Generator.GetSplitDataObjectName().ToString(), *GeneratedStruct.RelativePath);
+						*Generator.GetSplitDataObjectName().ToString(), *ChildInfo.RelativePath);
 					return EPollStatus::Error;
 				}
 			}
 		}
 
-		Generator.SetGeneratorSaveStateComplete(FGeneratorPackage::ESaveState::ClearOldPackagesLastAttempt);
+		Info.SetSaveStateComplete(FCookGenerationInfo::ESaveState::ClearOldPackagesLastAttempt);
 	}
 
 	UPackage* Owner = PackageData.GetPackage();
 	FName OwnerName = Owner->GetFName();
-	if (Generator.GetGeneratorSaveState() <= FGeneratorPackage::ESaveState::QueueGeneratedPackages)
+	if (Info.GetSaveState() <= FCookGenerationInfo::ESaveState::QueueGeneratedPackages)
 	{
-		for (const FGeneratorPackage::FGeneratedStruct& GeneratedStruct : Generator.GetPackagesToGenerate())
+		for (const FCookGenerationInfo& ChildInfo: Generator.GetPackagesToGenerate())
 		{
-			FPackageData* GeneratedPackageData = GeneratedStruct.PackageData;
-			GeneratedPackageData->ClearCookedPlatformData();
-			QueueDiscoveredPackageData(*GeneratedPackageData, FInstigator(EInstigator::GeneratedPackage, OwnerName));
+			FPackageData* ChildPackageData = ChildInfo.PackageData;
+			ChildPackageData->ClearCookedPlatformData();
+			QueueDiscoveredPackageData(*ChildPackageData, FInstigator(EInstigator::GeneratedPackage, OwnerName));
 		}
-		Generator.SetGeneratorSaveStateComplete(FGeneratorPackage::ESaveState::QueueGeneratedPackages);
+		Info.SetSaveStateComplete(FCookGenerationInfo::ESaveState::QueueGeneratedPackages);
 	}
 	return EPollStatus::Success;
 }
 
-UE::Cook::EPollStatus UCookOnTheFlyServer::PrepareSaveGeneratorPackage(UE::Cook::FGeneratorPackage& Generator,
-	UE::Cook::FPackageData& PackageData, UE::Cook::FCookerTimer& Timer)
+UE::Cook::EPollStatus UCookOnTheFlyServer::PrepareSaveGeneratedPackage(UE::Cook::FGeneratorPackage& Generator,
+	UE::Cook::FPackageData& PackageData, UE::Cook::FCookerTimer& Timer, bool bPrecaching)
 {
 	using namespace UE::Cook;
 
-	if (Generator.GetGeneratorSaveState() <= FGeneratorPackage::ESaveState::FinishCachePreObjectsToMove)
+	FCookGenerationInfo* InfoPtr = Generator.FindInfo(PackageData);
+	if (!InfoPtr)
+	{
+		UE_LOG(LogCook, Error, TEXT("Generated package %s is missing its generation data and cannot be saved."),
+			*PackageData.GetPackageName().ToString());
+		return EPollStatus::Error;
+	}
+	FCookGenerationInfo& Info(*InfoPtr);
+
+	if (Info.GetSaveState() <= FCookGenerationInfo::ESaveState::FinishCachePreMove)
 	{
 		if (PackageData.GetNumPendingCookedPlatformData() > 0)
 		{
 			return EPollStatus::Incomplete;
 		}
-		Generator.SetGeneratorSaveStateComplete(FGeneratorPackage::ESaveState::FinishCachePreObjectsToMove);
+		Info.SetSaveStateComplete(FCookGenerationInfo::ESaveState::FinishCachePreMove);
 	}
 
 	// GeneratedPackagesForPresave is used by multiple steps, recreate it when needed each time we come in to this function
 	TArray<ICookPackageSplitter::FGeneratedPackageForPreSave> GeneratedPackagesForPresave;
-	if (Generator.GetGeneratorSaveState() <= FGeneratorPackage::ESaveState::FinishCacheObjectsToMove)
+	if (Info.GetSaveState() <= FCookGenerationInfo::ESaveState::FinishCacheObjectsToMove)
 	{
-		if (Generator.GetGeneratorSaveState() <= FGeneratorPackage::ESaveState::BeginCacheObjectsToMove)
+		if (Info.GetSaveState() <= FCookGenerationInfo::ESaveState::BeginCacheObjectsToMove)
 		{
-			EPollStatus Result = BeginCacheObjectsToMove(Generator, PackageData, Timer, GeneratedPackagesForPresave);
+			EPollStatus Result = BeginCacheObjectsToMove(Generator, Info, Timer, GeneratedPackagesForPresave);
 			if (Result != EPollStatus::Success)
 			{
 				return Result;
 			}
-			Generator.SetGeneratorSaveStateComplete(FGeneratorPackage::ESaveState::BeginCacheObjectsToMove);
+			Info.SetSaveStateComplete(FCookGenerationInfo::ESaveState::BeginCacheObjectsToMove);
 		}
-		check(Generator.GetGeneratorSaveState() <= FGeneratorPackage::ESaveState::FinishCacheObjectsToMove);
+		check(Info.GetSaveState() <= FCookGenerationInfo::ESaveState::FinishCacheObjectsToMove);
 		if (PackageData.GetNumPendingCookedPlatformData() > 0)
 		{
 			return EPollStatus::Incomplete;
 		}
-		Generator.SetGeneratorSaveStateComplete(FGeneratorPackage::ESaveState::FinishCacheObjectsToMove);
+		Info.SetSaveStateComplete(FCookGenerationInfo::ESaveState::FinishCacheObjectsToMove);
 	}
 
-	if (Generator.GetGeneratorSaveState() <= FGeneratorPackage::ESaveState::CallPreSaveGeneratorPackage)
+	if (Info.GetSaveState() <= FCookGenerationInfo::ESaveState::CallPopulate)
 	{
-		EPollStatus Result = PreSaveGeneratorPackage(PackageData, Generator, GeneratedPackagesForPresave);
+		if (bPrecaching)
+		{
+			// We're not allowed to populate when precaching, because we want to avoid 
+			// garbagecollection in between Populating and PostSaving the populates package,
+			// so we need to not Populate until we're ready to save
+			return EPollStatus::Incomplete;
+		}
+
+		EPollStatus Result;
+		if (Info.IsGenerator())
+		{
+			Result = PreSaveGeneratorPackage(PackageData, Generator, GeneratedPackagesForPresave);
+		}
+		else
+		{
+			Result = TryPopulateGeneratedPackage(Generator, Info);
+		}
 		if (Result != EPollStatus::Success)
 		{
 			return Result;
 		}
-		Generator.SetGeneratorSaveStateComplete(FGeneratorPackage::ESaveState::CallPreSaveGeneratorPackage);
+		Info.SetSaveStateComplete(FCookGenerationInfo::ESaveState::CallPopulate);
 	}
 
-	if (Generator.GetGeneratorSaveState() <= FGeneratorPackage::ESaveState::FinishCachePostMove)
+	if (Info.GetSaveState() <= FCookGenerationInfo::ESaveState::FinishCachePostMove)
 	{
-		if (Generator.GetGeneratorSaveState() <= FGeneratorPackage::ESaveState::BeginCachePostMove)
+		if (Info.GetSaveState() <= FCookGenerationInfo::ESaveState::BeginCachePostMove)
 		{
-			EPollStatus Result = BeginCachePostMove(Generator, PackageData, Timer);
+			EPollStatus Result = BeginCachePostMove(Generator, Info, Timer);
 			if (Result != EPollStatus::Success)
 			{
 				return Result;
 			}
-			Generator.SetGeneratorSaveStateComplete(FGeneratorPackage::ESaveState::BeginCachePostMove);
+			Info.SetSaveStateComplete(FCookGenerationInfo::ESaveState::BeginCachePostMove);
 		}
-		check(Generator.GetGeneratorSaveState() <= FGeneratorPackage::ESaveState::FinishCachePostMove);
+		check(Info.GetSaveState() <= FCookGenerationInfo::ESaveState::FinishCachePostMove);
 		if (PackageData.GetNumPendingCookedPlatformData() > 0)
 		{
 			return EPollStatus::Incomplete;
 		}
-		Generator.SetGeneratorSaveStateComplete(FGeneratorPackage::ESaveState::FinishCachePostMove);
+		Info.SetSaveStateComplete(FCookGenerationInfo::ESaveState::FinishCachePostMove);
 	}
-	check(Generator.GetGeneratorSaveState() == FGeneratorPackage::ESaveState::ReadyForSave);
+	check(Info.GetSaveState() == FCookGenerationInfo::ESaveState::ReadyForSave);
 
 	return EPollStatus::Success;
 }
 
+/** Find all subobjects of all objects in RootObjects and add them to the end of InObjects */
+static void ExtendWithSubObjects(TArray<UObject*>& InObjects)
+{
+	TArray<UObject*> ObjectsInOuter;
+	int32 OriginalNum = InObjects.Num();
+	for (int32 Index = 0; Index < OriginalNum; ++Index)
+	{
+		ObjectsInOuter.Reset();
+		GetObjectsWithOuter(InObjects[Index], ObjectsInOuter);
+		InObjects.Append(ObjectsInOuter);
+	}
+}
+
 UE::Cook::EPollStatus UCookOnTheFlyServer::BeginCacheObjectsToMove(UE::Cook::FGeneratorPackage& Generator,
-	UE::Cook::FPackageData& PackageData, UE::Cook::FCookerTimer& Timer,
+	UE::Cook::FCookGenerationInfo& Info, UE::Cook::FCookerTimer& Timer,
 	TArray<ICookPackageSplitter::FGeneratedPackageForPreSave>& GeneratedPackagesForPresave)
 {
 	using namespace UE::Cook;
 
+	check(Info.PackageData); // Caller validated this
+	FPackageData& PackageData(*Info.PackageData);
 	UPackage* Package = PackageData.GetPackage();
 	ICookPackageSplitter* Splitter = Generator.GetCookPackageSplitterInstance();
 	UObject* SplitDataObject = Generator.FindSplitDataObject();
@@ -2903,22 +2935,31 @@ UE::Cook::EPollStatus UCookOnTheFlyServer::BeginCacheObjectsToMove(UE::Cook::FGe
 		return EPollStatus::Error;
 	}
 
-	if (Generator.GetGeneratorSaveState() <= FGeneratorPackage::ESaveState::CallObjectsToMoveIntoGenerator)
+	if (Info.GetSaveState() <= FCookGenerationInfo::ESaveState::CallObjectsToMove)
 	{
-		ConstructGeneratedPackagesForPresave(PackageData, Generator, GeneratedPackagesForPresave);
-		Generator.SetObjectsToMoveIntoGenerator(Package, SplitDataObject, Splitter, PackageData.GetCachedObjectsInOuter(),
-			GeneratedPackagesForPresave);
-		Generator.SetGeneratorSaveStateComplete(FGeneratorPackage::ESaveState::CallObjectsToMoveIntoGenerator);
+		TArray<UObject*> ObjectsToMove;
+		if (Info.IsGenerator())
+		{
+			ConstructGeneratedPackagesForPresave(PackageData, Generator, GeneratedPackagesForPresave);
+			Splitter->GetObjectsToMoveIntoGenerator(Package, SplitDataObject, GeneratedPackagesForPresave, ObjectsToMove);
+		}
+		else
+		{
+			ICookPackageSplitter::FGeneratedPackageForPopulate SplitterInfo{ Info.RelativePath, Package, Info.IsCreateAsMap() };
+			Splitter->GetObjectsToMoveIntoPackage(Package, SplitDataObject, SplitterInfo, ObjectsToMove);
+		}
+		ExtendWithSubObjects(ObjectsToMove);
+		Info.TakeOverCachedObjectsAndAddNew(PackageData.GetCachedObjectsInOuter(), ObjectsToMove);
+		Info.SetSaveStateComplete(FCookGenerationInfo::ESaveState::CallObjectsToMove);
 	}
 
-	FBeginCacheObjects& BeginCacheObjects = Generator.GetGeneratorBeginCacheObjects();
-	EPollStatus Result = CallBeginCacheOnObjects(PackageData, Package, BeginCacheObjects.ObjectsInRound,
-		BeginCacheObjects.NextIndexInRound, Timer);
+	EPollStatus Result = CallBeginCacheOnObjects(PackageData, Package, Info.BeginCacheObjects.ObjectsInRound,
+		Info.BeginCacheObjects.NextIndexInRound, Timer);
 	if (Result != EPollStatus::Success)
 	{
 		return Result;
 	}
-	BeginCacheObjects.EndRound(PackageData.GetNumRequestedPlatforms());
+	Info.BeginCacheObjects.EndRound(PackageData.GetNumRequestedPlatforms());
 	return EPollStatus::Success;
 }
 
@@ -2958,30 +2999,31 @@ void UCookOnTheFlyServer::ConstructGeneratedPackagesForPresave(UE::Cook::FPackag
 
 	// We need to find or (create empty stub packages for) each of the PackagesToGenerate so that PreSaveGeneratorPackage
 	// can refer to them to create hardlinks in the cooked Generator package
-	TArrayView<FGeneratorPackage::FGeneratedStruct> PackagesToGenerate = Generator.GetPackagesToGenerate();
+	TArrayView<FCookGenerationInfo> PackagesToGenerate = Generator.GetPackagesToGenerate();
 	TArray<ICookPackageSplitter::FGeneratedPackageForPreSave> SplitterDatas;
 	SplitterDatas.Reserve(PackagesToGenerate.Num());
-	for (FGeneratorPackage::FGeneratedStruct& GeneratedStruct : PackagesToGenerate)
+	for (FCookGenerationInfo& Info : PackagesToGenerate)
 	{
 		ICookPackageSplitter::FGeneratedPackageForPreSave& SplitterData = GeneratedPackagesForPresave.Emplace_GetRef();
-		SplitterData.RelativePath = GeneratedStruct.RelativePath;
-		SplitterData.bCreatedAsMap = GeneratedStruct.bCreateAsMap;
+		SplitterData.RelativePath = Info.RelativePath;
+		SplitterData.bCreatedAsMap = Info.IsCreateAsMap();
 
-		const FString GeneratedPackageName = GeneratedStruct.PackageData->GetPackageName().ToString();
+		const FString GeneratedPackageName = Info.PackageData->GetPackageName().ToString();
 		SplitterData.Package = FindObject<UPackage>(nullptr, *GeneratedPackageName);
 		if (!SplitterData.Package)
 		{
-			SplitterData.Package = Generator.CreateGeneratedUPackage(GeneratedStruct, Package, *GeneratedPackageName);
+			SplitterData.Package = Generator.CreateGeneratedUPackage(Info, Package, *GeneratedPackageName);
 		}
 	}
 }
 
-
 UE::Cook::EPollStatus UCookOnTheFlyServer::BeginCachePostMove(UE::Cook::FGeneratorPackage& Generator,
-	UE::Cook::FPackageData& PackageData, UE::Cook::FCookerTimer& Timer)
+	UE::Cook::FCookGenerationInfo& Info, UE::Cook::FCookerTimer& Timer)
 {
 	using namespace UE::Cook;
 
+	check(Info.PackageData); // Caller has validated
+	UE::Cook::FPackageData& PackageData(*Info.PackageData);
 	UPackage* Package = PackageData.GetPackage();
 	ICookPackageSplitter* Splitter = Generator.GetCookPackageSplitterInstance();
 	UObject* SplitDataObject = Generator.FindSplitDataObject();
@@ -2993,16 +3035,15 @@ UE::Cook::EPollStatus UCookOnTheFlyServer::BeginCachePostMove(UE::Cook::FGenerat
 		return EPollStatus::Error;
 	}
 
-	if (Generator.GetGeneratorSaveState() <= FGeneratorPackage::ESaveState::CallGetPostMoveObjects)
+	if (Info.GetSaveState() <= FCookGenerationInfo::ESaveState::CallGetPostMoveObjects)
 	{
-		Generator.SetPostMoveObjects(Package);
-		Generator.SetGeneratorSaveStateComplete(FGeneratorPackage::ESaveState::CallGetPostMoveObjects);
+		Info.RefreshPackageObjects(Package);
+		Info.SetSaveStateComplete(FCookGenerationInfo::ESaveState::CallGetPostMoveObjects);
 	}
 
-	FBeginCacheObjects& BeginCacheObjects = Generator.GetGeneratorBeginCacheObjects();
-	EPollStatus Result = CallBeginCacheOnObjects(PackageData, Package, BeginCacheObjects.ObjectsInRound,
-		BeginCacheObjects.NextIndexInRound, Timer);
-	if (PackageData.GetNumPendingCookedPlatformData() > 0 && !Generator.HasIssuedUndeclaredMovedObjectsWarning())
+	EPollStatus Result = CallBeginCacheOnObjects(PackageData, Package, Info.BeginCacheObjects.ObjectsInRound,
+		Info.BeginCacheObjects.NextIndexInRound, Timer);
+	if (PackageData.GetNumPendingCookedPlatformData() > 0 && !Info.HasIssuedUndeclaredMovedObjectsWarning())
 	{
 		UObject* FirstPendingObject = nullptr;
 		FString FirstPendingObjectName;
@@ -3018,66 +3059,47 @@ UE::Cook::EPollStatus UCookOnTheFlyServer::BeginCachePostMove(UE::Cook::FGenerat
 				}
 			}
 		}
-		UE_LOG(LogCook, Warning, TEXT("CookPackageSplitter created or moved objects during PreSaveGeneratorPackage that are not yet ready to save. This will cause an error if garbage collection runs before the package is saved.\n")
-			TEXT("Change the splitter's GetObjectsToMoveIntoGenerator to construct new objects and declare existing objects that will be moved from other packages.\n")
-			TEXT("SplitterObject: %s. NumPendingObjects: %d. FirstPendingObject: %s"),
-			*SplitDataObject->GetFullName(), PackageData.GetNumPendingCookedPlatformData(),
+		UE_LOG(LogCook, Warning, TEXT("CookPackageSplitter created or moved objects during %s that are not yet ready to save. This will cause an error if garbage collection runs before the package is saved.\n")
+			TEXT("Change the splitter's %s to construct new objects and declare existing objects that will be moved from other packages.\n")
+			TEXT("SplitterObject: %s%s\n")
+			TEXT("NumPendingObjects: %d, FirstPendingObject: %s"),
+			Info.IsGenerator() ? TEXT("PreSaveGeneratorPackage") : TEXT("TryPopulatePackage"),
+			Info.IsGenerator() ? TEXT("GetObjectsToMoveIntoGenerator") : TEXT("GetObjectsToMoveIntoPackage"),
+			*SplitDataObject->GetFullName(),
+			Info.IsGenerator() ? TEXT("") : *FString::Printf(TEXT("\nGeneratedPackage: %s"), *PackageData.GetPackageName().ToString()),
+			PackageData.GetNumPendingCookedPlatformData(),
 			FirstPendingObject ? *FirstPendingObject->GetFullName() : TEXT("<unknown>"));
-		Generator.SetHasIssuedUndeclaredMovedObjectsWarning();
+		Info.SetHasIssuedUndeclaredMovedObjectsWarning(true);
 	}
 	if (Result != EPollStatus::Success)
 	{
 		return Result;
 	}
-	BeginCacheObjects.EndRound(PackageData.GetNumRequestedPlatforms());
+	Info.BeginCacheObjects.EndRound(PackageData.GetNumRequestedPlatforms());
 
 	return EPollStatus::Success;
 }
 
-UPackage* UCookOnTheFlyServer::TryPopulateGeneratedPackage(UE::Cook::FPopulatePackageContext& Context)
+UPackage* UCookOnTheFlyServer::TryCreateGeneratedPackage(UE::Cook::FGeneratorPackage& Generator, UE::Cook::FCookGenerationInfo& Info)
 {
 	using namespace UE::Cook;
+	// Caller is responsible for validating OwnerPackage and Generated PackageData
+	check(Info.PackageData); // Caller is responsible for validating
+	UE::Cook::FPackageData& GeneratedPackageData = *Info.PackageData;
+	UPackage* OwnerPackage = Generator.GetOwnerPackage();
+	check(OwnerPackage); // Caller is responsible for validating
 
-	FGeneratorPackage* GeneratorStruct = Context.GeneratorStruct;
-	FGeneratorPackage::FGeneratedStruct& GeneratedStruct = *Context.GeneratedStruct;
-	const UPackage* OwnerPackage = Context.OwnerPackage;
-	UE::Cook::FPackageData* GeneratedPackageData = Context.GeneratedStruct->PackageData;
-	const FString GeneratedPackageName = GeneratedPackageData->GetPackageName().ToString();
-
-	if (Context.OwnerObject == nullptr)
-	{
-		Context.OwnerObject = GeneratorStruct->FindSplitDataObject();
-		if (!Context.OwnerObject)
-		{
-			UE_LOG(LogCook, Error, TEXT("PopulateGeneratedPacakge could not find the original splitting object. Generated package can not be created. Splitter=%s, Generated=%s."),
-				*GeneratorStruct->GetSplitDataObjectName().ToString(), *GeneratedPackageName);
-			return nullptr;
-		}
-	}
-
-	if (Context.GeneratorPackagePath.IsEmpty())
-	{
-		Context.GeneratorPackagePath = FPackageName::GetLongPackagePath(Context.OwnerPackage->GetPathName());
-		Context.GeneratorPackageShortName = FPackageName::GetShortName(Context.OwnerPackage->GetName());
-		Context.GeneratedCookedRootPath = FPaths::RemoveDuplicateSlashes(FString::Printf(TEXT("/%s/%s/%s/"),
-			*Context.GeneratorPackagePath, *Context.GeneratorPackageShortName, UE::Cook::GeneratedPackageSubPath));
-	}
-
-	check(GeneratedPackageData); // Caller already checked this
-
-	ICookPackageSplitter* Splitter = GeneratorStruct->GetCookPackageSplitterInstance();
-
-	// Create package
+	const FString GeneratedPackageName = GeneratedPackageData.GetPackageName().ToString();
 	UPackage* GeneratedPackage = FindObject<UPackage>(nullptr, *GeneratedPackageName);
 	bool bPopulatedByPreSave = false;
 	if (GeneratedPackage)
 	{
-		if (!GeneratedStruct.bHasCreatedPackage)
+		if (!Info.HasCreatedPackage())
 		{
 			UE_LOG(LogCook, Error, TEXT("PackageSplitter found an existing copy of a package it was trying to populate;")
 				TEXT("this is unexpected since garbage has been collected and the package should have been unreferenced so it should have been collected.")
 				TEXT("Splitter=%s, Generated=%s."),
-				*GeneratorStruct->GetSplitDataObjectName().ToString(), *GeneratedPackageName);
+				*Generator.GetSplitDataObjectName().ToString(), *GeneratedPackageName);
 			EReferenceChainSearchMode SearchMode = EReferenceChainSearchMode::Shortest
 				| EReferenceChainSearchMode::PrintAllResults
 				| EReferenceChainSearchMode::FullChain;
@@ -3090,35 +3112,55 @@ UPackage* UCookOnTheFlyServer::TryPopulateGeneratedPackage(UE::Cook::FPopulatePa
 	}
 	else
 	{
-		GeneratedPackage = GeneratorStruct->CreateGeneratedUPackage(GeneratedStruct, OwnerPackage, *GeneratedPackageName);
+		GeneratedPackage = Generator.CreateGeneratedUPackage(Info, OwnerPackage, *GeneratedPackageName);
+	}
+	return GeneratedPackage;
+}
+
+UE::Cook::EPollStatus UCookOnTheFlyServer::TryPopulateGeneratedPackage(UE::Cook::FGeneratorPackage& Generator,
+	UE::Cook::FCookGenerationInfo& GeneratedInfo)
+{
+	using namespace UE::Cook;
+
+	UPackage* OwnerPackage = Generator.GetOwnerPackage();
+	check(GeneratedInfo.PackageData); // Caller already checked this
+	UE::Cook::FPackageData& GeneratedPackageData = *GeneratedInfo.PackageData;
+	const FString GeneratedPackageName = GeneratedPackageData.GetPackageName().ToString();
+	UPackage* GeneratedPackage = GeneratedPackageData.GetPackage();
+	check(GeneratedPackage); // We would have been kicked out of save if the package were gone
+
+	UObject* OwnerObject = Generator.FindSplitDataObject();
+	if (!OwnerObject)
+	{
+		UE_LOG(LogCook, Error, TEXT("PopulateGeneratedPacakge could not find the original splitting object. Generated package can not be created. Splitter=%s, Generated=%s."),
+			*Generator.GetSplitDataObjectName().ToString(), *GeneratedPackageName);
+		return EPollStatus::Error;
 	}
 
+	ICookPackageSplitter* Splitter = Generator.GetCookPackageSplitterInstance();
+
 	// Populate package using CookPackageSplitterInstance and pass GeneratedPackage's cooked name for it to
 	// properly setup any internal reference to this package (SoftObjectPaths or others)
 	ICookPackageSplitter::FGeneratedPackageForPopulate PopulateData;
-	PopulateData.RelativePath = GeneratedStruct.RelativePath;
+	PopulateData.RelativePath = GeneratedInfo.RelativePath;
 	PopulateData.Package = GeneratedPackage;
-	PopulateData.bPopulatedByPreSave = bPopulatedByPreSave;
-	PopulateData.bCreatedAsMap = GeneratedStruct.bCreateAsMap;
-	bool bWasOwnerReloaded = GeneratorStruct->GetWasOwnerReloaded();
-	GeneratorStruct->ClearWasOwnerReloaded();
-	if (!Splitter->TryPopulatePackage(OwnerPackage, Context.OwnerObject, PopulateData, bWasOwnerReloaded))
+	PopulateData.bCreatedAsMap = GeneratedInfo.IsCreateAsMap();
+	if (!Splitter->TryPopulatePackage(OwnerPackage, OwnerObject, PopulateData))
 	{
 		UE_LOG(LogCook, Error, TEXT("PackageSplitter returned false from TryPopulatePackage. Splitter=%s, Generated=%s."),
-			*GeneratorStruct->GetSplitDataObjectName().ToString(), *GeneratedPackageName);
-		return nullptr;
+			*Generator.GetSplitDataObjectName().ToString(), *GeneratedPackageName);
+		return EPollStatus::Error;
 	}
 	bool bPackageIsMap = GeneratedPackage->ContainsMap();
-	if (bPackageIsMap != GeneratedStruct.bCreateAsMap)
+	if (bPackageIsMap != GeneratedInfo.IsCreateAsMap())
 	{
 		UE_LOG(LogCook, Error, TEXT("PackageSplitter specified generated package is %s in GetGenerateList results, but then in TryPopulatePackage created it as %s. Splitter=%s, Generated=%s."),
-			(GeneratedStruct.bCreateAsMap ? TEXT("map") : TEXT("uasset")), (bPackageIsMap ? TEXT("map") : TEXT("uasset")),
-			*GeneratorStruct->GetSplitDataObjectName().ToString(), *GeneratedPackageName);
-		return nullptr;
+			(GeneratedInfo.IsCreateAsMap() ? TEXT("map") : TEXT("uasset")), (bPackageIsMap ? TEXT("map") : TEXT("uasset")),
+			*Generator.GetSplitDataObjectName().ToString(), *GeneratedPackageName);
+		return EPollStatus::Error;
 	}
-	GeneratedPackage->MarkAsFullyLoaded();
 
-	return GeneratedPackage;
+	return EPollStatus::Success;
 }
 
 UE::Cook::EPollStatus UCookOnTheFlyServer::PrepareSave(UE::Cook::FPackageData& PackageData,
@@ -3207,12 +3249,6 @@ UE::Cook::EPollStatus UCookOnTheFlyServer::PrepareSaveInternal(UE::Cook::FPackag
 		Generator = PackageData.GetGeneratorPackage();
 		if (Generator)
 		{
-			if (bPrecaching)
-			{
-				// We're not allowed to split the package while precaching, so don't progress any farther
-				// with PrepareSave while precaching
-				return EPollStatus::Incomplete;
-			}
 			Result = QueueGeneratedPackages(*Generator, PackageData);
 			if (Result != EPollStatus::Success)
 			{
@@ -3229,13 +3265,23 @@ UE::Cook::EPollStatus UCookOnTheFlyServer::PrepareSaveInternal(UE::Cook::FPackag
 
 	if (Generator)
 	{
-		if (bPrecaching)
+		EPollStatus Result = PrepareSaveGeneratedPackage(*Generator, PackageData, Timer, bPrecaching);
+		if (Result != EPollStatus::Success)
 		{
-			// We're not allowed to split the package while precaching, so don't progress any farther
-			// with PrepareSave while precaching
-			return EPollStatus::Incomplete;
+			return Result;
 		}
-		EPollStatus Result = PrepareSaveGeneratorPackage(*Generator, PackageData, Timer);
+	}
+	else if (PackageData.IsGenerated())
+	{
+		FGeneratorPackage* ParentGenerator = PackageData.GetGeneratedOwner();
+		if (!ParentGenerator)
+		{
+			UE_LOG(LogCook, Error, TEXT("Generated package %s is missing its Parent GeneratorPackage and cannot be saved."),
+				*PackageData.GetPackageName().ToString());
+			return EPollStatus::Error;
+		}
+
+		EPollStatus Result = PrepareSaveGeneratedPackage(*ParentGenerator, PackageData, Timer, bPrecaching);
 		if (Result != EPollStatus::Success)
 		{
 			return Result;
@@ -3343,6 +3389,11 @@ void UCookOnTheFlyServer::ReleaseCookedPlatformData(UE::Cook::FPackageData& Pack
 	}
 
 	FGeneratorPackage* Generator = PackageData.GetGeneratorPackage();
+	if (!Generator)
+	{
+		Generator = PackageData.GetGeneratedOwner();
+	}
+	FCookGenerationInfo* GenerationInfo = Generator ? Generator->FindInfo(PackageData) : nullptr;
 
 	// For every Object on which we called BeginCacheForCookedPlatformData, we need to call ClearAllCachedCookedPlatformData
 	if (ReleaseSaveReason == EReleaseSaveReason::Completed)
@@ -3351,13 +3402,12 @@ void UCookOnTheFlyServer::ReleaseCookedPlatformData(UE::Cook::FPackageData& Pack
 		if (!IsCookingInEditor()) // ClearAllCachedCookedPlatformData and WillNeverCacheCookedPlatformDataAgain calls are only used when not in editor
 		{
 			UE_SCOPED_HIERARCHICAL_COOKTIMER(ClearAllCachedCookedPlatformData);
-			if (Generator)
+			if (GenerationInfo)
 			{
-				check(Generator->HasTakenOverCachedCookedPlatformData());
+				check(GenerationInfo->HasTakenOverCachedCookedPlatformData());
 				check(PackageData.GetCachedObjectsInOuter().Num() == 0);
 
-				FBeginCacheObjects& BeginCacheObjects = Generator->GetGeneratorBeginCacheObjects();
-				for (FBeginCacheObject& BeginCacheObject : BeginCacheObjects.Objects)
+				for (FBeginCacheObject& BeginCacheObject : GenerationInfo->BeginCacheObjects.Objects)
 				{
 					UObject* Object = BeginCacheObject.Object.Get();
 					if (Object)
@@ -3417,15 +3467,14 @@ void UCookOnTheFlyServer::ReleaseCookedPlatformData(UE::Cook::FPackageData& Pack
 				TArray<UObject*> ObjectsToClear;
 				TArray<FWeakObjectPtr>* CachedObjects = &PackageData.GetCachedObjectsInOuter();
 				int32 NumIndexes = PackageData.GetCookedPlatformDataNextIndex();
-				if (Generator)
+				if (GenerationInfo)
 				{
-					FBeginCacheObjects& BeginCacheObjects = Generator->GetGeneratorBeginCacheObjects();
-					if (Generator->HasTakenOverCachedCookedPlatformData())
+					if (GenerationInfo->HasTakenOverCachedCookedPlatformData())
 					{
 						check(PackageData.GetCachedObjectsInOuter().Num() == 0);
 
 						// Add all objects that were cached in previous rounds
-						for (FBeginCacheObject& BeginCacheObject : BeginCacheObjects.Objects)
+						for (FBeginCacheObject& BeginCacheObject : GenerationInfo->BeginCacheObjects.Objects)
 						{
 							UObject* Object = BeginCacheObject.Object.Get();
 							if (Object && BeginCacheObject.bHasFinishedRound)
@@ -3434,12 +3483,12 @@ void UCookOnTheFlyServer::ReleaseCookedPlatformData(UE::Cook::FPackageData& Pack
 							}
 						}
 						// Add all objects cached in the latest round
-						CachedObjects = &BeginCacheObjects.ObjectsInRound;
-						NumIndexes = BeginCacheObjects.NextIndexInRound;
+						CachedObjects = &GenerationInfo->BeginCacheObjects.ObjectsInRound;
+						NumIndexes = GenerationInfo->BeginCacheObjects.NextIndexInRound;
 					}
 					else
 					{
-						check(BeginCacheObjects.Objects.Num() == 0);
+						check(GenerationInfo->BeginCacheObjects.Objects.Num() == 0);
 					}
 				}
 				check(NumIndexes <= NumPlatforms * CachedObjects->Num());
@@ -3484,35 +3533,34 @@ void UCookOnTheFlyServer::ReleaseCookedPlatformData(UE::Cook::FPackageData& Pack
 		}
 	}
 
-	if (Generator)
+	if (GenerationInfo)
 	{
-		Generator->ResetGeneratorSaveState(PackageData.GetPackage(), ReleaseSaveReason);
-		PackageData.SetInitializedGeneratorSave(false);
-
-		if (ReleaseSaveReason == EReleaseSaveReason::Completed)
+		Generator->ResetSaveState(*GenerationInfo, PackageData.GetPackage(), ReleaseSaveReason);
+		if (GenerationInfo->IsGenerator())
 		{
-			Generator->SetGeneratorSaved(PackageData.GetPackage());
-			if (Generator->IsComplete())
-			{
-				PackageData.DestroyGeneratorPackage();
-				Generator = nullptr; // Destroyed
-			}
+			PackageData.SetInitializedGeneratorSave(false);
 		}
-	}
 
-	FGeneratorPackage* OwnerGeneratorPackage = PackageData.GetGeneratedOwner();
-	if (OwnerGeneratorPackage)
-	{
 		if (ReleaseSaveReason == EReleaseSaveReason::Completed)
 		{
-			OwnerGeneratorPackage->SetGeneratedSaved(PackageData);
-			if (OwnerGeneratorPackage->IsComplete())
+			Generator->SetPackageSaved(*GenerationInfo, PackageData);
+			if (Generator->IsComplete())
 			{
-				OwnerGeneratorPackage->GetOwner().DestroyGeneratorPackage();
-				OwnerGeneratorPackage = nullptr; // Destroyed
+				if (GenerationInfo->IsGenerator())
+				{
+					PackageData.DestroyGeneratorPackage();
+				}
+				else
+				{
+					Generator->GetOwner().DestroyGeneratorPackage();
+				}
+				// Clear now-dangling pointers
+				Generator = nullptr;
+				GenerationInfo = nullptr;
 			}
 		}
 	}
+
 	PackageData.ClearCookedPlatformData();
 
 	if (ReleaseSaveReason != EReleaseSaveReason::RecreateObjectCache)
@@ -4475,6 +4523,27 @@ void UCookOnTheFlyServer::PreGarbageCollect()
 	TArray<UPackage*> GCKeepPackages;
 	TArray<FPackageData*> GCKeepPackageDatas;
 
+	// Demote any Generated/Generator packages we called PreSave on so they call their PostSave before the GC
+	// or prevent them from being garbage collected if the splitter wants to keep them referenced
+	for (FPackageData* PackageData : PackageDatas->GetSaveQueue())
+	{
+		FGeneratorPackage* Generator = PackageData->GetGeneratorPackage();
+		if (!Generator)
+		{
+			Generator = PackageData->GetGeneratedOwner();
+		}
+		FCookGenerationInfo* Info = Generator ? Generator->FindInfo(*PackageData) : nullptr;
+		if (Info)
+		{
+			bool bShouldDemote;
+			Generator->PreGarbageCollect(*Info, GCKeepObjects, GCKeepPackages, GCKeepPackageDatas, bShouldDemote);
+			if (bShouldDemote)
+			{
+				ReleaseCookedPlatformData(*PackageData, UE::Cook::EReleaseSaveReason::Demoted);
+			}
+		}
+	}
+	
 	// Find the packages that are waiting on async jobs to finish cooking data
 	// and make sure that they are not garbage collected until the jobs have
 	// completed.
@@ -4499,6 +4568,17 @@ void UCookOnTheFlyServer::PreGarbageCollect()
 		}
 	}
 
+	// Prevent GC of any objects on which we are still waiting for IsCachedCookedPlatformData
+	for (UE::Cook::FPendingCookedPlatformData& Pending : PackageDatas->GetPendingCookedPlatformDatas())
+	{
+		if (!Pending.PollIsComplete())
+		{
+			UObject* Object = Pending.Object.Get();
+			check(Object); // Otherwise PollIsComplete would have returned true
+			GCKeepObjects.Add(Object);
+		}
+	}
+
 	const bool bPartialGC = IsCookFlagSet(ECookInitializationFlags::EnablePartialGC);
 	if (bPartialGC)
 	{
diff --git a/Engine/Source/Editor/UnrealEd/Private/Cooker/CookPackageData.cpp b/Engine/Source/Editor/UnrealEd/Private/Cooker/CookPackageData.cpp
--- a/Engine/Source/Editor/UnrealEd/Private/Cooker/CookPackageData.cpp
+++ b/Engine/Source/Editor/UnrealEd/Private/Cooker/CookPackageData.cpp
@@ -1023,7 +1023,12 @@ void FPackageData::CheckCookedPlatformDataEmpty() const
 	check(!GetCookedPlatformDataCalled());
 	check(!GetCookedPlatformDataComplete());
 	check(!GetGeneratorPackage() ||
-		GetGeneratorPackage()->GetGeneratorSaveState() <= FGeneratorPackage::ESaveState::StartGeneratorSave);
+		GetGeneratorPackage()->GetOwnerInfo().GetSaveState() <= FCookGenerationInfo::ESaveState::StartPopulate);
+	if (GetGeneratedOwner())
+	{
+		FCookGenerationInfo* Info = GetGeneratedOwner()->FindInfo(*this);
+		check(!Info || Info->GetSaveState() <= FCookGenerationInfo::ESaveState::StartPopulate);
+	}
 }
 
 void FPackageData::ClearCookedPlatformData()
@@ -1086,7 +1091,7 @@ bool FPackageData::IsSaveInvalidated() const
 		return false;
 	}
 
-	return GetPackage() == nullptr || !GetPackage()->IsFullyLoaded() ||
+	if (GetPackage() == nullptr || !GetPackage()->IsFullyLoaded() ||
 		Algo::AnyOf(CachedObjectsInOuter, [](const FWeakObjectPtr& WeakPtr)
 			{
 				// TODO: Keep track of which objects were public, and only invalidate the save if the object
@@ -1094,7 +1099,29 @@ bool FPackageData::IsSaveInvalidated() const
 				// Until we make that change, we will unnecessarily invalidate and demote some packages after a
 				// garbage collect
 				return WeakPtr.Get() == nullptr;
-			});
+			}))
+	{
+		return true;
+	}
+	if (GeneratorPackage)
+	{
+		if (GeneratorPackage->IsSaveInvalidated(*this))
+		{
+			return true;
+		}
+	}
+	else if (IsGenerated())
+	{
+		if (!GeneratedOwner)
+		{
+			return true;
+		}
+		if (GeneratedOwner->IsSaveInvalidated(*this))
+		{
+			return true;
+		}
+	}
+	return false;
 }
 
 void FPackageData::SetGeneratedOwner(FGeneratorPackage* InGeneratedOwner)
@@ -1115,7 +1142,7 @@ UE::Cook::FGeneratorPackage* FPackageData::CreateGeneratorPackage(const UObject*
 	else
 	{
 		// The earlier exit from SaveState should have reset the progress back to StartGeneratorSave or earlier
-		check(GetGeneratorPackage()->GetGeneratorSaveState() <= FGeneratorPackage::ESaveState::StartGeneratorSave);
+		check(GetGeneratorPackage()->GetOwnerInfo().GetSaveState() <= FCookGenerationInfo::ESaveState::StartPopulate);
 	}
 	return GetGeneratorPackage();
 }
@@ -1125,11 +1152,12 @@ UE::Cook::FGeneratorPackage* FPackageData::CreateGeneratorPackage(const UObject*
 
 FGeneratorPackage::FGeneratorPackage(UE::Cook::FPackageData& InOwner, const UObject* InSplitDataObject,
 	ICookPackageSplitter* InCookPackageSplitterInstance)
-: Owner(InOwner)
+: OwnerInfo(InOwner, true /* bInGenerated */)
 , SplitDataObjectName(*InSplitDataObject->GetFullName())
 {
 	check(InCookPackageSplitterInstance);
 	CookPackageSplitterInstance.Reset(InCookPackageSplitterInstance);
+	SetOwnerPackage(InOwner.GetPackage());
 }
 
 FGeneratorPackage::~FGeneratorPackage()
@@ -1149,78 +1177,87 @@ void FGeneratorPackage::ConditionalNotifyCompletion(ICookPackageSplitter::ETeard
 
 void FGeneratorPackage::ClearGeneratedPackages()
 {
-	for (FGeneratedStruct& GeneratedStruct : PackagesToGenerate)
+	for (FCookGenerationInfo& Info: PackagesToGenerate)
 	{
-		if (GeneratedStruct.PackageData)
+		if (Info.PackageData)
 		{
-			check(GeneratedStruct.PackageData->GetGeneratedOwner() == this);
-			GeneratedStruct.PackageData->SetGeneratedOwner(nullptr);
-			GeneratedStruct.PackageData = nullptr;
+			check(Info.PackageData->GetGeneratedOwner() == this);
+			Info.PackageData->SetGeneratedOwner(nullptr);
+			Info.PackageData = nullptr;
 		}
 	}
 }
 
 bool FGeneratorPackage::TryGenerateList(UObject* OwnerObject, FPackageDatas& PackageDatas)
 {
-	UPackage* OwnerPackage = Owner.GetPackage();
-	check(OwnerPackage);
+	FPackageData& OwnerPackageData = GetOwner();
+	UPackage* LocalOwnerPackage = OwnerPackageData.GetPackage();
+	check(LocalOwnerPackage);
 	TArray<ICookPackageSplitter::FGeneratedPackage> GeneratorDatas =
-		CookPackageSplitterInstance->GetGenerateList(OwnerPackage, OwnerObject);
+		CookPackageSplitterInstance->GetGenerateList(LocalOwnerPackage, OwnerObject);
 	PackagesToGenerate.Reset(GeneratorDatas.Num());
 	for (ICookPackageSplitter::FGeneratedPackage& SplitterData : GeneratorDatas)
 	{
-		FGeneratedStruct& GeneratedStruct = PackagesToGenerate.Emplace_GetRef();
-		GeneratedStruct.RelativePath = MoveTemp(SplitterData.RelativePath);
-		GeneratedStruct.Dependencies = MoveTemp(SplitterData.Dependencies);
-		FString PackageName = FPaths::RemoveDuplicateSlashes(FString::Printf(TEXT("%s/%s/%s"),
-			*this->Owner.GetPackageName().ToString(), GeneratedPackageSubPath, *GeneratedStruct.RelativePath));
-
 		if (!SplitterData.GetCreateAsMap().IsSet())
 		{
 			UE_LOG(LogCook, Error, TEXT("PackageSplitter did not specify whether CreateAsMap is true for generated package. Splitter=%s, Generated=%s."),
-				*this->GetSplitDataObjectName().ToString(), *PackageName);
+				*this->GetSplitDataObjectName().ToString(), *OwnerPackageData.GetPackageName().ToString());
 			return false;
 		}
-		GeneratedStruct.bCreateAsMap = *SplitterData.GetCreateAsMap();
+		bool bCreateAsMap = *SplitterData.GetCreateAsMap();
 
+		FString PackageName = FPaths::RemoveDuplicateSlashes(FString::Printf(TEXT("%s/%s/%s"),
+			*OwnerPackageData.GetPackageName().ToString(), GeneratedPackageSubPath, *SplitterData.RelativePath));
 		const FName PackageFName(*PackageName);
 		UE::Cook::FPackageData* PackageData = PackageDatas.TryAddPackageDataByPackageName(PackageFName,
-			false /* bRequireExists */, GeneratedStruct.bCreateAsMap);
+			false /* bRequireExists */, bCreateAsMap);
 		if (!PackageData)
 		{
 			UE_LOG(LogCook, Error, TEXT("PackageSplitter could not find mounted filename for generated packagepath. Splitter=%s, Generated=%s."),
 				*this->GetSplitDataObjectName().ToString(), *PackageName);
 			return false;
 		}
+		PackageData->SetGenerated(true);
+		// No package should be generated by two different splitters. If an earlier run of this splitter generated
+		// the package, the package's owner should have been reset to null when we called ClearGeneratedPackages
+		// between then and now
+		check(PackageData->GetGeneratedOwner() == nullptr);
 		if (IFileManager::Get().FileExists(*PackageData->GetFileName().ToString()))
 		{
 			UE_LOG(LogCook, Warning, TEXT("PackageSplitter specified a generated package that already exists in the workspace domain. Splitter=%s, Generated=%s."),
 				*this->GetSplitDataObjectName().ToString(), *PackageName);
 			return false;
 		}
-		GeneratedStruct.PackageData = PackageData;
-		PackageData->SetGenerated(true);
-		// No package should be generated by two different splitters. If an earlier run of this splitter generated
-		// the package, the package's owner should have been reset to null when we called ClearGeneratedPackages
-		// between then and now
-		check(PackageData->GetGeneratedOwner() == nullptr); 
+
+		FCookGenerationInfo& GeneratedInfo = PackagesToGenerate.Emplace_GetRef(*PackageData, false /* bInGenerator */);
+		GeneratedInfo.RelativePath = MoveTemp(SplitterData.RelativePath);
+		GeneratedInfo.Dependencies = MoveTemp(SplitterData.Dependencies);
+		GeneratedInfo.SetIsCreateAsMap(bCreateAsMap);
 		PackageData->SetGeneratedOwner(this);
 	}
-	RemainingToPopulate = GeneratorDatas.Num();
+	RemainingToPopulate = GeneratorDatas.Num() + 1; // GeneratedPackaged plus one for the Generator
 	return true;
 }
 
-FGeneratorPackage::FGeneratedStruct* FGeneratorPackage::FindGeneratedStruct(FPackageData* PackageData)
+FCookGenerationInfo* FGeneratorPackage::FindInfo(const FPackageData& PackageData)
 {
-	for (FGeneratedStruct& GeneratedStruct : PackagesToGenerate)
+	if (&PackageData == &GetOwner())
 	{
-		if (GeneratedStruct.PackageData == PackageData)
+		return &OwnerInfo;
+	}
+	for (FCookGenerationInfo& Info : PackagesToGenerate)
+	{
+		if (Info.PackageData == &PackageData)
 		{
-			return &GeneratedStruct;
+			return &Info;
 		}
 	}
 	return nullptr;
 }
+const FCookGenerationInfo* FGeneratorPackage::FindInfo(const FPackageData& PackageData) const
+{
+	return const_cast<FGeneratorPackage*>(this)->FindInfo(PackageData);
+}
 
 UObject* FGeneratorPackage::FindSplitDataObject() const
 {
@@ -1236,8 +1273,55 @@ UObject* FGeneratorPackage::FindSplitDataObject() const
 	return FindObject<UObject>(nullptr, *ObjectPath);
 }
 
+void FGeneratorPackage::PreGarbageCollect(FCookGenerationInfo& Info, TArray<UObject*>& GCKeepObjects,
+	TArray<UPackage*>& GCKeepPackages, TArray<FPackageData*>& GCKeepPackageDatas, bool& bOutShouldDemote)
+{
+	bOutShouldDemote = false;
+	check(Info.PackageData); // Caller validates this is non-null
+	if (Info.GetSaveState() > FCookGenerationInfo::ESaveState::CallPopulate)
+	{
+		if (GetCookPackageSplitterInstance()->UseInternalReferenceToAvoidGarbageCollect())
+		{
+			UPackage* Package = Info.PackageData->GetPackage();
+			if (Package)
+			{
+				GCKeepPackages.Add(Package);
+				GCKeepPackageDatas.Add(Info.PackageData);
+			}
+		}
+		else
+		{
+			bOutShouldDemote = true;
+		}
+	}
+	if (Info.GetSaveState() > FCookGenerationInfo::ESaveState::CallObjectsToMove)
+	{
+		if (GetCookPackageSplitterInstance()->UseInternalReferenceToAvoidGarbageCollect())
+		{
+			// For the UseInternalReferenceToAvoidGarbageCollect case, part of the CookPackageSplitter contract is that
+			// the Cooker will keep referenced the package and all objects returned from GetObjectsToMove* functions
+			// until the PostSave function is called
+			UPackage* Package = Info.PackageData->GetPackage();
+			if (Package)
+			{
+				GCKeepPackages.Add(Package);
+				GCKeepPackageDatas.Add(Info.PackageData);
+			}
+			for (FBeginCacheObject& BeginCacheObject : Info.BeginCacheObjects.Objects)
+			{
+				UObject* Object = BeginCacheObject.Object.Get();
+				if (Object)
+				{
+					GCKeepObjects.Add(Object);
+				}
+			}
+		}
+	}
+}
+
 void FGeneratorPackage::PostGarbageCollect()
 {
+	FPackageData& Owner = GetOwner();
 	if (Owner.GetState() == EPackageState::Save)
 	{
 		// UCookOnTheFlyServer::PreCollectGarbage adds references for the Generator package and all its public
@@ -1254,8 +1338,8 @@ void FGeneratorPackage::PostGarbageCollect()
 		// If we have any packages left to populate, our splitter contract requires that it be garbage collected;
 		// we promise that the package is not partially GC'd during calls to TryPopulateGeneratedPackage
 		// The splitter can opt-out of this contract and keep it referenced itself if it desires.
-		UPackage* OwnerPackage = FindObject<UPackage>(nullptr, *Owner.GetPackageName().ToString());
-		if (OwnerPackage)
+		UPackage* LocalOwnerPackage = FindObject<UPackage>(nullptr, *Owner.GetPackageName().ToString());
+		if (LocalOwnerPackage)
 		{
 			if (RemainingToPopulate > 0 &&
 				!Owner.IsKeepReferencedDuringGC() &&
@@ -1267,107 +1351,159 @@ void FGeneratorPackage::PostGarbageCollect()
 				EReferenceChainSearchMode SearchMode = EReferenceChainSearchMode::Shortest
 					| EReferenceChainSearchMode::PrintAllResults
 					| EReferenceChainSearchMode::FullChain;
-				FReferenceChainSearch RefChainSearch(OwnerPackage, SearchMode);
+				FReferenceChainSearch RefChainSearch(LocalOwnerPackage, SearchMode);
 			}
 		}
-		else
-		{
-			bWasOwnerReloaded = true;
-		}
 	}
 
 	bool bHasIssuedWarning = false;
-	for (FGeneratedStruct& GeneratedStruct : PackagesToGenerate)
+	for (FCookGenerationInfo& Info : PackagesToGenerate)
 	{
-		if (FindObject<UPackage>(nullptr, *GeneratedStruct.PackageData->GetPackageName().ToString()))
+		if (FindObject<UPackage>(nullptr, *Info.PackageData->GetPackageName().ToString()))
 		{
-			if (!GeneratedStruct.PackageData->IsKeepReferencedDuringGC() &&
-				!GeneratedStruct.bHasSaved &&
-				!bHasIssuedWarning)
+			if (!Info.PackageData->IsKeepReferencedDuringGC() && !Info.HasSaved() && !bHasIssuedWarning)
 			{
 				UE_LOG(LogCook, Warning, TEXT("PackageSplitter found a package it generated that was not removed from memory during garbage collection. This will cause errors later during population.")
-					TEXT("\n\tSplitter=%s, Generated=%s."), *GetSplitDataObjectName().ToString(), *GeneratedStruct.PackageData->GetPackageName().ToString());
+					TEXT("\n\tSplitter=%s, Generated=%s."), *GetSplitDataObjectName().ToString(), *Info.PackageData->GetPackageName().ToString());
 				bHasIssuedWarning = true; // Only issue the warning once per GC
 			}
 		}
 		else
 		{
-			GeneratedStruct.bHasCreatedPackage = false;
+			Info.SetHasCreatedPackage(false);
 		}
 	}
 }
 
-UPackage* FGeneratorPackage::CreateGeneratedUPackage(FGeneratorPackage::FGeneratedStruct& GeneratedStruct,
-	const UPackage* OwnerPackage, const TCHAR* GeneratedPackageName)
+UPackage* FGeneratorPackage::CreateGeneratedUPackage(FCookGenerationInfo& GeneratedInfo,
+	const UPackage* InOwnerPackage, const TCHAR* GeneratedPackageName)
 {
 	UPackage* GeneratedPackage = CreatePackage(GeneratedPackageName);
 	PRAGMA_DISABLE_DEPRECATION_WARNINGS
-	GeneratedPackage->SetGuid(OwnerPackage->GetGuid());
+	GeneratedPackage->SetGuid(InOwnerPackage->GetGuid());
 	PRAGMA_ENABLE_DEPRECATION_WARNINGS
-	GeneratedPackage->SetPersistentGuid(OwnerPackage->GetPersistentGuid());
+	GeneratedPackage->SetPersistentGuid(InOwnerPackage->GetPersistentGuid());
 	GeneratedPackage->SetPackageFlags(PKG_CookGenerated);
-	GeneratedStruct.bHasCreatedPackage = true;
+	GeneratedInfo.SetHasCreatedPackage(true);
 	return GeneratedPackage;
 }
 
-void FGeneratorPackage::SetGeneratorSaved(UPackage* GeneratorUPackage)
+void FGeneratorPackage::SetPackageSaved(FCookGenerationInfo& Info, FPackageData& PackageData)
 {
-	bOwnerHasSaved = true;
+	if (Info.HasSaved())
+	{
+		return;
+	}
+	Info.SetHasSaved(true);
+	--RemainingToPopulate;
+	check(RemainingToPopulate >= 0);
 	if (IsComplete())
 	{
 		ConditionalNotifyCompletion(ICookPackageSplitter::ETeardown::Complete);
 	}
 }
 
-void FGeneratorPackage::SetGeneratedSaved(FPackageData& PackageData)
+bool FGeneratorPackage::IsComplete() const
 {
-	FGeneratedStruct* GeneratedStruct = FindGeneratedStruct(&PackageData);
-	if (!GeneratedStruct)
+	return RemainingToPopulate == 0;
+}
+
+void FGeneratorPackage::ResetSaveState(FCookGenerationInfo& Info, UPackage* Package, EReleaseSaveReason ReleaseSaveReason)
+{
+	if (Info.GetSaveState() > FCookGenerationInfo::ESaveState::CallPopulate)
 	{
-		UE_LOG(LogCook, Warning, TEXT("PackageSplitter called SetGeneratedSaved on a package that does not belong to the splitter.")
-			TEXT("\n\tSplitter=%s, Generated=%s."), *GetSplitDataObjectName().ToString(),
-			*PackageData.GetPackageName().ToString());
-		return;
+		UObject* SplitObject = FindSplitDataObject();
+		if (!SplitObject || !Package)
+		{
+			UE_LOG(LogCook, Warning, TEXT("PackageSplitter: %s on %s was GarbageCollected before we finished saving it. This prevents us from calling PostSave and may corrupt other packages that it altered during Populate. Splitter=%s."),
+				(!Package ? TEXT("UPackage") : TEXT("SplitDataObject")),
+				Info.PackageData ? *Info.PackageData->GetPackageName().ToString() : *Info.RelativePath,
+				*GetSplitDataObjectName().ToString());
+		}
+		else
+		{
+			if (Info.IsGenerator())
+			{
+				GetCookPackageSplitterInstance()->PostSaveGeneratorPackage(Package, SplitObject);
+			}
+			else
+			{
+				ICookPackageSplitter::FGeneratedPackageForPopulate PopulateInfo;
+				PopulateInfo.RelativePath = Info.RelativePath;
+				PopulateInfo.bCreatedAsMap = Info.IsCreateAsMap();
+				PopulateInfo.Package = Package;
+				GetCookPackageSplitterInstance()->PostSaveGeneratedPackage(GetOwnerPackage(), SplitObject, PopulateInfo);
+			}
+		}
 	}
-	if (GeneratedStruct->bHasSaved)
+
+	if (Info.IsGenerator())
 	{
-		return;
+		if (ReleaseSaveReason == EReleaseSaveReason::RecreateObjectCache ||
+			ReleaseSaveReason == EReleaseSaveReason::Demoted ||
+			ReleaseSaveReason == EReleaseSaveReason::DoneForNow)
+		{
+			if (Info.GetSaveState() >= FCookGenerationInfo::ESaveState::StartPopulate)
+			{
+				Info.SetSaveState(FCookGenerationInfo::ESaveState::StartPopulate);
+			}
+			else
+			{
+				// Redo all the steps since we didn't make it to the FinishCachePreObjectsToMove.
+				// Restarting in the middle of that flow after a GarbageCollect is not robust
+				Info.SetSaveState(FCookGenerationInfo::ESaveState::StartGenerate);
+			}
+		}
+		else
+		{
+			Info.SetSaveState(FCookGenerationInfo::ESaveState::StartGenerate);
+		}
 	}
-	GeneratedStruct->bHasSaved = true;
-	--RemainingToPopulate;
-	check(RemainingToPopulate >= 0);
-	if (IsComplete())
+	else
 	{
-		ConditionalNotifyCompletion(ICookPackageSplitter::ETeardown::Complete);
+		Info.SetSaveState(FCookGenerationInfo::ESaveState::StartPopulate);
 	}
+	Info.BeginCacheObjects.Reset();
+	Info.SetHasTakenOverCachedCookedPlatformData(false);
+	Info.SetHasIssuedUndeclaredMovedObjectsWarning(false);
 }
 
-bool FGeneratorPackage::IsComplete() const
+bool FGeneratorPackage::IsSaveInvalidated(const FPackageData& PackageData) const
 {
-	return bOwnerHasSaved && RemainingToPopulate == 0;
-}
+	const FCookGenerationInfo* Info = FindInfo(PackageData);
+	if (!Info)
+	{
+		return true;
+	}
 
-void FGeneratorPackage::GetIntermediateMountPoint(FString& OutPackagePath, FString& OutLocalFilePath) const
-{
-	const FString OwnerShortName = FPackageName::GetShortName(Owner.GetPackageName().ToString());
-	OutPackagePath = FPaths::RemoveDuplicateSlashes(FString::Printf(TEXT("/%s%s/"),
-		*OwnerShortName, UE::Cook::GeneratedPackageSubPath));
-	OutLocalFilePath = FPaths::RemoveDuplicateSlashes(FString::Printf(TEXT("%s/Cooked/%s/%s/"),
-		*FPaths::ProjectIntermediateDir(), *OwnerShortName, UE::Cook::GeneratedPackageSubPath));
+	for (const FBeginCacheObject& Object : Info->BeginCacheObjects.Objects)
+	{
+		if (!Object.Object.IsValid())
+		{
+			return true;
+		}
+	}
 
+	if (!Info->IsGenerator())
+	{
+		UPackage* LocalPackage = OwnerPackage.Get();
+		if (!LocalPackage || !LocalPackage->IsFullyLoaded())
+		{
+			return true;
+		}
+	}
+	return false;
 }
-FString FGeneratorPackage::GetIntermediateLocalPath(const FGeneratorPackage::FGeneratedStruct& GeneratedStruct) const
+
+FCookGenerationInfo::FCookGenerationInfo(FPackageData& InPackageData, bool bInGenerator)
+	: PackageData(&InPackageData)
+	, GeneratorSaveState(bInGenerator ? ESaveState::StartGenerate : ESaveState::StartPopulate)
+	, bCreateAsMap(false), bHasCreatedPackage(false), bHasSaved(false), bTakenOverCachedCookedPlatformData(false)
+	, bIssuedUndeclaredMovedObjectsWarning(false), bGenerator(bInGenerator)
 {
-	FString UnusedPackagePath;
-	FString MountLocalFilePath;
-	GetIntermediateMountPoint(UnusedPackagePath, MountLocalFilePath);
-	FString Extension = FPaths::GetExtension(GeneratedStruct.PackageData->GetFileName().ToString(),
-		true /* bIncludeDot */);
-	return FPaths::RemoveDuplicateSlashes(FString::Printf(TEXT("%s/%s%s"),
-		*MountLocalFilePath, *GeneratedStruct.RelativePath, *Extension));
 }
 
-void FGeneratorPackage::SetGeneratorSaveStateComplete(ESaveState CompletedState)
+void FCookGenerationInfo::SetSaveStateComplete(ESaveState CompletedState)
 {
 	GeneratorSaveState = CompletedState;
 	if (GeneratorSaveState < ESaveState::Last)
@@ -1376,13 +1512,8 @@ void FGeneratorPackage::SetGeneratorSaveStateComplete(ESaveState CompletedState)
 	}
 }
 
-void FGeneratorPackage::SetObjectsToMoveIntoGenerator(UPackage* Package, UObject* SplitDataObject,
-	ICookPackageSplitter* Splitter, TArray<FWeakObjectPtr>& CachedObjectsInOuter,
-	TArray<ICookPackageSplitter::FGeneratedPackageForPreSave>& GeneratedPackagesForPresave)
+void FCookGenerationInfo::TakeOverCachedObjectsAndAddNew(TArray<FWeakObjectPtr>& CachedObjectsInOuter, TArray<UObject*>& NewObjects)
 {
-	TArray<UObject*> ObjectsToMove;
-	Splitter->GetObjectsToMoveIntoGenerator(Package, SplitDataObject, GeneratedPackagesForPresave, ObjectsToMove);
-
 	TMap<UObject*, TOptional<bool>> ObjectSet;
 	for (FWeakObjectPtr& ObjectInOuter : CachedObjectsInOuter)
 	{
@@ -1393,7 +1524,7 @@ void FGeneratorPackage::SetObjectsToMoveIntoGenerator(UPackage* Package, UObject
 			HasFinishedRound = true;
 		}
 	}
-	for (UObject* Object : ObjectsToMove)
+	for (UObject* Object : NewObjects)
 	{
 		TOptional<bool>& HasFinishedRound = ObjectSet.FindOrAdd(Object);
 		if (!HasFinishedRound.IsSet())
@@ -1402,21 +1533,21 @@ void FGeneratorPackage::SetObjectsToMoveIntoGenerator(UPackage* Package, UObject
 		}
 	}
 
-	GeneratorBeginCacheObjects.SetObjects(ObjectSet);
-	GeneratorBeginCacheObjects.StartRound();
+	BeginCacheObjects.SetObjects(ObjectSet);
+	BeginCacheObjects.StartRound();
 
 	CachedObjectsInOuter.Reset();
-	SetHasTakenOverCachedCookedPlatformData();
+	SetHasTakenOverCachedCookedPlatformData(true);
 }
 
-void FGeneratorPackage::SetPostMoveObjects(UPackage* Package)
+void FCookGenerationInfo::RefreshPackageObjects(UPackage* Package)
 {
 	TArray<UObject*> ObjectsInOuter;
 	GetObjectsWithOuter(Package, ObjectsInOuter);
 
 	TMap<UObject*, TOptional<bool>> ObjectSet;
-	ObjectSet.Reserve(GeneratorBeginCacheObjects.Objects.Num());
-	for (FBeginCacheObject& ExistingObject : GeneratorBeginCacheObjects.Objects)
+	ObjectSet.Reserve(BeginCacheObjects.Objects.Num());
+	for (FBeginCacheObject& ExistingObject : BeginCacheObjects.Objects)
 	{
 		UObject* Object = ExistingObject.Object.Get();
 		if (Object)
@@ -1437,49 +1568,8 @@ void FGeneratorPackage::SetPostMoveObjects(UPackage* Package)
 		}
 	}
 
-	GeneratorBeginCacheObjects.SetObjects(ObjectSet);
-	GeneratorBeginCacheObjects.StartRound();
-}
-
-void FGeneratorPackage::ResetGeneratorSaveState(UPackage* Package, UE::Cook::EReleaseSaveReason ReleaseSaveReason)
-{
-	if (GeneratorSaveState > ESaveState::CallPreSaveGeneratorPackage)
-	{
-		UObject* SplitObject = FindSplitDataObject();
-		if (!SplitObject || !Package)
-		{
-			UE_LOG(LogCook, Display, TEXT("PackageSplitter: %s was GarbageCollected before we finished saving it. This will cause a failure later when we attempt to save it again and generate a second time. Splitter=%s."),
-				(!Package ? TEXT("UPackage") : TEXT("SplitDataObject")), *GetSplitDataObjectName().ToString());
-		}
-		else
-		{
-			check(GeneratorSaveState > ESaveState::CallPreSaveGeneratorPackage); // SetGeneratorSaved should only be called on a completed save, which means we passed the state
-			GetCookPackageSplitterInstance()->PostSaveGeneratorPackage(Package, SplitObject);
-		}
-	}
-
-	if (ReleaseSaveReason == EReleaseSaveReason::RecreateObjectCache ||
-		ReleaseSaveReason == EReleaseSaveReason::Demoted ||
-		ReleaseSaveReason == EReleaseSaveReason::DoneForNow)
-	{
-		if (GeneratorSaveState >= ESaveState::FinishCachePreObjectsToMove)
-		{
-			GeneratorSaveState = ESaveState::FinishCachePreObjectsToMove;
-		}
-		else
-		{
-			// Redo all the steps since we didn't make it to the FinishCachePreObjectsToMove.
-			// Restarting in the middle of that flow after a GarbageCollect is not robust
-			GeneratorSaveState = ESaveState::Start;
-		}
-	}
-	else
-	{
-		GeneratorSaveState = ESaveState::Start;
-	}
-	GeneratorBeginCacheObjects.Reset();
-	bTakenOverCachedCookedPlatformData = false;
-	bIssuedUndeclaredMovedObjectsWarning = false;
+	BeginCacheObjects.SetObjects(ObjectSet);
+	BeginCacheObjects.StartRound();
 }
 
 void FBeginCacheObjects::Reset()
@@ -1500,7 +1590,6 @@ void FBeginCacheObjects::SetObjects(TMap<UObject*, TOptional<bool>>& InObjectSet
 	}
 }
 
-
 void FBeginCacheObjects::StartRound()
 {
 	ObjectsInRound.Reset(Objects.Num());
diff --git a/Engine/Source/Editor/UnrealEd/Private/Cooker/CookPackageData.h b/Engine/Source/Editor/UnrealEd/Private/Cooker/CookPackageData.h
--- a/Engine/Source/Editor/UnrealEd/Private/Cooker/CookPackageData.h
+++ b/Engine/Source/Editor/UnrealEd/Private/Cooker/CookPackageData.h
@@ -593,22 +593,85 @@ struct FBeginCacheObjects
 };
 
 /**
- * Helper that wraps a ICookPackageSplitter, gets/caches packages to generate and provide
- * the necessary to iterate over all of them iteratively.
+ * Struct uses with ICookPackageSplitter, which contains the state information for the save of either the
+ * Generator package or one of its Generated packages.
  */
-struct FGeneratorPackage
+struct FCookGenerationInfo
 {
 public:
-	struct FGeneratedStruct
+	/** State variable for reentrant SplitPackage calls */
+	enum class ESaveState : uint8
 	{
-		FString RelativePath;
-		TArray<FName> Dependencies;
-		FPackageData* PackageData = nullptr;
-		bool bCreateAsMap = false;
-		bool bHasCreatedPackage = false;
-		bool bHasSaved = false;
+		StartGenerate = 0,
+
+		GenerateList = StartGenerate,
+		ClearOldPackagesFirstAttempt,
+		ClearOldPackagesLastAttempt,
+		QueueGeneratedPackages,
+
+		StartPopulate,
+		FinishCachePreMove = StartPopulate,
+		CallObjectsToMove,
+		BeginCacheObjectsToMove,
+		FinishCacheObjectsToMove,
+		CallPopulate,
+		CallGetPostMoveObjects,
+		BeginCachePostMove,
+		FinishCachePostMove,
+
+		ReadyForSave,
+		Last = ReadyForSave,
 	};
 
+	FCookGenerationInfo(FPackageData& InPackageData, bool bInGenerator);
+
+	ESaveState GetSaveState() const { return GeneratorSaveState; }
+	void SetSaveState(ESaveState InValue) { GeneratorSaveState = InValue; }
+	void SetSaveStateComplete(ESaveState CompletedState);
+
+	bool IsCreateAsMap() const { return bCreateAsMap; }
+	void SetIsCreateAsMap(bool bValue) { bCreateAsMap = bValue; }
+	bool HasCreatedPackage() const { return bHasCreatedPackage; }
+	void SetHasCreatedPackage(bool bValue) { bHasCreatedPackage = bValue; }
+	bool HasSaved() const { return bHasSaved; }
+	void SetHasSaved(bool bValue) { bHasSaved = bValue; }
+	bool HasTakenOverCachedCookedPlatformData() const { return bTakenOverCachedCookedPlatformData; }
+	void SetHasTakenOverCachedCookedPlatformData(bool bValue) { bTakenOverCachedCookedPlatformData = bValue; }
+	bool HasIssuedUndeclaredMovedObjectsWarning() const { return bIssuedUndeclaredMovedObjectsWarning; }
+	void SetHasIssuedUndeclaredMovedObjectsWarning(bool bValue) { bIssuedUndeclaredMovedObjectsWarning = bValue; }
+	bool IsGenerator() const { return bGenerator; }
+	void SetIsGenerator(bool bValue) { bGenerator = bValue; }
+
+	/**
+	 * Steal the list of cached objects to call BeginCacheForCookedPlatformData on from the PackageData,
+	 * and add them to this. Also add the list of NewObjects reported by the splitter that will be moved into the package.
+	 */
+	void TakeOverCachedObjectsAndAddNew(TArray<FWeakObjectPtr>& CachedObjectsInOuter, TArray<UObject*>& NewObjects);
+	/** Fetch all the objects currently in the package and add them the list of objects that need BeginCacheForCookedPlatformData. */
+	void RefreshPackageObjects(UPackage* Package);
+
+public:
+	FBeginCacheObjects BeginCacheObjects;
+	FString RelativePath;
+	TArray<FName> Dependencies;
+	FPackageData* PackageData = nullptr;
+private:
+	ESaveState GeneratorSaveState = ESaveState::StartGenerate;
+	bool bCreateAsMap : 1;
+	bool bHasCreatedPackage : 1;
+	bool bHasSaved : 1;
+	bool bTakenOverCachedCookedPlatformData : 1;
+	bool bIssuedUndeclaredMovedObjectsWarning : 1;
+	bool bGenerator : 1;
+};
+
+/**
+ * Helper that wraps a ICookPackageSplitter, gets/caches packages to generate and provide
+ * the necessary to iterate over all of them iteratively.
+ */
+struct FGeneratorPackage
+{
+public:
 	/** Store the provided CookPackageSplitter and prepare the packages to generate. */
 	FGeneratorPackage(UE::Cook::FPackageData& InOwner, const UObject* InSplitDataObject,
 		ICookPackageSplitter* InCookPackageSplitterInstance);
@@ -620,110 +683,63 @@ public:
 	bool TryGenerateList(UObject* OwnerObject, FPackageDatas& PackageDatas);
 
 	/** Accessor for the packages to generate */
-	TArrayView<FGeneratedStruct> GetPackagesToGenerate() { return PackagesToGenerate; }
+	TArrayView<UE::Cook::FCookGenerationInfo> GetPackagesToGenerate() { return PackagesToGenerate; }
+	/** Return the GenerationInfo used to save the Generator's UPackage */
+	UE::Cook::FCookGenerationInfo& GetOwnerInfo() { return OwnerInfo; }
+	/** Return owner FPackageData. */
+	UE::Cook::FPackageData& GetOwner() { return *OwnerInfo.PackageData; }
+	/** Return the GenerationInfo for the given PackageData, or null if not found. */
+	UE::Cook::FCookGenerationInfo* FindInfo(const FPackageData& PackageData);
+	const UE::Cook::FCookGenerationInfo* FindInfo(const FPackageData& PackageData) const;
 
 	/** Return CookPackageSplitter. */
 	ICookPackageSplitter* GetCookPackageSplitterInstance() const { return CookPackageSplitterInstance.Get(); }
-	/** Return owner FPackageData. */
-	UE::Cook::FPackageData& GetOwner() { return Owner; }
 	/** Return the SplitDataObject's FullObjectPath. */
 	const FName GetSplitDataObjectName() const { return SplitDataObjectName; }
-	/**
-	 * Return the packagepath and localpath to the owner's mount point used for loading generated packages,
-	 * if we need to save them to disk.
-	 */
-	void GetIntermediateMountPoint(FString& OutPackagePath, FString& OutLocalPath) const;
-	/**
-	 * Return the standard path to the temporary file (including extension) where we store the given
-	 * generated package, if we need to save it to disk for later retrieval in the cook session.
-	 */
-	FString GetIntermediateLocalPath(const FGeneratorPackage::FGeneratedStruct& GeneratedStruct) const;
-
-	/** Return the FGeneratedStruct with the given PackageData, or null if not found. */
-	FGeneratedStruct* FindGeneratedStruct(FPackageData* PackageData);
 	/**
 	 * Find again the split object from its name, or return null if no longer in memory.
 	 * It may have been GC'd and reloaded since the last time we used it.
 	 */
 	UObject* FindSplitDataObject() const;
 
-	/** State variable for reentrant SplitPackage calls */
-	enum class ESaveState : uint8
-	{
-		Start = 0,
-
-		GenerateList = Start,
-		ClearOldPackagesFirstAttempt,
-		ClearOldPackagesLastAttempt,
-		QueueGeneratedPackages,
-
-		StartGeneratorSave,
-		FinishCachePreObjectsToMove = StartGeneratorSave,
-		CallObjectsToMoveIntoGenerator,
-		BeginCacheObjectsToMove,
-		FinishCacheObjectsToMove,
-		CallPreSaveGeneratorPackage,
-		CallGetPostMoveObjects,
-		BeginCachePostMove,
-		FinishCachePostMove,
-
-		ReadyForSave,
-		Last = ReadyForSave,
-	};
-	ESaveState GetGeneratorSaveState() const { return GeneratorSaveState; }
-	void SetGeneratorSaveState(ESaveState InValue) { GeneratorSaveState = InValue; }
-	void SetGeneratorSaveStateComplete(ESaveState CompletedState);
-	bool HasTakenOverCachedCookedPlatformData() const { return bTakenOverCachedCookedPlatformData; }
-	void SetHasTakenOverCachedCookedPlatformData() { bTakenOverCachedCookedPlatformData = true; }
-	bool HasIssuedUndeclaredMovedObjectsWarning() const { return bIssuedUndeclaredMovedObjectsWarning; }
-	void SetHasIssuedUndeclaredMovedObjectsWarning() { bIssuedUndeclaredMovedObjectsWarning = true; }
+	void ResetSaveState(FCookGenerationInfo& Info, UPackage* Package, UE::Cook::EReleaseSaveReason ReleaseSaveReason);
 
 	int32& GetNextPopulateIndex() { return NextPopulateIndex; }
 
-	/** Callback during garbage collection */
+	/** Callbacks during garbage collection */
+	void PreGarbageCollect(FCookGenerationInfo& Info, TArray<UObject*>& GCKeepObjects,
+		TArray<UPackage*>& GCKeepPackages, TArray<FPackageData*>& GCKeepPackageDatas, bool& bOutShouldDemote);
 	void PostGarbageCollect();
 
-	/** Get/Clear whether the Owner was garbage collected since the previous call to TryPopulatePackage */
-	bool GetWasOwnerReloaded() const { return bWasOwnerReloaded; }
-	void ClearWasOwnerReloaded() { bWasOwnerReloaded = false; }
 	/** Call CreatePackage and set PackageData for deterministic generated packages, and update status. */
-	UPackage* CreateGeneratedUPackage(FGeneratedStruct& GeneratedStruct,
+	UPackage* CreateGeneratedUPackage(FCookGenerationInfo& GenerationInfo,
 		const UPackage* OwnerPackage, const TCHAR* GeneratedPackageName);
-	/** Mark that the generator package has saved, to keep track of when *this is no longer needed. */
-	void SetGeneratorSaved(UPackage* GeneratorUPackage);
-	/** Mark that the generated package has saved, to keep track of when *this is no longer needed. */
-	void SetGeneratedSaved(FPackageData& PackageData);
+	/** Mark that the generator or a generated package has saved, to keep track of when *this is no longer needed. */
+	void SetPackageSaved(FCookGenerationInfo& Info, FPackageData& PackageData);
 	/** Return whether list has been generated and all generated packages have been populated */
 	bool IsComplete() const;
-	FBeginCacheObjects& GetGeneratorBeginCacheObjects() { return GeneratorBeginCacheObjects; }
 
-	void SetObjectsToMoveIntoGenerator(UPackage* Package, UObject* SplitDataObject,
-		ICookPackageSplitter* Splitter, TArray<FWeakObjectPtr>& CachedObjectsInOuter,
-		TArray<ICookPackageSplitter::FGeneratedPackageForPreSave>& GeneratedPackagesForPresave);
-	void SetPostMoveObjects(UPackage* Package);
-	void ResetGeneratorSaveState(UPackage* Package, UE::Cook::EReleaseSaveReason ReleaseSaveReason);
+	bool IsSaveInvalidated(const FPackageData& PackageData) const;
+
+	UPackage* GetOwnerPackage() const { return OwnerPackage.Get(); };
+	void SetOwnerPackage(UPackage* InPackage) { OwnerPackage = InPackage; }
 
 private:
 	void ConditionalNotifyCompletion(ICookPackageSplitter::ETeardown Status);
 
 	/** PackageData for the package that is being split */
-	UE::Cook::FPackageData& Owner;
+	FCookGenerationInfo OwnerInfo;
 	/** Name of the object that prompted the splitter creation */
 	FName SplitDataObjectName;
 	/** Cached CookPackageSplitter */
 	TUniquePtr<ICookPackageSplitter> CookPackageSplitterInstance;
 	/** Recorded list of packages to generate from the splitter, and data we need about them */
-	TArray<FGeneratedStruct> PackagesToGenerate;
-	FBeginCacheObjects GeneratorBeginCacheObjects;
+	TArray<FCookGenerationInfo> PackagesToGenerate;
+	TWeakObjectPtr<UPackage> OwnerPackage;
 
 	int32 NextPopulateIndex = 0;
 	int32 RemainingToPopulate = 0;
-	ESaveState GeneratorSaveState = ESaveState::Start;
-	bool bTakenOverCachedCookedPlatformData = false;
-	bool bIssuedUndeclaredMovedObjectsWarning = false;
 
-	bool bWasOwnerReloaded = false;
-	bool bOwnerHasSaved = false;
 	bool bNotifiedCompletion = false;
 };
 
diff --git a/Engine/Source/Editor/UnrealEd/Public/CookPackageSplitter.h b/Engine/Source/Editor/UnrealEd/Public/CookPackageSplitter.h
--- a/Engine/Source/Editor/UnrealEd/Public/CookPackageSplitter.h
+++ b/Engine/Source/Editor/UnrealEd/Public/CookPackageSplitter.h
@@ -66,9 +66,10 @@ public:
 		UPackage* Package = nullptr;
 		/** *GetCreateAsMap returned from GetGenerateList. The package filename extension has already been set based on this. */
 		bool bCreatedAsMap = false;
-		/** Whether the Package was populated by PreSave or constructed empty */
-		bool bPopulatedByPreSave = false;
 	};
+
+	virtual void GetObjectsToMoveIntoPackage(UPackage* OwnerPackage, UObject* OwnerObject,
+		const FGeneratedPackageForPopulate& GeneratedPackage, TArray<UObject*>& OutObjectsToMove) {}
 	/**
 	 * Try to populate a generated package.
 	 *
@@ -80,12 +81,18 @@ public:
 	 * @param OwnerPackage				The parent package being split
 	 * @param OwnerObject				The SplitDataClass instance that this CookPackageSplitter instance was created for
 	 * @param GeneratedPackage			Pointer and information about the package to populate
-	 * @param bWasOwnerReloaded			True iff the OwnerPackage was garbage collected and reloaded since the previous call
-	 *                                  to GetGenerateList or TryPopulatePackage
 	 * @return							True if successfully populates,  false on error (this will cause a cook error).
 	 */
-	virtual bool TryPopulatePackage(const UPackage* OwnerPackage, const UObject* OwnerObject,
-		const FGeneratedPackageForPopulate& GeneratedPackage, bool bWasOwnerReloaded) = 0;
+	virtual bool TryPopulatePackage(UPackage* OwnerPackage, UObject* OwnerObject,
+		const FGeneratedPackageForPopulate& GeneratedPackage) = 0;
+	/**
+	 * Called after saving the generated package. Undo any required adjustments to the parent package that
+	 * were made in TryPopulatePackage, so that the package is once again ready for use in the editor or in
+	 * future GetGenerateList or TryPopulatePackage calls
+	 */
+	virtual void PostSaveGeneratedPackage(UPackage* OwnerPackage, UObject* OwnerObject,
+		const FGeneratedPackageForPopulate& GeneratedPackage) {}
+
 
 	struct FGeneratedPackageForPreSave
 	{
@@ -118,6 +125,9 @@ public:
 	 * future GetGenerateList or TryPopulatePackage calls
 	 */
 	virtual void PostSaveGeneratorPackage(UPackage* OwnerPackage, UObject* OwnerObject) {}
+
+	/** Called when the Owner package needs to be reloaded after a garbage collect in order to populate a generated package. */
+	virtual void OnOwnerReloaded(UPackage* OwnerPackage, UObject* OwnerObject) {}
 };
 
 namespace UE
diff --git a/Engine/Source/Runtime/Engine/Private/WorldPartition/WorldPartitionCookPackageSplitter.cpp b/Engine/Source/Runtime/Engine/Private/WorldPartition/WorldPartitionCookPackageSplitter.cpp
--- a/Engine/Source/Runtime/Engine/Private/WorldPartition/WorldPartitionCookPackageSplitter.cpp
+++ b/Engine/Source/Runtime/Engine/Private/WorldPartition/WorldPartitionCookPackageSplitter.cpp
@@ -117,14 +117,10 @@ TArray<ICookPackageSplitter::FGeneratedPackage> FWorldPartitionCookPackageSplitt
 	return PackagesToGenerate;
 }
 
-bool FWorldPartitionCookPackageSplitter::TryPopulatePackage(const UPackage* OwnerPackage, const UObject* OwnerObject,
-	const FGeneratedPackageForPopulate& GeneratedPackage, bool bWasOwnerReloaded)
+bool FWorldPartitionCookPackageSplitter::TryPopulatePackage(UPackage* OwnerPackage, UObject* OwnerObject,
+	const FGeneratedPackageForPopulate& GeneratedPackage)
 {
-	check(!bWasOwnerReloaded);
-
-	// TODO: Make PopulateGeneratedPackageForCook const so we can honor the constness of the OwnerObject in this API function
-	const UWorld* ConstPartitionedWorld = ValidateDataObject(OwnerObject);
-	UWorld* PartitionedWorld = const_cast<UWorld*>(ConstPartitionedWorld);
+	UWorld* PartitionedWorld = ValidateDataObject(OwnerObject);
 	UWorldPartition* WorldPartition = PartitionedWorld->PersistentLevel->GetWorldPartition();
 
 	TArray<UObject*> LoadedObjects;
@@ -141,12 +137,12 @@ void FWorldPartitionCookPackageSplitter::PreSaveGeneratorPackage(UPackage* Owner
 {
 	UWorld* PartitionedWorld = ValidateDataObject(OwnerObject);
 	UWorldPartition* WorldPartition = PartitionedWorld->PersistentLevel->GetWorldPartition();
-
-	TArray<UObject*> LoadedObjects;
-	if (WorldPartition->LoadGeneratorPackageObjectsForCook(LoadedObjects))
-	{
-		WorldPartition->FinalizeGeneratorPackageForCook(GeneratedPackages);
-	}
+	WorldPartition->FinalizeGeneratorPackageForCook(GeneratedPackages);
 }
 
+void FWorldPartitionCookPackageSplitter::OnOwnerReloaded(UPackage* OwnerPackage, UObject* OwnerObject)
+{
+	// It should not be possible for the owner to reload due to garbage collection while we are active and keeping it referenced
+	check(!ReferencedWorld); 
+}
 #endif
diff --git a/Engine/Source/Runtime/Engine/Private/WorldPartition/WorldPartitionCookPackageSplitter.h b/Engine/Source/Runtime/Engine/Private/WorldPartition/WorldPartitionCookPackageSplitter.h
--- a/Engine/Source/Runtime/Engine/Private/WorldPartition/WorldPartitionCookPackageSplitter.h
+++ b/Engine/Source/Runtime/Engine/Private/WorldPartition/WorldPartitionCookPackageSplitter.h
@@ -20,10 +20,11 @@ public:
 	virtual void Teardown(ETeardown Status) override;
 	virtual bool UseInternalReferenceToAvoidGarbageCollect() override { return true; }
 	virtual TArray<ICookPackageSplitter::FGeneratedPackage> GetGenerateList(const UPackage* OwnerPackage, const UObject* OwnerObject) override;
-	virtual bool TryPopulatePackage(const UPackage* OwnerPackage, const UObject* OwnerObject,
-		const ICookPackageSplitter::FGeneratedPackageForPopulate& GeneratedPackage, bool bWasOwnerReloaded) override;
+	virtual bool TryPopulatePackage(UPackage* OwnerPackage, UObject* OwnerObject,
+		const ICookPackageSplitter::FGeneratedPackageForPopulate& GeneratedPackage) override;
 	virtual void PreSaveGeneratorPackage(UPackage* OwnerPackage, UObject* OwnerObject,
 		const TArray<ICookPackageSplitter::FGeneratedPackageForPreSave>& GeneratedPackages) override;
+	virtual void OnOwnerReloaded(UPackage* OwnerPackage, UObject* OwnerObject) override;
 	//~ End of ICookPackageSplitter
 
 private:
@@ -34,7 +35,7 @@ private:
 	const UWorld* ValidateDataObject(const UObject* SplitData);
 	UWorld* ValidateDataObject(UObject* SplitData);
 	
-	TObjectPtr<UWorld> ReferencedWorld;
+	TObjectPtr<UWorld> ReferencedWorld = nullptr;
 
 	bool bInitializedWorldPartition = false;
 	bool bForceInitializedWorld = false;
