CookPackageSplitter: Downgrade Warning "created or moved objects during PreSaveGeneratorPackage that are not yet ready to save" to a Display until the fix is submitted.
#rnx

diff --git a/Engine/Source/Editor/UnrealEd/Private/CookOnTheFlyServer.cpp b/Engine/Source/Editor/UnrealEd/Private/CookOnTheFlyServer.cpp
--- a/Engine/Source/Editor/UnrealEd/Private/CookOnTheFlyServer.cpp
+++ b/Engine/Source/Editor/UnrealEd/Private/CookOnTheFlyServer.cpp
@@ -3059,7 +3059,7 @@ UE::Cook::EPollStatus UCookOnTheFlyServer::BeginCachePostMove(UE::Cook::FGenerat
 				}
 			}
 		}
-		UE_LOG(LogCook, Warning, TEXT("CookPackageSplitter created or moved objects during %s that are not yet ready to save. This will cause an error if garbage collection runs before the package is saved.\n")
+		UE_LOG(LogCook, Display, TEXT("CookPackageSplitter created or moved objects during %s that are not yet ready to save. This will cause an error if garbage collection runs before the package is saved.\n")
 			TEXT("Change the splitter's %s to construct new objects and declare existing objects that will be moved from other packages.\n")
 			TEXT("SplitterObject: %s%s\n")
 			TEXT("NumPendingObjects: %d, FirstPendingObject: %s"),
