Fix non-unity
#rnx

diff --git a/Engine/Source/Editor/UnrealEd/Private/WorldPartition/WorldPartitionRenameDuplicateBuilder.cpp b/Engine/Source/Editor/UnrealEd/Private/WorldPartition/WorldPartitionRenameDuplicateBuilder.cpp
--- a/Engine/Source/Editor/UnrealEd/Private/WorldPartition/WorldPartitionRenameDuplicateBuilder.cpp
+++ b/Engine/Source/Editor/UnrealEd/Private/WorldPartition/WorldPartitionRenameDuplicateBuilder.cpp
@@ -12,7 +12,9 @@
 #include "GameFramework/Actor.h"
 #include "EditorWorldUtils.h"
 #include "UObject/SavePackage.h"
+#include "UObject/Linker.h"
 #include "UObject/MetaData.h"
+#include "UObject/ObjectRedirector.h"
 #include "ProfilingDebugging/ScopedTimers.h"
 #include "HAL/FileManager.h"
 
diff --git a/Engine/Source/Runtime/Engine/Private/WorldPartition/WorldPartitionActorDescViewProxy.cpp b/Engine/Source/Runtime/Engine/Private/WorldPartition/WorldPartitionActorDescViewProxy.cpp
--- a/Engine/Source/Runtime/Engine/Private/WorldPartition/WorldPartitionActorDescViewProxy.cpp
+++ b/Engine/Source/Runtime/Engine/Private/WorldPartition/WorldPartitionActorDescViewProxy.cpp
@@ -1,6 +1,7 @@
 // Copyright Epic Games, Inc. All Rights Reserved.
 
 #include "WorldPartition/WorldPartitionActorDescViewProxy.h"
+#include "GameFramework/Actor.h"
 
 #if WITH_EDITOR
 #include "WorldPartition/WorldPartitionActorDesc.h"
