World Partition Rename Builder:
- Fixed crash when renaming a world for the first time.
- Fixed crash when providing an invalid destination package.
- Converted to use the generic cluster reference code instead of the streaming generation version.

#rnx

diff --git a/Engine/Source/Editor/UnrealEd/Private/WorldPartition/WorldPartitionRenameDuplicateBuilder.cpp b/Engine/Source/Editor/UnrealEd/Private/WorldPartition/WorldPartitionRenameDuplicateBuilder.cpp
--- a/Engine/Source/Editor/UnrealEd/Private/WorldPartition/WorldPartitionRenameDuplicateBuilder.cpp
+++ b/Engine/Source/Editor/UnrealEd/Private/WorldPartition/WorldPartitionRenameDuplicateBuilder.cpp
@@ -5,7 +5,7 @@
 #include "WorldPartition/WorldPartition.h"
 #include "WorldPartition/WorldPartitionHelpers.h"
 #include "WorldPartition/WorldPartitionActorDescView.h"
-#include "WorldPartition/WorldPartitionActorCluster.h"
+#include "ReferenceCluster.h"
 #include "PackageSourceControlHelper.h"
 #include "SourceControlHelpers.h"
 #include "Engine/World.h"
@@ -75,49 +75,54 @@ static bool DeleteExistingMapPackages(const FString& ExistingPackageName, FPacka
 	UE_SCOPED_TIMER(TEXT("Delete existing destination packages"), LogWorldPartitionRenameDuplicateBuilder, Display);
 	TArray<FString> PackagesToDelete;
 
-	FString ExistingMapPackageFilePath = FPackageName::LongPackageNameToFilename(ExistingPackageName, FPackageName::GetMapPackageExtension());
-	if (IFileManager::Get().FileExists(*ExistingMapPackageFilePath))
+	FString ExistingMapPackageFilePath;	
+	if (FPackageName::TryConvertLongPackageNameToFilename(ExistingPackageName, ExistingMapPackageFilePath, FPackageName::GetMapPackageExtension()))
 	{
-		PackagesToDelete.Add(ExistingMapPackageFilePath);
-	}
+		if (IFileManager::Get().FileExists(*ExistingMapPackageFilePath))
+		{
+			PackagesToDelete.Add(ExistingMapPackageFilePath);
+		}
 
-	FString BuildDataPackageName = ExistingPackageName + TEXT("_BuiltData");
-	FString ExistingBuildDataPackageFilePath = FPackageName::LongPackageNameToFilename(BuildDataPackageName, FPackageName::GetAssetPackageExtension());
-	if (IFileManager::Get().FileExists(*ExistingBuildDataPackageFilePath))
-	{
-		PackagesToDelete.Add(ExistingBuildDataPackageFilePath);
-	}
+		FString BuildDataPackageName = ExistingPackageName + TEXT("_BuiltData");
+		FString ExistingBuildDataPackageFilePath = FPackageName::LongPackageNameToFilename(BuildDataPackageName, FPackageName::GetAssetPackageExtension());
+		if (IFileManager::Get().FileExists(*ExistingBuildDataPackageFilePath))
+		{
+			PackagesToDelete.Add(ExistingBuildDataPackageFilePath);
+		}
 
-	// Search for external object packages
-	const TArray<FString> ExternalPackagesPaths = ULevel::GetExternalObjectsPaths(ExistingPackageName);
-	for (const FString& ExternalPackagesPath : ExternalPackagesPaths)
-	{
-		FString ExternalPackagesFilePath = FPackageName::LongPackageNameToFilename(ExternalPackagesPath);
-		if (IFileManager::Get().DirectoryExists(*ExternalPackagesFilePath))
+		// Search for external object packages
+		const TArray<FString> ExternalPackagesPaths = ULevel::GetExternalObjectsPaths(ExistingPackageName);
+		for (const FString& ExternalPackagesPath : ExternalPackagesPaths)
 		{
-			const bool bSuccess = IFileManager::Get().IterateDirectoryRecursively(*ExternalPackagesFilePath, [&PackagesToDelete](const TCHAR* FilenameOrDirectory, bool bIsDirectory)
+			FString ExternalPackagesFilePath = FPackageName::LongPackageNameToFilename(ExternalPackagesPath);
+			if (IFileManager::Get().DirectoryExists(*ExternalPackagesFilePath))
 			{
-				if (!bIsDirectory)
+				const bool bSuccess = IFileManager::Get().IterateDirectoryRecursively(*ExternalPackagesFilePath, [&PackagesToDelete](const TCHAR* FilenameOrDirectory, bool bIsDirectory)
 				{
-					FString Filename(FilenameOrDirectory);
-					if (Filename.EndsWith(FPackageName::GetAssetPackageExtension()))
+					if (!bIsDirectory)
 					{
-						PackagesToDelete.Add(Filename);
+						FString Filename(FilenameOrDirectory);
+						if (Filename.EndsWith(FPackageName::GetAssetPackageExtension()))
+						{
+							PackagesToDelete.Add(Filename);
+						}
 					}
-				}
-				// Continue Directory Iteration
-				return true;
-			});
+					// Continue Directory Iteration
+					return true;
+				});
 
-			if (!bSuccess)
-			{
-				UE_LOG(LogWorldPartitionRenameDuplicateBuilder, Error, TEXT("Failed to iterate existing external actors path: %s"), *ExternalPackagesPath);
-				return false;
+				if (!bSuccess)
+				{
+					UE_LOG(LogWorldPartitionRenameDuplicateBuilder, Error, TEXT("Failed to iterate existing external actors path: %s"), *ExternalPackagesPath);
+					return false;
+				}
 			}
 		}
+
+		return UWorldPartitionBuilder::DeletePackages(PackagesToDelete, PackageHelper);
 	}
 
-	return UWorldPartitionBuilder::DeletePackages(PackagesToDelete, PackageHelper);
+	return true;
 }
 
 UWorldPartitionRenameDuplicateBuilder::UWorldPartitionRenameDuplicateBuilder(const FObjectInitializer& ObjectInitializer)
@@ -136,21 +141,14 @@ bool UWorldPartitionRenameDuplicateBuilder::RunInternal(UWorld* World, const FCe
 		return false;
 	}
 
-	FActorDescViewMap ActorDescViewMap;
-	for (UActorDescContainer::TConstIterator<> ActorDescIterator(WorldPartition); ActorDescIterator; ++ActorDescIterator)
+	// Build actor clusters
+	TArray<TPair<FGuid, TArray<FGuid>>> ActorsWithRefs;
+	for (FActorDescList::TIterator<> ActorDescIterator(WorldPartition); ActorDescIterator; ++ActorDescIterator)
 	{
-		const FWorldPartitionActorDesc* ActorDesc = *ActorDescIterator;
-		FWorldPartitionActorDescView ActorDescView(ActorDesc);
-		// Invalidate data layers to avoid clustering errors. It doesn't matter here as we need clustering only to gather references
-		ActorDescView.SetInvalidDataLayers();
-		ActorDescViewMap.Emplace(ActorDescView.GetGuid(), ActorDescView);
+		ActorsWithRefs.Emplace(ActorDescIterator->GetGuid(), ActorDescIterator->GetReferences());
 	}
 
-	TArray<FActorCluster> ActorClusters;
-	{
-		UE_SCOPED_TIMER(TEXT("Create actor clusters"), LogWorldPartitionRenameDuplicateBuilder, Display);
-		FActorClusterContext::CreateActorClusters(World, ActorDescViewMap.GetActorDescViewsByGuid(), ActorClusters);
-	}
+	TArray<TArray<FGuid>> ActorClusters = GenerateObjectsClusters(ActorsWithRefs);
 
 	UPackage* OriginalPackage = World->GetPackage();
 	OriginalWorldName = World->GetName();
@@ -182,6 +180,12 @@ bool UWorldPartitionRenameDuplicateBuilder::RunInternal(UWorld* World, const FCe
 	}
 				
 	UPackage* NewPackage = CreatePackage(*NewPackageName);
+	if (!NewPackage)
+	{
+		UE_LOG(LogWorldPartitionRenameDuplicateBuilder, Error, TEXT("Failed to create destination package."));
+		return false;
+	}
+
 	TMap<FGuid, FGuid> DuplicatedActorGuids;
 	TArray<UPackage*> DuplicatedPackagesToSave;
 	UWorld* NewWorld = nullptr;
@@ -326,10 +330,10 @@ bool UWorldPartitionRenameDuplicateBuilder::RunInternal(UWorld* World, const FCe
 			};
 
 			TArray<FWorldPartitionReference> ActorReferences;
-			for (const FActorCluster& ActorCluster : ActorClusters)
+			for (const TArray<FGuid>& ActorCluster : ActorClusters)
 			{
-				UE_LOG(LogWorldPartitionRenameDuplicateBuilder, Display, TEXT("Processing cluster with %d actor(s)"), ActorCluster.Actors.Num());
-				for (const FGuid& ActorGuid : ActorCluster.Actors)
+				UE_LOG(LogWorldPartitionRenameDuplicateBuilder, Display, TEXT("Processing cluster with %d actor(s)"), ActorCluster.Num());
+				for (const FGuid& ActorGuid : ActorCluster)
 				{
 					// Duplicated actors don't need to be processed
 					if (!DuplicatedActorGuids.Contains(ActorGuid))
@@ -344,10 +348,10 @@ bool UWorldPartitionRenameDuplicateBuilder::RunInternal(UWorld* World, const FCe
 					// If we are renaming add package to delete
 					if (bRename)
 					{
-						const FWorldPartitionActorDescView* ActorDescView = const_cast<const FActorDescViewMap&>(ActorDescViewMap).FindByGuid(ActorGuid);
-						check(ActorDescView);
+						const FWorldPartitionActorDesc* ActorDesc = WorldPartition->GetActorDesc(ActorGuid);
+						check(ActorDesc);
 
-						PackagesToDelete.Add(ActorDescView->GetActorPackage().ToString());
+						PackagesToDelete.Add(ActorDesc->GetActorPackage().ToString());
 					}
 				}
 						
