WorldPartition: Implement ICookPackageSplitter GetObjectsToMoveIntoPackage and GetObjectsToMoveIntoGenerator so the cooker can wait on them to be ready before doing the move.
#rnx

diff --git a/Engine/Source/Editor/UnrealEd/Private/CookOnTheFlyServer.cpp b/Engine/Source/Editor/UnrealEd/Private/CookOnTheFlyServer.cpp
--- a/Engine/Source/Editor/UnrealEd/Private/CookOnTheFlyServer.cpp
+++ b/Engine/Source/Editor/UnrealEd/Private/CookOnTheFlyServer.cpp
@@ -2904,15 +2904,33 @@ UE::Cook::EPollStatus UCookOnTheFlyServer::PrepareSaveGeneratedPackage(UE::Cook:
 }
 
 /** Find all subobjects of all objects in RootObjects and add them to the end of InObjects */
-static void ExtendWithSubObjects(TArray<UObject*>& InObjects)
+static void ValidateAndExtendObjectsToMove(TArray<UObject*>& InObjects, UE::Cook::FGeneratorPackage& Generator, UE::Cook::FCookGenerationInfo& Info)
 {
 	TArray<UObject*> ObjectsInOuter;
 	int32 OriginalNum = InObjects.Num();
 	for (int32 Index = 0; Index < OriginalNum; ++Index)
 	{
 		ObjectsInOuter.Reset();
-		GetObjectsWithOuter(InObjects[Index], ObjectsInOuter);
-		InObjects.Append(ObjectsInOuter);
+		UObject* OriginalObject = InObjects[Index];
+		if (!IsValid(OriginalObject))
+		{
+			UE_LOG(LogCook, Warning, TEXT("CookPackageSplitter found non-valid object %s returned from %s on Splitter %s%s. Ignoring it."),
+				OriginalObject ? *OriginalObject->GetFullName() : TEXT("<null>"),
+				Info.IsGenerator() ? TEXT("GetObjectsToMoveIntoGeneratorPackage") : TEXT("GetObjectsToMoveIntoGeneratedPackage"),
+				*Generator.GetSplitDataObjectName().ToString(),
+				Info.IsGenerator() ? TEXT("") : *FString::Printf(TEXT(", Package %s"), *Info.PackageData->GetPackageName().ToString()));
+			--OriginalNum;
+			--Index;
+			InObjects.RemoveAt(Index);
+			continue;
+		}
+		GetObjectsWithOuter(OriginalObject, ObjectsInOuter, true /* bIncludeNestedObjects */, RF_NoFlags, EInternalObjectFlags::Garbage);
+		InObjects.Reserve(InObjects.Num() + ObjectsInOuter.Num());
+		for (UObject* Object : ObjectsInOuter)
+		{
+			check(IsValid(Object));
+			InObjects.Add(Object);
+		}
 	}
 }
 
@@ -2941,14 +2959,14 @@ UE::Cook::EPollStatus UCookOnTheFlyServer::BeginCacheObjectsToMove(UE::Cook::FGe
 		if (Info.IsGenerator())
 		{
 			ConstructGeneratedPackagesForPresave(PackageData, Generator, GeneratedPackagesForPresave);
-			Splitter->GetObjectsToMoveIntoGenerator(Package, SplitDataObject, GeneratedPackagesForPresave, ObjectsToMove);
+			Splitter->GetObjectsToMoveIntoGeneratorPackage(Package, SplitDataObject, GeneratedPackagesForPresave, ObjectsToMove);
 		}
 		else
 		{
 			ICookPackageSplitter::FGeneratedPackageForPopulate SplitterInfo{ Info.RelativePath, Package, Info.IsCreateAsMap() };
-			Splitter->GetObjectsToMoveIntoPackage(Package, SplitDataObject, SplitterInfo, ObjectsToMove);
+			Splitter->GetObjectsToMoveIntoGeneratedPackage(Package, SplitDataObject, SplitterInfo, ObjectsToMove);
 		}
-		ExtendWithSubObjects(ObjectsToMove);
+		ValidateAndExtendObjectsToMove(ObjectsToMove, Generator, Info);
 		Info.TakeOverCachedObjectsAndAddNew(PackageData.GetCachedObjectsInOuter(), ObjectsToMove);
 		Info.SetSaveStateComplete(FCookGenerationInfo::ESaveState::CallObjectsToMove);
 	}
@@ -3064,7 +3082,7 @@ UE::Cook::EPollStatus UCookOnTheFlyServer::BeginCachePostMove(UE::Cook::FGenerat
 			TEXT("SplitterObject: %s%s\n")
 			TEXT("NumPendingObjects: %d, FirstPendingObject: %s"),
 			Info.IsGenerator() ? TEXT("PreSaveGeneratorPackage") : TEXT("TryPopulatePackage"),
-			Info.IsGenerator() ? TEXT("GetObjectsToMoveIntoGenerator") : TEXT("GetObjectsToMoveIntoPackage"),
+			Info.IsGenerator() ? TEXT("GetObjectsToMoveIntoGeneratorPackage") : TEXT("GetObjectsToMoveIntoGeneratedPackage"),
 			*SplitDataObject->GetFullName(),
 			Info.IsGenerator() ? TEXT("") : *FString::Printf(TEXT("\nGeneratedPackage: %s"), *PackageData.GetPackageName().ToString()),
 			PackageData.GetNumPendingCookedPlatformData(),
@@ -4680,7 +4698,9 @@ void UCookOnTheFlyServer::PostGarbageCollect()
 	TArray<FPackageData*> Demotes;
 	for (FPackageData* PackageData : PackageDatas->GetSaveQueue())
 	{
-		if (PackageData->IsSaveInvalidated())
+		bool bOutDemote;
+		PackageData->UpdateSaveAfterGarbageCollect(bOutDemote);
+		if (bOutDemote)
 		{
 			Demotes.Add(PackageData);
 		}
diff --git a/Engine/Source/Editor/UnrealEd/Private/Cooker/CookPackageData.cpp b/Engine/Source/Editor/UnrealEd/Private/Cooker/CookPackageData.cpp
--- a/Engine/Source/Editor/UnrealEd/Private/Cooker/CookPackageData.cpp
+++ b/Engine/Source/Editor/UnrealEd/Private/Cooker/CookPackageData.cpp
@@ -1084,11 +1084,12 @@ void FPackageData::RemapTargetPlatforms(const TMap<ITargetPlatform*, ITargetPlat
 	PlatformDatas = MoveTemp(NewPlatformDatas);
 }
 
-bool FPackageData::IsSaveInvalidated() const
+void FPackageData::UpdateSaveAfterGarbageCollect(bool& bOutDemote)
 {
+	bOutDemote = false;
 	if (GetState() != EPackageState::Save)
 	{
-		return false;
+		return;
 	}
 
 	if (GetPackage() == nullptr || !GetPackage()->IsFullyLoaded() ||
@@ -1101,27 +1102,24 @@ bool FPackageData::IsSaveInvalidated() const
 				return WeakPtr.Get() == nullptr;
 			}))
 	{
-		return true;
+		bOutDemote = true;
+		return;
 	}
 	if (GeneratorPackage)
 	{
-		if (GeneratorPackage->IsSaveInvalidated(*this))
-		{
-			return true;
-		}
+		GeneratorPackage->UpdateSaveAfterGarbageCollect(*this, bOutDemote);
 	}
 	else if (IsGenerated())
 	{
 		if (!GeneratedOwner)
 		{
-			return true;
+			bOutDemote = true;
 		}
-		if (GeneratedOwner->IsSaveInvalidated(*this))
+		else
 		{
-			return true;
+			GeneratedOwner->UpdateSaveAfterGarbageCollect(*this, bOutDemote);
 		}
 	}
-	return false;
 }
 
 void FPackageData::SetGeneratedOwner(FGeneratorPackage* InGeneratedOwner)
@@ -1324,7 +1322,7 @@ void FGeneratorPackage::PostGarbageCollect()
 	FPackageData& Owner = GetOwner();
 	if (Owner.GetState() == EPackageState::Save)
 	{
-		// UCookOnTheFlyServer::PreCollectGarbage adds references for the Generator package and all its public
+		// UCookOnTheFlyServer::PreGarbageCollect adds references for the Generator package and all its public
 		// objects, so it should still be loaded
 		if (!Owner.GetPackage() || !FindSplitDataObject())
 		{
@@ -1463,36 +1461,63 @@ void FGeneratorPackage::ResetSaveState(FCookGenerationInfo& Info, UPackage* Pack
 	{
 		Info.SetSaveState(FCookGenerationInfo::ESaveState::StartPopulate);
 	}
+	if (Info.BeginCacheObjects.Objects.Num() != 0 &&
+		GetCookPackageSplitterInstance()->UseInternalReferenceToAvoidGarbageCollect() &&
+		(ReleaseSaveReason == EReleaseSaveReason::Demoted || ReleaseSaveReason == EReleaseSaveReason::RecreateObjectCache))
+	{
+		UE_LOG(LogCook, Error, TEXT("CookPackageSplitter failure: We are demoting a generated package from save and removing our references that keep its objects loaded.\n")
+			TEXT("This will allow the objects to be garbage collected and cause failures in the splitter which expects them to remain loaded.\n")
+			TEXT("Package=%s, Splitter=%s, ReleaseSaveReason=%s"),
+			Info.PackageData ? *Info.PackageData->GetPackageName().ToString() : *Info.RelativePath,
+			*GetSplitDataObjectName().ToString(), LexToString(ReleaseSaveReason));
+	}
 	Info.BeginCacheObjects.Reset();
 	Info.SetHasTakenOverCachedCookedPlatformData(false);
 	Info.SetHasIssuedUndeclaredMovedObjectsWarning(false);
 }
 
-bool FGeneratorPackage::IsSaveInvalidated(const FPackageData& PackageData) const
+void FGeneratorPackage::UpdateSaveAfterGarbageCollect(const FPackageData& PackageData, bool& bInOutDemote)
 {
-	const FCookGenerationInfo* Info = FindInfo(PackageData);
+	FCookGenerationInfo* Info = FindInfo(PackageData);
 	if (!Info)
 	{
-		return true;
+		bInOutDemote = true;
+		return;
 	}
 
-	for (const FBeginCacheObject& Object : Info->BeginCacheObjects.Objects)
+	if (!Info->IsGenerator())
 	{
-		if (!Object.Object.IsValid())
+		UPackage* LocalPackage = OwnerPackage.Get();
+		if (!LocalPackage || !LocalPackage->IsFullyLoaded())
 		{
-			return true;
+			bInOutDemote = true;
+			return;
 		}
 	}
 
-	if (!Info->IsGenerator())
+	for (TArray<FBeginCacheObject>::TIterator Iter(Info->BeginCacheObjects.Objects); Iter; ++Iter)
 	{
-		UPackage* LocalPackage = OwnerPackage.Get();
-		if (!LocalPackage || !LocalPackage->IsFullyLoaded())
+		const FBeginCacheObject& Object = *Iter;
+		if (!Object.Object.IsValid())
 		{
-			return true;
+			if (GetCookPackageSplitterInstance()->UseInternalReferenceToAvoidGarbageCollect())
+			{
+				// No objects should be allowed to be deleted; we are supposed to keep them referenced
+				// But allowing demotion will break things for sure.
+				// Log a cook error but remove the invalidated object.
+				UE_LOG(LogCook, Error, TEXT("PackageSplitter found an object returned from %s that was removed from memory during garbage collection. This will cause errors during save of the package.")
+					TEXT("\n\tSplitter=%s%s."),
+					Info->IsGenerator() ? TEXT("GetObjectsToMoveIntoGeneratorPackage") : TEXT("GetObjectsToMoveIntoGeneratedPackage"),
+					*GetSplitDataObjectName().ToString(),
+					Info->IsGenerator() ? TEXT("") : *FString::Printf(TEXT(", Generated=%s."), *Info->PackageData->GetPackageName().ToString()));
+				Iter.RemoveCurrent();
+			}
+			else
+			{
+				bInOutDemote = true;
+			}
 		}
 	}
-	return false;
 }
 
 FCookGenerationInfo::FCookGenerationInfo(FPackageData& InPackageData, bool bInGenerator)
diff --git a/Engine/Source/Editor/UnrealEd/Private/Cooker/CookPackageData.h b/Engine/Source/Editor/UnrealEd/Private/Cooker/CookPackageData.h
--- a/Engine/Source/Editor/UnrealEd/Private/Cooker/CookPackageData.h
+++ b/Engine/Source/Editor/UnrealEd/Private/Cooker/CookPackageData.h
@@ -387,8 +387,11 @@ public:
 	 */
 	bool GetCookedPlatformDataComplete() const { return static_cast<bool>(bCookedPlatformDataComplete); }
 	void SetCookedPlatformDataComplete(bool bValue) { bCookedPlatformDataComplete = bValue != 0; }
-	/** Check whether savestate contracts on the PackageData were invalidated by by e.g. garbage collection. */
-	bool IsSaveInvalidated() const;
+	/**
+	 * Check whether savestate contracts on the PackageData were invalidated by by e.g. garbage collection.
+	 * Request demotion if so unless we have a contract to keep it, in which case it is fixed up.
+	 */
+	void UpdateSaveAfterGarbageCollect(bool& bOutDemote);
 
 	/** Get/Set the flag for whether PrepareSave has been called and returned an error. */
 	bool HasPrepareSaveFailed() const { return static_cast<bool>(bPrepareSaveFailed); }
@@ -719,7 +722,7 @@ public:
 	/** Return whether list has been generated and all generated packages have been populated */
 	bool IsComplete() const;
 
-	bool IsSaveInvalidated(const FPackageData& PackageData) const;
+	void UpdateSaveAfterGarbageCollect(const FPackageData& PackageData, bool& bInOutDemote);
 
 	UPackage* GetOwnerPackage() const { return OwnerPackage.Get(); };
 	void SetOwnerPackage(UPackage* InPackage) { OwnerPackage = InPackage; }
diff --git a/Engine/Source/Editor/UnrealEd/Private/Cooker/CookTypes.cpp b/Engine/Source/Editor/UnrealEd/Private/Cooker/CookTypes.cpp
--- a/Engine/Source/Editor/UnrealEd/Private/Cooker/CookTypes.cpp
+++ b/Engine/Source/Editor/UnrealEd/Private/Cooker/CookTypes.cpp
@@ -10,6 +10,19 @@
 namespace UE::Cook
 {
 
+const TCHAR* LexToString(EReleaseSaveReason Reason)
+{
+	switch (Reason)
+	{
+	case EReleaseSaveReason::Completed: return TEXT("Completed");
+	case EReleaseSaveReason::DoneForNow: return TEXT("DoneForNow");
+	case EReleaseSaveReason::Demoted: return TEXT("Demoted");
+	case EReleaseSaveReason::AbortSave: return TEXT("AbortSave");
+	case EReleaseSaveReason::RecreateObjectCache: return TEXT("RecreateObjectCache");
+	default: return TEXT("Invalid");
+	}
+}
+
 FCookerTimer::FCookerTimer(float InTimeSlice)
 	: StartTime(FPlatformTime::Seconds()), TimeSlice(InTimeSlice)
 {
diff --git a/Engine/Source/Editor/UnrealEd/Private/Cooker/CookTypes.h b/Engine/Source/Editor/UnrealEd/Private/Cooker/CookTypes.h
--- a/Engine/Source/Editor/UnrealEd/Private/Cooker/CookTypes.h
+++ b/Engine/Source/Editor/UnrealEd/Private/Cooker/CookTypes.h
@@ -111,6 +111,7 @@ namespace UE::Cook
 		AbortSave,
 		RecreateObjectCache,
 	};
+	const TCHAR* LexToString(UE::Cook::EReleaseSaveReason Reason);
 
 	/** The type of callback for External Requests that needs to be executed within the Scheduler's lock. */
 	typedef TUniqueFunction<void()> FSchedulerCallback;
@@ -346,7 +347,6 @@ struct FBeginCookContext
 	TArray<ITargetPlatform*> TargetPlatforms;
 };
 
-
 void LogCookerMessage(const FString& MessageText, EMessageSeverity::Type Severity);
 LLM_DECLARE_TAG(Cooker);
 
diff --git a/Engine/Source/Editor/UnrealEd/Public/CookPackageSplitter.h b/Engine/Source/Editor/UnrealEd/Public/CookPackageSplitter.h
--- a/Engine/Source/Editor/UnrealEd/Public/CookPackageSplitter.h
+++ b/Engine/Source/Editor/UnrealEd/Public/CookPackageSplitter.h
@@ -68,7 +68,7 @@ public:
 		bool bCreatedAsMap = false;
 	};
 
-	virtual void GetObjectsToMoveIntoPackage(UPackage* OwnerPackage, UObject* OwnerObject,
+	virtual void GetObjectsToMoveIntoGeneratedPackage(UPackage* OwnerPackage, UObject* OwnerObject,
 		const FGeneratedPackageForPopulate& GeneratedPackage, TArray<UObject*>& OutObjectsToMove) {}
 	/**
 	 * Try to populate a generated package.
@@ -110,7 +110,7 @@ public:
 	 * Return a list of all the objects that will be moved into the Generator package during its save, so the cooker
 	 * can call BeginCacheForCookedPlatformData on them before the move
 	 */
-	virtual void GetObjectsToMoveIntoGenerator(UPackage* OwnerPackage, UObject* OwnerObject,
+	virtual void GetObjectsToMoveIntoGeneratorPackage(UPackage* OwnerPackage, UObject* OwnerObject,
 		const TArray<ICookPackageSplitter::FGeneratedPackageForPreSave>& GeneratedPackages, TArray<UObject*>& OutObjectsToMove) {}
 
 	/**
diff --git a/Engine/Source/Runtime/Engine/Private/WorldPartition/WorldPartitionCookPackageSplitter.cpp b/Engine/Source/Runtime/Engine/Private/WorldPartition/WorldPartitionCookPackageSplitter.cpp
--- a/Engine/Source/Runtime/Engine/Private/WorldPartition/WorldPartitionCookPackageSplitter.cpp
+++ b/Engine/Source/Runtime/Engine/Private/WorldPartition/WorldPartitionCookPackageSplitter.cpp
@@ -117,19 +117,28 @@ TArray<ICookPackageSplitter::FGeneratedPackage> FWorldPartitionCookPackageSplitt
 	return PackagesToGenerate;
 }
 
+void FWorldPartitionCookPackageSplitter::GetObjectsToMoveIntoGeneratedPackage(UPackage* OwnerPackage, UObject* OwnerObject,
+	const FGeneratedPackageForPopulate& GeneratedPackage, TArray<UObject*>& OutObjectsToMove)
+{
+	UWorld* PartitionedWorld = ValidateDataObject(OwnerObject);
+	UWorldPartition* WorldPartition = PartitionedWorld->PersistentLevel->GetWorldPartition();
+	WorldPartition->LoadGeneratedPackageObjectsForCook(GeneratedPackage.RelativePath, OutObjectsToMove);
+}
+
 bool FWorldPartitionCookPackageSplitter::TryPopulatePackage(UPackage* OwnerPackage, UObject* OwnerObject,
 	const FGeneratedPackageForPopulate& GeneratedPackage)
 {
 	UWorld* PartitionedWorld = ValidateDataObject(OwnerObject);
 	UWorldPartition* WorldPartition = PartitionedWorld->PersistentLevel->GetWorldPartition();
+	return WorldPartition->PopulateGeneratedPackageForCook(GeneratedPackage.Package, GeneratedPackage.RelativePath);
+}
 
-	TArray<UObject*> LoadedObjects;
-	bool bResult = WorldPartition->LoadGeneratedPackageObjectsForCook(GeneratedPackage.RelativePath, LoadedObjects);
-	if (bResult)
-	{
-		bResult = WorldPartition->PopulateGeneratedPackageForCook(GeneratedPackage.Package, GeneratedPackage.RelativePath);
-	}
-	return bResult;
+void FWorldPartitionCookPackageSplitter::GetObjectsToMoveIntoGeneratorPackage(UPackage* OwnerPackage, UObject* OwnerObject,
+	const TArray<ICookPackageSplitter::FGeneratedPackageForPreSave>& GeneratedPackages, TArray<UObject*>& OutObjectsToMove)
+{
+	UWorld* PartitionedWorld = ValidateDataObject(OwnerObject);
+	UWorldPartition* WorldPartition = PartitionedWorld->PersistentLevel->GetWorldPartition();
+	WorldPartition->LoadGeneratorPackageObjectsForCook(OutObjectsToMove);
 }
 
 void FWorldPartitionCookPackageSplitter::PreSaveGeneratorPackage(UPackage* OwnerPackage, UObject* OwnerObject,
@@ -137,12 +146,7 @@ void FWorldPartitionCookPackageSplitter::PreSaveGeneratorPackage(UPackage* Owner
 {
 	UWorld* PartitionedWorld = ValidateDataObject(OwnerObject);
 	UWorldPartition* WorldPartition = PartitionedWorld->PersistentLevel->GetWorldPartition();
-
-	TArray<UObject*> LoadedObjects;
-	if (WorldPartition->LoadGeneratorPackageObjectsForCook(LoadedObjects))
-	{
-		WorldPartition->FinalizeGeneratorPackageForCook(GeneratedPackages);
-	}
+	WorldPartition->FinalizeGeneratorPackageForCook(GeneratedPackages);
 }
 
 void FWorldPartitionCookPackageSplitter::OnOwnerReloaded(UPackage* OwnerPackage, UObject* OwnerObject)
diff --git a/Engine/Source/Runtime/Engine/Private/WorldPartition/WorldPartitionCookPackageSplitter.h b/Engine/Source/Runtime/Engine/Private/WorldPartition/WorldPartitionCookPackageSplitter.h
--- a/Engine/Source/Runtime/Engine/Private/WorldPartition/WorldPartitionCookPackageSplitter.h
+++ b/Engine/Source/Runtime/Engine/Private/WorldPartition/WorldPartitionCookPackageSplitter.h
@@ -20,8 +20,12 @@ public:
 	virtual void Teardown(ETeardown Status) override;
 	virtual bool UseInternalReferenceToAvoidGarbageCollect() override { return true; }
 	virtual TArray<ICookPackageSplitter::FGeneratedPackage> GetGenerateList(const UPackage* OwnerPackage, const UObject* OwnerObject) override;
+	virtual void GetObjectsToMoveIntoGeneratedPackage(UPackage* OwnerPackage, UObject* OwnerObject,
+		const FGeneratedPackageForPopulate& GeneratedPackage, TArray<UObject*>& OutObjectsToMove) override;
 	virtual bool TryPopulatePackage(UPackage* OwnerPackage, UObject* OwnerObject,
 		const ICookPackageSplitter::FGeneratedPackageForPopulate& GeneratedPackage) override;
+	virtual void GetObjectsToMoveIntoGeneratorPackage(UPackage* OwnerPackage, UObject* OwnerObject,
+		const TArray<ICookPackageSplitter::FGeneratedPackageForPreSave>& GeneratedPackages, TArray<UObject*>& OutObjectsToMove) override;
 	virtual void PreSaveGeneratorPackage(UPackage* OwnerPackage, UObject* OwnerObject,
 		const TArray<ICookPackageSplitter::FGeneratedPackageForPreSave>& GeneratedPackages) override;
 	virtual void OnOwnerReloaded(UPackage* OwnerPackage, UObject* OwnerObject) override;
diff --git a/Engine/Source/Runtime/Engine/Private/WorldPartition/WorldPartitionRuntimeLevelStreamingCell.cpp b/Engine/Source/Runtime/Engine/Private/WorldPartition/WorldPartitionRuntimeLevelStreamingCell.cpp
--- a/Engine/Source/Runtime/Engine/Private/WorldPartition/WorldPartitionRuntimeLevelStreamingCell.cpp
+++ b/Engine/Source/Runtime/Engine/Private/WorldPartition/WorldPartitionRuntimeLevelStreamingCell.cpp
@@ -138,31 +138,48 @@ UWorldPartitionLevelStreamingDynamic* UWorldPartitionRuntimeLevelStreamingCell::
 
 bool UWorldPartitionRuntimeLevelStreamingCell::LoadCellObjectsForCook(TArray<UObject*>& OutLoadedObjects)
 {
+	check(!bCookHasLoadedCellObjects);
 	if (GetActorCount() > 0)
 	{
 		const bool bLoadAsync = false;
 		UWorld* OuterWorld = GetOuterUWorldPartition()->GetTypedOuter<UWorld>();
 		verify(FWorldPartitionLevelHelper::LoadActors(OuterWorld, nullptr, Packages, CookPackageReferencer, [](bool) {}, bLoadAsync, FLinkerInstancingContext()));
 
+		TArray<UObject*> ObjectsInPackage;
 		for (const FWorldPartitionRuntimeCellObjectMapping& Mapping : Packages)
 		{
 			if (AActor* Actor = FindObject<AActor>(nullptr, *Mapping.LoadedPath.ToString()))
 			{
 				OutLoadedObjects.Add(Actor);
-				CookLoadedActors.Add(Actor);
+				CookLoadedObjects.Add(Actor);
+
+				// Include objects found in the source actor package
+				const bool bIncludeNestedSubobjects = false;
+				ObjectsInPackage.Reset();
+				GetObjectsWithOuter(Actor->GetExternalPackage(), ObjectsInPackage, bIncludeNestedSubobjects);
+				for (UObject* Object : ObjectsInPackage)
+				{
+					if (Object->GetFName() != NAME_PackageMetaData && Object != Actor)
+					{
+						OutLoadedObjects.Add(Object);
+						CookLoadedObjects.Add(Object);
+					}
+				}
 			}
 		}
 	}
+	bCookHasLoadedCellObjects = true;
 	return true;
 }
 
-void UWorldPartitionRuntimeLevelStreamingCell::MoveAlwaysLoadedContentToPersistentLevel()
+bool UWorldPartitionRuntimeLevelStreamingCell::MoveAlwaysLoadedContentToPersistentLevel()
 {
 	check(IsAlwaysLoaded());
+	check(bCookHasLoadedCellObjects);
 	if (GetActorCount() > 0)
 	{
 		// At this point actors should be loaded
-		Algo::ForEach(CookLoadedActors, [](const TWeakObjectPtr<AActor>& Actor) { check(Actor.Get()); });
+		Algo::ForEach(CookLoadedObjects, [](const TWeakObjectPtr<UObject>& Object) { check(Object.Get()); });
 
 		UWorld* OuterWorld = GetOuterUWorldPartition()->GetTypedOuter<UWorld>();
 		FWorldPartitionLevelHelper::MoveExternalActorsToLevel(Packages, OuterWorld->PersistentLevel);
@@ -171,8 +188,10 @@ void UWorldPartitionRuntimeLevelStreamingCell::MoveAlwaysLoadedContentToPersiste
 		Packages.Empty();
 	}
 
-	CookLoadedActors.Empty();
+	CookLoadedObjects.Empty();
 	CookPackageReferencer.RemoveReferences();
+	bCookHasLoadedCellObjects  = false;
+	return true;
 }
 
 // Do all necessary work to prepare cell object for cook.
@@ -199,6 +218,7 @@ bool UWorldPartitionRuntimeLevelStreamingCell::PopulateGeneratedPackageForCook(U
 	{
 		return false;
 	}
+	check(bCookHasLoadedCellObjects);
 
 	if (GetActorCount() > 0)
 	{
@@ -209,7 +229,7 @@ bool UWorldPartitionRuntimeLevelStreamingCell::PopulateGeneratedPackageForCook(U
 		}
 
 		// At this point actors should be loaded
-		Algo::ForEach(CookLoadedActors, [](const TWeakObjectPtr<AActor>& Actor) { check(Actor.Get()); });
+		Algo::ForEach(CookLoadedObjects, [](const TWeakObjectPtr<UObject>& Object) { check(Object.Get()); });
 
 		// Create a level and move loaded actors in it
 		UWorldPartition* WorldPartition = GetOuterUWorldPartition();
@@ -222,8 +242,9 @@ bool UWorldPartitionRuntimeLevelStreamingCell::PopulateGeneratedPackageForCook(U
 		FWorldPartitionLevelHelper::RemapLevelSoftObjectPaths(NewLevel, WorldPartition);
 	}
 
-	CookLoadedActors.Empty();
+	CookLoadedObjects.Empty();
 	CookPackageReferencer.RemoveReferences();
+	bCookHasLoadedCellObjects  = false;
 
 	return true;
 }
diff --git a/Engine/Source/Runtime/Engine/Private/WorldPartition/WorldPartitionRuntimeSpatialHash.cpp b/Engine/Source/Runtime/Engine/Private/WorldPartition/WorldPartitionRuntimeSpatialHash.cpp
--- a/Engine/Source/Runtime/Engine/Private/WorldPartition/WorldPartitionRuntimeSpatialHash.cpp
+++ b/Engine/Source/Runtime/Engine/Private/WorldPartition/WorldPartitionRuntimeSpatialHash.cpp
@@ -1225,7 +1225,10 @@ bool UWorldPartitionRuntimeSpatialHash::FinalizeGeneratorPackageForCook(const TA
 	for (UWorldPartitionRuntimeCell* Cell : GetAlwaysLoadedCells())
 	{
 		check(Cell->IsAlwaysLoaded());
-		Cell->MoveAlwaysLoadedContentToPersistentLevel();
+		if (!Cell->MoveAlwaysLoadedContentToPersistentLevel())
+		{
+			return false;
+		}
 	}
 	
 	for (const ICookPackageSplitter::FGeneratedPackageForPreSave& GeneratedPackage : InGeneratedPackages)
diff --git a/Engine/Source/Runtime/Engine/Public/WorldPartition/WorldPartitionRuntimeCell.h b/Engine/Source/Runtime/Engine/Public/WorldPartition/WorldPartitionRuntimeCell.h
--- a/Engine/Source/Runtime/Engine/Public/WorldPartition/WorldPartitionRuntimeCell.h
+++ b/Engine/Source/Runtime/Engine/Public/WorldPartition/WorldPartitionRuntimeCell.h
@@ -204,7 +204,7 @@ class UWorldPartitionRuntimeCell : public UObject, public IWorldPartitionCell
 	virtual bool PrepareCellForCook(UPackage* InPackage) { return false; }
 	virtual bool LoadCellObjectsForCook(TArray<UObject*>& OutLoadedObjects) PURE_VIRTUAL(UWorldPartitionRuntimeCell::LoadCellObjectsForCook, return false;);
 	virtual bool PopulateGeneratedPackageForCook(UPackage* InPackage) PURE_VIRTUAL(UWorldPartitionRuntimeCell::PopulateGeneratedPackageForCook, return false;);
-	virtual void MoveAlwaysLoadedContentToPersistentLevel() PURE_VIRTUAL(UWorldPartitionRuntimeCell::MoveAlwaysLoadedContentToPersistentLevel);
+	virtual bool MoveAlwaysLoadedContentToPersistentLevel() PURE_VIRTUAL(UWorldPartitionRuntimeCell::MoveAlwaysLoadedContentToPersistentLevel, return false;);
 	virtual FString GetPackageNameToCreate() const PURE_VIRTUAL(UWorldPartitionRuntimeCell::GetPackageNameToCreate, return FString(""););
 
 	void SetIsHLOD(bool bInIsHLOD) { bIsHLOD = bInIsHLOD; }
diff --git a/Engine/Source/Runtime/Engine/Public/WorldPartition/WorldPartitionRuntimeLevelStreamingCell.h b/Engine/Source/Runtime/Engine/Public/WorldPartition/WorldPartitionRuntimeLevelStreamingCell.h
--- a/Engine/Source/Runtime/Engine/Public/WorldPartition/WorldPartitionRuntimeLevelStreamingCell.h
+++ b/Engine/Source/Runtime/Engine/Public/WorldPartition/WorldPartitionRuntimeLevelStreamingCell.h
@@ -42,6 +42,7 @@ class UWorldPartitionRuntimeLevelStreamingCell : public UWorldPartitionRuntimeSp
 	virtual bool PrepareCellForCook(UPackage* InPackage) override;
 	virtual bool LoadCellObjectsForCook(TArray<UObject*>& OutLoadedObjects) override;
 	virtual bool PopulateGeneratedPackageForCook(UPackage* InPackage) override;
+	virtual bool MoveAlwaysLoadedContentToPersistentLevel() override;
 	virtual FString GetPackageNameToCreate() const override;
 	//~End UWorldPartitionRuntimeCell Interface
 
@@ -58,11 +59,11 @@ private:
 	class UWorldPartitionLevelStreamingDynamic* GetOrCreateLevelStreaming() const;
 
 #if WITH_EDITOR
-	void MoveAlwaysLoadedContentToPersistentLevel();
 	class UWorldPartitionLevelStreamingDynamic* CreateLevelStreaming(const FString& InPackageName = FString()) const;
 
 	FWorldPartitionLevelHelper::FPackageReferencer CookPackageReferencer;
-	TArray<TWeakObjectPtr<AActor>> CookLoadedActors;
+	TArray<TWeakObjectPtr<UObject>> CookLoadedObjects;
+	bool bCookHasLoadedCellObjects = false;
 #endif
 
 #if WITH_EDITORONLY_DATA
