Removed the boolean bShouldMinimizeExtraWindowsInSelectedViewport and used existing boolean bShouldMinimizeEditorOnNonVRPIE.

PR #9196: Option to minimize windows when play in selected viewport. (Contributed by [REDACTED_PERSON])

diff --git a/Engine/Source/Editor/UnrealEd/Classes/Editor/EditorEngine.h b/Engine/Source/Editor/UnrealEd/Classes/Editor/EditorEngine.h
--- a/Engine/Source/Editor/UnrealEd/Classes/Editor/EditorEngine.h
+++ b/Engine/Source/Editor/UnrealEd/Classes/Editor/EditorEngine.h
@@ -3248,6 +3248,9 @@ private:
 	/** Delegate handle for game viewport close requests in PIE sessions. */
 	FDelegateHandle ViewportCloseRequestedDelegateHandle;
 
+	/** Minimized Windows during PIE */
+	TArray<TWeakPtr<SWindow>> MinimizedWindowsDuringPIE;
+
 public:
 	/**
 	 * Get a Subsystem of specified type
diff --git a/Engine/Source/Editor/UnrealEd/Private/PlayLevel.cpp b/Engine/Source/Editor/UnrealEd/Private/PlayLevel.cpp
--- a/Engine/Source/Editor/UnrealEd/Private/PlayLevel.cpp
+++ b/Engine/Source/Editor/UnrealEd/Private/PlayLevel.cpp
@@ -180,6 +180,17 @@ void UEditorEngine::EndPlayMap()
 
 	FEditorDelegates::PrePIEEnded.Broadcast( bIsSimulatingInEditor );
 
+	// Restore optionally minimized windows.
+	// We always restore no matter what the setting is since it could be toggled during PIE.
+	for (TWeakPtr<SWindow>& Window : MinimizedWindowsDuringPIE)
+	{
+		if (Window.IsValid())
+		{
+			Window.Pin()->Restore();
+		}
+	}
+	MinimizedWindowsDuringPIE.Empty();
+		
 	// Clean up Soft Object Path remaps
 	FSoftObjectPath::ClearPIEPackageNames();
 
@@ -2743,6 +2754,41 @@ void UEditorEngine::StartPlayInEditorSession(FRequestPlaySessionParams& InReques
 	PlayInEditorSessionInfo->bUsingOnlinePlatform = bUseOnlineSubsystemForLogin;
 	PlayInEditorSessionInfo->bAnyBlueprintErrors = ErroredBlueprints.Num() > 0;
 
+	// Optionally minimize all extra windows for precious framerate.
+	if (InRequestParams.EditorPlaySettings->bShouldMinimizeEditorOnNonVRPIE
+		&& !bIsSimulatingInEditor
+		&& InRequestParams.DestinationSlateViewport.IsSet())
+	{
+		TArray<TSharedRef<SWindow>> Windows;
+		FSlateApplication::Get().GetAllVisibleWindowsOrdered(Windows);
+		if (Windows.Num() > 0)
+		{
+			TSharedRef<SWindow> RootWindow = Windows[0];
+			if (TSharedPtr<IAssetViewport> DestinationViewport = InRequestParams.DestinationSlateViewport.GetValue().Pin())
+			{
+				TSharedPtr<SWindow> DestinationWindow = FSlateApplication::Get().FindWidgetWindow(DestinationViewport->AsWidget());
+			
+				for (TSharedRef<SWindow>& Window : Windows)
+				{
+					// Don't minimize the root,
+					// Don't minimize any free floating windows like other PIE windows.
+					// Don't minimize the viewport that is being targeted.
+					// Don't minimize already minimized windows.
+					if (Window == RootWindow
+						|| Window == DestinationWindow
+						|| Window->GetParentWindow() != RootWindow
+						|| Window->IsWindowMinimized())
+					{
+						continue;
+					}
+					
+					Window->Minimize();
+					MinimizedWindowsDuringPIE.Add(Window);
+				}
+			}
+		}
+	}
+	
 	// Now that we've gotten all of the editor house-keeping out of the way we can finally
 	// start creating world instances and multi player clients!
 	{
