Add support for previewing individual slices of a Texture2DArray in the texture editor.

diff --git a/Engine/Plugins/VirtualProduction/CameraCalibration/Source/CameraCalibrationEditor/Private/AssetEditor/Viewport/SimulcamEditorViewportClient.cpp b/Engine/Plugins/VirtualProduction/CameraCalibration/Source/CameraCalibrationEditor/Private/AssetEditor/Viewport/SimulcamEditorViewportClient.cpp
--- a/Engine/Plugins/VirtualProduction/CameraCalibration/Source/CameraCalibrationEditor/Private/AssetEditor/Viewport/SimulcamEditorViewportClient.cpp
+++ b/Engine/Plugins/VirtualProduction/CameraCalibration/Source/CameraCalibrationEditor/Private/AssetEditor/Viewport/SimulcamEditorViewportClient.cpp
@@ -50,9 +50,10 @@ void FSimulcamEditorViewportClient::Draw(FViewport* Viewport, FCanvas* Canvas)
 
 	constexpr float MipLevel = 0;
 	constexpr float LayerIndex = 0;
+	constexpr float SliceIndex = 0;
 
 	TRefCountPtr<FBatchedElementParameters> BatchedElementParameters;
-	BatchedElementParameters = new FBatchedElementTexture2DPreviewParameters(MipLevel, LayerIndex, false, false, false, false, false);
+	BatchedElementParameters = new FBatchedElementTexture2DPreviewParameters(MipLevel, LayerIndex, SliceIndex, false, false, false, false, false);
 
 	if (Texture->GetResource() != nullptr)
 	{
diff --git a/Engine/Shaders/Private/SimpleElementTexture2DPreviewPixelShader.usf b/Engine/Shaders/Private/SimpleElementTexture2DPreviewPixelShader.usf
--- a/Engine/Shaders/Private/SimpleElementTexture2DPreviewPixelShader.usf
+++ b/Engine/Shaders/Private/SimpleElementTexture2DPreviewPixelShader.usf
@@ -50,6 +50,7 @@ float4 PackedParams;
 
 #if TEXTURE_ARRAY
 float NumSlices;
+float SliceIndex;
 #endif
 
 void Main(
@@ -72,10 +73,13 @@ void Main(
 	float4 Sample;
 
 #if TEXTURE_ARRAY
-	// Divide up the texture array into horizontal slices.
-	float SliceSize = 1.0 / NumSlices;
-	float SliceIndex = floor(TextureCoordinate.y / SliceSize);
-	float3 SampleCoordinates = float3(TextureCoordinate.xy, SliceIndex);
+	float CurrentSliceIndex = SliceIndex;
+	if (CurrentSliceIndex < 0.0)
+	{
+		// If the slice index is not specified, display a portion of each slice from the texture array.
+		CurrentSliceIndex = floor(TextureCoordinate.y * NumSlices);
+	}
+	float3 SampleCoordinates = float3(TextureCoordinate.xy, CurrentSliceIndex);
 
 	if( MipLevel >= 0.0f )
 	{
diff --git a/Engine/Source/Editor/TextureEditor/Private/Models/TextureEditorViewportClient.cpp b/Engine/Source/Editor/TextureEditor/Private/Models/TextureEditorViewportClient.cpp
--- a/Engine/Source/Editor/TextureEditor/Private/Models/TextureEditorViewportClient.cpp
+++ b/Engine/Source/Editor/TextureEditor/Private/Models/TextureEditorViewportClient.cpp
@@ -101,6 +101,7 @@ void FTextureEditorViewportClient::Draw(FViewport* Viewport, FCanvas* Canvas)
 	TextureEditorPinned->CalculateTextureDimensions(Width, Height, Depth, ArraySize);
 	const float MipLevel = (float)TextureEditorPinned->GetMipLevel();
 	const float LayerIndex = (float)TextureEditorPinned->GetLayer();
+	const float SliceIndex = (float)TextureEditorPinned->GetSlice();
 
 	bool bIsVirtualTexture = false;
 
@@ -138,26 +139,26 @@ void FTextureEditorViewportClient::Draw(FViewport* Viewport, FCanvas* Canvas)
 			bool bIsSingleChannel = Texture2D->CompressionSettings == TC_Grayscale || Texture2D->CompressionSettings == TC_Alpha;
 			bool bSingleVTPhysicalSpace = Texture2D->IsVirtualTexturedWithSinglePhysicalSpace();
 			bIsVirtualTexture = Texture2D->IsCurrentlyVirtualTextured();
-			BatchedElementParameters = new FBatchedElementTexture2DPreviewParameters(MipLevel, LayerIndex, bIsNormalMap, bIsSingleChannel, bSingleVTPhysicalSpace, bIsVirtualTexture, false);
+			BatchedElementParameters = new FBatchedElementTexture2DPreviewParameters(MipLevel, LayerIndex, SliceIndex, bIsNormalMap, bIsSingleChannel, bSingleVTPhysicalSpace, bIsVirtualTexture, false);
 		}
 		else if (Texture2DArray) 
 		{
 			bool bIsNormalMap = Texture2DArray->IsNormalMap();
 			bool bIsSingleChannel = Texture2DArray->CompressionSettings == TC_Grayscale || Texture2DArray->CompressionSettings == TC_Alpha;
-			BatchedElementParameters = new FBatchedElementTexture2DPreviewParameters(MipLevel, LayerIndex, bIsNormalMap, bIsSingleChannel, false, false, true);
+			BatchedElementParameters = new FBatchedElementTexture2DPreviewParameters(MipLevel, LayerIndex, SliceIndex, bIsNormalMap, bIsSingleChannel, false, false, true);
 		}
 		else if (TextureRT2D)
 		{
-			BatchedElementParameters = new FBatchedElementTexture2DPreviewParameters(MipLevel, LayerIndex, false, false, false, false, false);
+			BatchedElementParameters = new FBatchedElementTexture2DPreviewParameters(MipLevel, LayerIndex, SliceIndex, false, false, false, false, false);
 		}
 		else if (TextureRT2DArray)
 		{
-			BatchedElementParameters = new FBatchedElementTexture2DPreviewParameters(MipLevel, LayerIndex, false, false, false, false, true);
+			BatchedElementParameters = new FBatchedElementTexture2DPreviewParameters(MipLevel, LayerIndex, SliceIndex, false, false, false, false, true);
 		}
 		else
 		{
 			// Default to treating any UTexture derivative as a 2D texture resource
-			BatchedElementParameters = new FBatchedElementTexture2DPreviewParameters(MipLevel, LayerIndex, false, false, false, false, false);
+			BatchedElementParameters = new FBatchedElementTexture2DPreviewParameters(MipLevel, LayerIndex, SliceIndex, false, false, false, false, false);
 		}
 	}
 
diff --git a/Engine/Source/Editor/TextureEditor/Private/TextureEditorToolkit.cpp b/Engine/Source/Editor/TextureEditor/Private/TextureEditorToolkit.cpp
--- a/Engine/Source/Editor/TextureEditor/Private/TextureEditorToolkit.cpp
+++ b/Engine/Source/Editor/TextureEditor/Private/TextureEditorToolkit.cpp
@@ -242,6 +242,9 @@ void FTextureEditorToolkit::InitTextureEditor( const EToolkitMode::Type Mode, co
 
 	SpecifiedLayer = 0;
 
+	SpecifiedSlice = 0;
+	bUseSpecifiedSlice = false;
+
 	SavedCompressionSetting = false;
 
 	// Start at whatever the last used zoom mode was
@@ -457,6 +460,11 @@ int32 FTextureEditorToolkit::GetLayer() const
 	return SpecifiedLayer;
 }
 
+int32 FTextureEditorToolkit::GetSlice() const
+{
+	return GetUseSpecifiedSlice() ? HandleSliceEntryBoxValue().Get(-1) : -1;
+}
+
 UTexture* FTextureEditorToolkit::GetTexture( ) const
 {
 	return Texture;
@@ -486,6 +494,10 @@ bool FTextureEditorToolkit::GetUseSpecifiedMip( ) const
 	return false;
 }
 
+bool FTextureEditorToolkit::GetUseSpecifiedSlice() const
+{
+	return HandleSliceCheckBoxIsEnabled() && bUseSpecifiedSlice;
+}
 
 double FTextureEditorToolkit::GetCustomZoomLevel( ) const
 {
@@ -1692,6 +1704,7 @@ void FTextureEditorToolkit::FillToolbar(FToolBarBuilder& ToolbarBuilder)
 	TSharedRef<SWidget> ChannelControl = MakeChannelControlWidget();
 	TSharedRef<SWidget> LODControl = MakeLODControlWidget();
 	TSharedRef<SWidget> LayerControl = MakeLayerControlWidget();
+	TSharedRef<SWidget> SliceControl = MakeSliceControlWidget();
 	TSharedRef<SWidget> ExposureControl = MakeExposureContolWidget();
 	TSharedPtr<SWidget> OptionalOpacityControl = IsVolumeTexture() ? TSharedPtr<SWidget>(MakeOpacityControlWidget()) : nullptr;
 	TSharedRef<SWidget> ZoomControl = MakeZoomControlWidget();
@@ -1728,6 +1741,15 @@ void FTextureEditorToolkit::FillToolbar(FToolBarBuilder& ToolbarBuilder)
 			ToolbarBuilder.EndSection();
 		}
 
+		if (HasSlices())
+		{
+			ToolbarBuilder.BeginSection("Slices");
+			{
+				ToolbarBuilder.AddWidget(SliceControl);
+			}
+			ToolbarBuilder.EndSection();
+		}
+
 		if (OptionalOpacityControl.IsValid())
 		{
 			ToolbarBuilder.BeginSection("Opacity");
@@ -1875,6 +1897,11 @@ TOptional<int32> FTextureEditorToolkit::GetMaxLayer() const
 	return FMath::Max(Texture->Source.GetNumLayers() - 1, 1);
 }
 
+TOptional<int32> FTextureEditorToolkit::GetMaxSlice() const
+{
+	return FMath::Max(0, GetNumSlices() - 1);
+}
+
 bool FTextureEditorToolkit::IsCubeTexture( ) const
 {
 	return (Texture->IsA(UTextureCube::StaticClass()) || Texture->IsA(UTextureCubeArray::StaticClass()) || Texture->IsA(UTextureRenderTargetCube::StaticClass()));
@@ -2159,6 +2186,48 @@ bool FTextureEditorToolkit::HasLayers() const
 	return Texture->Source.GetNumLayers() > 1;
 }
 
+bool FTextureEditorToolkit::HasSlices() const
+{
+	// currently only supports Texture2DArray, but can be extended for other texture types which are using multiple slices
+	// slice selection should be supported even for a texture array with less than two elements, because array elements can be added dynamically
+	return Texture->IsA(UTexture2DArray::StaticClass());
+}
+
+int32 FTextureEditorToolkit::GetNumSlices() const
+{
+	if (Texture->IsA(UTexture2DArray::StaticClass()))
+	{
+		return Cast<UTexture2DArray>(Texture)->GetArraySize();
+	}
+	return 1;
+}
+
+bool FTextureEditorToolkit::HandleSliceCheckBoxIsEnabled() const
+{
+	return GetNumSlices() > 1;
+}
+
+ECheckBoxState FTextureEditorToolkit::HandleSliceCheckBoxIsChecked() const
+{
+	return GetUseSpecifiedSlice() ? ECheckBoxState::Checked : ECheckBoxState::Unchecked;
+}
+
+void FTextureEditorToolkit::HandleSliceCheckBoxCheckedStateChanged(ECheckBoxState InNewState)
+{
+	bUseSpecifiedSlice = InNewState == ECheckBoxState::Checked;
+}
+
+TOptional<int32> FTextureEditorToolkit::HandleSliceEntryBoxValue() const
+{
+	return FMath::Clamp<int32>(SpecifiedSlice, 0, GetMaxSlice().Get(0));
+}
+
+void FTextureEditorToolkit::HandleSliceEntryBoxChanged(int32 NewSlice)
+{
+	SpecifiedSlice = FMath::Clamp<int32>(NewSlice, 0, GetMaxSlice().Get(0));
+	PopulateQuickInfo();
+}
+
 bool FTextureEditorToolkit::HandleReimportActionCanExecute( ) const
 {
 	if (Texture->IsA<ULightMapTexture2D>() || Texture->IsA<UShadowMapTexture2D>() || Texture->IsA<UTexture2DDynamic>() || Texture->IsA<UTextureRenderTarget>() || Texture->IsA<UCurveLinearColorAtlas>())
@@ -2659,6 +2728,47 @@ TSharedRef<SWidget> FTextureEditorToolkit::MakeLayerControlWidget()
 	return LayerControl;
 }
 
+TSharedRef<SWidget> FTextureEditorToolkit::MakeSliceControlWidget()
+{
+	TSharedRef<SWidget> SliceControl = SNew(SBox)
+		.WidthOverride(212.0f)
+		[
+			SNew(SHorizontalBox)
+			.IsEnabled(this, &FTextureEditorToolkit::HandleSliceCheckBoxIsEnabled)
+			+ SHorizontalBox::Slot()
+			.Padding(4.0f, 0.0f, 4.0f, 0.0f)
+			.VAlign(VAlign_Center)
+			.AutoWidth()
+			[
+				SNew(STextBlock)
+				.Text(NSLOCTEXT("TextureEditor", "Slice", "Slice"))
+			]
+			+ SHorizontalBox::Slot()
+			.Padding(4.0f, 0.0f, 2.0f, 0.0f)
+			.VAlign(VAlign_Center)
+			.AutoWidth()
+			[
+				SNew(SCheckBox)
+				.IsChecked(this, &FTextureEditorToolkit::HandleSliceCheckBoxIsChecked)
+				.OnCheckStateChanged(this, &FTextureEditorToolkit::HandleSliceCheckBoxCheckedStateChanged)
+			]
+			+ SHorizontalBox::Slot()
+			.Padding(0.0f, 0.0f, 4.0f, 0.0f)
+			.VAlign(VAlign_Center)
+			[
+				SNew(SNumericEntryBox<int32>)
+				.IsEnabled(this, &FTextureEditorToolkit::GetUseSpecifiedSlice)
+				.AllowSpin(true)
+				.MinSliderValue(0)
+				.MaxSliderValue(this, &FTextureEditorToolkit::GetMaxSlice)
+				.Value(this, &FTextureEditorToolkit::HandleSliceEntryBoxValue)
+				.OnValueChanged(this, &FTextureEditorToolkit::HandleSliceEntryBoxChanged)
+			]
+		];
+
+	return SliceControl;
+}
+
 TSharedRef<SWidget> FTextureEditorToolkit::MakeExposureContolWidget()
 {
 	TSharedRef<SWidget> ExposureControl = SNew(SBox)
diff --git a/Engine/Source/Editor/TextureEditor/Private/TextureEditorToolkit.h b/Engine/Source/Editor/TextureEditor/Private/TextureEditorToolkit.h
--- a/Engine/Source/Editor/TextureEditor/Private/TextureEditorToolkit.h
+++ b/Engine/Source/Editor/TextureEditor/Private/TextureEditorToolkit.h
@@ -65,9 +65,11 @@ public:
 	virtual ESimpleElementBlendMode GetColourChannelBlendMode( ) const override;
 	virtual int32 GetMipLevel( ) const override;
 	virtual int32 GetLayer() const override;
+	virtual int32 GetSlice() const override;
 	virtual UTexture* GetTexture( ) const override;
 	virtual bool HasValidTextureResource( ) const override;
 	virtual bool GetUseSpecifiedMip( ) const override;
+	virtual bool GetUseSpecifiedSlice() const override;
 	virtual double GetCustomZoomLevel( ) const override;
 	virtual void SetCustomZoomLevel(double ZoomValue) override;
 	virtual void PopulateQuickInfo( ) override;
@@ -145,8 +147,16 @@ protected:
 	 */
 	EPixelFormat GetPixelFormat() const;
 
+	/**
+	 * Gets the highest layer index for this texture.
+	 */
 	TOptional<int32> GetMaxLayer() const;
 
+	/**
+	 * Gets the highest slice index for this texture.
+	 */
+	TOptional<int32> GetMaxSlice() const;
+
 	/**
 	 * Checks whether the texture being edited is a cube map texture.
 	 */
@@ -230,6 +240,27 @@ private:
 
 	bool HasLayers() const;
 
+	// Returns true if the current texture supports slice selection.
+	bool HasSlices() const;
+
+	// Gets the number of slices that can be selected for this texture.
+	int32 GetNumSlices() const;
+
+	// Callback for determining whether the slice check box is enabled.
+	bool HandleSliceCheckBoxIsEnabled() const;
+
+	// Callback for getting the checked state of the slice check box.
+	ECheckBoxState HandleSliceCheckBoxIsChecked() const;
+
+	// Callback for changing the checked state of the slice check box.
+	void HandleSliceCheckBoxCheckedStateChanged(ECheckBoxState InNewState);
+
+	// Callback for getting the value of the slice index entry box.
+	TOptional<int32> HandleSliceEntryBoxValue() const;
+
+	// Callback for changing the value of the slice index entry box.
+	void HandleSliceEntryBoxChanged(int32 Slice);
+
 	// Callback for determining whether the Reimport action can execute.
 	bool HandleReimportActionCanExecute( ) const;
 
@@ -310,6 +341,7 @@ private:
 	TSharedRef<SWidget> MakeChannelControlWidget();
 	TSharedRef<SWidget> MakeLODControlWidget();
 	TSharedRef<SWidget> MakeLayerControlWidget();
+	TSharedRef<SWidget> MakeSliceControlWidget();
 	TSharedRef<SWidget> MakeExposureContolWidget();
 	TSharedRef<SWidget> MakeOpacityControlWidget();
 	TSharedRef<SWidget> MakeZoomControlWidget();
@@ -381,11 +413,19 @@ private:
 
 	/** Which mip level should be shown */
 	int32 SpecifiedMipLevel;
-	/* When true, the specified mip value is used. Top mip is used when false.*/
+
+	/** When true, the specified mip value is used. Top mip is used when false.*/
 	bool bUseSpecifiedMipLevel;
 
+	/** Which layer should be shown */
 	int32 SpecifiedLayer;
 
+	/** Which slice should be shown */
+	int32 SpecifiedSlice;
+
+	/** When true, the specified slice index value is used, otherwise the value of -1 is used.*/
+	bool bUseSpecifiedSlice;
+
 	/** During re-import, cache this setting so it can be restored if necessary */
 	bool SavedCompressionSetting;
 
diff --git a/Engine/Source/Editor/TextureEditor/Public/Interfaces/ITextureEditorToolkit.h b/Engine/Source/Editor/TextureEditor/Public/Interfaces/ITextureEditorToolkit.h
--- a/Engine/Source/Editor/TextureEditor/Public/Interfaces/ITextureEditorToolkit.h
+++ b/Engine/Source/Editor/TextureEditor/Public/Interfaces/ITextureEditorToolkit.h
@@ -32,8 +32,10 @@ public:
 	/** Accessors */ 
 	virtual int32 GetMipLevel() const = 0;
 	virtual int32 GetLayer() const = 0;
+	virtual int32 GetSlice() const = 0;
 	virtual ESimpleElementBlendMode GetColourChannelBlendMode() const = 0;
 	virtual bool GetUseSpecifiedMip() const = 0;
+	virtual bool GetUseSpecifiedSlice() const = 0;
 	virtual double GetCustomZoomLevel() const = 0;
 	virtual void SetCustomZoomLevel( double ZoomValue ) = 0;
 	virtual void ZoomIn() = 0;
diff --git a/Engine/Source/Editor/UnrealEd/Private/Texture2DPreview.cpp b/Engine/Source/Editor/UnrealEd/Private/Texture2DPreview.cpp
--- a/Engine/Source/Editor/UnrealEd/Private/Texture2DPreview.cpp
+++ b/Engine/Source/Editor/UnrealEd/Private/Texture2DPreview.cpp
@@ -47,6 +47,7 @@ public:
 		ColorWeights.Bind(Initializer.ParameterMap,TEXT("ColorWeights"));
 		PackedParameters.Bind(Initializer.ParameterMap,TEXT("PackedParams"));
 		NumSlices.Bind(Initializer.ParameterMap, TEXT("NumSlices"));
+		SliceIndex.Bind(Initializer.ParameterMap, TEXT("SliceIndex"));
 	}
 	FSimpleElementTexture2DPreviewPS() {}
 
@@ -61,57 +62,58 @@ public:
 		return IsFeatureLevelSupported(Parameters.Platform, ERHIFeatureLevel::SM5) && !IsConsolePlatform(Parameters.Platform);
 	}
 
-	void SetParameters(FRHICommandList& RHICmdList, const FTexture* TextureValue, const FMatrix44f& ColorWeightsValue, float GammaValue, float MipLevel, float LayerIndex, bool bIsNormalMap, bool bIsSingleVTPhysicalSpace, bool bIsVirtualTexture, bool bIsTextureArray)
+	void SetParameters(FRHICommandList& InRHICmdList, const FTexture* InTextureValue, const FMatrix44f& InColorWeightsValue, float InGammaValue, float InMipLevel, float InLayerIndex, float InSliceIndex, bool bInIsNormalMap, bool bInIsSingleVTPhysicalSpace, bool bInIsVirtualTexture, bool bInIsTextureArray)
 	{
-		FRHIPixelShader* ShaderRHI = RHICmdList.GetBoundPixelShader();
-		if (bIsVirtualTexture)
+		FRHIPixelShader* ShaderRHI = InRHICmdList.GetBoundPixelShader();
+		if (bInIsVirtualTexture)
 		{
-			FVirtualTexture2DResource* VirtualTextureValue = (FVirtualTexture2DResource*)TextureValue;
+			FVirtualTexture2DResource* VirtualTextureValue = (FVirtualTexture2DResource*)InTextureValue;
 			IAllocatedVirtualTexture* AllocatedVT = VirtualTextureValue->AcquireAllocatedVT();
 	
-			FRHIShaderResourceView* PhysicalView = AllocatedVT->GetPhysicalTextureSRV((uint32)LayerIndex, TextureValue->bSRGB);
-			SetSRVParameter(RHICmdList, ShaderRHI, InTexture, PhysicalView);
-			SetSamplerParameter(RHICmdList, ShaderRHI, InTextureSampler, VirtualTextureValue->SamplerStateRHI);
+			FRHIShaderResourceView* PhysicalView = AllocatedVT->GetPhysicalTextureSRV((uint32)InLayerIndex, InTextureValue->bSRGB);
+			SetSRVParameter(InRHICmdList, ShaderRHI, InTexture, PhysicalView);
+			SetSamplerParameter(InRHICmdList, ShaderRHI, InTextureSampler, VirtualTextureValue->SamplerStateRHI);
 
-			SetTextureParameter(RHICmdList, ShaderRHI, InPageTableTexture0, AllocatedVT->GetPageTableTexture(0u));
+			SetTextureParameter(InRHICmdList, ShaderRHI, InPageTableTexture0, AllocatedVT->GetPageTableTexture(0u));
 			if (AllocatedVT->GetNumPageTableTextures() > 1u)
 			{
-				SetTextureParameter(RHICmdList, ShaderRHI, InPageTableTexture1, AllocatedVT->GetPageTableTexture(1u));
+				SetTextureParameter(InRHICmdList, ShaderRHI, InPageTableTexture1, AllocatedVT->GetPageTableTexture(1u));
 			}
 			else
 			{
-				SetTextureParameter(RHICmdList, ShaderRHI, InPageTableTexture1, GBlackTexture->TextureRHI);
+				SetTextureParameter(InRHICmdList, ShaderRHI, InPageTableTexture1, GBlackTexture->TextureRHI);
 			}
 
 			FUintVector4 PageTableUniform[2];
 			FUintVector4 Uniform;
 
 			AllocatedVT->GetPackedPageTableUniform(PageTableUniform);
-			AllocatedVT->GetPackedUniform(&Uniform, (uint32)LayerIndex);
+			AllocatedVT->GetPackedUniform(&Uniform, (uint32)InLayerIndex);
 
-			SetShaderValueArray(RHICmdList, ShaderRHI, VTPackedPageTableUniform, PageTableUniform, UE_ARRAY_COUNT(PageTableUniform));
-			SetShaderValue(RHICmdList, ShaderRHI, VTPackedUniform, Uniform);
+			SetShaderValueArray(InRHICmdList, ShaderRHI, VTPackedPageTableUniform, PageTableUniform, UE_ARRAY_COUNT(PageTableUniform));
+			SetShaderValue(InRHICmdList, ShaderRHI, VTPackedUniform, Uniform);
 		}
 		else
 		{
-			SetTextureParameter(RHICmdList, ShaderRHI, InPageTableTexture0, GBlackTexture->TextureRHI);
-			SetTextureParameter(RHICmdList, ShaderRHI, InPageTableTexture1, GBlackTexture->TextureRHI);
-			SetTextureParameter(RHICmdList, ShaderRHI, InTexture, InTextureSampler, TextureValue);
+			SetTextureParameter(InRHICmdList, ShaderRHI, InPageTableTexture0, GBlackTexture->TextureRHI);
+			SetTextureParameter(InRHICmdList, ShaderRHI, InPageTableTexture1, GBlackTexture->TextureRHI);
+			SetTextureParameter(InRHICmdList, ShaderRHI, InTexture, InTextureSampler, InTextureValue);
 		}
 		
-		SetShaderValue(RHICmdList, ShaderRHI,ColorWeights,ColorWeightsValue);
-		FVector4f PackedParametersValue(GammaValue, MipLevel, bIsNormalMap ? 1.0 : -1.0f, bIsSingleVTPhysicalSpace ? 0 : LayerIndex);
-		SetShaderValue(RHICmdList, ShaderRHI, PackedParameters, PackedParametersValue);
+		SetShaderValue(InRHICmdList, ShaderRHI,ColorWeights,InColorWeightsValue);
+		FVector4f PackedParametersValue(InGammaValue, InMipLevel, bInIsNormalMap ? 1.0 : -1.0f, bInIsSingleVTPhysicalSpace ? 0 : InLayerIndex);
+		SetShaderValue(InRHICmdList, ShaderRHI, PackedParameters, PackedParametersValue);
 
-		// Store slice count for texture array
-		if (bIsTextureArray)
+		// Store slice count and selected slice index for texture array
+		if (bInIsTextureArray)
 		{
-			const float NumSlicesData = TextureValue ? TextureValue->GetSizeZ() : 1;
-			SetShaderValue(RHICmdList, ShaderRHI, NumSlices, NumSlicesData);
+			const float NumSlicesData = InTextureValue ? InTextureValue->GetSizeZ() : 1;
+			SetShaderValue(InRHICmdList, ShaderRHI, NumSlices, NumSlicesData);
+			SetShaderValue(InRHICmdList, ShaderRHI, SliceIndex, InSliceIndex);
 		}
 
-		SetShaderValue(RHICmdList, ShaderRHI, TextureComponentReplicate, (TextureValue && TextureValue->bGreyScaleFormat) ? FLinearColor(1,0,0,0) : FLinearColor(0,0,0,0));
-		SetShaderValue(RHICmdList, ShaderRHI, TextureComponentReplicateAlpha, (TextureValue && TextureValue->bGreyScaleFormat) ? FLinearColor(1,0,0,0) : FLinearColor(0,0,0,1));
+		SetShaderValue(InRHICmdList, ShaderRHI, TextureComponentReplicate, (InTextureValue && InTextureValue->bGreyScaleFormat) ? FLinearColor(1,0,0,0) : FLinearColor(0,0,0,0));
+		SetShaderValue(InRHICmdList, ShaderRHI, TextureComponentReplicateAlpha, (InTextureValue && InTextureValue->bGreyScaleFormat) ? FLinearColor(1,0,0,0) : FLinearColor(0,0,0,1));
 	}
 
 private:
@@ -126,6 +128,7 @@ private:
 	LAYOUT_FIELD(FShaderParameter, ColorWeights);
 	LAYOUT_FIELD(FShaderParameter, PackedParameters);
 	LAYOUT_FIELD(FShaderParameter, NumSlices);
+	LAYOUT_FIELD(FShaderParameter, SliceIndex);
 };
 
 IMPLEMENT_GLOBAL_SHADER(FSimpleElementTexture2DPreviewPS, "/Engine/Private/SimpleElementTexture2DPreviewPixelShader.usf", "Main", SF_Pixel);
@@ -161,5 +164,5 @@ void FBatchedElementTexture2DPreviewParameters::BindShaders(
 	SetGraphicsPipelineState(RHICmdList, GraphicsPSOInit, 0);
 
 	VertexShader->SetParameters(RHICmdList, InTransform);
-	PixelShader->SetParameters(RHICmdList, Texture, FMatrix44f(ColorWeights), InGamma, MipLevel, LayerIndex, bIsNormalMap, bIsSingleVTPhysicalSpace, bIsVirtualTexture, bIsTextureArray);
+	PixelShader->SetParameters(RHICmdList, Texture, FMatrix44f(ColorWeights), InGamma, MipLevel, LayerIndex, SliceIndex, bIsNormalMap, bIsSingleVTPhysicalSpace, bIsVirtualTexture, bIsTextureArray);
 }
diff --git a/Engine/Source/Editor/UnrealEd/Private/ThumbnailRendering/TextureThumbnailRenderer.cpp b/Engine/Source/Editor/UnrealEd/Private/ThumbnailRendering/TextureThumbnailRenderer.cpp
--- a/Engine/Source/Editor/UnrealEd/Private/ThumbnailRendering/TextureThumbnailRenderer.cpp
+++ b/Engine/Source/Editor/UnrealEd/Private/ThumbnailRendering/TextureThumbnailRenderer.cpp
@@ -85,7 +85,7 @@ void UTextureThumbnailRenderer::Draw(UObject* Object, int32 X, int32 Y, uint32 W
 		{
 			bool bIsNormalMap = Texture2DArray->IsNormalMap();
 			bool bIsSingleChannel = true;
-			BatchedElementParameters = new FBatchedElementTexture2DPreviewParameters(0, 0, bIsNormalMap, bIsSingleChannel, false, false, true);
+			BatchedElementParameters = new FBatchedElementTexture2DPreviewParameters(0, 0, 0, bIsNormalMap, bIsSingleChannel, false, false, true);
 		}
 		else if (TextureLightProfile)
 		{
diff --git a/Engine/Source/Editor/UnrealEd/Public/Texture2DPreview.h b/Engine/Source/Editor/UnrealEd/Public/Texture2DPreview.h
--- a/Engine/Source/Editor/UnrealEd/Public/Texture2DPreview.h
+++ b/Engine/Source/Editor/UnrealEd/Public/Texture2DPreview.h
@@ -15,9 +15,10 @@
 class UNREALED_API FBatchedElementTexture2DPreviewParameters : public FBatchedElementParameters
 {
 public:
-	FBatchedElementTexture2DPreviewParameters(float InMipLevel, float InLayerIndex, bool bInIsNormalMap, bool bInIsSingleChannel, bool bInIsSingleVTPhysicalSpace, bool bInIsVirtualTexture, bool bInIsTextureArray)
+	FBatchedElementTexture2DPreviewParameters(float InMipLevel, float InLayerIndex, float InSliceIndex, bool bInIsNormalMap, bool bInIsSingleChannel, bool bInIsSingleVTPhysicalSpace, bool bInIsVirtualTexture, bool bInIsTextureArray)
 		: MipLevel(InMipLevel)
 		, LayerIndex(InLayerIndex)
+		, SliceIndex(InSliceIndex)
 		, bIsNormalMap( bInIsNormalMap )
 		, bIsSingleChannelFormat( bInIsSingleChannel )
 		, bIsSingleVTPhysicalSpace(bInIsSingleVTPhysicalSpace)
@@ -34,6 +35,7 @@ private:
 	/** Parameters that need to be passed to the shader */
 	float MipLevel;
 	float LayerIndex;
+	float SliceIndex;
 	bool bIsNormalMap;
 	bool bIsSingleChannelFormat;
 	bool bIsSingleVTPhysicalSpace;
