World Partition: Fix minimap build capturing scene outside of the requested bounds

diff --git a/Engine/Source/Editor/UnrealEd/Private/WorldPartition/WorldPartitionMiniMapBuilder.cpp b/Engine/Source/Editor/UnrealEd/Private/WorldPartition/WorldPartitionMiniMapBuilder.cpp
--- a/Engine/Source/Editor/UnrealEd/Private/WorldPartition/WorldPartitionMiniMapBuilder.cpp
+++ b/Engine/Source/Editor/UnrealEd/Private/WorldPartition/WorldPartitionMiniMapBuilder.cpp
@@ -128,15 +128,16 @@ bool UWorldPartitionMiniMapBuilder::RunInternal(UWorld* World, const FCellInfo&
 {
 	check(World != nullptr);
 
+	// Clamp input bounds
+	const FBox ClampedBounds = InCellInfo.Bounds.Overlap(IterativeWorldBounds);
+
 	// World X,Y to minimap X,Y
-	const FVector3d MinimapMin = WorldToMinimap.TransformPosition(InCellInfo.Bounds.Min);
-	const FVector3d MinimapMax = WorldToMinimap.TransformPosition(InCellInfo.Bounds.Max);
+	const FVector3d MinimapMin = WorldToMinimap.TransformPosition(ClampedBounds.Min);
+	const FVector3d MinimapMax = WorldToMinimap.TransformPosition(ClampedBounds.Max);
 	const FIntVector2 DstMin(FMath::Floor(MinimapMin.X), FMath::Floor(MinimapMin.Y));
 	const FIntVector2 DstMax(FMath::Floor(MinimapMax.X), FMath::Floor(MinimapMax.Y));
-	const FIntVector2 DstClampedMin = FIntVector2(FMath::Clamp(DstMin.X, 0, MinimapImageSizeX), FMath::Clamp(DstMin.Y, 0, MinimapImageSizeY));
-	const FIntVector2 DstClampedMax = FIntVector2(FMath::Clamp(DstMax.X, 0, MinimapImageSizeX), FMath::Clamp(DstMax.Y, 0, MinimapImageSizeY));
-	const uint32 CaptureWidthPixels = DstClampedMax.X - DstClampedMin.X;
-	const uint32 CaptureHeightPixels = DstClampedMax.Y - DstClampedMin.Y;
+	const uint32 CaptureWidthPixels = DstMax.X - DstMin.X;
+	const uint32 CaptureHeightPixels = DstMax.Y - DstMin.Y;
 
 	// Capture a tile if the region to capture is not empty
 	if (CaptureWidthPixels > 0 && CaptureHeightPixels > 0)
@@ -147,7 +148,7 @@ bool UWorldPartitionMiniMapBuilder::RunInternal(UWorld* World, const FCellInfo&
 		TileTexture->Source.Init(CaptureWidthPixels, CaptureHeightPixels, 1, 1, TSF_BGRA8);
 		TileTexture->PowerOfTwoMode = ETexturePowerOfTwoSetting::PadToPowerOfTwo;
 
-		FWorldPartitionMiniMapHelper::CaptureBoundsMiniMapToTexture(World, GetTransientPackage(), CaptureWidthPixels, CaptureHeightPixels, TileTexture, TextureName, InCellInfo.Bounds, WorldMiniMap->CaptureSource, WorldMiniMap->CaptureWarmupFrames);
+		FWorldPartitionMiniMapHelper::CaptureBoundsMiniMapToTexture(World, GetTransientPackage(), CaptureWidthPixels, CaptureHeightPixels, TileTexture, TextureName, ClampedBounds, WorldMiniMap->CaptureSource, WorldMiniMap->CaptureWarmupFrames);
 
 		// Copy captured image to VT minimap
 		const uint32 BPP = TileTexture->Source.GetBytesPerPixel();
@@ -157,7 +158,7 @@ bool UWorldPartitionMiniMapBuilder::RunInternal(UWorld* World, const FCellInfo&
 		check(SrcDataPtr);
 
 		const uint32 DstDataStrideBytes = WorldMiniMap->MiniMapTexture->Source.GetSizeX() * BPP;
-		uint8* const DstDataPtr = MiniMapSourcePtr + (DstClampedMin.Y * DstDataStrideBytes) + (DstClampedMin.X * BPP);
+		uint8* const DstDataPtr = MiniMapSourcePtr + (DstMin.Y * DstDataStrideBytes) + (DstMin.X * BPP);
 		check(DstDataPtr);
 
 		for (uint32 RowIdx = 0; RowIdx < CaptureHeightPixels; ++RowIdx)
