Allow the class icon finder to handle any generated class type

#rnx

diff --git a/Engine/Source/Editor/UnrealEd/Private/ClassIconFinder.cpp b/Engine/Source/Editor/UnrealEd/Private/ClassIconFinder.cpp
--- a/Engine/Source/Editor/UnrealEd/Private/ClassIconFinder.cpp
+++ b/Engine/Source/Editor/UnrealEd/Private/ClassIconFinder.cpp
@@ -113,7 +113,7 @@ const UClass* FClassIconFinder::GetIconClassForAssetData(const FAssetData& InAss
 	}
 	
 	static const FName IgnoreClassThumbnail(TEXT("IgnoreClassThumbnail"));
-	if ( (AssetClass->IsChildOf<UBlueprint>() || AssetClass->IsChildOf<UBlueprintGeneratedClass>()) &&
+	if ( (AssetClass->IsChildOf<UBlueprint>() || AssetClass->IsChildOf<UClass>()) &&
 		!AssetClass->HasMetaDataHierarchical(IgnoreClassThumbnail))
 	{
 		if ( bOutIsClassType )
@@ -121,7 +121,7 @@ const UClass* FClassIconFinder::GetIconClassForAssetData(const FAssetData& InAss
 			*bOutIsClassType = true;
 		}
 
-		// We need to use the asset data to get the parent class as the blueprint may not be loaded
+		// We need to use the asset data to get the parent class as the generated class may not be loaded
 		FString ParentClassName;
 		if ( !InAssetData.GetTagValue(FBlueprintTags::NativeParentClassPath, ParentClassName) )
 		{
