Cooker: Add diagnostics to log when GarbageCollection is not impactful.
#rnx

diff --git a/Engine/Source/Editor/UnrealEd/Classes/CookOnTheSide/CookOnTheFlyServer.h b/Engine/Source/Editor/UnrealEd/Classes/CookOnTheSide/CookOnTheFlyServer.h
--- a/Engine/Source/Editor/UnrealEd/Classes/CookOnTheSide/CookOnTheFlyServer.h
+++ b/Engine/Source/Editor/UnrealEd/Classes/CookOnTheSide/CookOnTheFlyServer.h
@@ -36,6 +36,7 @@ class ITargetPlatform;
 enum class ODSCRecompileCommand;
 struct FBeginCookContext;
 struct FPropertyChangedEvent;
+struct FResourceSizeEx;
 
 namespace UE::LinkerLoad { enum class EImportBehavior : uint8; }
 
@@ -148,6 +149,23 @@ namespace UE::Cook::Private
 	class FRegisteredCookPackageSplitter;
 }
 
+namespace UE::Cook
+{
+
+struct FStatHistoryInt
+{
+public:
+	void Initialize(int64 InitialValue);
+	void AddInstance(int64 CurrentValue);
+	int64 GetMinimum() const { return Minimum; }
+	int64 GetMaximum() const { return Maximum; }
+
+private:
+	int64 Minimum = 0;
+	int64 Maximum = 0;
+};
+
+}
 UCLASS()
 class UNREALED_API UCookOnTheFlyServer : public UObject, public FTickableEditorObject, public FExec, public UE::Cook::ICookInfo
 {
@@ -255,6 +273,7 @@ private:
 	uint64 MemoryMaxUsedPhysical;
 	uint64 MemoryMinFreeVirtual;
 	uint64 MemoryMinFreePhysical;
+	float MemoryExpectedFreedToSpreadRatio;
 	/** Max number of packages to save before we partial gc */
 	int32 MaxNumPackagesBeforePartialGC;
 	/** Max number of concurrent shader jobs reducing this too low will increase cook time */
@@ -444,7 +463,7 @@ private:
 	void PollFlushRenderingCommands();
 	TRefCountPtr<FPollable> CreatePollableLLM();
 	TRefCountPtr<FPollable> CreatePollableTriggerGC();
-	void PollTimeForGC(UE::Cook::FTickStackData& StackData);
+	void PollGarbageCollection(UE::Cook::FTickStackData& StackData);
 	void PollQueuedCancel(UE::Cook::FTickStackData& StackData);
 	void WaitForAsync(UE::Cook::FTickStackData& StackData);
 	void TickRecompileShaderRequestsPrivate();
@@ -721,6 +740,7 @@ public:
 	}
 
 	bool HasExceededMaxMemory() const;
+	void EvaluateGarbageCollectionResults(int32 NumObjectsBeforeGC, const FPlatformMemoryStats& MemStatsBeforeGC);
 
 	/**
 	* RequestPackage to be cooked
@@ -785,6 +805,12 @@ private:
 	/** Update accumulators of editor data and consume editor changes since the last cook when starting cooks in the editor. */
 	void BeginCookEditorSystems();
 	void BeginCookPackageWriters(FBeginCookContext& BeginContext);
+	void InitializeSession();
+	void PreGarbageCollectImpl(TArray<UPackage*>& GCKeepPackages, TArray<UE::Cook::FPackageData*>& GCKeepPackageDatas,
+		TArray<UObject*>& LocalGCKeepObjects);
+	void PostGarbageCollectImpl(TArray<UObject*>& LocalGCKeepObjects);
+	void GetDirectAndTransitiveResourceSize(FResourceSizeEx& OutDirectSize, FResourceSizeEx& OutTransitiveSize,
+		int64& OutNumDirectPackages, int64& OutNumTransitivePackages, TSet<UPackage*>&& DirectPackages);
 
 	//////////////////////////////////////////////////////////////////////////
 	// cook by the book specific functions
@@ -1206,6 +1232,8 @@ private:
 	static UCookOnTheFlyServer* ActiveCOTFS;
 	uint32		StatLoadedPackageCount = 0;
 	uint32		StatSavedPackageCount = 0;
+	UE::Cook::FStatHistoryInt NumObjectsHistory;
+	UE::Cook::FStatHistoryInt VirtualMemoryHistory;
 
 	/** This is set to true when the decision about which packages we need to cook changes because e.g. a platform was added to the sessionplatforms. */
 	bool bPackageFilterDirty = false;
diff --git a/Engine/Source/Editor/UnrealEd/Private/Commandlets/CookCommandlet.cpp b/Engine/Source/Editor/UnrealEd/Private/Commandlets/CookCommandlet.cpp
--- a/Engine/Source/Editor/UnrealEd/Private/Commandlets/CookCommandlet.cpp
+++ b/Engine/Source/Editor/UnrealEd/Private/Commandlets/CookCommandlet.cpp
@@ -972,6 +972,7 @@ void UCookCommandlet::ConditionalCollectGarbage(uint32 TickResults, UCookOnTheFl
 #endif
 	UE_SCOPED_COOKTIMER(CookCommandlet_GC);
 
+	const FPlatformMemoryStats MemStatsBeforeGC = FPlatformMemory::GetStats();
 	int32 NumObjectsBeforeGC = GUObjectArray.GetObjectArrayNumMinusAvailable();
 	int32 NumObjectsAvailableBeforeGC = GUObjectArray.GetObjectArrayEstimatedAvailable();
 	UE_LOG(LogCookCommandlet, Display, TEXT("GarbageCollection...%s (%s)"), (bPartialGC ? TEXT(" partial gc") : TEXT("")), *GCReason);
@@ -1004,4 +1005,9 @@ void UCookCommandlet::ConditionalCollectGarbage(uint32 TickResults, UCookOnTheFl
 			(bPartialGC ? TEXT("Partial") : TEXT("Full")), NumObjectsBeforeGC, NumObjectsAvailableBeforeGC, NumObjectsAfterGC, NumObjectsAvailableAfterGC);
 		DumpMemStats();
 	}
+
+	if (TickResults & UCookOnTheFlyServer::COSR_RequiresGC_OOM)
+	{
+		COTFS.EvaluateGarbageCollectionResults(NumObjectsBeforeGC, MemStatsBeforeGC);
+	} 
 }
\ No newline at end of file
diff --git a/Engine/Source/Editor/UnrealEd/Private/CookOnTheFlyServer.cpp b/Engine/Source/Editor/UnrealEd/Private/CookOnTheFlyServer.cpp
--- a/Engine/Source/Editor/UnrealEd/Private/CookOnTheFlyServer.cpp
+++ b/Engine/Source/Editor/UnrealEd/Private/CookOnTheFlyServer.cpp
@@ -96,8 +96,10 @@
 #include "PlatformInfo.h"
 #include "ProfilingDebugging/CookStats.h"
 #include "ProfilingDebugging/PlatformFileTrace.h"
+#include "ProfilingDebugging/ResourceSize.h"
 #include "ProjectDescriptor.h"
 #include "SceneUtils.h"
+#include "Serialization/ArchiveCountMem.h"
 #include "Serialization/ArchiveStackTrace.h"
 #include "Serialization/ArchiveUObject.h"
 #include "Serialization/ArrayReader.h"
@@ -718,6 +720,10 @@ void UCookOnTheFlyServer::AddCookOnTheFlyPlatformFromGameThread(ITargetPlatform*
 
 void UCookOnTheFlyServer::StartCookOnTheFlySessionFromGameThread(ITargetPlatform* TargetPlatform)
 {
+	if (PlatformManager->GetSessionPlatforms().Num() == 0)
+	{
+		InitializeSession();
+	}
 	PlatformManager->AddSessionPlatform(*this, TargetPlatform);
 	bPackageFilterDirty = true;
 
@@ -1888,13 +1894,10 @@ TRefCountPtr<UCookOnTheFlyServer::FPollable> UCookOnTheFlyServer::CreatePollable
 TRefCountPtr<UCookOnTheFlyServer::FPollable> UCookOnTheFlyServer::CreatePollableTriggerGC()
 {
 	bool bTestCook = IsCookFlagSet(ECookInitializationFlags::TestCook);
-	if (PackagesPerGC == 0 && !bTestCook && !IsCookOnTheFlyMode())
-	{
-		return TRefCountPtr<FPollable>();
-	}
 
-	float PeriodSeconds = MAX_flt;
-	float IdlePeriodSeconds = MAX_flt;
+	// Collect statistics every 2 minutes even if we are not tracking time between garbage collects
+	float PeriodSeconds = 120.f;
+	float IdlePeriodSeconds = 120.f;
 	constexpr float SecondsPerPackage = .01f;
 	if (bTestCook)
 	{
@@ -1915,11 +1918,14 @@ TRefCountPtr<UCookOnTheFlyServer::FPollable> UCookOnTheFlyServer::CreatePollable
 	IdlePeriodSeconds = FMath::Min(IdlePeriodSeconds, PeriodSeconds);
 
 	return TRefCountPtr<FPollable>(new FPollable(TEXT("TimeForGC"), PeriodSeconds, IdlePeriodSeconds,
-		[this](UE::Cook::FTickStackData& StackData) { PollTimeForGC(StackData); }));
+		[this](UE::Cook::FTickStackData& StackData) { PollGarbageCollection(StackData); }));
 }
 
-void UCookOnTheFlyServer::PollTimeForGC(UE::Cook::FTickStackData& StackData)
+void UCookOnTheFlyServer::PollGarbageCollection(UE::Cook::FTickStackData& StackData)
 {
+	NumObjectsHistory.AddInstance(GUObjectArray.GetObjectArrayNumMinusAvailable());
+	VirtualMemoryHistory.AddInstance(FPlatformMemory::GetStats().UsedVirtual);
+
 	if (IsCookFlagSet(ECookInitializationFlags::TestCook))
 	{
 		StackData.ResultFlags |= COSR_RequiresGC | COSR_YieldTick;
@@ -1950,6 +1956,17 @@ void UCookOnTheFlyServer::PollTimeForGC(UE::Cook::FTickStackData& StackData)
 namespace UE::Cook
 {
 
+void FStatHistoryInt::Initialize(int64 InitialValue)
+{
+	Maximum = Minimum = InitialValue;
+}
+
+void FStatHistoryInt::AddInstance(int64 CurrentValue)
+{
+	Maximum = FMath::Max(CurrentValue, Maximum);
+	Minimum = FMath::Min(CurrentValue, Minimum);
+}
+
 static void ProcessDeferredCommands(UCookOnTheFlyServer& COTFS)
 {
 #if OUTPUT_COOKTIMING
@@ -4244,6 +4261,162 @@ bool UCookOnTheFlyServer::HasExceededMaxMemory() const
 	}
 }
 
+void UCookOnTheFlyServer::EvaluateGarbageCollectionResults(int32 NumObjectsBeforeGC, const FPlatformMemoryStats& MemStatsBeforeGC)
+{
+	using namespace UE::Cook;
+
+	if (IsCookingInEditor())
+	{
+		return;
+	}
+
+	int64 NumObjectsMin = NumObjectsHistory.GetMinimum();
+	int64 NumObjectsMax = NumObjectsHistory.GetMaximum();
+	int64 NumObjectsSpread = NumObjectsMax - NumObjectsMin;
+	int64 NumObjectsAfterGC = GUObjectArray.GetObjectArrayNumMinusAvailable();
+	int64 NumObjectsFreed = NumObjectsBeforeGC - NumObjectsAfterGC;
+	FPlatformMemoryStats MemStatsAfterGC = FPlatformMemory::GetStats();
+	int64 VirtualMemMin = VirtualMemoryHistory.GetMinimum();
+	int64 VirtualMemMax = VirtualMemoryHistory.GetMaximum();
+	int64 VirtualMemSpread = VirtualMemMax - VirtualMemMin;
+	int64 VirtualMemBeforeGC = MemStatsBeforeGC.UsedVirtual;
+	int64 VirtualMemAfterGC = MemStatsAfterGC.UsedVirtual;
+	int64 VirtualMemFreed = MemStatsBeforeGC.UsedVirtual - MemStatsAfterGC.UsedVirtual;
+
+	int64 ExpectedObjectsFreed = MemoryExpectedFreedToSpreadRatio * NumObjectsSpread;
+	double ExpectedMemFreed = MemoryExpectedFreedToSpreadRatio * VirtualMemSpread;
+	if ((NumObjectsFreed >= ExpectedObjectsFreed || NumObjectsBeforeGC - NumObjectsMin < ExpectedObjectsFreed) &&
+		(VirtualMemFreed >= ExpectedMemFreed || VirtualMemBeforeGC - VirtualMemMin <= ExpectedMemFreed))
+	{
+		// Nothing to report, Garbage Collection was as impactful as expected
+		return;
+	}
+
+	TArray<UPackage*> GCKeepPackages;
+	TArray<FPackageData*> GCKeepPackageDatas;
+	TArray<UObject*> LocalGCKeepObjects;
+	PreGarbageCollectImpl(GCKeepPackages, GCKeepPackageDatas, LocalGCKeepObjects);
+	PostGarbageCollectImpl(LocalGCKeepObjects);
+
+	TSet<UPackage*> DirectPackages(GCKeepPackages);
+	DirectPackages.Append(GCKeepPackages);
+	// Some Objects can be in KeepObjects without their Package being added, because they are PollPendingCookedPlatformDatas
+	// Add their packages, since we do the transitive search based on package dependencies
+	for (UObject* Object : LocalGCKeepObjects)
+	{
+		DirectPackages.Add(Object->GetPackage());
+	}
+
+	FResourceSizeEx DirectResourceSize;
+	FResourceSizeEx TransitiveResourceSize;
+	int64 NumDirectPackages;
+	int64 NumTransitivePackages;
+	GetDirectAndTransitiveResourceSize(DirectResourceSize, TransitiveResourceSize,
+		NumDirectPackages, NumTransitivePackages, MoveTemp(DirectPackages));
+
+	constexpr int32 BytesPerMeg = 1000000;
+	UE_LOG(LogCook, Display, TEXT("Garbage Collection was not very impactful.\n")
+		TEXT("\tNumObjects:\n")
+		TEXT("\t\tProcess Min:    %10" INT64_FMT "\n")
+		TEXT("\t\tProcess Max:    %10" INT64_FMT "\n")
+		TEXT("\t\tProcess Spread: %10" INT64_FMT "\n")
+		TEXT("\t\tBefore GC:      %10" INT64_FMT "\n")
+		TEXT("\t\tAfter GC:       %10" INT64_FMT "\n")
+		TEXT("\t\tFreed by GC:    %10" INT64_FMT "\n")
+		TEXT("\tVirtual Memory:\n")
+		TEXT("\t\tProcess Min:    %10" INT64_FMT " MB\n")
+		TEXT("\t\tProcess Max:    %10" INT64_FMT " MB\n")
+		TEXT("\t\tProcess Spread: %10" INT64_FMT " MB\n")
+		TEXT("\t\tBefore GC:      %10" INT64_FMT " MB\n")
+		TEXT("\t\tAfter GC:       %10" INT64_FMT " MB\n")
+		TEXT("\t\tFreed by GC:    %10" INT64_FMT " MB\n")
+		TEXT("\tReferences:\n")
+		TEXT("\t\tCooker direct packages:         %10" INT64_FMT "\n")
+		TEXT("\t\tCooker transitive packages:     %10" INT64_FMT "\n")
+		TEXT("\t\tCooker direct package size:     %10" INT64_FMT " MB\n")
+		TEXT("\t\tCooker transitive package size: %10" INT64_FMT " MB\n"),
+		NumObjectsMin, NumObjectsMax, NumObjectsSpread,
+		(int64)NumObjectsBeforeGC, NumObjectsAfterGC, NumObjectsFreed,
+		VirtualMemMin / BytesPerMeg, VirtualMemMax / BytesPerMeg, VirtualMemSpread / BytesPerMeg,
+		VirtualMemBeforeGC / BytesPerMeg, VirtualMemAfterGC / BytesPerMeg, VirtualMemFreed / BytesPerMeg,
+		NumDirectPackages, NumTransitivePackages,
+		DirectResourceSize.GetTotalMemoryBytes() / BytesPerMeg,
+		TransitiveResourceSize.GetTotalMemoryBytes() / BytesPerMeg
+	);
+#if ENABLE_LOW_LEVEL_MEM_TRACKER
+	FLowLevelMemTracker::Get().DumpToLog();
+#endif
+}
+
+void UCookOnTheFlyServer::GetDirectAndTransitiveResourceSize(FResourceSizeEx& OutDirectSize,
+	FResourceSizeEx& OutTransitiveSize, int64& OutNumDirectPackages, int64& OutNumTransitivePackages,
+	TSet<UPackage*>&& DirectPackages)
+{
+	OutDirectSize = FResourceSizeEx(EResourceSizeMode::Exclusive);
+
+	TArray<UObject*> ObjectsInPackage;
+	auto AddObjectSize = [](UObject* Object, FResourceSizeEx& InResourceSize)
+	{
+		FArchiveCountMem MemoryCount(Object, true);
+		InResourceSize.AddDedicatedSystemMemoryBytes(MemoryCount.GetMax());
+		Object->GetResourceSizeEx(InResourceSize);
+	};
+	auto AddPackageSize = [&ObjectsInPackage, &AddObjectSize](UPackage* InPackage, FResourceSizeEx& InResourceSize)
+	{
+		InPackage->GetResourceSizeEx(InResourceSize);
+		AddObjectSize(InPackage, InResourceSize);
+
+		ObjectsInPackage.Reset();
+		GetObjectsWithPackage(InPackage, ObjectsInPackage);
+		for (UObject* Object : ObjectsInPackage)
+		{
+			AddObjectSize(Object, InResourceSize);
+		}
+	};
+
+	OutNumTransitivePackages = 0;
+	OutNumDirectPackages = 0;
+	TRingBuffer<UPackage*> NextPackages;
+	NextPackages.Reserve(DirectPackages.Num());
+	for (UPackage* Package : DirectPackages)
+	{
+		check(Package);
+		if (!Package->IsRooted())
+		{
+			AddPackageSize(Package, OutDirectSize);
+			NextPackages.Add(Package);
+			++OutNumDirectPackages;
+			++OutNumTransitivePackages;
+		}
+	}
+
+	OutTransitiveSize = OutDirectSize;
+	TSet<UPackage*>& TransitivePackages = DirectPackages;
+	TArray<FName> Dependencies;
+	while (!NextPackages.IsEmpty())
+	{
+		UPackage* NextPackage = NextPackages.PopFrontValue();
+		Dependencies.Reset();
+		AssetRegistry->GetDependencies(NextPackage->GetFName(), Dependencies, UE::AssetRegistry::EDependencyCategory::Package, UE::AssetRegistry::EDependencyQuery::Hard);
+		for (FName DependencyName : Dependencies)
+		{
+			UPackage* DependencyPackage = FindObject<UPackage>(FTopLevelAssetPath(DependencyName, NAME_None));
+			if (!DependencyPackage)
+			{
+				continue;
+			}
+			bool bAlreadyInSet;
+			TransitivePackages.Add(DependencyPackage, &bAlreadyInSet);
+			if (!bAlreadyInSet && !DependencyPackage->IsRooted())
+			{
+				++OutNumTransitivePackages;
+				NextPackages.Add(DependencyPackage);
+				AddPackageSize(DependencyPackage, OutTransitiveSize);
+			}
+		}
+	}
+}
+
 TArray<UPackage*> UCookOnTheFlyServer::GetUnsolicitedPackages(const TArray<const ITargetPlatform*>& TargetPlatforms) const
 {
 	// No longer supported
@@ -4518,12 +4691,23 @@ const TArray<FName>& UCookOnTheFlyServer::GetFullPackageDependencies(const FName
 
 void UCookOnTheFlyServer::PreGarbageCollect()
 {
-	using namespace UE::Cook;
 	if (!IsInSession())
 	{
 		return;
 	}
 
+	NumObjectsHistory.AddInstance(GUObjectArray.GetObjectArrayNumMinusAvailable());
+	VirtualMemoryHistory.AddInstance(FPlatformMemory::GetStats().UsedVirtual);
+	TArray<UPackage*> GCKeepPackages;
+	TArray<UE::Cook::FPackageData*> GCKeepPackageDatas;
+	PreGarbageCollectImpl(GCKeepPackages, GCKeepPackageDatas, GCKeepObjects);
+}
+
+void UCookOnTheFlyServer::PreGarbageCollectImpl(TArray<UPackage*>& GCKeepPackages,
+	TArray<UE::Cook::FPackageData*>& GCKeepPackageDatas, TArray<UObject*>& LocalGCKeepObjects)
+{
+	using namespace UE::Cook;
+
 #if COOK_CHECKSLOW_PACKAGEDATA
 	// Verify that only packages in the save state have pointers to objects
 	for (const FPackageData* PackageData : *PackageDatas.Get())
@@ -4534,11 +4718,9 @@ void UCookOnTheFlyServer::PreGarbageCollect()
 	if (SavingPackageData)
 	{
 		check(SavingPackageData->GetPackage());
-		GCKeepObjects.Add(SavingPackageData->GetPackage());
+		LocalGCKeepObjects.Add(SavingPackageData->GetPackage());
 	}
 
-	TArray<UPackage*> GCKeepPackages;
-	TArray<FPackageData*> GCKeepPackageDatas;
 
 	// Demote any Generated/Generator packages we called PreSave on so they call their PostSave before the GC
 	// or prevent them from being garbage collected if the splitter wants to keep them referenced
@@ -4553,7 +4735,7 @@ void UCookOnTheFlyServer::PreGarbageCollect()
 		if (Info)
 		{
 			bool bShouldDemote;
-			Generator->PreGarbageCollect(*Info, GCKeepObjects, GCKeepPackages, GCKeepPackageDatas, bShouldDemote);
+			Generator->PreGarbageCollect(*Info, LocalGCKeepObjects, GCKeepPackages, GCKeepPackageDatas, bShouldDemote);
 			if (bShouldDemote)
 			{
 				ReleaseCookedPlatformData(*PackageData, UE::Cook::EReleaseSaveReason::Demoted);
@@ -4592,14 +4774,14 @@ void UCookOnTheFlyServer::PreGarbageCollect()
 		{
 			UObject* Object = Pending.Object.Get();
 			check(Object); // Otherwise PollIsComplete would have returned true
-			GCKeepObjects.Add(Object);
+			LocalGCKeepObjects.Add(Object);
 		}
 	}
 
 	const bool bPartialGC = IsCookFlagSet(ECookInitializationFlags::EnablePartialGC);
 	if (bPartialGC)
 	{
-		GCKeepObjects.Empty(1000);
+		LocalGCKeepObjects.Empty(1000);
 
 		// Keep all inprogress packages (including packages that have only made it to the request list) that have been partially loaded
 		// Additionally, keep all partially loaded packages that are transitively dependended on by any inprogress packages
@@ -4660,18 +4842,18 @@ void UCookOnTheFlyServer::PreGarbageCollect()
 		}
 	}
 
-	// Add packages and all RF_Public objects outered to them to GCKeepObjects
+	// Add packages and all RF_Public objects outered to them to LocalGCKeepObjects
 	TArray<UObject*> ObjectsWithOuter;
 	for (UPackage* Package : GCKeepPackages)
 	{
-		GCKeepObjects.Add(Package);
+		LocalGCKeepObjects.Add(Package);
 		ObjectsWithOuter.Reset();
 		GetObjectsWithOuter(Package, ObjectsWithOuter);
 		for (UObject* Obj : ObjectsWithOuter)
 		{
 			if (Obj->HasAnyFlags(RF_Public))
 			{
-				GCKeepObjects.Add(Obj);
+				LocalGCKeepObjects.Add(Obj);
 			}
 		}
 	}
@@ -4690,6 +4872,13 @@ void UCookOnTheFlyServer::CookerAddReferencedObjects(FReferenceCollector& Collec
 }
 
 void UCookOnTheFlyServer::PostGarbageCollect()
+{
+	NumObjectsHistory.AddInstance(GUObjectArray.GetObjectArrayNumMinusAvailable());
+	VirtualMemoryHistory.AddInstance(FPlatformMemory::GetStats().UsedVirtual);
+	PostGarbageCollectImpl(GCKeepObjects);
+}
+
+void UCookOnTheFlyServer::PostGarbageCollectImpl(TArray<UObject*>& LocalGCKeepObjects)
 {
 	using namespace UE::Cook;
 
@@ -4716,7 +4905,7 @@ void UCookOnTheFlyServer::PostGarbageCollect()
 	// However, if garbage collection deleted the package WHILE WE WERE SAVING IT, then we have problems.
 	check(!SavingPackageData || SavingPackageData->GetPackage() != nullptr);
 
-	GCKeepObjects.Empty();
+	LocalGCKeepObjects.Empty();
 
 	for (FPackageData* PackageData : *PackageDatas.Get())
 	{
@@ -5311,6 +5500,9 @@ void UCookOnTheFlyServer::Initialize( ECookMode::Type DesiredCookMode, ECookInit
 	ReadMemorySetting(TEXT("MemoryMaxUsedPhysical"), MemoryMaxUsedPhysical);
 	ReadMemorySetting(TEXT("MemoryMinFreeVirtual"), MemoryMinFreeVirtual);
 	ReadMemorySetting(TEXT("MemoryMinFreePhysical"), MemoryMinFreePhysical);
+	MemoryExpectedFreedToSpreadRatio = 0.10f;
+	GConfig->GetFloat(TEXT("CookSettings"), TEXT("MemoryExpectedFreedToSpreadRatio"),
+		MemoryExpectedFreedToSpreadRatio, GEditorIni);
 
 	MaxPreloadAllocated = 16;
 	DesiredSaveQueueLength = 8;
@@ -5417,6 +5609,12 @@ void UCookOnTheFlyServer::Initialize( ECookMode::Type DesiredCookMode, ECookInit
 	});
 }
 
+void UCookOnTheFlyServer::InitializeSession()
+{
+	NumObjectsHistory.Initialize(GUObjectArray.GetObjectArrayNumMinusAvailable());
+	VirtualMemoryHistory.Initialize(FPlatformMemory::GetStats().UsedVirtual);
+}
+
 bool UCookOnTheFlyServer::Exec(class UWorld* InWorld, const TCHAR* Cmd, FOutputDevice& Ar)
 {
 	if (FParse::Command(&Cmd, TEXT("package")))
@@ -8932,6 +9130,7 @@ void UCookOnTheFlyServer::BeginCookPackageWriters(FBeginCookContext& BeginContex
 
 void UCookOnTheFlyServer::SelectSessionPlatforms(FBeginCookContext& BeginContext)
 {
+	InitializeSession();
 	PlatformManager->SelectSessionPlatforms(*this, BeginContext.TargetPlatforms);
 	if (PackageTracker->HasBeenConsumed())
 	{
diff --git a/Engine/Source/Runtime/Core/Private/HAL/LowLevelMemTracker.cpp b/Engine/Source/Runtime/Core/Private/HAL/LowLevelMemTracker.cpp
--- a/Engine/Source/Runtime/Core/Private/HAL/LowLevelMemTracker.cpp
+++ b/Engine/Source/Runtime/Core/Private/HAL/LowLevelMemTracker.cpp
@@ -9,6 +9,7 @@
 #include "LowLevelMemTrackerPrivate.h"
 #include "MemPro/MemProProfiler.h"
 #include "Misc/CString.h"
+#include "Misc/OutputDeviceRedirector.h"
 #include "Misc/Paths.h"
 #include "Misc/ScopeLock.h"
 #include "Misc/ScopeRWLock.h"
@@ -78,8 +79,6 @@ FAutoConsoleCommand DumpLLM(
 	TEXT("Logs out the current and peak sizes of all tracked LLM tags"),
 	FConsoleCommandWithWorldArgsAndOutputDeviceDelegate::CreateLambda([](const TArray<FString>& Args, UWorld* InWorld, FOutputDevice& Ar)
 	{
-		const float InvToMb = 1.0 / (1024 * 1024);
-
 		bool bCSV = false;
 		for (const FString& Arg : Args)
 		{
@@ -88,47 +87,77 @@ FAutoConsoleCommand DumpLLM(
 				bCSV = true;
 			}
 		}
+		FLowLevelMemTracker::Get().DumpToLog(bCSV ? FLowLevelMemTracker::EDumpFormat::CSV : FLowLevelMemTracker::EDumpFormat::PlainText, &Ar);
+	}));
+
+void FLowLevelMemTracker::DumpToLog(EDumpFormat DumpFormat, FOutputDevice* OutputDevice)
+{
+	if (!IsEnabled())
+	{
+		return;
+	}
+	const float InvToMb = 1.0 / (1024 * 1024);
+	FOutputDevice& Ar = *(OutputDevice ? OutputDevice : GLog);
+	if (DumpFormat == EDumpFormat::CSV)
+	{
+		Ar.Logf(TEXT(",TagName,SizeMB,PeakMB,Tracker,PathName"));
+	}
+	else
+	{
+		Ar.Logf(TEXT("%40s %12s %12s  Tracker PathName"), TEXT("TagName"), TEXT("SizeMB"), TEXT("PeakMB"));
+	}
 
-		if (FLowLevelMemTracker::Get().IsEnabled())
+	for (ELLMTracker TrackerType : {ELLMTracker::Default, ELLMTracker::Platform})
+	{
+		const TCHAR* TrackerName = TrackerType == ELLMTracker::Default ? TEXT("Default") : TEXT("Platform");
+
+		struct FTagLine
 		{
-			if (bCSV)
+			FString TagName;
+			FString Line;
+			int64 Size;
+		};
+		TArray<FTagLine> TagLines;
+		for (const UE::LLMPrivate::FTagData* TagData : GetTrackedTags(TrackerType))
+		{
+			int64 CurrentAmount = GetTagAmountForTracker(TrackerType, TagData, false /* bPeakAmount */);
+			int64 PeakAmount = GetTagAmountForTracker(TrackerType, TagData, true /* bPeakAmount */);
+			FString TagName = GetTagUniqueName(TagData).ToString();
+			FString TagLine;
+			if (DumpFormat == EDumpFormat::CSV)
 			{
-				Ar.Logf(TEXT(",TagName,SizeMB,PeakMB,Tracker,PathName"));
+				TagLine = FString::Printf(TEXT(",%s,%.3f,%.3f,%s,%s"),
+					*TagName,
+					static_cast<float>(CurrentAmount) * InvToMb,
+					static_cast<float>(PeakAmount) * InvToMb,
+					TrackerName,
+					*GetTagDisplayPathName(TagData));
 			}
 			else
 			{
-				Ar.Logf(TEXT("%40s %12s %12s  Tracker PathName"), TEXT("TagName"), TEXT("SizeMB"), TEXT("PeakMB"));
+				TagLine = FString::Printf(TEXT("%40s %12.3f %12.3f %8s %s"),
+					*TagName,
+					static_cast<float>(CurrentAmount) * InvToMb,
+					static_cast<float>(PeakAmount) * InvToMb,
+					TrackerName,
+					*GetTagDisplayPathName(TagData));
 			}
-
-			for (ELLMTracker TrackerType : {ELLMTracker::Default, ELLMTracker::Platform})
+			TagLines.Add(FTagLine{ MoveTemp(TagName), MoveTemp(TagLine), CurrentAmount });
+		}
+		TagLines.Sort([](const FTagLine& A, const FTagLine& B)
 			{
-				const TCHAR* TrackerName = TrackerType == ELLMTracker::Default ? TEXT("Default") : TEXT("Platform");
-
-				for (const UE::LLMPrivate::FTagData* TagData : FLowLevelMemTracker::Get().GetTrackedTags(TrackerType))
+				if (A.Size != B.Size)
 				{
-					if (bCSV)
-					{
-						Ar.Logf(TEXT(",%s,%.3f,%.3f,%s,%s"),
-							*FLowLevelMemTracker::Get().GetTagUniqueName(TagData).ToString(),
-							static_cast<float>(FLowLevelMemTracker::Get().GetTagAmountForTracker(TrackerType, TagData, false)) * InvToMb,
-							static_cast<float>(FLowLevelMemTracker::Get().GetTagAmountForTracker(TrackerType, TagData, true)) * InvToMb,
-							TrackerName,
-							*FLowLevelMemTracker::Get().GetTagDisplayPathName(TagData));
-					}
-					else
-					{
-						Ar.Logf(TEXT("%40s %12.3f %12.3f %8s %s"),
-							*FLowLevelMemTracker::Get().GetTagUniqueName(TagData).ToString(),
-							static_cast<float>(FLowLevelMemTracker::Get().GetTagAmountForTracker(TrackerType, TagData, false)) * InvToMb,
-							static_cast<float>(FLowLevelMemTracker::Get().GetTagAmountForTracker(TrackerType, TagData, true)) * InvToMb,
-							TrackerName,
-							*FLowLevelMemTracker::Get().GetTagDisplayPathName(TagData));
-					}
+					return A.Size > B.Size;
 				}
-			}
+				return A.TagName < B.TagName;
+			});
+		for (FTagLine& TagLine : TagLines)
+		{
+			Ar.Logf(TEXT("%s"), *TagLine.Line);
 		}
-	})
-);
+	}
+}
 
 DECLARE_LLM_MEMORY_STAT(TEXT("LLM Overhead"), STAT_LLMOverheadTotal, STATGROUP_LLMOverhead);
 
diff --git a/Engine/Source/Runtime/Core/Public/HAL/LowLevelMemTracker.h b/Engine/Source/Runtime/Core/Public/HAL/LowLevelMemTracker.h
--- a/Engine/Source/Runtime/Core/Public/HAL/LowLevelMemTracker.h
+++ b/Engine/Source/Runtime/Core/Public/HAL/LowLevelMemTracker.h
@@ -590,6 +590,13 @@ public:
 	// Publishes the active LLM stats in the active frame, useful for single targeted LLM snapshots
 	void PublishDataSingleFrame();
 
+	enum class EDumpFormat
+	{
+		PlainText,
+		CSV,
+	};
+	void DumpToLog(EDumpFormat DumpFormat = EDumpFormat::PlainText, FOutputDevice* OutputDevice = nullptr);
+
 	void OnPreFork();
 
 private:
