Actor Descriptor Properties: added support for to filter actors to resave by properties in WorldPartitionResaveActorsBuilder.

#rnx

diff --git a/Engine/Source/Editor/UnrealEd/Private/WorldPartition/WorldPartitionResaveActorsBuilder.cpp b/Engine/Source/Editor/UnrealEd/Private/WorldPartition/WorldPartitionResaveActorsBuilder.cpp
--- a/Engine/Source/Editor/UnrealEd/Private/WorldPartition/WorldPartitionResaveActorsBuilder.cpp
+++ b/Engine/Source/Editor/UnrealEd/Private/WorldPartition/WorldPartitionResaveActorsBuilder.cpp
@@ -74,6 +74,11 @@ bool UWorldPartitionResaveActorsBuilder::PreRun(UWorld* World, FPackageSourceCon
 			UE_LOG(LogWorldPartitionResaveActorsBuilder, Error, TEXT("SwitchActorPackagingSchemeToReduced is not compatible with ActorTags"));
 			return false;
 		}
+		else if (!ActorProperties.IsEmpty())
+		{
+			UE_LOG(LogWorldPartitionResaveActorsBuilder, Error, TEXT("SwitchActorPackagingSchemeToReduced is not compatible with ActorProperties"));
+			return false;
+		}
 		else if (bResaveDirtyActorDescsOnly)
 		{
 			UE_LOG(LogWorldPartitionResaveActorsBuilder, Error, TEXT("SwitchActorPackagingSchemeToReduced is not compatible with ResaveDirtyActorDescsOnly"));
@@ -272,26 +277,35 @@ bool UWorldPartitionResaveActorsBuilder::RunInternal(UWorld* World, const FCellI
 	else
 	{
 		TArray<UPackage*> PackagesToSave;
-		TSet<FName> ActorTagsSet(ActorTags);
 
 		FWorldPartitionHelpers::FForEachActorWithLoadingParams ForEachActorWithLoadingParams;
 
 		ForEachActorWithLoadingParams.ActorClass = ActorClass;
 
-		if (ActorTags.Num())
-		{			
-			ForEachActorWithLoadingParams.FilterActorDesc = [this, &ActorTagsSet](const FWorldPartitionActorDesc* ActorDesc) -> bool
+		ForEachActorWithLoadingParams.FilterActorDesc = [this](const FWorldPartitionActorDesc* ActorDesc) -> bool
+		{
+			for (const FName& ActorTag : ActorDesc->GetTags())
 			{
-				for (const FName& ActorTag : ActorDesc->GetTags())
+				if (ActorTags.Contains(ActorTag))
 				{
-					if (ActorTagsSet.Contains(ActorTag))
+					return true;
+				}
+			}
+			
+			for (auto [PropertyName, PropertyValue] : ActorProperties)
+			{
+				FName Value;
+				if (ActorDesc->GetProperty(PropertyName, &Value))
+				{
+					if (Value == PropertyValue)
 					{
 						return true;
 					}
 				}
-				return false;
-			};
-		}
+			}
+
+			return false;
+		};
 
 		ForEachActorWithLoadingParams.OnPreGarbageCollect = [&PackagesToSave, &PackageHelper]()
 		{
diff --git a/Engine/Source/Editor/UnrealEd/Public/WorldPartition/WorldPartitionResaveActorsBuilder.h b/Engine/Source/Editor/UnrealEd/Public/WorldPartition/WorldPartitionResaveActorsBuilder.h
--- a/Engine/Source/Editor/UnrealEd/Public/WorldPartition/WorldPartitionResaveActorsBuilder.h
+++ b/Engine/Source/Editor/UnrealEd/Public/WorldPartition/WorldPartitionResaveActorsBuilder.h
@@ -5,7 +5,7 @@
 #include "WorldPartition/WorldPartitionBuilder.h"
 #include "WorldPartitionResaveActorsBuilder.generated.h"
 
-// Example Command Line: ProjectName MapName -run=WorldPartitionBuilderCommandlet -SCCProvider=Perforce -Builder=WorldPartitionResaveActorsBuilder [-ActorClassName=StaticMeshActor] [-SwitchActorPackagingSchemeToReduced] [-ActorTags=(Tag1,Tag2,...)]
+// Example Command Line: ProjectName MapName -run=WorldPartitionBuilderCommandlet -SCCProvider=Perforce -Builder=WorldPartitionResaveActorsBuilder [-ActorClassName=StaticMeshActor] [-SwitchActorPackagingSchemeToReduced] [-ActorTags=(Tag1,Tag2,...)] [-ActorProperties=((Property1,Value1),(Property2,Value2),...)]
 
 UCLASS()
 class UWorldPartitionResaveActorsBuilder : public UWorldPartitionBuilder
@@ -37,5 +37,8 @@ private:
 	bool bEnableActorFolders;
 
 	UPROPERTY()
-	TArray<FName> ActorTags;
+	TSet<FName> ActorTags;
+
+	UPROPERTY()
+	TMap<FName, FName> ActorProperties;
 };
\ No newline at end of file
