Whitespace: CookRequests.h/.cpp: Remove indentation from code in UE::Cook namespace to match the standard on the core team.
#rnx

diff --git a/Engine/Source/Editor/UnrealEd/Private/Cooker/CookRequests.cpp b/Engine/Source/Editor/UnrealEd/Private/Cooker/CookRequests.cpp
--- a/Engine/Source/Editor/UnrealEd/Private/Cooker/CookRequests.cpp
+++ b/Engine/Source/Editor/UnrealEd/Private/Cooker/CookRequests.cpp
@@ -11,319 +11,316 @@
 #include "Misc/ScopeLock.h"
 #include "Templates/UnrealTemplate.h"
 
-namespace UE
-{
-namespace Cook
+namespace UE::Cook
 {
 
-	//////////////////////////////////////////////////////////////////////////
-	// FFilePlatformRequest
+//////////////////////////////////////////////////////////////////////////
+// FFilePlatformRequest
 
-	FFilePlatformRequest::FFilePlatformRequest(const FName& InFilename, FInstigator&& InInstigator)
-		:FFilePlatformRequest(InFilename, MoveTemp(InInstigator),
-			TArray<const ITargetPlatform*>(), FCompletionCallback())
-	{
-	}
+FFilePlatformRequest::FFilePlatformRequest(const FName& InFilename, FInstigator&& InInstigator)
+	:FFilePlatformRequest(InFilename, MoveTemp(InInstigator),
+		TArray<const ITargetPlatform*>(), FCompletionCallback())
+{
+}
 
-	FFilePlatformRequest::FFilePlatformRequest(const FName& InFilename, FInstigator&& InInstigator, 
-		const ITargetPlatform* InPlatform, FCompletionCallback&& InCompletionCallback)
-		:FFilePlatformRequest(InFilename, MoveTemp(InInstigator),
-			TArrayView<const ITargetPlatform* const>({ InPlatform }), MoveTemp(InCompletionCallback))
-	{
-	}
+FFilePlatformRequest::FFilePlatformRequest(const FName& InFilename, FInstigator&& InInstigator, 
+	const ITargetPlatform* InPlatform, FCompletionCallback&& InCompletionCallback)
+	:FFilePlatformRequest(InFilename, MoveTemp(InInstigator),
+		TArrayView<const ITargetPlatform* const>({ InPlatform }), MoveTemp(InCompletionCallback))
+{
+}
 
-	FFilePlatformRequest::FFilePlatformRequest(const FName& InFilename, FInstigator&& InInstigator,
-		const TArrayView<const ITargetPlatform* const>& InPlatforms, FCompletionCallback&& InCompletionCallback)
-		:FFilePlatformRequest(InFilename, MoveTemp(InInstigator),
-			TArray<const ITargetPlatform*>(InPlatforms.GetData(), InPlatforms.Num()), MoveTemp(InCompletionCallback))
-	{
-	}
+FFilePlatformRequest::FFilePlatformRequest(const FName& InFilename, FInstigator&& InInstigator,
+	const TArrayView<const ITargetPlatform* const>& InPlatforms, FCompletionCallback&& InCompletionCallback)
+	:FFilePlatformRequest(InFilename, MoveTemp(InInstigator),
+		TArray<const ITargetPlatform*>(InPlatforms.GetData(), InPlatforms.Num()), MoveTemp(InCompletionCallback))
+{
+}
 
-	FFilePlatformRequest::FFilePlatformRequest(const FName& InFilename, FInstigator&& InInstigator,
-		TArray<const ITargetPlatform*>&& InPlatforms, FCompletionCallback&& InCompletionCallback)
-		: Platforms(MoveTemp(InPlatforms))
-		, CompletionCallback(MoveTemp(InCompletionCallback))
-		, Instigator(MoveTemp(InInstigator))
-	{
-		SetFilename(InFilename.ToString());
-	}
+FFilePlatformRequest::FFilePlatformRequest(const FName& InFilename, FInstigator&& InInstigator,
+	TArray<const ITargetPlatform*>&& InPlatforms, FCompletionCallback&& InCompletionCallback)
+	: Platforms(MoveTemp(InPlatforms))
+	, CompletionCallback(MoveTemp(InCompletionCallback))
+	, Instigator(MoveTemp(InInstigator))
+{
+	SetFilename(InFilename.ToString());
+}
 
-	FFilePlatformRequest::FFilePlatformRequest(const FFilePlatformRequest& InFilePlatformRequest)
-		: Filename(InFilePlatformRequest.Filename)
-		, Platforms(InFilePlatformRequest.Platforms)
-		, Instigator(InFilePlatformRequest.Instigator)
-	{
-		check(!InFilePlatformRequest.CompletionCallback); // CompletionCallbacks can not be copied, so the caller's intent is not clear in this constructor if the input has one
-	}
+FFilePlatformRequest::FFilePlatformRequest(const FFilePlatformRequest& InFilePlatformRequest)
+	: Filename(InFilePlatformRequest.Filename)
+	, Platforms(InFilePlatformRequest.Platforms)
+	, Instigator(InFilePlatformRequest.Instigator)
+{
+	check(!InFilePlatformRequest.CompletionCallback); // CompletionCallbacks can not be copied, so the caller's intent is not clear in this constructor if the input has one
+}
 		
-	FFilePlatformRequest::FFilePlatformRequest(FFilePlatformRequest&& InFilePlatformRequest)
-		: Filename(MoveTemp(InFilePlatformRequest.Filename))
-		, Platforms(MoveTemp(InFilePlatformRequest.Platforms))
-		, CompletionCallback(MoveTemp(InFilePlatformRequest.CompletionCallback))
-		, Instigator(MoveTemp(InFilePlatformRequest.Instigator))
-	{
-	}
+FFilePlatformRequest::FFilePlatformRequest(FFilePlatformRequest&& InFilePlatformRequest)
+	: Filename(MoveTemp(InFilePlatformRequest.Filename))
+	, Platforms(MoveTemp(InFilePlatformRequest.Platforms))
+	, CompletionCallback(MoveTemp(InFilePlatformRequest.CompletionCallback))
+	, Instigator(MoveTemp(InFilePlatformRequest.Instigator))
+{
+}
 
-	FFilePlatformRequest& FFilePlatformRequest::operator=(FFilePlatformRequest&& InFileRequest)
-	{
-		Filename = MoveTemp(InFileRequest.Filename);
-		Platforms = MoveTemp(InFileRequest.Platforms);
-		check(!CompletionCallback); // We don't support holding multiple completion callbacks
-		CompletionCallback = MoveTemp(InFileRequest.CompletionCallback);
-		Instigator = MoveTemp(InFileRequest.Instigator);
-		return *this;
-	}
+FFilePlatformRequest& FFilePlatformRequest::operator=(FFilePlatformRequest&& InFileRequest)
+{
+	Filename = MoveTemp(InFileRequest.Filename);
+	Platforms = MoveTemp(InFileRequest.Platforms);
+	check(!CompletionCallback); // We don't support holding multiple completion callbacks
+	CompletionCallback = MoveTemp(InFileRequest.CompletionCallback);
+	Instigator = MoveTemp(InFileRequest.Instigator);
+	return *this;
+}
 
-	void FFilePlatformRequest::SetFilename(FString InFilename)
-	{
-		Filename = FPackageDatas::GetStandardFileName(InFilename);
-	}
+void FFilePlatformRequest::SetFilename(FString InFilename)
+{
+	Filename = FPackageDatas::GetStandardFileName(InFilename);
+}
 
-	const FName& FFilePlatformRequest::GetFilename() const
-	{
-		return Filename;
-	}
+const FName& FFilePlatformRequest::GetFilename() const
+{
+	return Filename;
+}
 
-	FInstigator& FFilePlatformRequest::GetInstigator()
-	{
-		return Instigator;
-	}
+FInstigator& FFilePlatformRequest::GetInstigator()
+{
+	return Instigator;
+}
 
-	const TArray<const ITargetPlatform*>& FFilePlatformRequest::GetPlatforms() const
-	{
-		return Platforms;
-	}
+const TArray<const ITargetPlatform*>& FFilePlatformRequest::GetPlatforms() const
+{
+	return Platforms;
+}
 
-	TArray<const ITargetPlatform*>& FFilePlatformRequest::GetPlatforms()
-	{
-		return Platforms;
-	}
+TArray<const ITargetPlatform*>& FFilePlatformRequest::GetPlatforms()
+{
+	return Platforms;
+}
 
-	void FFilePlatformRequest::RemovePlatform(const ITargetPlatform* Platform)
-	{
-		Platforms.Remove(Platform);
-	}
+void FFilePlatformRequest::RemovePlatform(const ITargetPlatform* Platform)
+{
+	Platforms.Remove(Platform);
+}
 
-	void FFilePlatformRequest::AddPlatform(const ITargetPlatform* Platform)
-	{
-		check(Platform != nullptr);
-		Platforms.Add(Platform);
-	}
+void FFilePlatformRequest::AddPlatform(const ITargetPlatform* Platform)
+{
+	check(Platform != nullptr);
+	Platforms.Add(Platform);
+}
 
-	bool FFilePlatformRequest::HasPlatform(const ITargetPlatform* Platform) const
-	{
-		return Platforms.Find(Platform) != INDEX_NONE;
-	}
+bool FFilePlatformRequest::HasPlatform(const ITargetPlatform* Platform) const
+{
+	return Platforms.Find(Platform) != INDEX_NONE;
+}
 
-	FCompletionCallback& FFilePlatformRequest::GetCompletionCallback()
-	{
-		return CompletionCallback;
-	}
+FCompletionCallback& FFilePlatformRequest::GetCompletionCallback()
+{
+	return CompletionCallback;
+}
 
-	bool FFilePlatformRequest::IsValid() const
-	{
-		return Filename != NAME_None;
-	}
+bool FFilePlatformRequest::IsValid() const
+{
+	return Filename != NAME_None;
+}
 
-	void FFilePlatformRequest::Clear()
-	{
-		Filename = TEXT("");
-		Platforms.Empty();
-	}
+void FFilePlatformRequest::Clear()
+{
+	Filename = TEXT("");
+	Platforms.Empty();
+}
 
-	bool FFilePlatformRequest::operator==(const FFilePlatformRequest& InFileRequest) const
+bool FFilePlatformRequest::operator==(const FFilePlatformRequest& InFileRequest) const
+{
+	if (InFileRequest.Filename == Filename)
 	{
-		if (InFileRequest.Filename == Filename)
+		if (InFileRequest.Platforms == Platforms)
 		{
-			if (InFileRequest.Platforms == Platforms)
-			{
-				return true;
-			}
+			return true;
 		}
-		return false;
 	}
+	return false;
+}
 
-	FString FFilePlatformRequest::ToString() const
-	{
-		FString Result = FString::Printf(TEXT("%s;"), *Filename.ToString());
-
-		for (const ITargetPlatform* Platform : Platforms)
-		{
-			Result += FString::Printf(TEXT("%s,"), *Platform->PlatformName());
-		}
-		return Result;
-	}
+FString FFilePlatformRequest::ToString() const
+{
+	FString Result = FString::Printf(TEXT("%s;"), *Filename.ToString());
 
-	void FFilePlatformRequest::RemapTargetPlatforms(const TMap<ITargetPlatform*, ITargetPlatform*>& Remap)
+	for (const ITargetPlatform* Platform : Platforms)
 	{
-		RemapArrayElements(Platforms, Remap);
+		Result += FString::Printf(TEXT("%s,"), *Platform->PlatformName());
 	}
+	return Result;
+}
 
-	//////////////////////////////////////////////////////////////////////////
-	// FExternalRequests
+void FFilePlatformRequest::RemapTargetPlatforms(const TMap<ITargetPlatform*, ITargetPlatform*>& Remap)
+{
+	RemapArrayElements(Platforms, Remap);
+}
 
-	int32 FExternalRequests::GetNumRequests() const
-	{
-		return RequestCount;
-	}
+//////////////////////////////////////////////////////////////////////////
+// FExternalRequests
 
-	bool FExternalRequests::HasRequests() const
-	{
-		return RequestCount > 0;
-	}
+int32 FExternalRequests::GetNumRequests() const
+{
+	return RequestCount;
+}
 
-	void FExternalRequests::AddCallback(FSchedulerCallback&& Callback)
-	{
-		FScopeLock ScopeLock(&RequestLock);
+bool FExternalRequests::HasRequests() const
+{
+	return RequestCount > 0;
+}
 
-		Callbacks.Add(MoveTemp(Callback));
-		++RequestCount;
-	}
+void FExternalRequests::AddCallback(FSchedulerCallback&& Callback)
+{
+	FScopeLock ScopeLock(&RequestLock);
+
+	Callbacks.Add(MoveTemp(Callback));
+	++RequestCount;
+}
 
-	void FExternalRequests::EnqueueUnique(FFilePlatformRequest&& FileRequest, bool bForceFrontOfQueue)
+void FExternalRequests::EnqueueUnique(FFilePlatformRequest&& FileRequest, bool bForceFrontOfQueue)
+{
+	FScopeLock ScopeLock(&RequestLock);
+	FName Filename = FileRequest.GetFilename();
+	FFilePlatformRequest* ExistingRequest = RequestMap.Find(Filename);
+	if (!ExistingRequest)
 	{
-		FScopeLock ScopeLock(&RequestLock);
-		FName Filename = FileRequest.GetFilename();
-		FFilePlatformRequest* ExistingRequest = RequestMap.Find(Filename);
-		if (!ExistingRequest)
+		RequestMap.Add(Filename, MoveTemp(FileRequest));
+		if (bForceFrontOfQueue)
 		{
-			RequestMap.Add(Filename, MoveTemp(FileRequest));
-			if (bForceFrontOfQueue)
-			{
-				Queue.AddFront(Filename);
-			}
-			else
-			{
-				Queue.Add(Filename);
-			}
-
-			++RequestCount;
+			Queue.AddFront(Filename);
 		}
 		else
 		{
-			if (FileRequest.GetCompletionCallback())
-			{
-				check(!ExistingRequest->GetCompletionCallback()); // We don't support multiple callbacks
-				ExistingRequest->GetCompletionCallback() = MoveTemp(FileRequest.GetCompletionCallback());
-			}
-
-			// add the requested platforms to the platform list
-			for (const ITargetPlatform* Platform : FileRequest.GetPlatforms())
-			{
-				ExistingRequest->GetPlatforms().AddUnique(Platform);
-			}
-
-			if (bForceFrontOfQueue)
-			{
-				FName* ExistingName = Algo::Find(Queue, Filename);
-				int32 Index = Queue.ConvertPointerToIndex(ExistingName);
-				check(Index != INDEX_NONE);
-				if (Index != 0)
-				{
-					Queue[Index] = Queue[0];
-					Queue[0] = Filename;
-				}
-			}
+			Queue.Add(Filename);
 		}
-	}
 
-	EExternalRequestType FExternalRequests::DequeueNextCluster(TArray<FSchedulerCallback>& OutCallbacks, TArray<FFilePlatformRequest>& OutBuildRequests)
+		++RequestCount;
+	}
+	else
 	{
-		FScopeLock ScopeLock(&RequestLock);
+		if (FileRequest.GetCompletionCallback())
+		{
+			check(!ExistingRequest->GetCompletionCallback()); // We don't support multiple callbacks
+			ExistingRequest->GetCompletionCallback() = MoveTemp(FileRequest.GetCompletionCallback());
+		}
 
-		if (ThreadUnsafeDequeueCallbacks(OutCallbacks))
+		// add the requested platforms to the platform list
+		for (const ITargetPlatform* Platform : FileRequest.GetPlatforms())
 		{
-			return EExternalRequestType::Callback;
+			ExistingRequest->GetPlatforms().AddUnique(Platform);
 		}
-		else if (Queue.Num())
+
+		if (bForceFrontOfQueue)
 		{
-			OutBuildRequests.Reserve(Queue.Num() + OutBuildRequests.Num());
-			while (Queue.Num())
+			FName* ExistingName = Algo::Find(Queue, Filename);
+			int32 Index = Queue.ConvertPointerToIndex(ExistingName);
+			check(Index != INDEX_NONE);
+			if (Index != 0)
 			{
-				FName Filename = Queue.PopFrontValue();
-				OutBuildRequests.Add(RequestMap.FindAndRemoveChecked(Filename));
-				--RequestCount;
+				Queue[Index] = Queue[0];
+				Queue[0] = Filename;
 			}
-			return EExternalRequestType::Cook;
-		}
-		else
-		{
-			return EExternalRequestType::None;
 		}
 	}
+}
+
+EExternalRequestType FExternalRequests::DequeueNextCluster(TArray<FSchedulerCallback>& OutCallbacks, TArray<FFilePlatformRequest>& OutBuildRequests)
+{
+	FScopeLock ScopeLock(&RequestLock);
 
-	bool FExternalRequests::DequeueCallbacks(TArray<FSchedulerCallback>& OutCallbacks)
+	if (ThreadUnsafeDequeueCallbacks(OutCallbacks))
 	{
-		FScopeLock ScopeLock(&RequestLock);
-		return ThreadUnsafeDequeueCallbacks(OutCallbacks);
+		return EExternalRequestType::Callback;
 	}
-
-	bool FExternalRequests::ThreadUnsafeDequeueCallbacks(TArray<FSchedulerCallback>& OutCallbacks)
+	else if (Queue.Num())
 	{
-		if (Callbacks.Num() > 0)
-		{
-			OutCallbacks = MoveTemp(Callbacks);
-			Callbacks.Empty();
-			RequestCount -= OutCallbacks.Num();
-			return true;
-		}
-		else
+		OutBuildRequests.Reserve(Queue.Num() + OutBuildRequests.Num());
+		while (Queue.Num())
 		{
-			return false;
+			FName Filename = Queue.PopFrontValue();
+			OutBuildRequests.Add(RequestMap.FindAndRemoveChecked(Filename));
+			--RequestCount;
 		}
+		return EExternalRequestType::Cook;
 	}
-
-	void FExternalRequests::EmptyRequests()
+	else
 	{
-		FScopeLock ScopeLock(&RequestLock);
-
-		Queue.Empty();
-		RequestMap.Empty();
-		Callbacks.Empty();
-		RequestCount = 0;
+		return EExternalRequestType::None;
 	}
+}
 
-	void FExternalRequests::DequeueAll(TArray<FSchedulerCallback>& OutCallbacks, TArray<FFilePlatformRequest>& OutCookRequests)
-	{
-		FScopeLock ScopeLock(&RequestLock);
+bool FExternalRequests::DequeueCallbacks(TArray<FSchedulerCallback>& OutCallbacks)
+{
+	FScopeLock ScopeLock(&RequestLock);
+	return ThreadUnsafeDequeueCallbacks(OutCallbacks);
+}
 
+bool FExternalRequests::ThreadUnsafeDequeueCallbacks(TArray<FSchedulerCallback>& OutCallbacks)
+{
+	if (Callbacks.Num() > 0)
+	{
 		OutCallbacks = MoveTemp(Callbacks);
 		Callbacks.Empty();
+		RequestCount -= OutCallbacks.Num();
+		return true;
+	}
+	else
+	{
+		return false;
+	}
+}
 
-		for (TPair<FName, FFilePlatformRequest>& Request : RequestMap)
-		{
-			OutCookRequests.Add(MoveTemp(Request.Value));
-		}
-		RequestMap.Empty();
-		Queue.Empty();
+void FExternalRequests::EmptyRequests()
+{
+	FScopeLock ScopeLock(&RequestLock);
 
-		RequestCount = 0;
-	}
+	Queue.Empty();
+	RequestMap.Empty();
+	Callbacks.Empty();
+	RequestCount = 0;
+}
 
-	void FExternalRequests::OnRemoveSessionPlatform(const ITargetPlatform* TargetPlatform)
-	{
-		FScopeLock ScopeLock(&RequestLock);
+void FExternalRequests::DequeueAll(TArray<FSchedulerCallback>& OutCallbacks, TArray<FFilePlatformRequest>& OutCookRequests)
+{
+	FScopeLock ScopeLock(&RequestLock);
 
-		// The caller should not be removing platforms if we have an active request referencing that platform, but in case they did, remove the platform
-		// from all pending requests
-		for (TPair<FName, FFilePlatformRequest>& kvpair : RequestMap)
-		{
-			FFilePlatformRequest& Request = kvpair.Value;
-			Request.GetPlatforms().Remove(TargetPlatform);
-			if (Request.GetPlatforms().Num() == 0)
-			{
-				UE_LOG(LogCook, Error, TEXT("RemovePlatform call has left an empty list of platforms requested in CookOnTheSide request."));
-			}
-		}
+	OutCallbacks = MoveTemp(Callbacks);
+	Callbacks.Empty();
+
+	for (TPair<FName, FFilePlatformRequest>& Request : RequestMap)
+	{
+		OutCookRequests.Add(MoveTemp(Request.Value));
 	}
+	RequestMap.Empty();
+	Queue.Empty();
+
+	RequestCount = 0;
+}
+
+void FExternalRequests::OnRemoveSessionPlatform(const ITargetPlatform* TargetPlatform)
+{
+	FScopeLock ScopeLock(&RequestLock);
 
-	void FExternalRequests::RemapTargetPlatforms(const TMap<ITargetPlatform*, ITargetPlatform*>& Remap)
+	// The caller should not be removing platforms if we have an active request referencing that platform, but in case they did, remove the platform
+	// from all pending requests
+	for (TPair<FName, FFilePlatformRequest>& kvpair : RequestMap)
 	{
-		for (TPair<FName, FFilePlatformRequest>& KVPair : RequestMap)
+		FFilePlatformRequest& Request = kvpair.Value;
+		Request.GetPlatforms().Remove(TargetPlatform);
+		if (Request.GetPlatforms().Num() == 0)
 		{
-			KVPair.Value.RemapTargetPlatforms(Remap);
+			UE_LOG(LogCook, Error, TEXT("RemovePlatform call has left an empty list of platforms requested in CookOnTheSide request."));
 		}
 	}
+}
 
+void FExternalRequests::RemapTargetPlatforms(const TMap<ITargetPlatform*, ITargetPlatform*>& Remap)
+{
+	for (TPair<FName, FFilePlatformRequest>& KVPair : RequestMap)
+	{
+		KVPair.Value.RemapTargetPlatforms(Remap);
+	}
 }
+
 }
\ No newline at end of file
diff --git a/Engine/Source/Editor/UnrealEd/Private/Cooker/CookRequests.h b/Engine/Source/Editor/UnrealEd/Private/Cooker/CookRequests.h
--- a/Engine/Source/Editor/UnrealEd/Private/Cooker/CookRequests.h
+++ b/Engine/Source/Editor/UnrealEd/Private/Cooker/CookRequests.h
@@ -19,118 +19,119 @@ class ITargetPlatform;
 namespace UE::Cook
 {
 
+/**
+ * Structure holding the data for a request for the CookOnTheFlyServer to cook a FileName. Includes platform which file is requested for.
+ * These requests are external to the cooker's scheduler, and do not use the FPackageData the scheduler uses internally.
+ */
+struct FFilePlatformRequest
+{
+protected:
+	FName Filename;
+	TArray<const ITargetPlatform*> Platforms;
+	FCompletionCallback CompletionCallback;
+	FInstigator Instigator;
+
+public:
+	FFilePlatformRequest() = default;
+
+	FFilePlatformRequest(const FName& InFilename, FInstigator&& InInstigator);
+	FFilePlatformRequest(const FName& InFilename, FInstigator&& InInstigator, const ITargetPlatform* InPlatform,
+		FCompletionCallback&& InCompletionCallback = FCompletionCallback());
+	FFilePlatformRequest(const FName& InFilename, FInstigator&& InInstigator,
+		const TArrayView<const ITargetPlatform* const>& InPlatforms, FCompletionCallback&& InCompletionCallback = FCompletionCallback());
+	FFilePlatformRequest(const FName& InFilename, FInstigator&& InInstigator,
+		TArray<const ITargetPlatform*>&& InPlatforms, FCompletionCallback&& InCompletionCallback = FCompletionCallback());
+	FFilePlatformRequest(const FFilePlatformRequest& InFilePlatformRequest);
+	FFilePlatformRequest(FFilePlatformRequest&& InFilePlatformRequest);
+	FFilePlatformRequest& operator=(FFilePlatformRequest&& InFileRequest);
+
+	void SetFilename(FString InFilename);
+	const FName& GetFilename() const;
+	FInstigator& GetInstigator();
+
+	const TArray<const ITargetPlatform*>& GetPlatforms() const;
+	TArray<const ITargetPlatform*>& GetPlatforms();
+	void RemovePlatform(const ITargetPlatform* Platform);
+	void AddPlatform(const ITargetPlatform* Platform);
+	bool HasPlatform(const ITargetPlatform* Platform) const;
+
+	/** A callback that the scheduler will call after the request is processed and is cooked, fails to cook, is canceled, or is skipped because it already exists. */
+	FCompletionCallback& GetCompletionCallback();
+
+	bool IsValid() const;
+	void Clear();
+	bool operator==(const FFilePlatformRequest& InFileRequest) const;
+	FString ToString() const;
+
+	void RemapTargetPlatforms(const TMap<ITargetPlatform*, ITargetPlatform*>& Remap);
+};
+
+
+/**
+ * A container class for External Requests made to the cooker.
+ * External Requests are cook requests that are made outside of the scheduler's lock and hence need to be separately synchronized.
+ * External Requests can be either a request to cook a given FileName (packages are identified by FileName in this container) on given platforms,
+ * or a request to run an arbitrary callback inside the scheduler's lock.
+ * This class is threadsafe; all methods are guarded by a CriticalSection. 
+ */
+class FExternalRequests
+{
+public:
+
+	/**
+	 * Lockless value for the number of External Requests in the container.
+	 * May be out of date after calling; do not assume the number of actual requests is any one of equal, greater than, or less than the returned value.
+	 * Intended usage is for the Scheduler to be the only consumer of requests, and to use this value for rough reporting of periodic progress.
+	 */
+	int32 GetNumRequests() const;
+
 	/**
-	 * Structure holding the data for a request for the CookOnTheFlyServer to cook a FileName. Includes platform which file is requested for.
-	 * These requests are external to the cooker's scheduler, and do not use the FPackageData the scheduler uses internally.
+	 * Lockless value for the number of External Requests in the container.
+	 * May be out of date after calling; do not assume a true return value means Requests are actually present or a false value means no Requests are present.
+	 * Intended usage is for the Scheduler to be the only consumer of requests, and to use this value for periodic checking of whether there is any work that justifies the expense of taking the lock.
+	 * In a single-consumer case, HasRequests will eventually correctly return true as long as the consumer is not consuming.
 	 */
-	struct FFilePlatformRequest
-	{
-	protected:
-		FName Filename;
-		TArray<const ITargetPlatform*> Platforms;
-		FCompletionCallback CompletionCallback;
-		FInstigator Instigator;
-
-	public:
-		FFilePlatformRequest() = default;
-
-		FFilePlatformRequest(const FName& InFilename, FInstigator&& InInstigator);
-		FFilePlatformRequest(const FName& InFilename, FInstigator&& InInstigator, const ITargetPlatform* InPlatform,
-			FCompletionCallback&& InCompletionCallback = FCompletionCallback());
-		FFilePlatformRequest(const FName& InFilename, FInstigator&& InInstigator,
-			const TArrayView<const ITargetPlatform* const>& InPlatforms, FCompletionCallback&& InCompletionCallback = FCompletionCallback());
-		FFilePlatformRequest(const FName& InFilename, FInstigator&& InInstigator,
-			TArray<const ITargetPlatform*>&& InPlatforms, FCompletionCallback&& InCompletionCallback = FCompletionCallback());
-		FFilePlatformRequest(const FFilePlatformRequest& InFilePlatformRequest);
-		FFilePlatformRequest(FFilePlatformRequest&& InFilePlatformRequest);
-		FFilePlatformRequest& operator=(FFilePlatformRequest&& InFileRequest);
-
-		void SetFilename(FString InFilename);
-		const FName& GetFilename() const;
-		FInstigator& GetInstigator();
-
-		const TArray<const ITargetPlatform*>& GetPlatforms() const;
-		TArray<const ITargetPlatform*>& GetPlatforms();
-		void RemovePlatform(const ITargetPlatform* Platform);
-		void AddPlatform(const ITargetPlatform* Platform);
-		bool HasPlatform(const ITargetPlatform* Platform) const;
-
-		/** A callback that the scheduler will call after the request is processed and is cooked, fails to cook, is canceled, or is skipped because it already exists. */
-		FCompletionCallback& GetCompletionCallback();
-
-		bool IsValid() const;
-		void Clear();
-		bool operator==(const FFilePlatformRequest& InFileRequest) const;
-		FString ToString() const;
-
-		void RemapTargetPlatforms(const TMap<ITargetPlatform*, ITargetPlatform*>& Remap);
-	};
+	bool HasRequests() const;
 
+	/** Add a callback-type request. The scheduler will run all callbacks (in FIFO order) as soon as it completes its current task. */
+	void AddCallback(FSchedulerCallback&& Callback);
+
+	/** Add the given cook-type request, merging its list of platforms with any existing request if one already exists. */
+	void EnqueueUnique(FFilePlatformRequest&& FileRequest, bool bForceFrontOfQueue = false);
 
 	/**
-	 * A container class for External Requests made to the cooker.
-	 * External Requests are cook requests that are made outside of the scheduler's lock and hence need to be separately synchronized.
-	 * External Requests can be either a request to cook a given FileName (packages are identified by FileName in this container) on given platforms,
-	 * or a request to run an arbitrary callback inside the scheduler's lock.
-	 * This class is threadsafe; all methods are guarded by a CriticalSection. 
+	 * If this FExternalRequests has any callbacks, dequeue them all into OutCallbacks and return EExternalRequestType::Callback; Callbacks take priority over cook requests.
+	 * Otherwise, if there are any cook requests, dequeue them all into OutBuildRequests and return EExternalRequestType::Cook.
+	 * Otherwise, return EExternalRequestType::None.
 	 */
-	class FExternalRequests
-	{
-	public:
-
-		/**
-		 * Lockless value for the number of External Requests in the container.
-		 * May be out of date after calling; do not assume the number of actual requests is any one of equal, greater than, or less than the returned value.
-		 * Intended usage is for the Scheduler to be the only consumer of requests, and to use this value for rough reporting of periodic progress.
-		 */
-		int32 GetNumRequests() const;
-
-		/**
-		 * Lockless value for the number of External Requests in the container.
-		 * May be out of date after calling; do not assume a true return value means Requests are actually present or a false value means no Requests are present.
-		 * Intended usage is for the Scheduler to be the only consumer of requests, and to use this value for periodic checking of whether there is any work that justifies the expense of taking the lock.
-		 * In a single-consumer case, HasRequests will eventually correctly return true as long as the consumer is not consuming.
-		 */
-		bool HasRequests() const;
-
-		/** Add a callback-type request. The scheduler will run all callbacks (in FIFO order) as soon as it completes its current task. */
-		void AddCallback(FSchedulerCallback&& Callback);
-
-		/** Add the given cook-type request, merging its list of platforms with any existing request if one already exists. */
-		void EnqueueUnique(FFilePlatformRequest&& FileRequest, bool bForceFrontOfQueue = false);
-
-		/**
-		 * If this FExternalRequests has any callbacks, dequeue them all into OutCallbacks and return EExternalRequestType::Callback; Callbacks take priority over cook requests.
-		 * Otherwise, if there are any cook requests, dequeue them all into OutBuildRequests and return EExternalRequestType::Cook.
-		 * Otherwise, return EExternalRequestType::None.
-		 */
-		EExternalRequestType DequeueNextCluster(TArray<FSchedulerCallback>& OutCallbacks, TArray<FFilePlatformRequest>& OutBuildRequests);
-		/* Move any existing callbacks onto OutCallbacks, and return whether any were added. */
-		bool DequeueCallbacks(TArray<FSchedulerCallback>& OutCallbacks);
-		/** Eliminate all callbacks and cook requests and free memory */
-		void EmptyRequests();
-		/** Move all callbacks into OutCallbacks, and all cook requests into OutCookRequests. This is used when canceling a cook session. */
-		void DequeueAll(TArray<FSchedulerCallback>& OutCallbacks, TArray<FFilePlatformRequest>& OutCookRequests);
-
-		/** Remove references to the given platform from all cook requests. */
-		void OnRemoveSessionPlatform(const ITargetPlatform* TargetPlatform);
-
-		/** Swap all ITargetPlatform* stored on this instance according to the mapping in @param Remap. */
-		void RemapTargetPlatforms(const TMap<ITargetPlatform*, ITargetPlatform*>& Remap);
-
-	public:
-		/** An FEvent the scheduler can sleep on when waiting for new cookonthefly requests. */
-		FEvent* CookRequestEvent = nullptr;
-	private:
-		/* Implementation for DequeueCallbacks that assumes the caller has entered the RequestLock. */
-		bool ThreadUnsafeDequeueCallbacks(TArray<FSchedulerCallback>& OutCallbacks);
-
-		/** Queue of the FileName for the cook-type requests in this instance. The FileName can be used to look up the rest of the data for the request. */
-		TRingBuffer<FName> Queue;
-		/** Map of the extended information for the cook-type requests in this instance, keyed by the FileName of the request. */
-		TMap<FName, FFilePlatformRequest> RequestMap;
-		TArray<FSchedulerCallback> Callbacks;
-		FCriticalSection RequestLock;
-		int32 RequestCount = 0;
-	};
+	EExternalRequestType DequeueNextCluster(TArray<FSchedulerCallback>& OutCallbacks, TArray<FFilePlatformRequest>& OutBuildRequests);
+	/* Move any existing callbacks onto OutCallbacks, and return whether any were added. */
+	bool DequeueCallbacks(TArray<FSchedulerCallback>& OutCallbacks);
+	/** Eliminate all callbacks and cook requests and free memory */
+	void EmptyRequests();
+	/** Move all callbacks into OutCallbacks, and all cook requests into OutCookRequests. This is used when canceling a cook session. */
+	void DequeueAll(TArray<FSchedulerCallback>& OutCallbacks, TArray<FFilePlatformRequest>& OutCookRequests);
+
+	/** Remove references to the given platform from all cook requests. */
+	void OnRemoveSessionPlatform(const ITargetPlatform* TargetPlatform);
+
+	/** Swap all ITargetPlatform* stored on this instance according to the mapping in @param Remap. */
+	void RemapTargetPlatforms(const TMap<ITargetPlatform*, ITargetPlatform*>& Remap);
+
+public:
+	/** An FEvent the scheduler can sleep on when waiting for new cookonthefly requests. */
+	FEvent* CookRequestEvent = nullptr;
+private:
+	/* Implementation for DequeueCallbacks that assumes the caller has entered the RequestLock. */
+	bool ThreadUnsafeDequeueCallbacks(TArray<FSchedulerCallback>& OutCallbacks);
+
+	/** Queue of the FileName for the cook-type requests in this instance. The FileName can be used to look up the rest of the data for the request. */
+	TRingBuffer<FName> Queue;
+	/** Map of the extended information for the cook-type requests in this instance, keyed by the FileName of the request. */
+	TMap<FName, FFilePlatformRequest> RequestMap;
+	TArray<FSchedulerCallback> Callbacks;
+	FCriticalSection RequestLock;
+	int32 RequestCount = 0;
+};
+
 }
\ No newline at end of file
