[WP navmesh] WorldPartitionNavigationDataBuilder -> Do not checkout package that are being deleted to prevent overriding the deleting status

diff --git a/Engine/Source/Editor/UnrealEd/Private/WorldPartition/WorldPartitionNavigationDataBuilder.cpp b/Engine/Source/Editor/UnrealEd/Private/WorldPartition/WorldPartitionNavigationDataBuilder.cpp
--- a/Engine/Source/Editor/UnrealEd/Private/WorldPartition/WorldPartitionNavigationDataBuilder.cpp
+++ b/Engine/Source/Editor/UnrealEd/Private/WorldPartition/WorldPartitionNavigationDataBuilder.cpp
@@ -213,18 +213,18 @@ bool UWorldPartitionNavigationDataBuilder::RunInternal(UWorld* World, const FCel
 	if (!PackagesToSave.IsEmpty())
 	{
 		{
-			// Checkout or remove read-only for packages to save
+			// Checkout or remove read-only for packages to add
 			TRACE_CPUPROFILER_EVENT_SCOPE(CheckoutPackages);
-			UE_LOG(LogWorldPartitionNavigationDataBuilder, Log, TEXT("Checking out %d packages."), PackagesToSave.Num());
+			UE_LOG(LogWorldPartitionNavigationDataBuilder, Log, TEXT("Checking out %d packages."), PackagesToAdd.Num());
 
 			if (PackageHelper.UseSourceControl())
 			{
-				FEditorFileUtils::CheckoutPackages(PackagesToSave, /*OutPackagesCheckedOut*/nullptr, /*bErrorIfAlreadyCheckedOut*/false);
+				FEditorFileUtils::CheckoutPackages(PackagesToAdd, /*OutPackagesCheckedOut*/nullptr, /*bErrorIfAlreadyCheckedOut*/false);
 			}
 			else
 			{
 				// Remove read-only
-				for (const UPackage* Package : PackagesToSave)
+				for (const UPackage* Package : PackagesToAdd)
 				{
 					const FString PackageFilename = SourceControlHelpers::PackageFilename(Package);
 					if (IPlatformFile::GetPlatformPhysical().FileExists(*PackageFilename))
