Added ToolMenus and IMovieScenePlayer.

Limiting the build to 8 processors:
Before: 1490.40 seconds - 27.3 GB (29,323,796,480 bytes)
After: 1470.59 seconds - 27.4 GB (29,443,416,064 bytes)

diff --git a/Engine/Source/Editor/UnrealEd/Public/UnrealEdSharedPCH.h b/Engine/Source/Editor/UnrealEd/Public/UnrealEdSharedPCH.h
--- a/Engine/Source/Editor/UnrealEd/Public/UnrealEdSharedPCH.h
+++ b/Engine/Source/Editor/UnrealEd/Public/UnrealEdSharedPCH.h
@@ -618,3 +618,9 @@
 #include "EditorViewportClient.h"
 #include "Factories/Factory.h"
 #include "Settings/EditorLoadingSavingSettings.h"
+
+// From ToolMenus
+#include "ToolMenus.h"
+
+// From Movie
+#include "IMovieScenePlayer.h"
\ No newline at end of file
