Add Filtering mechanism to MainFrame's MRUFavoritesList and LevelEditorMenu's 'Open Recent Level'

diff --git a/Engine/Source/Editor/LevelEditor/Private/LevelEditorMenu.cpp b/Engine/Source/Editor/LevelEditor/Private/LevelEditorMenu.cpp
--- a/Engine/Source/Editor/LevelEditor/Private/LevelEditorMenu.cpp
+++ b/Engine/Source/Editor/LevelEditor/Private/LevelEditorMenu.cpp
@@ -161,6 +161,11 @@ void FLevelEditorMenu::RegisterLevelEditorMenus()
 						{
 							TSharedPtr< FUICommandInfo > OpenRecentFile = FLevelEditorCommands::Get().OpenRecentFileCommands[ CurRecentIndex ];
 
+							if (!MRUFavorites.MRUItemPassesCurrentFilter(CurRecentIndex))
+							{
+								continue;
+							}
+
 							const FString CurRecent = MRUFavorites.GetMRUItem( CurRecentIndex );
 
 							const FText ToolTip = FText::Format( LOCTEXT( "RecentLevelToolTip", "Opens recent level: {0}" ), FText::FromString( CurRecent ) );
diff --git a/Engine/Source/Editor/UnrealEd/Private/MRUFavoritesList.cpp b/Engine/Source/Editor/UnrealEd/Private/MRUFavoritesList.cpp
--- a/Engine/Source/Editor/UnrealEd/Private/MRUFavoritesList.cpp
+++ b/Engine/Source/Editor/UnrealEd/Private/MRUFavoritesList.cpp
@@ -153,3 +153,58 @@ bool FMainMRUFavoritesList::VerifyFavoritesFile( int32 ItemIndex )
 
 	return true;
 }
+
+/**
+* Supplies an optional delegate that can be used to filter a given MRUFavorites item
+* Useful for dynamically verifying which items should be utilized at a given time
+*
+* @param DoesMRUFavoritesItemPassFilterDelegate The delegate to use
+*/
+void FMainMRUFavoritesList::RegisterDoesMRUFavoritesItemPassFilterDelegate(FDoesMRUFavoritesItemPassFilter DoesMRUFavoritesItemPassFilterDelegate)
+{
+	DoesMRUFavoritesItemPassFilter = DoesMRUFavoritesItemPassFilterDelegate;
+}
+
+/**
+* Unregisters the optional filter delegate
+*/
+void FMainMRUFavoritesList::UnregisterDoesMRUFavoritesItemPassFilterDelegate()
+{
+	DoesMRUFavoritesItemPassFilter = FDoesMRUFavoritesItemPassFilter();
+}
+
+/**
+* Checks the favorites item specified by the provided index against the optional 'DoesMRUFavoritesItemPassFilterDelegate'.
+*
+* @param ItemIndex Index of the favorites item to check
+*
+* @return true if the item specified by the index passes the filter or if no filter has been provided; false if it does not pass the filter
+*/
+bool FMainMRUFavoritesList::FavoritesItemPassesCurrentFilter(int32 ItemIndex) const
+{
+	check(FavoriteItems.IsValidIndex(ItemIndex));
+	if (DoesMRUFavoritesItemPassFilter.IsBound())
+	{
+		const FString& FavoriteItem = FavoriteItems[ItemIndex];
+		return DoesMRUFavoritesItemPassFilter.Execute(FavoriteItem);
+	}
+	return true;
+}
+
+/**
+* Checks the MRU item specified by the provided index against the optional 'DoesMRUFavoritesItemPassFilterDelegate'.
+*
+* @param ItemIndex Index of the MRU item to check
+*
+* @return true if the item specified by the index passes the filter or if no filter has been provided; false if it does not pass the filter
+*/
+bool FMainMRUFavoritesList::MRUItemPassesCurrentFilter(int32 ItemIndex) const
+{
+	check(Items.IsValidIndex(ItemIndex));
+	if (DoesMRUFavoritesItemPassFilter.IsBound())
+	{
+		const FString& MRUItem = Items[ItemIndex];
+		return DoesMRUFavoritesItemPassFilter.Execute(MRUItem);
+	}
+	return true;
+}
diff --git a/Engine/Source/Editor/UnrealEd/Public/MRUFavoritesList.h b/Engine/Source/Editor/UnrealEd/Public/MRUFavoritesList.h
--- a/Engine/Source/Editor/UnrealEd/Public/MRUFavoritesList.h
+++ b/Engine/Source/Editor/UnrealEd/Public/MRUFavoritesList.h
@@ -87,8 +87,43 @@ public:
 	 */
 	void MoveFavoritesItemToHead( const FString& Item );
 
+	DECLARE_DELEGATE_RetVal_OneParam(bool, FDoesMRUFavoritesItemPassFilter, const FString& MRUFavoritesItem);
+	/**
+	* Supplies an optional delegate that can be used to filter a given MRUFavorites item
+	* Useful for dynamically verifying which items should be utilized at a given time
+	*
+	* @param DoesMRUFavoritesItemPassFilterDelegate The delegate to use
+	*/
+	void RegisterDoesMRUFavoritesItemPassFilterDelegate(FDoesMRUFavoritesItemPassFilter DoesMRUFavoritesItemPassFilterDelegate);
+
+	/**
+	* Unregisters the optional filter delegate
+	*/
+	void UnregisterDoesMRUFavoritesItemPassFilterDelegate();
+
+	/**
+	* Checks the favorites item specified by the provided index against the optional 'DoesMRUFavoritesItemPassFilterDelegate'.
+	*
+	* @param ItemIndex Index of the favorites item to check
+	*
+	* @return true if the item specified by the index passes the filter or if no filter has been provided; false if it does not pass the filter
+	*/
+	bool FavoritesItemPassesCurrentFilter(int32 ItemIndex) const;
+
+	/**
+	* Checks the MRU item specified by the provided index against the optional 'DoesMRUFavoritesItemPassFilterDelegate'.
+	*
+	* @param ItemIndex Index of the MRU item to check
+	*
+	* @return true if the item specified by the index passes the filter or if no filter has been provided; false if it does not pass the filter
+	*/
+	bool MRUItemPassesCurrentFilter(int32 ItemIndex) const;
+
 private:
 
+	/** Filter delegate */
+	FDoesMRUFavoritesItemPassFilter DoesMRUFavoritesItemPassFilter;
+
 	/** Favorited items */
 	TArray<FString> FavoriteItems;
 
