- Save Current Level / Ctrl+S: Allow skipping the checkout dialog and save only the writable files.
- Fix issue when Level was being saved and it had some non writable dirty external packages
- Fix reentrance issue between FEditorFileUtils::PromptForCheckoutAndSave and FEditorFileUtils::PromptToCheckoutPackages

diff --git a/Engine/Source/Editor/PackagesDialog/Private/SPackagesDialog.cpp b/Engine/Source/Editor/PackagesDialog/Private/SPackagesDialog.cpp
--- a/Engine/Source/Editor/PackagesDialog/Private/SPackagesDialog.cpp
+++ b/Engine/Source/Editor/PackagesDialog/Private/SPackagesDialog.cpp
@@ -443,17 +443,15 @@ void SPackagesDialog::RefreshButtons()
 		FPackageButton& Button = *Buttons[ButtonIndex];
 		if(Button.GetType() == DRT_MakeWritable)
 		{
-			if(UndeterminedItems > 0 || CheckedItems > 0)
-				Button.SetDisabled(false);
-			else
-				Button.SetDisabled(true);
+			Button.SetDisabled(UndeterminedItems == 0 && CheckedItems == 0);
 		}
 		else if(Button.GetType() == DRT_CheckOut)
 		{
-			if(CheckedItems > 0)
-				Button.SetDisabled(false);
-			else
-				Button.SetDisabled(true);
+			Button.SetDisabled(CheckedItems == 0);
+		}
+		else if (Button.GetType() == DRT_Skip)
+		{
+			Button.SetDisabled(CheckedItems > 0);
 		}
 	}
 }
diff --git a/Engine/Source/Editor/PackagesDialog/Public/PackagesDialog.h b/Engine/Source/Editor/PackagesDialog/Public/PackagesDialog.h
--- a/Engine/Source/Editor/PackagesDialog/Public/PackagesDialog.h
+++ b/Engine/Source/Editor/PackagesDialog/Public/PackagesDialog.h
@@ -17,6 +17,7 @@ enum EDialogReturnType
 	DRT_CheckOut,
 	DRT_MakeWritable,
 	DRT_Cancel,
+	DRT_Skip,
 	DRT_None
 };
 
diff --git a/Engine/Source/Editor/UnrealEd/Private/FileHelpers.cpp b/Engine/Source/Editor/UnrealEd/Private/FileHelpers.cpp
--- a/Engine/Source/Editor/UnrealEd/Private/FileHelpers.cpp
+++ b/Engine/Source/Editor/UnrealEd/Private/FileHelpers.cpp
@@ -88,6 +88,7 @@ DEFINE_LOG_CATEGORY_STATIC(LogFileHelpers, Log, All);
 //definition of flag used to do special work when we're attempting to load the "startup map"
 bool FEditorFileUtils::bIsLoadingDefaultStartupMap = false;
 bool FEditorFileUtils::bIsPromptingForCheckoutAndSave = false;
+bool FEditorFileUtils::bSkipExternalObjectSave = false;
 TSet<FString> FEditorFileUtils::PackagesNotSavedDuringSaveAll;
 TSet<FString> FEditorFileUtils::PackagesNotToPromptAnyMore;
 
@@ -845,7 +846,7 @@ static bool SaveWorld(UWorld* World,
 		// This makes UEditorEngine::Save's own call to InitializePhysicsSceneForSaveIfNecessary redundant but wasn't removed to avoid breaking other code paths
 		const bool bInitializedPhysicsSceneForSave = GEditor->InitializePhysicsSceneForSaveIfNecessary(SaveWorld, bForceInitializedWorld);
 				
-		if(!bAutosaving)
+		if(!bAutosaving && !FEditorFileUtils::ShouldSkipExternalObjectSave())
 		{
 			if (bSuccess)
 			{
@@ -1733,7 +1734,20 @@ void FEditorFileUtils::UpdateCheckoutPackageItems(bool bCheckDirty, TArray<UPack
 	AddCheckoutPackageItems(bCheckDirty, PackagesToCheckOut, OutPackagesNotNeedingCheckout, nullptr);
 }
 
-bool FEditorFileUtils::PromptToCheckoutPackages(bool bCheckDirty, const TArray<UPackage*>& PackagesToCheckOut, TArray<UPackage*>* OutPackagesCheckedOutOrMadeWritable, TArray<UPackage*>* OutPackagesNotNeedingCheckout, const bool bPromptingAfterModify )
+bool FEditorFileUtils::PromptToCheckoutPackages(bool bCheckDirty, const TArray<UPackage*>& PackagesToCheckOut, TArray<UPackage*>* OutPackagesCheckedOutOrMadeWritable, TArray<UPackage*>* OutPackagesNotNeedingCheckout, const bool bPromptingAfterModify, const bool bAllowSkip)
+{
+	if (bIsPromptingForCheckoutAndSave)
+	{
+		return false;
+	}
+
+	// Prevent re-entrance into this function by setting up a guard value (also used by FEditorFileUtils::PromptForCheckoutAndSave)
+	TGuardValue<bool> PromptForCheckoutAndSaveGuard(bIsPromptingForCheckoutAndSave, true);
+
+	return PromptToCheckoutPackagesInternal(bCheckDirty, PackagesToCheckOut, OutPackagesCheckedOutOrMadeWritable, OutPackagesNotNeedingCheckout, bPromptingAfterModify, bAllowSkip);
+}
+
+bool FEditorFileUtils::PromptToCheckoutPackagesInternal(bool bCheckDirty, const TArray<UPackage*>& PackagesToCheckOut, TArray<UPackage*>* OutPackagesCheckedOutOrMadeWritable, TArray<UPackage*>* OutPackagesNotNeedingCheckout, const bool bPromptingAfterModify, const bool bAllowSkip )
 {
 	bool bResult = true;
 
@@ -1795,6 +1809,12 @@ bool FEditorFileUtils::PromptToCheckoutPackages(bool bCheckDirty, const TArray<U
 		// Make writable button to make checked files writable
 		CheckoutPackagesDialogModule.AddButton(DRT_MakeWritable, NSLOCTEXT("PackagesDialogModule", "Dlg_MakeWritableButton", "Make Writable"), NSLOCTEXT("PackagesDialogModule", "Dlg_MakeWritableTooltip", "Makes selected files writable on disk"));
 
+		if (bAllowSkip)
+		{
+			// Skip button to skip checkout step
+			CheckoutPackagesDialogModule.AddButton(DRT_Skip, NSLOCTEXT("PackagesDialogModule", "Dlg_SkipButton", "Skip"), NSLOCTEXT("PackagesDialogModule", "Dlg_SkipTooltip", "Save all files that are writable, but don't check any files out from source control or make them writable."));
+		}
+
 		// The cancel button should be different if we are prompting during a modify.
 		const FText CancelButtonText  = bPromptingAfterModify ? NSLOCTEXT("PackagesDialogModule", "Dlg_AskMeLater", "Ask Me Later") : NSLOCTEXT("PackagesDialogModule", "Dlg_Cancel", "Cancel");
 		const FText CancelButtonToolTip = bPromptingAfterModify ? NSLOCTEXT("PackagesDialogModule", "Dlg_AskMeLaterToolTip", "Don't ask again until this asset is saved") : NSLOCTEXT("PackagesDialogModule", "Dlg_CancelTooltip", "Cancel Request"); 
@@ -1889,7 +1909,7 @@ bool FEditorFileUtils::PromptToCheckoutPackages(bool bCheckDirty, const TArray<U
 
 				bPerformedOperation = true;
 			}
-			else if (UserResponse == DRT_Save)
+			else if (UserResponse == DRT_Save || UserResponse == DRT_Skip)
 			{
 				bResult = true;
 				bPerformedOperation = true;
@@ -3855,31 +3875,35 @@ bool FEditorFileUtils::SaveCurrentLevel()
 	ULevel* Level = GWorld->GetCurrentLevel();
 	if (Level)
 	{
-		// Check dirtiness if the level is using external actors, no need to save it needlessly
-		bool bCheckDirty = Level->IsUsingExternalActors();
-		if (FEditorFileUtils::PromptToCheckoutLevels(bCheckDirty, Level))
+		// Check dirtiness if the level is using external objects, no need to save it needlessly
+		const bool bCheckDirty = Level->IsUsingExternalObjects();
+
+		TArray<UPackage*> PackagesToSave;
+		
+		UPackage* LevelPackage = Level->GetPackage();
+		// Get Packages to save
+		if (!bCheckDirty || LevelPackage->IsDirty() || LevelPackage->HasAnyPackageFlags(PKG_NewlyCreated))
 		{
-			UPackage* LevelPackage = Level->GetPackage();
-			// Save the level
-			if (!bCheckDirty || LevelPackage->IsDirty() || LevelPackage->HasAnyPackageFlags(PKG_NewlyCreated))
-			{
-				bReturnCode &= FEditorFileUtils::SaveLevel(Level);
-			}
+			PackagesToSave.Add(LevelPackage);
+		}
 
-			// Gather the level owned packages (i.e external actors and save them)
-			TArray<UPackage*> PackagesToSave = Level->GetLoadedExternalObjectPackages();
-			for (auto It = PackagesToSave.CreateIterator(); It; ++It)
+		// Get External Packages to save
+		const TArray<UPackage*> ExternalPackages = Level->GetLoadedExternalObjectPackages();
+		for (UPackage* ExternalPackage : ExternalPackages)
+		{
+			if (FPackageName::IsValidLongPackageName(ExternalPackage->GetName()))
 			{
-				UPackage* Package = *It;
-				if (bCheckDirty && !Package->IsDirty() && !UPackage::IsEmptyPackage(Package))
-				{
-					It.RemoveCurrent();
-				} // Remove package unsaved packages located in /Temp (they should be saved with the level's first save)
-				else if (!FPackageName::IsValidLongPackageName(Package->GetName()))
+				if (!bCheckDirty || ExternalPackage->IsDirty() || UPackage::IsEmptyPackage(ExternalPackage))
 				{
-					It.RemoveCurrent(); 
+					PackagesToSave.Add(ExternalPackage);
 				}
 			}
+		}
+
+		if (PackagesToSave.Num())
+		{
+			// If Level gets saved we don't want it to save its external packages because we've already filtered out the ones that need saving and they are part of the PackagesToSave array
+			TGuardValue<bool> GuardValue(bSkipExternalObjectSave, true);
 			bReturnCode &= InternalSavePackages(PackagesToSave, false, false, false);
 		}
 	}
@@ -4093,7 +4117,7 @@ FEditorFileUtils::EPromptReturnCode FEditorFileUtils::PromptForCheckoutAndSave(
 		return PR_Cancelled;
 	}
 
-	// Prevent re-entrance into this function by setting up a guard value
+	// Prevent re-entrance into this function by setting up a guard value (also used by FEditorFileUtils::PromptToCheckoutPackages)
 	TGuardValue<bool> PromptForCheckoutAndSaveGuard(bIsPromptingForCheckoutAndSave, true);
 
 	// Initialize the value we will return to indicate success
@@ -4275,7 +4299,9 @@ FEditorFileUtils::EPromptReturnCode FEditorFileUtils::PromptForCheckoutAndSave(
 
 		if (!bAlreadyCheckedOut)
 		{
-			bUserResponse = FEditorFileUtils::PromptToCheckoutPackages(false, PackagesToSave, &PackagesCheckedOutOrMadeWritable, &PackagesNotNeedingCheckout);
+			const bool bPromptingAfterModify = false;
+			const bool bAllowSkip = true;
+			bUserResponse = FEditorFileUtils::PromptToCheckoutPackagesInternal(false, PackagesToSave, &PackagesCheckedOutOrMadeWritable, &PackagesNotNeedingCheckout, bPromptingAfterModify, bAllowSkip);
 		}
 
 		if( bAlreadyCheckedOut || (bUserResponse && (PackagesCheckedOutOrMadeWritable.Num() > 0 || PackagesNotNeedingCheckout.Num() > 0)) )
diff --git a/Engine/Source/Editor/UnrealEd/Public/FileHelpers.h b/Engine/Source/Editor/UnrealEd/Public/FileHelpers.h
--- a/Engine/Source/Editor/UnrealEd/Public/FileHelpers.h
+++ b/Engine/Source/Editor/UnrealEd/Public/FileHelpers.h
@@ -467,11 +467,12 @@ public:
 	 * @param	OutPackagesCheckedOutOrMadeWritable		If not NULL, this array will be populated with packages that the user selected to check out or make writable.
 	 * @param	OutPackagesNotNeedingCheckout			If not NULL, this array will be populated with packages that the user was not prompted about and do not need to be checked out to save.  Useful for saving packages even if the user canceled the checkout dialog.
 	 * @param	bPromptingAfterModify					If true, we are prompting the user after an object has been modified, which changes the cancel button to "Ask me later".
+	 * @param	bAllowSkip								If true, user can skip the checkout / make writable.
 	 *
 	 * @return	true if the user did not cancel out of the dialog and has potentially checked out some files
 	 *			(or if there is no source control integration); false if the user cancelled the dialog
 	 */
-	UNREALED_API static bool PromptToCheckoutPackages(bool bCheckDirty, const TArray<UPackage*>& PackagesToCheckOut, TArray< UPackage* >* OutPackagesCheckedOutOrMadeWritable = NULL, TArray< UPackage* >* OutPackagesNotNeedingCheckout = NULL, const bool bPromptingAfterModify = false );
+	UNREALED_API static bool PromptToCheckoutPackages(bool bCheckDirty, const TArray<UPackage*>& PackagesToCheckOut, TArray< UPackage* >* OutPackagesCheckedOutOrMadeWritable = NULL, TArray< UPackage* >* OutPackagesNotNeedingCheckout = NULL, const bool bPromptingAfterModify = false, const bool bAllowSkip = false );
 	
 	/**
 	 * Check out the specified packages from source control and report any errors while checking out
@@ -541,7 +542,10 @@ public:
 	static bool SaveWorlds(UWorld* InWorld, const FString& RootPath, const TCHAR* Prefix, TArray<FString>& OutFilenames);
 
 	/** Whether or not we're in the middle of loading the simple startup map */
-	static bool IsLoadingStartupMap() {return bIsLoadingDefaultStartupMap;}
+	static bool IsLoadingStartupMap() { return bIsLoadingDefaultStartupMap; }
+
+	/** Whether or not saving the map package should save the external objects packages */
+	static bool ShouldSkipExternalObjectSave() { return bSkipExternalObjectSave; }
 
 	/**
 	 * Returns a file filter string appropriate for a specific file interaction.
@@ -619,12 +623,18 @@ private:
 
 	/** Callback from PackagesDialog used to update the list of items when the source control state changes */
 	static void UpdateCheckoutPackageItems(bool bCheckDirty, TArray<UPackage*> PackagesToCheckOut, TArray<UPackage*>* OutPackagesNotNeedingCheckout);
+	
+	/** Prompts for package checkout without reentrance check */
+	static bool PromptToCheckoutPackagesInternal(bool bCheckDirty, const TArray<UPackage*>& PackagesToCheckOut, TArray<UPackage*>* OutPackagesCheckedOutOrMadeWritable, TArray<UPackage*>* OutPackagesNotNeedingCheckout, const bool bPromptingAfterModify, const bool bAllowSkip);
 
 	static bool bIsLoadingDefaultStartupMap;
 
 	/** Flag used to determine if the checkout and save prompt is already open to prevent re-entrance */
 	static bool bIsPromptingForCheckoutAndSave;
 	
+	/** Flag used in SaveMap to skip saving external objects and only save the map */
+	static bool bSkipExternalObjectSave;
+
 	// Set of packages to ignore for save/checkout when using SaveAll.
 	static TSet<FString> PackagesNotSavedDuringSaveAll;
 
