Replace direct access to SkeletalMesh object from USkinnedMeshComponent with GetSkeletalMesh function

diff --git a/Engine/Source/Editor/PhysicsAssetEditor/Private/PhysicsAssetEditorSkeletonTreeBuilder.cpp b/Engine/Source/Editor/PhysicsAssetEditor/Private/PhysicsAssetEditorSkeletonTreeBuilder.cpp
--- a/Engine/Source/Editor/PhysicsAssetEditor/Private/PhysicsAssetEditorSkeletonTreeBuilder.cpp
+++ b/Engine/Source/Editor/PhysicsAssetEditor/Private/PhysicsAssetEditorSkeletonTreeBuilder.cpp
@@ -116,9 +116,9 @@ void FPhysicsAssetEditorSkeletonTreeBuilder::AddBodies(FSkeletonTreeBuilderOutpu
 	if (PreviewScenePtr.IsValid())
 	{
 		UDebugSkelMeshComponent* PreviewMeshComponent = PreviewScenePtr.Pin()->GetPreviewMeshComponent();
-		if (PreviewMeshComponent->SkeletalMesh)
+		if (PreviewMeshComponent->GetSkeletalMesh())
 		{
-			FReferenceSkeleton& RefSkeleton = PreviewMeshComponent->SkeletalMesh->GetRefSkeleton();
+			FReferenceSkeleton& RefSkeleton = PreviewMeshComponent->GetSkeletalMesh()->GetRefSkeleton();
 
 			for (int32 BoneIndex = 0; BoneIndex < RefSkeleton.GetRawBoneNum(); ++BoneIndex)
 			{
diff --git a/Engine/Source/Editor/SequenceRecorder/Private/AnimationRecorder.cpp b/Engine/Source/Editor/SequenceRecorder/Private/AnimationRecorder.cpp
--- a/Engine/Source/Editor/SequenceRecorder/Private/AnimationRecorder.cpp
+++ b/Engine/Source/Editor/SequenceRecorder/Private/AnimationRecorder.cpp
@@ -119,7 +119,7 @@ bool FAnimationRecorder::TriggerRecordAnimation(USkeletalMeshComponent* Componen
 	FFrameRate SampleRate;
 	float MaximumLength;
 
-	if (!Component || !Component->SkeletalMesh || !Component->SkeletalMesh->GetSkeleton())
+	if (!Component || !Component->GetSkeletalMesh() || !Component->GetSkeletalMesh()->GetSkeleton())
 	{
 		return false;
 	}
@@ -136,7 +136,7 @@ bool FAnimationRecorder::TriggerRecordAnimation(USkeletalMeshComponent* Componen
 
 bool FAnimationRecorder::TriggerRecordAnimation(USkeletalMeshComponent* Component, const FString& InAssetPath, const FString& InAssetName)
 {
-	if (!Component || !Component->SkeletalMesh || !Component->SkeletalMesh->GetSkeleton())
+	if (!Component || !Component->GetSkeletalMesh() || !Component->GetSkeletalMesh()->GetSkeleton())
 	{
 		return false;
 	}
@@ -182,7 +182,7 @@ bool FAnimationRecorder::TriggerRecordAnimation(USkeletalMeshComponent* Componen
 	if (NewSeq)
 	{
 		// set skeleton
-		NewSeq->SetSkeleton(Component->SkeletalMesh->GetSkeleton());
+		NewSeq->SetSkeleton(Component->GetSkeletalMesh()->GetSkeleton());
 		// Notify the asset registry
 		FAssetRegistryModule::AssetCreated(NewSeq);
 		StartRecord(Component, NewSeq);
@@ -267,7 +267,7 @@ void FAnimationRecorder::StartRecord(USkeletalMeshComponent* Component, UAnimSeq
 	for (int32 BoneIndex=0; BoneIndex <PreviousSpacesBases.Num(); ++BoneIndex)
 	{
 		// verify if this bone exists in skeleton
-		const int32 BoneTreeIndex = AnimSkeleton->GetSkeletonBoneIndexFromMeshBoneIndex(Component->MasterPoseComponent != nullptr ? Component->MasterPoseComponent->SkeletalMesh : Component->SkeletalMesh, BoneIndex);
+		const int32 BoneTreeIndex = AnimSkeleton->GetSkeletonBoneIndexFromMeshBoneIndex(Component->MasterPoseComponent != nullptr ? Component->MasterPoseComponent->GetSkeletalMesh() : Component->GetSkeletalMesh(), BoneIndex);
 		if (BoneTreeIndex != INDEX_NONE)
 		{
 			// add tracks for the bone existing
@@ -277,10 +277,10 @@ void FAnimationRecorder::StartRecord(USkeletalMeshComponent* Component, UAnimSeq
 		}
 	}
 
-	AnimationObject->RetargetSource = Component->SkeletalMesh ? AnimSkeleton->GetRetargetSourceForMesh(Component->SkeletalMesh) : NAME_None;
+	AnimationObject->RetargetSource = Component->GetSkeletalMesh() ? AnimSkeleton->GetRetargetSourceForMesh(Component->GetSkeletalMesh()) : NAME_None;
 	if (AnimationObject->RetargetSource == NAME_None)
 	{
-		AnimationObject->RetargetSourceAsset = Component->SkeletalMesh;
+		AnimationObject->RetargetSourceAsset = Component->GetSkeletalMesh();
 	}
 
 	// record the first frame
@@ -649,7 +649,7 @@ void FAnimationRecorder::ProcessRecordedTimes(UAnimSequence* AnimSequence, USkel
 	{
 		for (int32 BoneIndex = 0; BoneIndex < SpaceBases->Num(); ++BoneIndex)
 		{
-			const int32 BoneTreeIndex = AnimSkeleton->GetSkeletonBoneIndexFromMeshBoneIndex(SkeletalMeshComponent->MasterPoseComponent != nullptr ? SkeletalMeshComponent->MasterPoseComponent->SkeletalMesh : SkeletalMeshComponent->SkeletalMesh, BoneIndex);
+			const int32 BoneTreeIndex = AnimSkeleton->GetSkeletonBoneIndexFromMeshBoneIndex(SkeletalMeshComponent->MasterPoseComponent != nullptr ? SkeletalMeshComponent->MasterPoseComponent->GetSkeletalMesh() : SkeletalMeshComponent->GetSkeletalMesh(), BoneIndex);
 			if (BoneTreeIndex != INDEX_NONE)
 			{
 				FName BoneTreeName = AnimSkeleton->GetReferenceSkeleton().GetBoneName(BoneTreeIndex);
@@ -670,7 +670,7 @@ void FAnimationRecorder::ProcessRecordedTimes(UAnimSequence* AnimSequence, USkel
 	for (int32 BoneIndex = 0; BoneIndex < SpaceBases->Num(); ++BoneIndex)
 	{
 		// verify if this bone exists in skeleton
-		const int32 BoneTreeIndex = AnimSkeleton->GetSkeletonBoneIndexFromMeshBoneIndex(SkeletalMeshComponent->MasterPoseComponent != nullptr ? SkeletalMeshComponent->MasterPoseComponent->SkeletalMesh : SkeletalMeshComponent->SkeletalMesh, BoneIndex);
+		const int32 BoneTreeIndex = AnimSkeleton->GetSkeletonBoneIndexFromMeshBoneIndex(SkeletalMeshComponent->MasterPoseComponent != nullptr ? SkeletalMeshComponent->MasterPoseComponent->GetSkeletalMesh() : SkeletalMeshComponent->GetSkeletalMesh(), BoneIndex);
 		if (BoneTreeIndex != INDEX_NONE)
 		{
 			// add tracks for the bone existing
@@ -823,7 +823,7 @@ bool FAnimationRecorder::Record(USkeletalMeshComponent* Component, FTransform co
 	if (ensure(AnimationObject))
 	{
 		IAnimationDataController& Controller = AnimationObject->GetController();
-		USkeletalMesh* SkeletalMesh = Component->MasterPoseComponent != nullptr ? ToRawPtr(Component->MasterPoseComponent->SkeletalMesh) : ToRawPtr(Component->SkeletalMesh);
+		USkeletalMesh* SkeletalMesh = Component->MasterPoseComponent != nullptr ? ToRawPtr(Component->MasterPoseComponent->GetSkeletalMesh()) : ToRawPtr(Component->GetSkeletalMesh());
 
 		const TArray<FBoneAnimationTrack>& BoneAnimationTracks = AnimationObject->GetDataModel()->GetBoneAnimationTracks();
 
@@ -917,7 +917,7 @@ bool FAnimationRecorder::Record(USkeletalMeshComponent* Component, FTransform co
 				{
 					if (FrameToAdd == 0)
 					{
-						const FTransform RefPose = Component->SkeletalMesh->GetRefSkeleton().GetRefBonePose()[BoneIndex];
+						const FTransform RefPose = Component->GetSkeletalMesh()->GetRefSkeleton().GetRefBonePose()[BoneIndex];
 						RawTrack.PosKeys.Add((FVector3f)RefPose.GetTranslation());
 						RawTrack.RotKeys.Add(FQuat4f(RefPose.GetRotation()));
 						RawTrack.ScaleKeys.Add((FVector3f)RefPose.GetScale3D());
diff --git a/Engine/Source/Editor/SequenceRecorder/Private/Sections/MovieScene3DTransformSectionRecorder.cpp b/Engine/Source/Editor/SequenceRecorder/Private/Sections/MovieScene3DTransformSectionRecorder.cpp
--- a/Engine/Source/Editor/SequenceRecorder/Private/Sections/MovieScene3DTransformSectionRecorder.cpp
+++ b/Engine/Source/Editor/SequenceRecorder/Private/Sections/MovieScene3DTransformSectionRecorder.cpp
@@ -132,7 +132,7 @@ void FMovieScene3DTransformSectionRecorder::FinalizeSection(float CurrentTime)
 		USkeletalMeshComponent* SkeletalMeshComponent = AnimRecorder->GetSkeletalMeshComponent();
 		if (SkeletalMeshComponent)
 		{
-			USkeletalMesh* SkeletalMesh = SkeletalMeshComponent->MasterPoseComponent != nullptr ? ToRawPtr(SkeletalMeshComponent->MasterPoseComponent->SkeletalMesh) : ToRawPtr(SkeletalMeshComponent->SkeletalMesh);
+			USkeletalMesh* SkeletalMesh = SkeletalMeshComponent->MasterPoseComponent != nullptr ? ToRawPtr(SkeletalMeshComponent->MasterPoseComponent->GetSkeletalMesh()) : ToRawPtr(SkeletalMeshComponent->GetSkeletalMesh());
 			if (AnimSequence && SkeletalMesh)
 			{
 				// find the root bone
diff --git a/Engine/Source/Editor/SequenceRecorder/Private/Sections/MovieSceneAnimationSectionRecorder.cpp b/Engine/Source/Editor/SequenceRecorder/Private/Sections/MovieSceneAnimationSectionRecorder.cpp
--- a/Engine/Source/Editor/SequenceRecorder/Private/Sections/MovieSceneAnimationSectionRecorder.cpp
+++ b/Engine/Source/Editor/SequenceRecorder/Private/Sections/MovieSceneAnimationSectionRecorder.cpp
@@ -26,7 +26,7 @@ bool FMovieSceneAnimationSectionRecorderFactory::CanRecordObject(UObject* InObje
 	if (InObjectToRecord->IsA<USkeletalMeshComponent>())
 	{
 		USkeletalMeshComponent* SkeletalMeshComponent = Cast<USkeletalMeshComponent>(InObjectToRecord);
-		if (SkeletalMeshComponent && SkeletalMeshComponent->SkeletalMesh)
+		if (SkeletalMeshComponent && SkeletalMeshComponent->GetSkeletalMesh())
 		{
 			return true;
 		}
@@ -64,7 +64,7 @@ void FMovieSceneAnimationSectionRecorder::CreateSection(UObject* InObjectToRecor
 
 	if(SkeletalMeshComponent.IsValid())
 	{
-		SkeletalMesh = SkeletalMeshComponent->SkeletalMesh;
+		SkeletalMesh = SkeletalMeshComponent->GetSkeletalMesh();
 		if (SkeletalMesh != nullptr)
 		{
 			ComponentTransform = SkeletalMeshComponent->GetComponentToWorld().GetRelativeTransform(SkeletalMeshComponent->GetOwner()->GetTransform());
@@ -99,7 +99,7 @@ void FMovieSceneAnimationSectionRecorder::CreateSection(UObject* InObjectToRecor
 					FAssetRegistryModule::AssetCreated(AnimSequence.Get());
 
 					// set skeleton
-					AnimSequence->SetSkeleton(SkeletalMeshComponent->SkeletalMesh->GetSkeleton());
+					AnimSequence->SetSkeleton(SkeletalMeshComponent->GetSkeletalMesh()->GetSkeleton());
 				}
 			}
 
diff --git a/Engine/Source/Editor/SkeletonEditor/Private/BoneProxy.cpp b/Engine/Source/Editor/SkeletonEditor/Private/BoneProxy.cpp
--- a/Engine/Source/Editor/SkeletonEditor/Private/BoneProxy.cpp
+++ b/Engine/Source/Editor/SkeletonEditor/Private/BoneProxy.cpp
@@ -62,7 +62,7 @@ void UBoneProxy::Tick(float DeltaTime)
 					Scale = BoneTransform.GetScale3D();
 				}
 
-				FTransform ReferenceTransform = Component->SkeletalMesh->GetRefSkeleton().GetRefBonePose()[BoneIndex];
+				FTransform ReferenceTransform = Component->GetSkeletalMesh()->GetRefSkeleton().GetRefBonePose()[BoneIndex];
 				ReferenceLocation = ReferenceTransform.GetLocation();
 				ReferenceRotation = ReferenceTransform.GetRotation().Rotator();
 				ReferenceScale = ReferenceTransform.GetScale3D();
diff --git a/Engine/Source/Editor/SkeletonEditor/Private/EditableSkeleton.cpp b/Engine/Source/Editor/SkeletonEditor/Private/EditableSkeleton.cpp
--- a/Engine/Source/Editor/SkeletonEditor/Private/EditableSkeleton.cpp
+++ b/Engine/Source/Editor/SkeletonEditor/Private/EditableSkeleton.cpp
@@ -796,7 +796,7 @@ void FEditableSkeleton::HandleRemoveAllAssets(TSharedPtr<IPersonaPreviewScene> I
 
 	DeleteAttachedObjects(Skeleton->PreviewAttachedAssetContainer, InPreviewScene);
 
-	USkeletalMesh* SkeletalMesh = InPreviewScene->GetPreviewMeshComponent()->SkeletalMesh;
+	USkeletalMesh* SkeletalMesh = InPreviewScene->GetPreviewMeshComponent()->GetSkeletalMesh();
 	if (SkeletalMesh)
 	{
 		SkeletalMesh->Modify();
@@ -910,7 +910,7 @@ void FEditableSkeleton::HandleAttachAssets(const TArray<UObject*>& InObjects, co
 
 		if (bAttachToMesh)
 		{
-			USkeletalMesh* SkeletalMesh = InPreviewScene->GetPreviewMeshComponent()->SkeletalMesh;
+			USkeletalMesh* SkeletalMesh = InPreviewScene->GetPreviewMeshComponent()->GetSkeletalMesh();
 			if (SkeletalMesh != nullptr)
 			{
 				FScopedTransaction Transaction(LOCTEXT("DragDropAttachMeshUndo", "Attach Assets to Mesh"));
@@ -957,7 +957,7 @@ void FEditableSkeleton::HandleDeleteAttachedAssets(const TArray<FPreviewAttached
 	if (InAttachedObjects.Num() > 0)
 	{
 		Skeleton->Modify();
-		USkeletalMesh* SkeletalMesh = InPreviewScene->GetPreviewMeshComponent()->SkeletalMesh;
+		USkeletalMesh* SkeletalMesh = InPreviewScene->GetPreviewMeshComponent()->GetSkeletalMesh();
 		if (SkeletalMesh != nullptr)
 		{
 			SkeletalMesh->Modify();
@@ -993,7 +993,7 @@ void FEditableSkeleton::HandleDeleteSockets(const TArray<FSelectedSocketInfo>& I
 		}
 		else
 		{
-			USkeletalMesh* SkeletalMesh = InPreviewScene->GetPreviewMeshComponent()->SkeletalMesh;
+			USkeletalMesh* SkeletalMesh = InPreviewScene->GetPreviewMeshComponent()->GetSkeletalMesh();
 			if (SkeletalMesh != nullptr)
 			{
 				UObject* Object = SkeletalMesh->GetPreviewAttachedAssetContainer().GetAttachedObjectByAttachName(SocketName);
diff --git a/Engine/Source/Editor/SkeletonEditor/Private/SSkeletonTree.cpp b/Engine/Source/Editor/SkeletonEditor/Private/SSkeletonTree.cpp
--- a/Engine/Source/Editor/SkeletonEditor/Private/SSkeletonTree.cpp
+++ b/Engine/Source/Editor/SkeletonEditor/Private/SSkeletonTree.cpp
@@ -1077,7 +1077,7 @@ void SSkeletonTree::RemoveFromLOD(int32 LODIndex, bool bIncludeSelected, bool bI
 	}
 
 	UDebugSkelMeshComponent* PreviewMeshComponent = GetPreviewScene()->GetPreviewMeshComponent();
-	if (!PreviewMeshComponent->SkeletalMesh)
+	if (!PreviewMeshComponent->GetSkeletalMesh())
 	{
 		return;
 	}
@@ -1094,7 +1094,7 @@ void SSkeletonTree::RemoveFromLOD(int32 LODIndex, bool bIncludeSelected, bool bI
 
 		//Scoped post edit change
 		{
-			FScopedSkeletalMeshPostEditChange ScopedPostEditChange(PreviewMeshComponent->SkeletalMesh);
+			FScopedSkeletalMeshPostEditChange ScopedPostEditChange(PreviewMeshComponent->GetSkeletalMesh());
 
 			for (const TSharedPtr<FSkeletonTreeBoneItem>& Item : TreeSelection.GetSelectedItems<FSkeletonTreeBoneItem>())
 			{
@@ -1104,7 +1104,7 @@ void SSkeletonTree::RemoveFromLOD(int32 LODIndex, bool bIncludeSelected, bool bI
 				{
 					if (bIncludeSelected)
 					{
-						PreviewMeshComponent->SkeletalMesh->AddBoneToReductionSetting(LODIndex, BoneName);
+						PreviewMeshComponent->GetSkeletalMesh()->AddBoneToReductionSetting(LODIndex, BoneName);
 						BonesToRemove.AddUnique(BoneName);
 					}
 					else
@@ -1114,7 +1114,7 @@ void SSkeletonTree::RemoveFromLOD(int32 LODIndex, bool bIncludeSelected, bool bI
 							if (RefSkeleton.GetParentIndex(ChildIndex) == BoneIndex)
 							{
 								FName ChildBoneName = RefSkeleton.GetBoneName(ChildIndex);
-								PreviewMeshComponent->SkeletalMesh->AddBoneToReductionSetting(LODIndex, ChildBoneName);
+								PreviewMeshComponent->GetSkeletalMesh()->AddBoneToReductionSetting(LODIndex, ChildBoneName);
 								BonesToRemove.AddUnique(ChildBoneName);
 							}
 						}
@@ -1122,20 +1122,20 @@ void SSkeletonTree::RemoveFromLOD(int32 LODIndex, bool bIncludeSelected, bool bI
 				}
 			}
 
-			int32 TotalLOD = PreviewMeshComponent->SkeletalMesh->GetLODNum();
+			int32 TotalLOD = PreviewMeshComponent->GetSkeletalMesh()->GetLODNum();
 			IMeshUtilities& MeshUtilities = FModuleManager::Get().LoadModuleChecked<IMeshUtilities>("MeshUtilities");
 
 			if (bIncludeBelowLODs)
 			{
 				for (int32 Index = LODIndex + 1; Index < TotalLOD; ++Index)
 				{
-					MeshUtilities.RemoveBonesFromMesh(PreviewMeshComponent->SkeletalMesh, Index, &BonesToRemove);
-					PreviewMeshComponent->SkeletalMesh->AddBoneToReductionSetting(Index, BonesToRemove);
+					MeshUtilities.RemoveBonesFromMesh(PreviewMeshComponent->GetSkeletalMesh(), Index, &BonesToRemove);
+					PreviewMeshComponent->GetSkeletalMesh()->AddBoneToReductionSetting(Index, BonesToRemove);
 				}
 			}
 
 			// remove from current LOD
-			MeshUtilities.RemoveBonesFromMesh(PreviewMeshComponent->SkeletalMesh, LODIndex, &BonesToRemove);
+			MeshUtilities.RemoveBonesFromMesh(PreviewMeshComponent->GetSkeletalMesh(), LODIndex, &BonesToRemove);
 		}
 		// update UI to reflect the change
 		OnLODSwitched();
@@ -1256,7 +1256,7 @@ void SSkeletonTree::OnPasteSockets(bool bPasteToSelectedBone)
 	if ( TreeSelection.IsSingleOfTypeSelected<FSkeletonTreeBoneItem>() )
 	{
 		FName DestBoneName = bPasteToSelectedBone ? TreeSelection.GetSingleSelectedItem()->GetRowItemName() : NAME_None;
-		USkeletalMesh* SkeletalMesh = GetPreviewScene().IsValid() ? ToRawPtr(GetPreviewScene()->GetPreviewMeshComponent()->SkeletalMesh) : nullptr;
+		USkeletalMesh* SkeletalMesh = GetPreviewScene().IsValid() ? ToRawPtr(GetPreviewScene()->GetPreviewMeshComponent()->GetSkeletalMesh()) : nullptr;
 		GetEditableSkeletonInternal()->HandlePasteSockets(DestBoneName, SkeletalMesh);
 
 		CreateFromSkeleton();
@@ -1317,7 +1317,7 @@ void SSkeletonTree::OnCustomizeSocket()
 	if(TreeSelection.IsSingleOfTypeSelected<FSkeletonTreeSocketItem>())
 	{
 		USkeletalMeshSocket* SocketToCustomize = StaticCastSharedPtr<FSkeletonTreeSocketItem>(TreeSelection.GetSingleSelectedItem())->GetSocket();
-		USkeletalMesh* SkeletalMesh = GetPreviewScene().IsValid() ? ToRawPtr(GetPreviewScene()->GetPreviewMeshComponent()->SkeletalMesh) : nullptr;
+		USkeletalMesh* SkeletalMesh = GetPreviewScene().IsValid() ? ToRawPtr(GetPreviewScene()->GetPreviewMeshComponent()->GetSkeletalMesh()) : nullptr;
 		GetEditableSkeletonInternal()->HandleCustomizeSocket(SocketToCustomize, SkeletalMesh);
 		CreateFromSkeleton();
 	}
@@ -1387,7 +1387,7 @@ void  SSkeletonTree::OnRemoveAllAssets()
 
 bool SSkeletonTree::CanRemoveAllAssets() const
 {
-	USkeletalMesh* SkeletalMesh = GetPreviewScene().IsValid() ? ToRawPtr(GetPreviewScene()->GetPreviewMeshComponent()->SkeletalMesh) : nullptr;
+	USkeletalMesh* SkeletalMesh = GetPreviewScene().IsValid() ? ToRawPtr(GetPreviewScene()->GetPreviewMeshComponent()->GetSkeletalMesh()) : nullptr;
 
 	const bool bHasPreviewAttachedObjects = GetEditableSkeletonInternal()->GetSkeleton().PreviewAttachedAssetContainer.Num() > 0;
 	const bool bHasMeshPreviewAttachedObjects = ( SkeletalMesh && SkeletalMesh->GetPreviewAttachedAssetContainer().Num() );
@@ -2273,7 +2273,7 @@ void SSkeletonTree::SelectItemsBy(TFunctionRef<bool(const TSharedRef<ISkeletonTr
 
 void SSkeletonTree::DuplicateAndSelectSocket(const FSelectedSocketInfo& SocketInfoToDuplicate, const FName& NewParentBoneName /*= FName()*/)
 {
-	USkeletalMesh* SkeletalMesh = GetPreviewScene().IsValid() ? ToRawPtr(GetPreviewScene()->GetPreviewMeshComponent()->SkeletalMesh) : nullptr;
+	USkeletalMesh* SkeletalMesh = GetPreviewScene().IsValid() ? ToRawPtr(GetPreviewScene()->GetPreviewMeshComponent()->GetSkeletalMesh()) : nullptr;
 	USkeletalMeshSocket* NewSocket = GetEditableSkeleton()->DuplicateSocket(SocketInfoToDuplicate, NewParentBoneName, SkeletalMesh);
 
 	if (GetPreviewScene().IsValid())
diff --git a/Engine/Source/Editor/SkeletonEditor/Private/SkeletonEditor.cpp b/Engine/Source/Editor/SkeletonEditor/Private/SkeletonEditor.cpp
--- a/Engine/Source/Editor/SkeletonEditor/Private/SkeletonEditor.cpp
+++ b/Engine/Source/Editor/SkeletonEditor/Private/SkeletonEditor.cpp
@@ -348,7 +348,7 @@ TStatId FSkeletonEditor::GetStatId() const
 bool FSkeletonEditor::CanRemoveBones() const
 {
 	UDebugSkelMeshComponent* PreviewMeshComponent = PersonaToolkit->GetPreviewMeshComponent();
-	return PreviewMeshComponent && PreviewMeshComponent->SkeletalMesh;
+	return PreviewMeshComponent && PreviewMeshComponent->GetSkeletalMesh();
 }
 
 void FSkeletonEditor::RemoveUnusedBones()
@@ -364,9 +364,9 @@ void FSkeletonEditor::TestSkeletonCurveNamesForUse() const
 
 void FSkeletonEditor::UpdateSkeletonRefPose()
 {
-	if (PersonaToolkit->GetPreviewMeshComponent()->SkeletalMesh)
+	if (PersonaToolkit->GetPreviewMeshComponent()->GetSkeletalMesh())
 	{
-		GetSkeletonTree()->GetEditableSkeleton()->UpdateSkeletonReferencePose(PersonaToolkit->GetPreviewMeshComponent()->SkeletalMesh);
+		GetSkeletonTree()->GetEditableSkeleton()->UpdateSkeletonReferencePose(PersonaToolkit->GetPreviewMeshComponent()->GetSkeletalMesh());
 	}
 }
 
diff --git a/Engine/Source/Editor/SkeletonEditor/Private/SkeletonTreeBoneItem.cpp b/Engine/Source/Editor/SkeletonEditor/Private/SkeletonTreeBoneItem.cpp
--- a/Engine/Source/Editor/SkeletonEditor/Private/SkeletonTreeBoneItem.cpp
+++ b/Engine/Source/Editor/SkeletonEditor/Private/SkeletonTreeBoneItem.cpp
@@ -427,7 +427,7 @@ FReply FSkeletonTreeBoneItem::HandleDrop(const FDragDropEvent& DragDropEvent)
 		else if (BoneName != SocketInfo.Socket->BoneName)
 		{
 			// The socket can be dropped here if we're a bone and NOT the socket's existing parent
-			USkeletalMesh* SkeletalMesh = GetSkeletonTree()->GetPreviewScene().IsValid() ? ToRawPtr(GetSkeletonTree()->GetPreviewScene()->GetPreviewMeshComponent()->SkeletalMesh) : nullptr;
+			USkeletalMesh* SkeletalMesh = GetSkeletonTree()->GetPreviewScene().IsValid() ? ToRawPtr(GetSkeletonTree()->GetPreviewScene()->GetPreviewMeshComponent()->GetSkeletalMesh()) : nullptr;
 			GetEditableSkeleton()->SetSocketParent(SocketInfo.Socket->SocketName, BoneName, SkeletalMesh);
 
 			return FReply::Handled();
@@ -454,15 +454,15 @@ bool FSkeletonTreeBoneItem::IsBoneWeighted(int32 MeshBoneIndex, UDebugSkelMeshCo
 {
 	// MeshBoneIndex must be an index into the mesh's skeleton, *not* the source skeleton!!!
 
-	if (!PreviewComponent || !PreviewComponent->SkeletalMesh || !PreviewComponent->SkeletalMesh->GetResourceForRendering() || !PreviewComponent->SkeletalMesh->GetResourceForRendering()->LODRenderData.Num())
+	if (!PreviewComponent || !PreviewComponent->GetSkeletalMesh() || !PreviewComponent->GetSkeletalMesh()->GetResourceForRendering() || !PreviewComponent->GetSkeletalMesh()->GetResourceForRendering()->LODRenderData.Num())
 	{
 		// If there's no mesh, then this bone can't possibly be weighted!
 		return false;
 	}
 
 	//Get current LOD
-	const int32 LODIndex = FMath::Clamp(PreviewComponent->GetPredictedLODLevel(), 0, PreviewComponent->SkeletalMesh->GetResourceForRendering()->LODRenderData.Num() - 1);
-	FSkeletalMeshLODRenderData& LODData = PreviewComponent->SkeletalMesh->GetResourceForRendering()->LODRenderData[LODIndex];
+	const int32 LODIndex = FMath::Clamp(PreviewComponent->GetPredictedLODLevel(), 0, PreviewComponent->GetSkeletalMesh()->GetResourceForRendering()->LODRenderData.Num() - 1);
+	FSkeletalMeshLODRenderData& LODData = PreviewComponent->GetSkeletalMesh()->GetResourceForRendering()->LODRenderData[LODIndex];
 
 	//Check whether the bone is vertex weighted
 	int32 Index = LODData.ActiveBoneIndices.Find(MeshBoneIndex);
@@ -474,15 +474,15 @@ bool FSkeletonTreeBoneItem::IsBoneRequired(int32 MeshBoneIndex, UDebugSkelMeshCo
 {
 	// MeshBoneIndex must be an index into the mesh's skeleton, *not* the source skeleton!!!
 
-	if (!PreviewComponent || !PreviewComponent->SkeletalMesh || !PreviewComponent->SkeletalMesh->GetResourceForRendering() || !PreviewComponent->SkeletalMesh->GetResourceForRendering()->LODRenderData.Num())
+	if (!PreviewComponent || !PreviewComponent->GetSkeletalMesh() || !PreviewComponent->GetSkeletalMesh()->GetResourceForRendering() || !PreviewComponent->GetSkeletalMesh()->GetResourceForRendering()->LODRenderData.Num())
 	{
 		// If there's no mesh, then this bone can't possibly be weighted!
 		return false;
 	}
 
 	//Get current LOD
-	const int32 LODIndex = FMath::Clamp(PreviewComponent->GetPredictedLODLevel(), 0, PreviewComponent->SkeletalMesh->GetResourceForRendering()->LODRenderData.Num() - 1);
-	FSkeletalMeshLODRenderData& LODData = PreviewComponent->SkeletalMesh->GetResourceForRendering()->LODRenderData[LODIndex];
+	const int32 LODIndex = FMath::Clamp(PreviewComponent->GetPredictedLODLevel(), 0, PreviewComponent->GetSkeletalMesh()->GetResourceForRendering()->LODRenderData.Num() - 1);
+	FSkeletalMeshLODRenderData& LODData = PreviewComponent->GetSkeletalMesh()->GetResourceForRendering()->LODRenderData[LODIndex];
 
 	//Check whether the bone is vertex weighted
 	int32 Index = LODData.RequiredBones.Find(MeshBoneIndex);
diff --git a/Engine/Source/Editor/SkeletonEditor/Private/SkeletonTreeBuilder.cpp b/Engine/Source/Editor/SkeletonEditor/Private/SkeletonTreeBuilder.cpp
--- a/Engine/Source/Editor/SkeletonEditor/Private/SkeletonTreeBuilder.cpp
+++ b/Engine/Source/Editor/SkeletonEditor/Private/SkeletonTreeBuilder.cpp
@@ -304,7 +304,7 @@ void FSkeletonTreeBuilder::AddSockets(FSkeletonTreeBuilderOutput& Output)
 	if (PreviewScenePtr.IsValid())
 	{
 		UDebugSkelMeshComponent* PreviewMeshComponent = PreviewScenePtr.Pin()->GetPreviewMeshComponent();
-		if (USkeletalMesh* const SkeletalMesh = PreviewMeshComponent->SkeletalMesh)
+		if (USkeletalMesh* const SkeletalMesh = PreviewMeshComponent->GetSkeletalMesh())
 		{
 			AddSocketsFromData(SkeletalMesh->GetMeshOnlySocketList(), ESocketParentType::Mesh, Output);
 		}
@@ -319,7 +319,7 @@ void FSkeletonTreeBuilder::AddAttachedAssets(FSkeletonTreeBuilderOutput& Output)
 	if (PreviewScenePtr.IsValid())
 	{
 		UDebugSkelMeshComponent* PreviewMeshComponent = PreviewScenePtr.Pin()->GetPreviewMeshComponent();
-		if (USkeletalMesh* const SkeletalMesh = PreviewMeshComponent->SkeletalMesh)
+		if (USkeletalMesh* const SkeletalMesh = PreviewMeshComponent->GetSkeletalMesh())
 		{
 			AddAttachedAssetContainer(SkeletalMesh->GetPreviewAttachedAssetContainer(), Output);
 		}
@@ -348,7 +348,7 @@ void FSkeletonTreeBuilder::AddSocketsFromData(const TArray< USkeletalMeshSocket*
 			if (PreviewScenePtr.IsValid())
 			{
 				UDebugSkelMeshComponent* PreviewMeshComponent = PreviewScenePtr.Pin()->GetPreviewMeshComponent();
-				if (USkeletalMesh* const SkeletalMesh = PreviewMeshComponent->SkeletalMesh)
+				if (USkeletalMesh* const SkeletalMesh = PreviewMeshComponent->GetSkeletalMesh())
 				{
 					bIsCustomized = EditableSkeletonPtr.Pin()->DoesSocketAlreadyExist(nullptr, FText::FromName(Socket->SocketName), ESocketParentType::Mesh, SkeletalMesh);
 				}
diff --git a/Engine/Source/Editor/SkeletonEditor/Private/SkeletonTreeSocketItem.cpp b/Engine/Source/Editor/SkeletonEditor/Private/SkeletonTreeSocketItem.cpp
--- a/Engine/Source/Editor/SkeletonEditor/Private/SkeletonTreeSocketItem.cpp
+++ b/Engine/Source/Editor/SkeletonEditor/Private/SkeletonTreeSocketItem.cpp
@@ -147,7 +147,7 @@ bool FSkeletonTreeSocketItem::CanCustomizeSocket() const
 	if (GetSkeletonTree()->GetPreviewScene().IsValid())
 	{
 		UDebugSkelMeshComponent* PreviewComponent = GetSkeletonTree()->GetPreviewScene()->GetPreviewMeshComponent();
-		return PreviewComponent && PreviewComponent->SkeletalMesh && !IsSocketCustomized();
+		return PreviewComponent && PreviewComponent->GetSkeletalMesh() && !IsSocketCustomized();
 	}
 	return false;
 }
@@ -172,7 +172,7 @@ bool FSkeletonTreeSocketItem::OnVerifySocketNameChanged( const FText& InText, FT
 	}
 	else
 	{
-		USkeletalMesh* SkeletalMesh = GetSkeletonTree()->GetPreviewScene().IsValid() ? ToRawPtr(GetSkeletonTree()->GetPreviewScene()->GetPreviewMeshComponent()->SkeletalMesh) : nullptr;
+		USkeletalMesh* SkeletalMesh = GetSkeletonTree()->GetPreviewScene().IsValid() ? ToRawPtr(GetSkeletonTree()->GetPreviewScene()->GetPreviewMeshComponent()->GetSkeletalMesh()) : nullptr;
 		bVerifyName = !GetEditableSkeleton()->DoesSocketAlreadyExist( Socket, NewText, ParentType, SkeletalMesh );
 
 		// Needs to be checked on verify.
@@ -192,7 +192,7 @@ void FSkeletonTreeSocketItem::OnCommitSocketName( const FText& InText, ETextComm
 	FText NewText = FText::TrimPrecedingAndTrailing(InText);
 
 	// Notify skeleton tree of socket rename
-	USkeletalMesh* SkeletalMesh = GetSkeletonTree()->GetPreviewScene().IsValid() ? ToRawPtr(GetSkeletonTree()->GetPreviewScene()->GetPreviewMeshComponent()->SkeletalMesh) : nullptr;
+	USkeletalMesh* SkeletalMesh = GetSkeletonTree()->GetPreviewScene().IsValid() ? ToRawPtr(GetSkeletonTree()->GetPreviewScene()->GetPreviewMeshComponent()->GetSkeletalMesh()) : nullptr;
 	GetEditableSkeleton()->RenameSocket(Socket->SocketName, FName(*NewText.ToString()), SkeletalMesh);
 }
 
diff --git a/Engine/Source/Editor/StatsViewer/Private/StatsPages/PrimitiveStatsPage.cpp b/Engine/Source/Editor/StatsViewer/Private/StatsPages/PrimitiveStatsPage.cpp
--- a/Engine/Source/Editor/StatsViewer/Private/StatsPages/PrimitiveStatsPage.cpp
+++ b/Engine/Source/Editor/StatsViewer/Private/StatsPages/PrimitiveStatsPage.cpp
@@ -156,7 +156,7 @@ struct PrimitiveStatsGenerator
 		// The skeletal mesh of a skeletal mesh component is its resource.
 		else if( SkeletalMeshComponent )
 		{
-			USkeletalMesh* Mesh = SkeletalMeshComponent->SkeletalMesh;
+			USkeletalMesh* Mesh = SkeletalMeshComponent->GetSkeletalMesh();
 			Resource = Mesh;
 			// Calculate vertex color usage for skeletal meshes
 			if( Mesh )
@@ -324,7 +324,7 @@ struct PrimitiveStatsGenerator
 				// ... in the case of skeletal mesh component.
 				else if( SkeletalMeshComponent )
 				{
-					USkeletalMesh* SkeletalMesh = SkeletalMeshComponent->SkeletalMesh;
+					USkeletalMesh* SkeletalMesh = SkeletalMeshComponent->GetSkeletalMesh();
 					if( SkeletalMesh )
 					{
 						FSkeletalMeshRenderData* SkelMeshRenderData = SkeletalMesh->GetResourceForRendering();
diff --git a/Engine/Source/Editor/UnrealEd/Classes/Animation/DebugSkelMeshComponent.h b/Engine/Source/Editor/UnrealEd/Classes/Animation/DebugSkelMeshComponent.h
--- a/Engine/Source/Editor/UnrealEd/Classes/Animation/DebugSkelMeshComponent.h
+++ b/Engine/Source/Editor/UnrealEd/Classes/Animation/DebugSkelMeshComponent.h
@@ -597,9 +597,9 @@ public:
 	 */
 	virtual const FReferenceSkeleton& GetReferenceSkeleton() const
 	{
-		if (SkeletalMesh)
+		if (GetSkeletalMesh())
 		{
-			return SkeletalMesh->GetRefSkeleton();
+			return GetSkeletalMesh()->GetRefSkeleton();
 		}
 
 		static FReferenceSkeleton EmptySkeleton;
diff --git a/Engine/Source/Editor/UnrealEd/Private/Animation/DebugSkelMeshComponent.cpp b/Engine/Source/Editor/UnrealEd/Private/Animation/DebugSkelMeshComponent.cpp
--- a/Engine/Source/Editor/UnrealEd/Private/Animation/DebugSkelMeshComponent.cpp
+++ b/Engine/Source/Editor/UnrealEd/Private/Animation/DebugSkelMeshComponent.cpp
@@ -97,9 +97,9 @@ FBoxSphereBounds UDebugSkelMeshComponent::CalcBounds(const FTransform& LocalToWo
 			}			
 		}
 
-		if ( SkeletalMesh )
+		if (GetSkeletalMesh())
 		{
-			Result = Result + SkeletalMesh->GetBounds();
+			Result = Result + GetSkeletalMesh()->GetBounds();
 		}
 	}
 
@@ -424,7 +424,7 @@ FPrimitiveSceneProxy* UDebugSkelMeshComponent::CreateSceneProxy()
 {
 	FDebugSkelMeshSceneProxy* Result = NULL;
 	ERHIFeatureLevel::Type SceneFeatureLevel = GetWorld()->FeatureLevel;
-	FSkeletalMeshRenderData* SkelMeshRenderData = SkeletalMesh ? SkeletalMesh->GetResourceForRendering() : NULL;
+	FSkeletalMeshRenderData* SkelMeshRenderData = GetSkeletalMesh() ? GetSkeletalMesh()->GetResourceForRendering() : NULL;
 
 	// only create a scene proxy for rendering if
 	// properly initialized
@@ -488,9 +488,9 @@ void UDebugSkelMeshComponent::InitAnim(bool bForceReinit)
 {
 	// If we already have PreviewInstance and its asset's Skeleton isn't compatible with the mesh's Skeleton
 	// then we need to clear it up to avoid an issue
-	if ( PreviewInstance && PreviewInstance->GetCurrentAsset() && SkeletalMesh )
+	if ( PreviewInstance && PreviewInstance->GetCurrentAsset() && GetSkeletalMesh())
 	{
-		if (!SkeletalMesh->GetSkeleton()->IsCompatible(PreviewInstance->GetCurrentAsset()->GetSkeleton()))
+		if (!GetSkeletalMesh()->GetSkeleton()->IsCompatible(PreviewInstance->GetCurrentAsset()->GetSkeleton()))
 		{
 			// if it doesn't match, just clear it
 			PreviewInstance->SetAnimationAsset(NULL);
@@ -717,7 +717,7 @@ void UDebugSkelMeshComponent::GenSpaceBases(TArray<FTransform>& OutSpaceBases)
 	FBlendedHeapCurve TempCurve;
 	UE::Anim::FMeshAttributeContainer TempAtttributes;
 	DoInstancePreEvaluation();
-	PerformAnimationEvaluation(SkeletalMesh, AnimScriptInstance, OutSpaceBases, TempBoneSpaceTransforms, TempRootBoneTranslation, TempCurve, TempAtttributes);
+	PerformAnimationEvaluation(GetSkeletalMesh(), AnimScriptInstance, OutSpaceBases, TempBoneSpaceTransforms, TempRootBoneTranslation, TempCurve, TempAtttributes);
 	DoInstancePostEvaluation();
 }
 
@@ -923,7 +923,7 @@ void UDebugSkelMeshComponent::ToggleClothSectionsVisibility(bool bShowOnlyClothS
 
 void UDebugSkelMeshComponent::RestoreClothSectionsVisibility()
 {
-	if (!SkeletalMesh)
+	if (!GetSkeletalMesh())
 	{
 		return;
 	}
@@ -976,11 +976,11 @@ void UDebugSkelMeshComponent::RebuildClothingSectionsFixedVerts(bool bInvalidate
 	// TODO: There is no need to rebuild all section/LODs at once.
 	//       It should only do the section associated to the current cloth asset being
 	//        painted instead, and only when the MaxDistance mask changes.
-	FScopedSkeletalMeshPostEditChange ScopedSkeletalMeshPostEditChange(SkeletalMesh);
+	FScopedSkeletalMeshPostEditChange ScopedSkeletalMeshPostEditChange(GetSkeletalMesh());
 
-	SkeletalMesh->PreEditChange(nullptr);
+	GetSkeletalMesh()->PreEditChange(nullptr);
 
-	TIndirectArray<FSkeletalMeshLODModel>& LODModels = SkeletalMesh->GetImportedModel()->LODModels;
+	TIndirectArray<FSkeletalMeshLODModel>& LODModels = GetSkeletalMesh()->GetImportedModel()->LODModels;
 
 	for (int32 LODIndex = 0; LODIndex < LODModels.Num(); ++LODIndex)
 	{
@@ -994,7 +994,7 @@ void UDebugSkelMeshComponent::RebuildClothingSectionsFixedVerts(bool bInvalidate
 
 	if (bInvalidateDerivedDataCache)
 	{
-		SkeletalMesh->InvalidateDeriveDataCacheGUID();  // Dirty the DDC key unless previewing
+		GetSkeletalMesh()->InvalidateDeriveDataCacheGUID();  // Dirty the DDC key unless previewing
 	}
 
 	ReregisterComponent();
@@ -1002,7 +1002,7 @@ void UDebugSkelMeshComponent::RebuildClothingSectionsFixedVerts(bool bInvalidate
 
 void UDebugSkelMeshComponent::RebuildClothingSectionFixedVerts(int32 LODIndex, int32 SectionIndex)
 {
-	FSkeletalMeshModel* const SkeletalMeshModel = SkeletalMesh->GetImportedModel();
+	FSkeletalMeshModel* const SkeletalMeshModel = GetSkeletalMesh()->GetImportedModel();
 	if (!ensure(SkeletalMeshModel && LODIndex < SkeletalMeshModel->LODModels.Num()))
 	{
 		return;
@@ -1020,7 +1020,7 @@ void UDebugSkelMeshComponent::RebuildClothingSectionFixedVerts(int32 LODIndex, i
 		return;
 	}
 
-	const UClothingAssetCommon* const ClothingAsset = Cast<UClothingAssetCommon>(SkeletalMesh->GetClothingAsset(UpdatedSection.ClothingData.AssetGuid));
+	const UClothingAssetCommon* const ClothingAsset = Cast<UClothingAssetCommon>(GetSkeletalMesh()->GetClothingAsset(UpdatedSection.ClothingData.AssetGuid));
 	check(ClothingAsset);  // Must have a valid clothing asset at this point, or something has gone terribly wrong
 
 	const FClothLODDataCommon& ClothLODData = ClothingAsset->LodData[UpdatedSection.ClothingData.AssetLodIndex];
@@ -1094,9 +1094,9 @@ void UDebugSkelMeshComponent::TickComponent(float DeltaTime, enum ELevelTick Tic
 
 void UDebugSkelMeshComponent::RefreshSelectedClothingSkinnedPositions()
 {
-	if(SkeletalMesh && SelectedClothingGuidForPainting.IsValid())
+	if(GetSkeletalMesh() && SelectedClothingGuidForPainting.IsValid())
 	{
-		UClothingAssetBase** Asset = SkeletalMesh->GetMeshClothingAssets().FindByPredicate([&](UClothingAssetBase* Item)
+		UClothingAssetBase** Asset = GetSkeletalMesh()->GetMeshClothingAssets().FindByPredicate([&](UClothingAssetBase* Item)
 		{
 			return Item && SelectedClothingGuidForPainting == Item->GetAssetGuid();
 		});
@@ -1305,7 +1305,7 @@ FDebugSkelMeshDynamicData::FDebugSkelMeshDynamicData(UDebugSkelMeshComponent* In
 		SkinnedPositions = InComponent->SkinnedSelectedClothingPositions;
 		SkinnedNormals = InComponent->SkinnedSelectedClothingNormals;
 
-		if(USkeletalMesh* Mesh = InComponent->SkeletalMesh)
+		if(USkeletalMesh* Mesh = InComponent->GetSkeletalMesh())
 		{
 			const int32 NumClothingAssets = Mesh->GetMeshClothingAssets().Num();
 			for(int32 ClothingAssetIndex = 0; ClothingAssetIndex < NumClothingAssets; ++ClothingAssetIndex)
@@ -1361,7 +1361,7 @@ FScopedSuspendAlternateSkinWeightPreview::FScopedSuspendAlternateSkinWeightPrevi
 		for (TObjectIterator<UDebugSkelMeshComponent> It; It; ++It)
 		{
 			UDebugSkelMeshComponent* DebugSKComp = *It;
-			if (DebugSKComp->SkeletalMesh == SkeletalMesh)
+			if (DebugSKComp->GetSkeletalMesh() == SkeletalMesh)
 			{
 				const FName ProfileName = DebugSKComp->GetCurrentSkinWeightProfileName();
 				if (ProfileName != NAME_None)
diff --git a/Engine/Source/Editor/UnrealEd/Private/ApexClothingUtils.cpp b/Engine/Source/Editor/UnrealEd/Private/ApexClothingUtils.cpp
--- a/Engine/Source/Editor/UnrealEd/Private/ApexClothingUtils.cpp
+++ b/Engine/Source/Editor/UnrealEd/Private/ApexClothingUtils.cpp
@@ -33,7 +33,7 @@ void ReregisterSkelMeshComponents(USkeletalMesh* SkelMesh)
 		USkeletalMeshComponent* MeshComponent = *It;
 		if( MeshComponent && 
 			!MeshComponent->IsTemplate() &&
-			MeshComponent->SkeletalMesh == SkelMesh )
+			MeshComponent->GetSkeletalMesh() == SkelMesh )
 		{
 			MeshComponent->ReregisterComponent();
 		}
@@ -47,7 +47,7 @@ void RefreshSkelMeshComponents(USkeletalMesh* SkelMesh)
 		USkeletalMeshComponent* MeshComponent = *It;
 		if( MeshComponent && 
 			!MeshComponent->IsTemplate() &&
-			MeshComponent->SkeletalMesh == SkelMesh )
+			MeshComponent->GetSkeletalMesh() == SkelMesh )
 		{
 			MeshComponent->RecreateRenderState_Concurrent();
 		}
diff --git a/Engine/Source/Editor/UnrealEd/Private/ComponentAssetBroker.cpp b/Engine/Source/Editor/UnrealEd/Private/ComponentAssetBroker.cpp
--- a/Engine/Source/Editor/UnrealEd/Private/ComponentAssetBroker.cpp
+++ b/Engine/Source/Editor/UnrealEd/Private/ComponentAssetBroker.cpp
@@ -81,7 +81,7 @@ public:
 	{
 		if (USkeletalMeshComponent* SkelMeshComp = Cast<USkeletalMeshComponent>(InComponent))
 		{
-			return SkelMeshComp->SkeletalMesh;
+			return SkelMeshComp->GetSkeletalMesh();
 		}
 		return NULL;
 	}
diff --git a/Engine/Source/Editor/UnrealEd/Private/EditorActor.cpp b/Engine/Source/Editor/UnrealEd/Private/EditorActor.cpp
--- a/Engine/Source/Editor/UnrealEd/Private/EditorActor.cpp
+++ b/Engine/Source/Editor/UnrealEd/Private/EditorActor.cpp
@@ -2357,7 +2357,7 @@ void UUnrealEdEngine::edactSelectMatchingSkeletalMesh(bool bAllClasses)
 		if(SkelMeshActor && SkelMeshActor->GetSkeletalMeshComponent())
 		{
 			bSelectSkelMeshActors = true;
-			SelectedMeshes.AddUnique(SkelMeshActor->GetSkeletalMeshComponent()->SkeletalMesh);
+			SelectedMeshes.AddUnique(SkelMeshActor->GetSkeletalMeshComponent()->GetSkeletalMesh());
 			SelectedWorlds.AddUnique(Actor->GetWorld());			
 		}
 
@@ -2369,7 +2369,7 @@ void UUnrealEdEngine::edactSelectMatchingSkeletalMesh(bool bAllClasses)
 			if (PawnSkeletalMesh)
 			{
 				bSelectPawns = true;
-				SelectedMeshes.AddUnique(PawnSkeletalMesh->SkeletalMesh);
+				SelectedMeshes.AddUnique(PawnSkeletalMesh->GetSkeletalMesh());
 				SelectedWorlds.AddUnique(Actor->GetWorld());
 			}
 		}
@@ -2405,7 +2405,7 @@ void UUnrealEdEngine::edactSelectMatchingSkeletalMesh(bool bAllClasses)
 				ASkeletalMeshActor* SkelMeshActor = Cast<ASkeletalMeshActor>(Actor);
 				if( SkelMeshActor && 
 					SkelMeshActor->GetSkeletalMeshComponent() && 
-					SelectedMeshes.Contains(SkelMeshActor->GetSkeletalMeshComponent()->SkeletalMesh) )
+					SelectedMeshes.Contains(SkelMeshActor->GetSkeletalMeshComponent()->GetSkeletalMesh()) )
 				{
 					bSelectActor = true;
 				}
@@ -2417,7 +2417,7 @@ void UUnrealEdEngine::edactSelectMatchingSkeletalMesh(bool bAllClasses)
 				if (Pawn)
 				{
 					USkeletalMeshComponent* PawnSkeletalMesh = Pawn->FindComponentByClass<USkeletalMeshComponent>();
-					if (PawnSkeletalMesh && SelectedMeshes.Contains(PawnSkeletalMesh->SkeletalMesh) )
+					if (PawnSkeletalMesh && SelectedMeshes.Contains(PawnSkeletalMesh->GetSkeletalMesh()) )
 					{
 						bSelectActor = true;
 					}
diff --git a/Engine/Source/Editor/UnrealEd/Private/EditorEngine.cpp b/Engine/Source/Editor/UnrealEd/Private/EditorEngine.cpp
--- a/Engine/Source/Editor/UnrealEd/Private/EditorEngine.cpp
+++ b/Engine/Source/Editor/UnrealEd/Private/EditorEngine.cpp
@@ -3496,7 +3496,7 @@ struct FConvertStaticMeshActorInfo
 		InternalGetFromActor(Actor);
 
 		// Copy over component properties.
-		SkeletalMesh			= MeshComp->SkeletalMesh;
+		SkeletalMesh			= MeshComp->GetSkeletalMesh();
 		OverrideMaterials		= MeshComp->OverrideMaterials;
 		CachedMaxDrawDistance	= MeshComp->CachedMaxDrawDistance;
 		CastShadow				= MeshComp->CastShadow;
@@ -3519,7 +3519,7 @@ struct FConvertStaticMeshActorInfo
 		InternalSetToActor(Actor);
 
 		// Set component properties.
-		if ( bComponentPropsDifferFromDefaults[0] ) MeshComp->SkeletalMesh			= SkeletalMesh;
+		if ( bComponentPropsDifferFromDefaults[0] ) MeshComp->GetSkeletalMesh()		= SkeletalMesh;
 		if ( bComponentPropsDifferFromDefaults[1] ) MeshComp->OverrideMaterials		= OverrideMaterials;
 		if ( bComponentPropsDifferFromDefaults[2] ) MeshComp->CachedMaxDrawDistance	= CachedMaxDrawDistance;
 		if ( bComponentPropsDifferFromDefaults[3] ) MeshComp->CastShadow			= CastShadow;
diff --git a/Engine/Source/Editor/UnrealEd/Private/EditorExporters.cpp b/Engine/Source/Editor/UnrealEd/Private/EditorExporters.cpp
--- a/Engine/Source/Editor/UnrealEd/Private/EditorExporters.cpp
+++ b/Engine/Source/Editor/UnrealEd/Private/EditorExporters.cpp
@@ -1570,7 +1570,7 @@ bool IsSomethingToExport(AActor* Actor)
 		UChildActorComponent* ChildActorComp = Cast<UChildActorComponent>(Component);
 
 		if ((StaticMeshComp && StaticMeshComp->GetStaticMesh()) ||
-			(SkelMeshComp && SkelMeshComp->SkeletalMesh) ||
+			(SkelMeshComp && SkelMeshComp->GetSkeletalMesh()) ||
 			Component->IsA(UCameraComponent::StaticClass()) ||
 			Component->IsA(ULightComponent::StaticClass()) ||
 			(ChildActorComp && ChildActorComp->GetChildActor() && IsSomethingToExport(ChildActorComp->GetChildActor()))
diff --git a/Engine/Source/Editor/UnrealEd/Private/Factories/ActorFactory.cpp b/Engine/Source/Editor/UnrealEd/Private/Factories/ActorFactory.cpp
--- a/Engine/Source/Editor/UnrealEd/Private/Factories/ActorFactory.cpp
+++ b/Engine/Source/Editor/UnrealEd/Private/Factories/ActorFactory.cpp
@@ -958,7 +958,7 @@ void UActorFactoryPhysicsAsset::PostSpawnActor(UObject* Asset, AActor* NewActor)
 	NewSkelActor->GetSkeletalMeshComponent()->UnregisterComponent();
 
 	// Change properties
-	NewSkelActor->GetSkeletalMeshComponent()->SkeletalMesh = UseSkelMesh;
+	NewSkelActor->GetSkeletalMeshComponent()->GetSkeletalMesh() = UseSkelMesh;
 	if (NewSkelActor->GetWorld()->IsPlayInEditor())
 	{
 		NewSkelActor->ReplicatedMesh = UseSkelMesh;
@@ -991,7 +991,7 @@ void UActorFactoryPhysicsAsset::PostCreateBlueprint( UObject* Asset, AActor* CDO
 
 			USkeletalMesh* UseSkelMesh = PhysicsAsset->PreviewSkeletalMesh.Get();
 
-			SkeletalPhysicsActor->GetSkeletalMeshComponent()->SkeletalMesh = UseSkelMesh;
+			SkeletalPhysicsActor->GetSkeletalMeshComponent()->GetSkeletalMesh() = UseSkelMesh;
 			SkeletalPhysicsActor->GetSkeletalMeshComponent()->PhysicsAssetOverride = PhysicsAsset;
 		}
 
@@ -1334,7 +1334,7 @@ void UActorFactorySkeletalMesh::PostSpawnActor( UObject* Asset, AActor* NewActor
 	NewSMActor->GetSkeletalMeshComponent()->UnregisterComponent();
 
 	// Change properties
-	NewSMActor->GetSkeletalMeshComponent()->SkeletalMesh = SkeletalMesh;
+	NewSMActor->GetSkeletalMeshComponent()->GetSkeletalMesh() = SkeletalMesh;
 	if (NewSMActor->GetWorld()->IsGameWorld())
 	{
 		NewSMActor->ReplicatedMesh = SkeletalMesh;
@@ -1366,7 +1366,7 @@ void UActorFactorySkeletalMesh::PostCreateBlueprint( UObject* Asset, AActor* CDO
 		UAnimBlueprint* AnimBlueprint = Cast<UAnimBlueprint>(Asset);
 
 		ASkeletalMeshActor* SkeletalMeshActor = CastChecked<ASkeletalMeshActor>(CDO);
-		SkeletalMeshActor->GetSkeletalMeshComponent()->SkeletalMesh = SkeletalMesh;
+		SkeletalMeshActor->GetSkeletalMeshComponent()->GetSkeletalMesh() = SkeletalMesh;
 		SkeletalMeshActor->GetSkeletalMeshComponent()->AnimClass = AnimBlueprint ? Cast<UAnimBlueprintGeneratedClass>(AnimBlueprint->GeneratedClass) : nullptr;
 	}
 }
diff --git a/Engine/Source/Editor/UnrealEd/Private/Fbx/FbxAnimationExport.cpp b/Engine/Source/Editor/UnrealEd/Private/Fbx/FbxAnimationExport.cpp
--- a/Engine/Source/Editor/UnrealEd/Private/Fbx/FbxAnimationExport.cpp
+++ b/Engine/Source/Editor/UnrealEd/Private/Fbx/FbxAnimationExport.cpp
@@ -634,7 +634,7 @@ void FFbxExporter::ExportMatineeGroup(class AMatineeActor* MatineeActor, USkelet
 	}
 	// Create the Skeleton
 	TArray<FbxNode*> BoneNodes;
-	FbxNode* SkeletonRootNode = CreateSkeleton(SkeletalMeshComponent->SkeletalMesh, BoneNodes);
+	FbxNode* SkeletonRootNode = CreateSkeleton(SkeletalMeshComponent->GetSkeletalMesh(), BoneNodes);
 	FbxSkeletonRoots.Add(SkeletalMeshComponent, SkeletonRootNode);
 	BaseNode->AddChild(SkeletonRootNode);
 
@@ -754,7 +754,7 @@ void FFbxExporter::ExportAnimTrack(IAnimTrackAdapter& AnimTrackAdapter, AActor*
 		// Add the animation data to the bone nodes
 		for(int32 BoneIndex = 0; BoneIndex < BoneNodes.Num(); ++BoneIndex)
 		{
-			FName BoneName = InSkeletalMeshComponent->SkeletalMesh->GetRefSkeleton().GetBoneName(BoneIndex);
+			FName BoneName = InSkeletalMeshComponent->GetSkeletalMesh()->GetRefSkeleton().GetBoneName(BoneIndex);
 			FbxNode* CurrentBoneNode = BoneNodes[BoneIndex];
 
 			// Create the AnimCurves
diff --git a/Engine/Source/Editor/UnrealEd/Private/Fbx/FbxMainExport.cpp b/Engine/Source/Editor/UnrealEd/Private/Fbx/FbxMainExport.cpp
--- a/Engine/Source/Editor/UnrealEd/Private/Fbx/FbxMainExport.cpp
+++ b/Engine/Source/Editor/UnrealEd/Private/Fbx/FbxMainExport.cpp
@@ -1224,7 +1224,7 @@ void FFbxExporter::ExportSkeletalMesh( AActor* Actor, USkeletalMeshComponent* Sk
 	if (Scene == NULL || Actor == NULL || SkeletalMeshComponent == NULL) return;
 
 	// Retrieve the skeletal mesh rendering information.
-	USkeletalMesh* SkeletalMesh = SkeletalMeshComponent->SkeletalMesh;
+	USkeletalMesh* SkeletalMesh = SkeletalMeshComponent->GetSkeletalMesh();
 
 	FString FbxNodeName = NodeNameAdapter.GetActorNodeName(Actor);
 
@@ -2217,7 +2217,7 @@ FbxNode* FFbxExporter::ExportActor(AActor* Actor, bool bExportComponents, INodeN
 				{
 					ComponentsToExport.Add( StaticMeshComp );
 				}
-				else if( SkelMeshComp && SkelMeshComp->SkeletalMesh )
+				else if( SkelMeshComp && SkelMeshComp->GetSkeletalMesh())
 				{
 					ComponentsToExport.Add( SkelMeshComp );
 				}
@@ -2312,7 +2312,7 @@ FbxNode* FFbxExporter::ExportActor(AActor* Actor, bool bExportComponents, INodeN
 						ExportStaticMeshToFbx(StaticMeshComp->GetStaticMesh(), LODIndex, *StaticMeshComp->GetName(), ExportNode, LightmapUVChannel, ColorBuffer, MaterialOrderOverride, &ToRawPtrTArrayUnsafe(StaticMeshComp->OverrideMaterials));
 					}
 				}
-				else if (SkelMeshComp && SkelMeshComp->SkeletalMesh)
+				else if (SkelMeshComp && SkelMeshComp->GetSkeletalMesh())
 				{
 					ExportSkeletalMeshComponent(SkelMeshComp, *SkelMeshComp->GetName(), ExportNode, NodeNameAdapter, bSaveAnimSeq);
 				}
diff --git a/Engine/Source/Editor/UnrealEd/Private/Fbx/FbxSkeletalMeshExport.cpp b/Engine/Source/Editor/UnrealEd/Private/Fbx/FbxSkeletalMeshExport.cpp
--- a/Engine/Source/Editor/UnrealEd/Private/Fbx/FbxSkeletalMeshExport.cpp
+++ b/Engine/Source/Editor/UnrealEd/Private/Fbx/FbxSkeletalMeshExport.cpp
@@ -587,11 +587,11 @@ void FFbxExporter::CreateBindPose(FbxNode* MeshRootNode)
 
 void FFbxExporter::ExportSkeletalMeshComponent(USkeletalMeshComponent* SkelMeshComp, const TCHAR* MeshName, FbxNode* ActorRootNode, INodeNameAdapter& NodeNameAdapter, bool bSaveAnimSeq)
 {
-	if (SkelMeshComp && SkelMeshComp->SkeletalMesh)
+	if (SkelMeshComp && SkelMeshComp->GetSkeletalMesh())
 	{
 		UAnimSequence* AnimSeq = (bSaveAnimSeq && SkelMeshComp->GetAnimationMode() == EAnimationMode::AnimationSingleNode) ? 
 			Cast<UAnimSequence>(SkelMeshComp->AnimationData.AnimToPlay) : NULL;
-		FbxNode* SkeletonRootNode = ExportSkeletalMeshToFbx(SkelMeshComp->SkeletalMesh, AnimSeq, MeshName, ActorRootNode, &ToRawPtrTArrayUnsafe(SkelMeshComp->OverrideMaterials));
+		FbxNode* SkeletonRootNode = ExportSkeletalMeshToFbx(SkelMeshComp->GetSkeletalMesh(), AnimSeq, MeshName, ActorRootNode, &ToRawPtrTArrayUnsafe(SkelMeshComp->OverrideMaterials));
 		if(SkeletonRootNode)
 		{
 			FbxSkeletonRoots.Add(SkelMeshComp, SkeletonRootNode);
diff --git a/Engine/Source/Editor/UnrealEd/Private/Fbx/ReimportFbxSceneFactory.cpp b/Engine/Source/Editor/UnrealEd/Private/Fbx/ReimportFbxSceneFactory.cpp
--- a/Engine/Source/Editor/UnrealEd/Private/Fbx/ReimportFbxSceneFactory.cpp
+++ b/Engine/Source/Editor/UnrealEd/Private/Fbx/ReimportFbxSceneFactory.cpp
@@ -894,7 +894,7 @@ void UReimportFbxSceneFactory::RecursivelySetComponentProperties(USCS_Node* Curr
 		{
 			USkeletalMeshComponent *CurrentNodeMeshComponent = Cast<USkeletalMeshComponent>(CurrentNodeSceneComponent);
 			USkeletalMeshComponent *MeshComponent = Cast<USkeletalMeshComponent>(SceneComponent);
-			if (CurrentNodeMeshComponent->SkeletalMesh != MeshComponent->SkeletalMesh)
+			if (CurrentNodeMeshComponent->GetSkeletalMesh() != MeshComponent->GetSkeletalMesh())
 				bShouldSerializeProperty = false;
 		}
 
diff --git a/Engine/Source/Editor/UnrealEd/Private/LevelEditorViewport.cpp b/Engine/Source/Editor/UnrealEd/Private/LevelEditorViewport.cpp
--- a/Engine/Source/Editor/UnrealEd/Private/LevelEditorViewport.cpp
+++ b/Engine/Source/Editor/UnrealEd/Private/LevelEditorViewport.cpp
@@ -689,7 +689,7 @@ static bool AttemptApplyObjToComponent(UObject* ObjToUse, USceneComponent* Compo
 						SkeletalMeshComponent->Modify();
 
 						// If the component doesn't have a mesh or the anim blueprint's skeleton isn't compatible with the existing mesh's skeleton, the mesh should change
-						const bool bShouldChangeMesh = !SkeletalMeshComponent->SkeletalMesh || !SkeletalMeshComponent->SkeletalMesh->GetSkeleton()->IsCompatible(AnimBPSkeleton);
+						const bool bShouldChangeMesh = !SkeletalMeshComponent->GetSkeletalMesh() || !SkeletalMeshComponent->GetSkeletalMesh()->GetSkeleton()->IsCompatible(AnimBPSkeleton);
 
 						if (bShouldChangeMesh)
 						{
@@ -697,7 +697,7 @@ static bool AttemptApplyObjToComponent(UObject* ObjToUse, USceneComponent* Compo
 						}
 
 						// Verify that the skeletons are compatible before changing the anim BP
-						if (SkeletalMeshComponent->SkeletalMesh && SkeletalMeshComponent->SkeletalMesh->GetSkeleton()->IsCompatible(AnimBPSkeleton))
+						if (SkeletalMeshComponent->GetSkeletalMesh() && SkeletalMeshComponent->GetSkeletalMesh()->GetSkeleton()->IsCompatible(AnimBPSkeleton))
 						{
 							SkeletalMeshComponent->SetAnimInstanceClass(DroppedObjAsAnimBlueprint->GeneratedClass);
 							bResult = true;
@@ -724,7 +724,7 @@ static bool AttemptApplyObjToComponent(UObject* ObjToUse, USceneComponent* Compo
 						SkeletalMeshComponent->Modify();
 
 						// If the component doesn't have a mesh or the anim blueprint's skeleton isn't compatible with the existing mesh's skeleton, the mesh should change
-						const bool bShouldChangeMesh = !SkeletalMeshComponent->SkeletalMesh || !SkeletalMeshComponent->SkeletalMesh->GetSkeleton()->IsCompatible(AnimSkeleton);
+						const bool bShouldChangeMesh = !SkeletalMeshComponent->GetSkeletalMesh() || !SkeletalMeshComponent->GetSkeletalMesh()->GetSkeleton()->IsCompatible(AnimSkeleton);
 
 						if (bShouldChangeMesh)
 						{
@@ -737,7 +737,7 @@ static bool AttemptApplyObjToComponent(UObject* ObjToUse, USceneComponent* Compo
 						// set runtime data
 						SkeletalMeshComponent->SetAnimation(DroppedObjAsAnimSequence);
 
-						if (SkeletalMeshComponent->SkeletalMesh)
+						if (SkeletalMeshComponent->GetSkeletalMesh())
 						{
 							bResult = true;
 							SkeletalMeshComponent->InitAnim(true);
diff --git a/Engine/Source/Editor/UnrealEd/Private/ThumbnailHelpers.cpp b/Engine/Source/Editor/UnrealEd/Private/ThumbnailHelpers.cpp
--- a/Engine/Source/Editor/UnrealEd/Private/ThumbnailHelpers.cpp
+++ b/Engine/Source/Editor/UnrealEd/Private/ThumbnailHelpers.cpp
@@ -476,9 +476,9 @@ void FSkeletalMeshThumbnailScene::GetViewMatrixParameters(const float InFOVDegre
 	const float TargetDistance = HalfMeshSize / FMath::Tan(HalfFOVRadians);
 
 	USceneThumbnailInfo* ThumbnailInfo = nullptr;
-	if(PreviewActor->GetSkeletalMeshComponent()->SkeletalMesh)
+	if(PreviewActor->GetSkeletalMeshComponent()->GetSkeletalMesh())
 	{
-		ThumbnailInfo = Cast<USceneThumbnailInfo>(PreviewActor->GetSkeletalMeshComponent()->SkeletalMesh->GetThumbnailInfo());
+		ThumbnailInfo = Cast<USceneThumbnailInfo>(PreviewActor->GetSkeletalMeshComponent()->GetSkeletalMesh()->GetThumbnailInfo());
 	}
 	
 	if ( ThumbnailInfo )
@@ -690,7 +690,7 @@ void FAnimationSequenceThumbnailScene::GetViewMatrixParameters(const float InFOV
 {
 	check(PreviewAnimation);
 	check(PreviewActor->GetSkeletalMeshComponent());
-	check(PreviewActor->GetSkeletalMeshComponent()->SkeletalMesh);
+	check(PreviewActor->GetSkeletalMeshComponent()->GetSkeletalMesh());
 
 	const float HalfFOVRadians = FMath::DegreesToRadians<float>(InFOVDegrees) * 0.5f;
 	// No need to add extra size to view slightly outside of the sphere to compensate for perspective since skeletal meshes already buffer bounds.
@@ -824,7 +824,7 @@ void FBlendSpaceThumbnailScene::GetViewMatrixParameters(const float InFOVDegrees
 {
 	check(PreviewAnimation);
 	check(PreviewActor->GetSkeletalMeshComponent());
-	check(PreviewActor->GetSkeletalMeshComponent()->SkeletalMesh);
+	check(PreviewActor->GetSkeletalMeshComponent()->GetSkeletalMesh());
 
 	const float HalfFOVRadians = FMath::DegreesToRadians<float>(InFOVDegrees) * 0.5f;
 	// No need to add extra size to view slightly outside of the sphere to compensate for perspective since skeletal meshes already buffer bounds.
@@ -945,7 +945,7 @@ void FAnimBlueprintThumbnailScene::GetViewMatrixParameters(const float InFOVDegr
 {
 	check(PreviewBlueprint);
 	check(PreviewActor->GetSkeletalMeshComponent());
-	check(PreviewActor->GetSkeletalMeshComponent()->SkeletalMesh);
+	check(PreviewActor->GetSkeletalMeshComponent()->GetSkeletalMesh());
 
 	const float HalfFOVRadians = FMath::DegreesToRadians<float>(InFOVDegrees) * 0.5f;
 	// No need to add extra size to view slightly outside of the sphere to compensate for perspective since skeletal meshes already buffer bounds.
@@ -1029,9 +1029,9 @@ void FPhysicsAssetThumbnailScene::GetViewMatrixParameters(const float InFOVDegre
 	const float TargetDistance = HalfMeshSize / FMath::Tan(HalfFOVRadians);
 
 	USceneThumbnailInfo* ThumbnailInfo = USceneThumbnailInfo::StaticClass()->GetDefaultObject<USceneThumbnailInfo>();
-	if(PreviewActor->GetSkeletalMeshComponent()->SkeletalMesh && PreviewActor->GetSkeletalMeshComponent()->SkeletalMesh->GetPhysicsAsset())
+	if(PreviewActor->GetSkeletalMeshComponent()->GetSkeletalMesh() && PreviewActor->GetSkeletalMeshComponent()->GetSkeletalMesh()->GetPhysicsAsset())
 	{
-		if ( USceneThumbnailInfo* InteralThumbnailInfo = Cast<USceneThumbnailInfo>(PreviewActor->GetSkeletalMeshComponent()->SkeletalMesh->GetPhysicsAsset()->ThumbnailInfo) )
+		if ( USceneThumbnailInfo* InteralThumbnailInfo = Cast<USceneThumbnailInfo>(PreviewActor->GetSkeletalMeshComponent()->GetSkeletalMesh()->GetPhysicsAsset()->ThumbnailInfo) )
 		{
 			ThumbnailInfo = InteralThumbnailInfo;
 			if ( TargetDistance + InteralThumbnailInfo->OrbitZoom < 0 )
@@ -1124,7 +1124,7 @@ bool FClassActorThumbnailScene::IsValidComponentForVisualization(UActorComponent
 		}
 
 		USkeletalMeshComponent* SkelMeshComp = Cast<USkeletalMeshComponent>(Component);
-		if ( SkelMeshComp && SkelMeshComp->SkeletalMesh )
+		if ( SkelMeshComp && SkelMeshComp->GetSkeletalMesh())
 		{
 			return true;
 		}
diff --git a/Engine/Source/Editor/UnrealEd/Private/VertexSnapping.cpp b/Engine/Source/Editor/UnrealEd/Private/VertexSnapping.cpp
--- a/Engine/Source/Editor/UnrealEd/Private/VertexSnapping.cpp
+++ b/Engine/Source/Editor/UnrealEd/Private/VertexSnapping.cpp
@@ -242,7 +242,7 @@ static TSharedPtr<FVertexIterator> MakeVertexIterator( UPrimitiveComponent* Comp
 	}
 
 	USkinnedMeshComponent* SkinnedComponent = Cast<USkinnedMeshComponent>( Component );
-	if( SkinnedComponent && SkinnedComponent->SkeletalMesh && SkinnedComponent->MeshObject )
+	if( SkinnedComponent && SkinnedComponent->GetSkeletalMesh() && SkinnedComponent->MeshObject )
 	{
 		return MakeShareable( new FSkeletalMeshVertexIterator( SkinnedComponent ) );
 	}
diff --git a/Engine/Source/Runtime/AnimGraphRuntime/Private/AnimNodes/AnimNode_AimOffsetLookAt.cpp b/Engine/Source/Runtime/AnimGraphRuntime/Private/AnimNodes/AnimNode_AimOffsetLookAt.cpp
--- a/Engine/Source/Runtime/AnimGraphRuntime/Private/AnimNodes/AnimNode_AimOffsetLookAt.cpp
+++ b/Engine/Source/Runtime/AnimGraphRuntime/Private/AnimNodes/AnimNode_AimOffsetLookAt.cpp
@@ -31,7 +31,7 @@ void FAnimNode_AimOffsetLookAt::OnInitializeAnimInstance(const FAnimInstanceProx
 	SocketBoneReference.BoneName = NAME_None;
 	if (USkeletalMeshComponent* SkelMeshComp = InAnimInstance->GetSkelMeshComponent())
 	{
-		if (USkeletalMesh* SkelMesh = SkelMeshComp->SkeletalMesh)
+		if (USkeletalMesh* SkelMesh = SkelMeshComp->GetSkeletalMesh())
 		{
 			if (const USkeletalMeshSocket* Socket = SkelMesh->FindSocket(SourceSocketName))
 			{
diff --git a/Engine/Source/Runtime/AnimGraphRuntime/Private/AnimNodes/AnimNode_CopyPoseFromMesh.cpp b/Engine/Source/Runtime/AnimGraphRuntime/Private/AnimNodes/AnimNode_CopyPoseFromMesh.cpp
--- a/Engine/Source/Runtime/AnimGraphRuntime/Private/AnimNodes/AnimNode_CopyPoseFromMesh.cpp
+++ b/Engine/Source/Runtime/AnimGraphRuntime/Private/AnimNodes/AnimNode_CopyPoseFromMesh.cpp
@@ -45,14 +45,14 @@ void FAnimNode_CopyPoseFromMesh::RefreshMeshComponent(USkeletalMeshComponent* Ta
 				ReinitializeMeshComponent(InMeshComponent, InTargetMeshComponent);
 			}
 			// if component is still same but mesh has been changed, we have to reinitialize
-			else if (CurrentMeshComponent->SkeletalMesh != CurrentlyUsedSourceMesh.Get())
+			else if (CurrentMeshComponent->GetSkeletalMesh() != CurrentlyUsedSourceMesh.Get())
 			{
 				ReinitializeMeshComponent(InMeshComponent, InTargetMeshComponent);
 			}
 			else if (InTargetMeshComponent)
 			{
 				// see if target mesh has changed
-				if (InTargetMeshComponent->SkeletalMesh != CurrentlyUsedTargetMesh.Get())
+				if (InTargetMeshComponent->GetSkeletalMesh() != CurrentlyUsedTargetMesh.Get())
 				{
 					ReinitializeMeshComponent(InMeshComponent, InTargetMeshComponent);
 				}
@@ -112,7 +112,7 @@ void FAnimNode_CopyPoseFromMesh::PreUpdate(const UAnimInstance* InAnimInstance)
 
 	USkeletalMeshComponent* CurrentMeshComponent = CurrentlyUsedSourceMeshComponent.IsValid() ? CurrentlyUsedSourceMeshComponent.Get() : nullptr;
 
-	if (CurrentMeshComponent && CurrentMeshComponent->SkeletalMesh && CurrentMeshComponent->IsRegistered())
+	if (CurrentMeshComponent && CurrentMeshComponent->GetSkeletalMesh() && CurrentMeshComponent->IsRegistered())
 	{
 		// If our source is running under master-pose, then get bone data from there
 		if(USkeletalMeshComponent* MasterPoseComponent = Cast<USkeletalMeshComponent>(CurrentMeshComponent->MasterPoseComponent.Get()))
@@ -121,7 +121,7 @@ void FAnimNode_CopyPoseFromMesh::PreUpdate(const UAnimInstance* InAnimInstance)
 		}
 
 		// re-check mesh component validity as it may have changed to master
-		if(CurrentMeshComponent->SkeletalMesh && CurrentMeshComponent->IsRegistered())
+		if(CurrentMeshComponent->GetSkeletalMesh() && CurrentMeshComponent->IsRegistered())
 		{
 			const bool bUROInSync = CurrentMeshComponent->ShouldUseUpdateRateOptimizations() && CurrentMeshComponent->AnimUpdateRateParams != nullptr && CurrentMeshComponent->AnimUpdateRateParams == InAnimInstance->GetSkelMeshComponent()->AnimUpdateRateParams;
 			const bool bUsingExternalInterpolation = CurrentMeshComponent->IsUsingExternalInterpolation();
@@ -133,7 +133,7 @@ void FAnimNode_CopyPoseFromMesh::PreUpdate(const UAnimInstance* InAnimInstance)
 			SourceMeshTransformArray.Append((bUROInSync || bUsingExternalInterpolation) && bArraySizesMatch ? CachedComponentSpaceTransforms : CurrentMeshComponent->GetComponentSpaceTransforms());
 
 			// Ref skeleton is need for parent index lookups later, so store it now
-			CurrentlyUsedMesh = CurrentMeshComponent->SkeletalMesh;
+			CurrentlyUsedMesh = CurrentMeshComponent->GetSkeletalMesh();
 
 			if(bCopyCurves)
 			{
@@ -170,7 +170,7 @@ void FAnimNode_CopyPoseFromMesh::Update_AnyThread(const FAnimationUpdateContext&
 	GetEvaluateGraphExposedInputs().Execute(Context);
 
 	TRACE_ANIM_NODE_VALUE(Context, TEXT("Component"), *GetNameSafe(CurrentlyUsedSourceMeshComponent.IsValid() ? CurrentlyUsedSourceMeshComponent.Get() : nullptr));
-	TRACE_ANIM_NODE_VALUE(Context, TEXT("Mesh"), *GetNameSafe(CurrentlyUsedSourceMeshComponent.IsValid() ? CurrentlyUsedSourceMeshComponent.Get()->SkeletalMesh : nullptr));
+	TRACE_ANIM_NODE_VALUE(Context, TEXT("Mesh"), *GetNameSafe(CurrentlyUsedSourceMeshComponent.IsValid() ? CurrentlyUsedSourceMeshComponent.Get()->GetSkeletalMesh() : nullptr));
 }
 
 void FAnimNode_CopyPoseFromMesh::Evaluate_AnyThread(FPoseContext& Output)
@@ -253,7 +253,7 @@ void FAnimNode_CopyPoseFromMesh::GatherDebugData(FNodeDebugData& DebugData)
 	DECLARE_SCOPE_HIERARCHICAL_COUNTER_ANIMNODE(GatherDebugData)
 	FString DebugLine = DebugData.GetNodeName(this);
 
-	DebugLine += FString::Printf(TEXT("('%s')"), *GetNameSafe(CurrentlyUsedSourceMeshComponent.IsValid() ? CurrentlyUsedSourceMeshComponent.Get()->SkeletalMesh : nullptr));
+	DebugLine += FString::Printf(TEXT("('%s')"), *GetNameSafe(CurrentlyUsedSourceMeshComponent.IsValid() ? CurrentlyUsedSourceMeshComponent.Get()->GetSkeletalMesh() : nullptr));
 	DebugData.AddDebugItem(DebugLine, true);
 }
 
@@ -266,10 +266,10 @@ void FAnimNode_CopyPoseFromMesh::ReinitializeMeshComponent(USkeletalMeshComponen
 	BoneMapToSource.Reset();
 	CurveNameToUIDMap.Reset();
 
-	if (TargetMeshComponent && IsValid(NewSourceMeshComponent) && NewSourceMeshComponent->SkeletalMesh)
+	if (TargetMeshComponent && IsValid(NewSourceMeshComponent) && NewSourceMeshComponent->GetSkeletalMesh())
 	{
-		USkeletalMesh* SourceSkelMesh = NewSourceMeshComponent->SkeletalMesh;
-		USkeletalMesh* TargetSkelMesh = TargetMeshComponent->SkeletalMesh;
+		USkeletalMesh* SourceSkelMesh = NewSourceMeshComponent->GetSkeletalMesh();
+		USkeletalMesh* TargetSkelMesh = TargetMeshComponent->GetSkeletalMesh();
 		
 		if (IsValid(SourceSkelMesh) && !SourceSkelMesh->HasAnyFlags(RF_NeedPostLoad) &&
 			IsValid(TargetSkelMesh) && !TargetSkelMesh->HasAnyFlags(RF_NeedPostLoad))
diff --git a/Engine/Source/Runtime/AnimGraphRuntime/Private/AnimNodes/AnimNode_PoseSnapshot.cpp b/Engine/Source/Runtime/AnimGraphRuntime/Private/AnimNodes/AnimNode_PoseSnapshot.cpp
--- a/Engine/Source/Runtime/AnimGraphRuntime/Private/AnimNodes/AnimNode_PoseSnapshot.cpp
+++ b/Engine/Source/Runtime/AnimGraphRuntime/Private/AnimNodes/AnimNode_PoseSnapshot.cpp
@@ -22,7 +22,7 @@ void FAnimNode_PoseSnapshot::PreUpdate(const UAnimInstance* InAnimInstance)
 	USkeletalMesh* CurrentSkeletalMesh = nullptr;
 	if (InAnimInstance->GetSkelMeshComponent() && InAnimInstance->GetSkelMeshComponent()->IsRegistered())
 	{
-		CurrentSkeletalMesh = InAnimInstance->GetSkelMeshComponent()->SkeletalMesh;
+		CurrentSkeletalMesh = InAnimInstance->GetSkelMeshComponent()->GetSkeletalMesh();
 	}
 
 	if (CurrentSkeletalMesh)
diff --git a/Engine/Source/Runtime/AnimGraphRuntime/Private/AnimSequencerInstance.cpp b/Engine/Source/Runtime/AnimGraphRuntime/Private/AnimSequencerInstance.cpp
--- a/Engine/Source/Runtime/AnimGraphRuntime/Private/AnimSequencerInstance.cpp
+++ b/Engine/Source/Runtime/AnimGraphRuntime/Private/AnimSequencerInstance.cpp
@@ -65,7 +65,7 @@ void UAnimSequencerInstance::SavePose()
 {
 	if (USkeletalMeshComponent* SkeletalMeshComponent = GetSkelMeshComponent())
 	{
-		if (SkeletalMeshComponent->SkeletalMesh && SkeletalMeshComponent->GetComponentSpaceTransforms().Num() > 0)
+		if (SkeletalMeshComponent->GetSkeletalMesh() && SkeletalMeshComponent->GetComponentSpaceTransforms().Num() > 0)
 		{
 			SavePoseSnapshot(UAnimSequencerInstance::SequencerPoseName);
 		}
diff --git a/Engine/Source/Runtime/AnimGraphRuntime/Private/BoneControllers/AnimNode_RigidBody.cpp b/Engine/Source/Runtime/AnimGraphRuntime/Private/BoneControllers/AnimNode_RigidBody.cpp
--- a/Engine/Source/Runtime/AnimGraphRuntime/Private/BoneControllers/AnimNode_RigidBody.cpp
+++ b/Engine/Source/Runtime/AnimGraphRuntime/Private/BoneControllers/AnimNode_RigidBody.cpp
@@ -958,7 +958,7 @@ void ComputeBodyInsertionOrder(TArray<FBoneIndexType>& InsertionOrder, const USk
 
 	InsertionOrder.Reset();
 
-	if (SKC.SkeletalMesh == nullptr)
+	if (SKC.GetSkeletalMesh() == nullptr)
 	{
 		return;
 	}
@@ -971,7 +971,7 @@ void ComputeBodyInsertionOrder(TArray<FBoneIndexType>& InsertionOrder, const USk
 		SKC.ComputeRequiredBones(RequiredBones0, ComponentSpaceTMs0, 0, /*bIgnorePhysicsAsset=*/ true);
 
 		TArray<bool> InSortedOrder;
-		InSortedOrder.AddZeroed(SKC.SkeletalMesh->GetRefSkeleton().GetNum());
+		InSortedOrder.AddZeroed(SKC.GetSkeletalMesh()->GetRefSkeleton().GetNum());
 
 		auto MergeIndices = [&InsertionOrder, &InSortedOrder](const TArray<FBoneIndexType>& RequiredBones) -> void
 		{
@@ -1008,7 +1008,7 @@ void FAnimNode_RigidBody::InitPhysics(const UAnimInstance* InAnimInstance)
 	DestroyPhysicsSimulation();
 
 	const USkeletalMeshComponent* SkeletalMeshComp = InAnimInstance->GetSkelMeshComponent();
-	const USkeletalMesh* SkeletalMeshAsset = SkeletalMeshComp->SkeletalMesh;
+	const USkeletalMesh* SkeletalMeshAsset = SkeletalMeshComp->GetSkeletalMesh();
 	USkeleton* SkeletonAsset = InAnimInstance->CurrentSkeleton;
 
 	if(!SkeletalMeshAsset || !SkeletonAsset)
diff --git a/Engine/Source/Runtime/AnimGraphRuntime/Private/BoneControllers/AnimNode_SplineIK.cpp b/Engine/Source/Runtime/AnimGraphRuntime/Private/BoneControllers/AnimNode_SplineIK.cpp
--- a/Engine/Source/Runtime/AnimGraphRuntime/Private/BoneControllers/AnimNode_SplineIK.cpp
+++ b/Engine/Source/Runtime/AnimGraphRuntime/Private/BoneControllers/AnimNode_SplineIK.cpp
@@ -35,9 +35,9 @@ void FAnimNode_SplineIK::GatherDebugData(FNodeDebugData& DebugData)
 
 void FAnimNode_SplineIK::OnInitializeAnimInstance(const FAnimInstanceProxy* InProxy, const UAnimInstance* InAnimInstance)
 {
-	if (InAnimInstance->GetSkelMeshComponent() && InAnimInstance->GetSkelMeshComponent()->SkeletalMesh)
+	if (InAnimInstance->GetSkelMeshComponent() && InAnimInstance->GetSkelMeshComponent()->GetSkeletalMesh())
 	{
-		GatherBoneReferences(InAnimInstance->GetSkelMeshComponent()->SkeletalMesh->GetRefSkeleton());
+		GatherBoneReferences(InAnimInstance->GetSkelMeshComponent()->GetSkeletalMesh()->GetRefSkeleton());
 	}
 }
 
diff --git a/Engine/Source/Runtime/ClothingSystemRuntimeCommon/Private/ClothingAsset.cpp b/Engine/Source/Runtime/ClothingSystemRuntimeCommon/Private/ClothingAsset.cpp
--- a/Engine/Source/Runtime/ClothingSystemRuntimeCommon/Private/ClothingAsset.cpp
+++ b/Engine/Source/Runtime/ClothingSystemRuntimeCommon/Private/ClothingAsset.cpp
@@ -833,7 +833,7 @@ void UClothingAssetCommon::ReregisterComponentsUsingClothing()
 		{
 			if(USkeletalMeshComponent* Component = *It)
 			{
-				if(Component->SkeletalMesh == OwnerMesh)
+				if(Component->GetSkeletalMesh() == OwnerMesh)
 				{
 					FComponentReregisterContext Context(Component);
 					// Context goes out of scope, causing Component to be re-registered
@@ -851,7 +851,7 @@ void UClothingAssetCommon::ForEachInteractorUsingClothing(TFunction<void(UClothi
 		{
 			if(USkeletalMeshComponent* Component = *It)
 			{
-				if(Component->SkeletalMesh == OwnerMesh)
+				if(Component->GetSkeletalMesh() == OwnerMesh)
 				{
 					if(UClothingSimulationInteractor* CurInteractor = Component->GetClothingSimulationInteractor())
 					{
diff --git a/Engine/Source/Runtime/ClothingSystemRuntimeCommon/Private/ClothingSimulation.cpp b/Engine/Source/Runtime/ClothingSystemRuntimeCommon/Private/ClothingSimulation.cpp
--- a/Engine/Source/Runtime/ClothingSystemRuntimeCommon/Private/ClothingSimulation.cpp
+++ b/Engine/Source/Runtime/ClothingSystemRuntimeCommon/Private/ClothingSimulation.cpp
@@ -60,7 +60,7 @@ void FClothingSimulationContextCommon::Fill(const USkeletalMeshComponent* InComp
 
 void FClothingSimulationContextCommon::FillBoneTransforms(const USkeletalMeshComponent* InComponent)
 {
-	const USkeletalMesh* const SkeletalMesh = InComponent->SkeletalMesh;
+	const USkeletalMesh* const SkeletalMesh = InComponent->GetSkeletalMesh();
 
 	if (USkinnedMeshComponent* const MasterComponent = InComponent->MasterPoseComponent.Get())
 	{
@@ -122,7 +122,7 @@ void FClothingSimulationContextCommon::FillRefToLocals(const USkeletalMeshCompon
 	// Constraints are initialized using bone distances upon initialization, so fill out reference pose
 	if (bIsInitialization)
 	{
-		const USkeletalMesh* const SkeletalMesh = InComponent->SkeletalMesh;
+		const USkeletalMesh* const SkeletalMesh = InComponent->GetSkeletalMesh();
 		if (SkeletalMesh)
 		{
 			const int32 NumBones = SkeletalMesh->GetRefSkeleton().GetNum();
diff --git a/Engine/Source/Runtime/Engine/Classes/Components/SkinnedMeshComponent.h b/Engine/Source/Runtime/Engine/Classes/Components/SkinnedMeshComponent.h
--- a/Engine/Source/Runtime/Engine/Classes/Components/SkinnedMeshComponent.h
+++ b/Engine/Source/Runtime/Engine/Classes/Components/SkinnedMeshComponent.h
@@ -360,6 +360,15 @@ protected:
 	TSharedPtr<FSkelMeshRefPoseOverride> RefPoseOverride;
 
 public:
+	TObjectPtr<class USkeletalMesh>& GetSkeletalMesh()
+	{
+		return SkeletalMesh;
+	}
+	const TObjectPtr<class USkeletalMesh>& GetSkeletalMesh() const
+	{
+		return SkeletalMesh;
+	}
+
 	USkinnedAsset* GetSkinnedAsset() const;
 
 	const TArray<int32>& GetMasterBoneMap() const { return MasterBoneMap; }
diff --git a/Engine/Source/Runtime/Engine/Private/Animation/AnimInstance.cpp b/Engine/Source/Runtime/Engine/Private/Animation/AnimInstance.cpp
--- a/Engine/Source/Runtime/Engine/Private/Animation/AnimInstance.cpp
+++ b/Engine/Source/Runtime/Engine/Private/Animation/AnimInstance.cpp
@@ -210,9 +210,9 @@ void UAnimInstance::InitializeAnimation(bool bInDeferRootNodeInitialization)
 	// make sure your skeleton is initialized
 	// you can overwrite different skeleton
 	USkeletalMeshComponent* OwnerComponent = GetSkelMeshComponent();
-	if (OwnerComponent->SkeletalMesh != NULL)
+	if (OwnerComponent->GetSkeletalMesh() != NULL)
 	{
-		CurrentSkeleton = OwnerComponent->SkeletalMesh->GetSkeleton();
+		CurrentSkeleton = OwnerComponent->GetSkeletalMesh()->GetSkeleton();
 	}
 	else
 	{
@@ -1228,9 +1228,9 @@ void UAnimInstance::RecalcRequiredBones()
 	USkeletalMeshComponent* SkelMeshComp = GetSkelMeshComponent();
 	check( SkelMeshComp )
 
-	if( SkelMeshComp->SkeletalMesh && SkelMeshComp->SkeletalMesh->GetSkeleton() )
+	if( SkelMeshComp->GetSkeletalMesh() && SkelMeshComp->GetSkeletalMesh()->GetSkeleton() )
 	{
-		GetProxyOnGameThread<FAnimInstanceProxy>().RecalcRequiredBones(SkelMeshComp, SkelMeshComp->SkeletalMesh);
+		GetProxyOnGameThread<FAnimInstanceProxy>().RecalcRequiredBones(SkelMeshComp, SkelMeshComp->GetSkeletalMesh());
 	}
 	else if( CurrentSkeleton != NULL )
 	{
@@ -1669,9 +1669,9 @@ void UAnimInstance::GetActiveCurveNames(EAnimCurveType CurveType, TArray<FName>&
 void UAnimInstance::GetAllCurveNames(TArray<FName>& OutNames) const
 {
 	USkeletalMeshComponent* SkelMeshComp = GetOwningComponent();
-	if (SkelMeshComp && SkelMeshComp->SkeletalMesh && SkelMeshComp->SkeletalMesh->GetSkeleton())
+	if (SkelMeshComp && SkelMeshComp->GetSkeletalMesh() && SkelMeshComp->GetSkeletalMesh()->GetSkeleton())
 	{
-		const USkeleton* CurSkeleton = SkelMeshComp->SkeletalMesh->GetSkeleton();
+		const USkeleton* CurSkeleton = SkelMeshComp->GetSkeletalMesh()->GetSkeleton();
 
 		const FSmartNameMapping* Mapping = CurSkeleton->GetSmartNameContainer(USkeleton::AnimCurveMappingName);
 		if (Mapping)
diff --git a/Engine/Source/Runtime/Engine/Private/Animation/AnimInstanceProxy.cpp b/Engine/Source/Runtime/Engine/Private/Animation/AnimInstanceProxy.cpp
--- a/Engine/Source/Runtime/Engine/Private/Animation/AnimInstanceProxy.cpp
+++ b/Engine/Source/Runtime/Engine/Private/Animation/AnimInstanceProxy.cpp
@@ -649,9 +649,9 @@ void FAnimInstanceProxy::InitializeObjects(UAnimInstance* InAnimInstance)
 		MainMontageEvaluationData = InAnimInstance->IsUsingMainInstanceMontageEvaluationData() ? &MainInstanceProxy->MontageEvaluationData : &MontageEvaluationData;
 	}
 
-	if (SkeletalMeshComponent->SkeletalMesh != nullptr)
+	if (SkeletalMeshComponent->GetSkeletalMesh() != nullptr)
 	{
-		Skeleton = SkeletalMeshComponent->SkeletalMesh->GetSkeleton();
+		Skeleton = SkeletalMeshComponent->GetSkeletalMesh()->GetSkeleton();
 	}
 	else
 	{
diff --git a/Engine/Source/Runtime/Engine/Private/Animation/AnimSequence.cpp b/Engine/Source/Runtime/Engine/Private/Animation/AnimSequence.cpp
--- a/Engine/Source/Runtime/Engine/Private/Animation/AnimSequence.cpp
+++ b/Engine/Source/Runtime/Engine/Private/Animation/AnimSequence.cpp
@@ -3866,9 +3866,9 @@ bool UAnimSequence::CreateAnimation(USkeletalMesh* Mesh)
 
 bool UAnimSequence::CreateAnimation(USkeletalMeshComponent* MeshComponent)
 {
-	if(MeshComponent && MeshComponent->SkeletalMesh)
+	if(MeshComponent && MeshComponent->GetSkeletalMesh())
 	{
-		USkeletalMesh * Mesh = MeshComponent->SkeletalMesh;
+		USkeletalMesh * Mesh = MeshComponent->GetSkeletalMesh();
 
 		ValidateModel();
 		const FReferenceSkeleton& RefSkeleton = Mesh->GetRefSkeleton();
diff --git a/Engine/Source/Runtime/Engine/Private/Animation/AnimSingleNodeInstance.cpp b/Engine/Source/Runtime/Engine/Private/Animation/AnimSingleNodeInstance.cpp
--- a/Engine/Source/Runtime/Engine/Private/Animation/AnimSingleNodeInstance.cpp
+++ b/Engine/Source/Runtime/Engine/Private/Animation/AnimSingleNodeInstance.cpp
@@ -40,7 +40,7 @@ void UAnimSingleNodeInstance::SetAnimationAsset(class UAnimationAsset* NewAsset,
 	USkeletalMeshComponent* MeshComponent = GetSkelMeshComponent();
 	if (MeshComponent)
 	{
-		if (MeshComponent->SkeletalMesh == nullptr)
+		if (MeshComponent->GetSkeletalMesh() == nullptr)
 		{
 			// if it does not have SkeletalMesh, we nullify it
 			CurrentAsset = nullptr;
@@ -48,7 +48,7 @@ void UAnimSingleNodeInstance::SetAnimationAsset(class UAnimationAsset* NewAsset,
 		else if (CurrentAsset != nullptr)
 		{
 			// if we have an asset, make sure their skeleton matches, otherwise, null it
-			if (!MeshComponent->SkeletalMesh->GetSkeleton()->IsCompatible(CurrentAsset->GetSkeleton()))
+			if (!MeshComponent->GetSkeletalMesh()->GetSkeleton()->IsCompatible(CurrentAsset->GetSkeleton()))
 			{
 				// clear asset since we do not have matching skeleton
 				CurrentAsset = nullptr;
diff --git a/Engine/Source/Runtime/Engine/Private/Animation/AnimTrace.cpp b/Engine/Source/Runtime/Engine/Private/Animation/AnimTrace.cpp
--- a/Engine/Source/Runtime/Engine/Private/Animation/AnimTrace.cpp
+++ b/Engine/Source/Runtime/Engine/Private/Animation/AnimTrace.cpp
@@ -534,7 +534,7 @@ void FAnimTrace::OutputSkeletalMeshComponent(const USkeletalMeshComponent* InCom
 	if(BoneCount > 0 || CurveCount > 0)
 	{
 		TRACE_OBJECT(InComponent);
-		TRACE_SKELETAL_MESH(InComponent->SkeletalMesh);
+		TRACE_SKELETAL_MESH(InComponent->GetSkeletalMesh());
 
 		TArray<float>& CurveValues = FAnimTraceScratchBuffers::Get().CurveValues;
 		CurveValues.Reset();
@@ -561,7 +561,7 @@ void FAnimTrace::OutputSkeletalMeshComponent(const USkeletalMeshComponent* InCom
 			<< SkeletalMeshComponent3.Cycle(FPlatformTime::Cycles64())
 			<< SkeletalMeshComponent3.RecordingTime(FObjectTrace::GetWorldElapsedTime(InComponent->GetWorld()))
 			<< SkeletalMeshComponent3.ComponentId(FObjectTrace::GetObjectId(InComponent))
-			<< SkeletalMeshComponent3.MeshId(FObjectTrace::GetObjectId(InComponent->SkeletalMesh))
+			<< SkeletalMeshComponent3.MeshId(FObjectTrace::GetObjectId(InComponent->GetSkeletalMesh()))
 			<< SkeletalMeshComponent3.ComponentToWorld(reinterpret_cast<const float*>(&InComponent->GetComponentToWorld()), sizeof(FTransform) / sizeof(float))
 			<< SkeletalMeshComponent3.Pose(reinterpret_cast<const float*>(InComponent->GetComponentSpaceTransforms().GetData()), BoneCount * (sizeof(FTransform) / sizeof(float)))
 			<< SkeletalMeshComponent3.CurveIds(CurveIds.GetData(), CurveIds.Num())
diff --git a/Engine/Source/Runtime/Engine/Private/Animation/PoseAsset.cpp b/Engine/Source/Runtime/Engine/Private/Animation/PoseAsset.cpp
--- a/Engine/Source/Runtime/Engine/Private/Animation/PoseAsset.cpp
+++ b/Engine/Source/Runtime/Engine/Private/Animation/PoseAsset.cpp
@@ -1191,12 +1191,12 @@ bool UPoseAsset::AddOrUpdatePoseWithUniqueName(USkeletalMeshComponent* MeshCompo
 void UPoseAsset::AddOrUpdatePose(const FSmartName& PoseName, USkeletalMeshComponent* MeshComponent, bool bUpdateCurves)
 {
 	USkeleton* MySkeleton = GetSkeleton();
-	if (MySkeleton && MeshComponent && MeshComponent->SkeletalMesh)
+	if (MySkeleton && MeshComponent && MeshComponent->GetSkeletalMesh())
 	{
 		TArray<FName> TrackNames;
 		// note this ignores root motion
 		TArray<FTransform> BoneTransform = MeshComponent->GetComponentSpaceTransforms();
-		const FReferenceSkeleton& RefSkeleton = MeshComponent->SkeletalMesh->GetRefSkeleton();
+		const FReferenceSkeleton& RefSkeleton = MeshComponent->GetSkeletalMesh()->GetRefSkeleton();
 		for (int32 BoneIndex = 0; BoneIndex < RefSkeleton.GetNum(); ++BoneIndex)
 		{
 			TrackNames.Add(RefSkeleton.GetBoneName(BoneIndex));
@@ -1212,7 +1212,7 @@ void UPoseAsset::AddOrUpdatePose(const FSmartName& PoseName, USkeletalMeshCompon
 			}
 		}
 
-		const USkeleton* MeshSkeleton = MeshComponent->SkeletalMesh->GetSkeleton();
+		const USkeleton* MeshSkeleton = MeshComponent->GetSkeletalMesh()->GetSkeleton();
 		const FSmartNameMapping* Mapping = MeshSkeleton->GetSmartNameContainer(USkeleton::AnimCurveMappingName);
 		
 		TArray<float> NewCurveValues;
diff --git a/Engine/Source/Runtime/Engine/Private/Animation/SkeletalMeshActor.cpp b/Engine/Source/Runtime/Engine/Private/Animation/SkeletalMeshActor.cpp
--- a/Engine/Source/Runtime/Engine/Private/Animation/SkeletalMeshActor.cpp
+++ b/Engine/Source/Runtime/Engine/Private/Animation/SkeletalMeshActor.cpp
@@ -74,7 +74,7 @@ void ASkeletalMeshActor::CheckForErrors()
 				->AddToken(FMapErrorToken::Create(FMapErrors::ActorLargeShadowCaster));
 		}
 
-		if (SkeletalMeshComponent->SkeletalMesh == NULL)
+		if (SkeletalMeshComponent->GetSkeletalMesh() == NULL)
 		{
 			FMessageLog("MapCheck").Warning()
 				->AddToken(FUObjectToken::Create(this))
@@ -116,7 +116,7 @@ void ASkeletalMeshActor::PostInitializeComponents()
 	// grab the current mesh for replication
 	if (GetLocalRole() == ROLE_Authority && SkeletalMeshComponent)
 	{
-		ReplicatedMesh = SkeletalMeshComponent->SkeletalMesh;
+		ReplicatedMesh = SkeletalMeshComponent->GetSkeletalMesh();
 	}
 
 	// Unfix bodies flagged as 'full anim weight'
@@ -148,8 +148,8 @@ void ASkeletalMeshActor::OnRep_ReplicatedMaterial1()
 
 bool ASkeletalMeshActor::CanPlayAnimation(class UAnimSequenceBase* AnimAssetBase/*=NULL*/) const
 {
-	return (SkeletalMeshComponent->SkeletalMesh && SkeletalMeshComponent->SkeletalMesh->GetSkeleton() &&
-		(!AnimAssetBase || SkeletalMeshComponent->SkeletalMesh->GetSkeleton()->IsCompatible(AnimAssetBase->GetSkeleton())));
+	return (SkeletalMeshComponent->GetSkeletalMesh() && SkeletalMeshComponent->GetSkeletalMesh()->GetSkeleton() &&
+		(!AnimAssetBase || SkeletalMeshComponent->GetSkeletalMesh()->GetSkeleton()->IsCompatible(AnimAssetBase->GetSkeleton())));
 }
 
 #if WITH_EDITOR
@@ -158,9 +158,9 @@ bool ASkeletalMeshActor::GetReferencedContentObjects(TArray<UObject*>& Objects)
 {
 	Super::GetReferencedContentObjects(Objects);
 
-	if (SkeletalMeshComponent->SkeletalMesh)
+	if (SkeletalMeshComponent->GetSkeletalMesh())
 	{
-		Objects.Add(SkeletalMeshComponent->SkeletalMesh);
+		Objects.Add(SkeletalMeshComponent->GetSkeletalMesh());
 	}
 
 	if (SkeletalMeshComponent->GetAnimationMode() == EAnimationMode::Type::AnimationSingleNode &&
@@ -180,9 +180,9 @@ void ASkeletalMeshActor::EditorReplacedActor(AActor* OldActor)
 	{
 		// if no skeletal mesh set, take one from previous actor
 		if (SkeletalMeshComponent && OldSkelMeshActor->SkeletalMeshComponent &&
-			SkeletalMeshComponent->SkeletalMesh == NULL)
+			SkeletalMeshComponent->GetSkeletalMesh() == NULL)
 		{
-			SkeletalMeshComponent->SetSkeletalMesh(OldSkelMeshActor->SkeletalMeshComponent->SkeletalMesh);
+			SkeletalMeshComponent->SetSkeletalMesh(OldSkelMeshActor->SkeletalMeshComponent->GetSkeletalMesh());
 		}
 	}
 }
diff --git a/Engine/Source/Runtime/Engine/Private/Animation/Skeleton.cpp b/Engine/Source/Runtime/Engine/Private/Animation/Skeleton.cpp
--- a/Engine/Source/Runtime/Engine/Private/Animation/Skeleton.cpp
+++ b/Engine/Source/Runtime/Engine/Private/Animation/Skeleton.cpp
@@ -2356,8 +2356,8 @@ void USkeleton::HandleVirtualBoneChanges()
 	{
 		USkinnedMeshComponent* MeshComponent = *It;
 		if (MeshComponent &&
-			MeshComponent->SkeletalMesh &&
-			MeshComponent->SkeletalMesh->GetSkeleton() == this &&
+			MeshComponent->GetSkeletalMesh() &&
+			MeshComponent->GetSkeletalMesh()->GetSkeleton() == this &&
 			!MeshComponent->IsTemplate())
 		{
 			FComponentReregisterContext Context(MeshComponent);
diff --git a/Engine/Source/Runtime/Engine/Private/Animation/SkinWeightProfile.cpp b/Engine/Source/Runtime/Engine/Private/Animation/SkinWeightProfile.cpp
--- a/Engine/Source/Runtime/Engine/Private/Animation/SkinWeightProfile.cpp
+++ b/Engine/Source/Runtime/Engine/Private/Animation/SkinWeightProfile.cpp
@@ -28,7 +28,7 @@ public:
 	{
 		for (TObjectIterator<USkinnedMeshComponent> It; It; ++It)
 		{
-			if (It->SkeletalMesh == InSkeletalMesh)
+			if (It->GetSkeletalMesh() == InSkeletalMesh)
 			{
 				checkf(!It->IsUnreachable(), TEXT("%s"), *It->GetFullName());
 
diff --git a/Engine/Source/Runtime/Engine/Private/Collision/Collision.cpp b/Engine/Source/Runtime/Engine/Private/Collision/Collision.cpp
--- a/Engine/Source/Runtime/Engine/Private/Collision/Collision.cpp
+++ b/Engine/Source/Runtime/Engine/Private/Collision/Collision.cpp
@@ -502,9 +502,9 @@ namespace CollisionResponseConsoleCommands
 		}
 		else if (const USkinnedMeshComponent* SkinnedMeshComp = Cast<USkinnedMeshComponent>(Comp))
 		{
-			if (SkinnedMeshComp->SkeletalMesh)
+			if (SkinnedMeshComp->GetSkeletalMesh())
 			{
-				AssetName = SkinnedMeshComp->SkeletalMesh->GetPathName();
+				AssetName = SkinnedMeshComp->GetSkeletalMesh()->GetPathName();
 			}
 		}
 
diff --git a/Engine/Source/Runtime/Engine/Private/Components/PoseableMeshComponent.cpp b/Engine/Source/Runtime/Engine/Private/Components/PoseableMeshComponent.cpp
--- a/Engine/Source/Runtime/Engine/Private/Components/PoseableMeshComponent.cpp
+++ b/Engine/Source/Runtime/Engine/Private/Components/PoseableMeshComponent.cpp
@@ -21,9 +21,9 @@ bool UPoseableMeshComponent::AllocateTransformData()
 	// Allocate transforms if not present.
 	if ( Super::AllocateTransformData() )
 	{
-		if(BoneSpaceTransforms.Num() != SkeletalMesh->GetRefSkeleton().GetNum() )
+		if(BoneSpaceTransforms.Num() != GetSkeletalMesh()->GetRefSkeleton().GetNum() )
 		{
-			BoneSpaceTransforms = SkeletalMesh->GetRefSkeleton().GetRefBonePose();
+			BoneSpaceTransforms = GetSkeletalMesh()->GetRefSkeleton().GetRefBonePose();
 
 			TArray<FBoneIndexType> RequiredBoneIndexArray;
 			RequiredBoneIndexArray.AddUninitialized(BoneSpaceTransforms.Num());
@@ -32,7 +32,7 @@ bool UPoseableMeshComponent::AllocateTransformData()
 				RequiredBoneIndexArray[BoneIndex] = BoneIndex;
 			}
 
-			RequiredBones.InitializeTo(RequiredBoneIndexArray, FCurveEvaluationOption(false), *SkeletalMesh);
+			RequiredBones.InitializeTo(RequiredBoneIndexArray, FCurveEvaluationOption(false), *GetSkeletalMesh());
 		}
 
 		FillComponentSpaceTransforms();
@@ -51,7 +51,7 @@ void UPoseableMeshComponent::RefreshBoneTransforms(FActorComponentTickFunction*
 	SCOPE_CYCLE_COUNTER(STAT_RefreshBoneTransforms);
 
 	// Can't do anything without a SkeletalMesh
-	if( !SkeletalMesh )
+	if( !GetSkeletalMesh())
 	{
 		return;
 	}
@@ -78,15 +78,15 @@ void UPoseableMeshComponent::FillComponentSpaceTransforms()
 {
 	ANIM_MT_SCOPE_CYCLE_COUNTER(FillComponentSpaceTransforms, IsRunningParallelEvaluation());
 
-	if( !SkeletalMesh )
+	if( !GetSkeletalMesh())
 	{
 		return;
 	}
 
 	// right now all this does is to convert to SpaceBases
-	check( SkeletalMesh->GetRefSkeleton().GetNum() == BoneSpaceTransforms.Num() );
-	check( SkeletalMesh->GetRefSkeleton().GetNum() == GetNumComponentSpaceTransforms());
-	check( SkeletalMesh->GetRefSkeleton().GetNum() == GetBoneVisibilityStates().Num() );
+	check( GetSkeletalMesh()->GetRefSkeleton().GetNum() == BoneSpaceTransforms.Num() );
+	check( GetSkeletalMesh()->GetRefSkeleton().GetNum() == GetNumComponentSpaceTransforms());
+	check( GetSkeletalMesh()->GetRefSkeleton().GetNum() == GetBoneVisibilityStates().Num() );
 
 	const int32 NumBones = BoneSpaceTransforms.Num();
 
@@ -113,7 +113,7 @@ void UPoseableMeshComponent::FillComponentSpaceTransforms()
 		BoneProcessed[BoneIndex] = 1;
 #endif
 		// For all bones below the root, final component-space transform is relative transform * component-space transform of parent.
-		const int32 ParentIndex = SkeletalMesh->GetRefSkeleton().GetParentIndex(BoneIndex);
+		const int32 ParentIndex = GetSkeletalMesh()->GetRefSkeleton().GetParentIndex(BoneIndex);
 		FPlatformMisc::Prefetch(SpaceBasesData + ParentIndex);
 
 #if DO_GUARD_SLOW
@@ -130,7 +130,7 @@ void UPoseableMeshComponent::FillComponentSpaceTransforms()
 
 void UPoseableMeshComponent::SetBoneTransformByName(FName BoneName, const FTransform& InTransform, EBoneSpaces::Type BoneSpace)
 {
-	if( !SkeletalMesh || !RequiredBones.IsValid() )
+	if( !GetSkeletalMesh() || !RequiredBones.IsValid() )
 	{
 		return;
 	}
@@ -216,7 +216,7 @@ FTransform GetBoneTransformByNameHelper(FName BoneName, EBoneSpaces::Type BoneSp
 
 FTransform UPoseableMeshComponent::GetBoneTransformByName(FName BoneName, EBoneSpaces::Type BoneSpace)
 {
-	if( !SkeletalMesh || !RequiredBones.IsValid() )
+	if( !GetSkeletalMesh() || !RequiredBones.IsValid() )
 	{
 		return FTransform();
 	}
@@ -261,7 +261,7 @@ FVector UPoseableMeshComponent::GetBoneScaleByName(FName BoneName, EBoneSpaces::
 
 void UPoseableMeshComponent::ResetBoneTransformByName(FName BoneName)
 {
-	if( !SkeletalMesh )
+	if( !GetSkeletalMesh())
 	{
 		return;
 	}
@@ -269,7 +269,7 @@ void UPoseableMeshComponent::ResetBoneTransformByName(FName BoneName)
 	const int32 BoneIndex = GetBoneIndex(BoneName);
 	if( BoneIndex != INDEX_NONE )
 	{
-		BoneSpaceTransforms[BoneIndex] = SkeletalMesh->GetRefSkeleton().GetRefBonePose()[BoneIndex];
+		BoneSpaceTransforms[BoneIndex] = GetSkeletalMesh()->GetRefSkeleton().GetRefBonePose()[BoneIndex];
 	}
 	else
 	{
@@ -283,7 +283,7 @@ void UPoseableMeshComponent::CopyPoseFromSkeletalComponent(USkeletalMeshComponen
 	if(RequiredBones.IsValid())
 	{
 		TArray<FTransform> LocalTransforms = InComponentToCopy->GetBoneSpaceTransforms();
-		if(this->SkeletalMesh == InComponentToCopy->SkeletalMesh 
+		if(this->GetSkeletalMesh() == InComponentToCopy->GetSkeletalMesh()
 			&& LocalTransforms.Num() == BoneSpaceTransforms.Num())
 		{
 			
@@ -294,10 +294,10 @@ void UPoseableMeshComponent::CopyPoseFromSkeletalComponent(USkeletalMeshComponen
 			// The meshes don't match, search bone-by-bone (slow path)
 
 			// first set the localatoms to ref pose from our current mesh
-			BoneSpaceTransforms = SkeletalMesh->GetRefSkeleton().GetRefBonePose();
+			BoneSpaceTransforms = GetSkeletalMesh()->GetRefSkeleton().GetRefBonePose();
 
 			// Now overwrite any matching bones
-			const int32 NumSourceBones = InComponentToCopy->SkeletalMesh->GetRefSkeleton().GetNum();
+			const int32 NumSourceBones = InComponentToCopy->GetSkeletalMesh()->GetRefSkeleton().GetNum();
 
 			for(int32 SourceBoneIndex = 0 ; SourceBoneIndex < NumSourceBones ; ++SourceBoneIndex)
 			{
