Cooker: Add feedback for why a requested package was requested when reporting that the package can not be found.
#rnx

diff --git a/Engine/Source/Editor/UnrealEd/Private/CookOnTheFlyServer.cpp b/Engine/Source/Editor/UnrealEd/Private/CookOnTheFlyServer.cpp
--- a/Engine/Source/Editor/UnrealEd/Private/CookOnTheFlyServer.cpp
+++ b/Engine/Source/Editor/UnrealEd/Private/CookOnTheFlyServer.cpp
@@ -8945,14 +8945,15 @@ void UCookOnTheFlyServer::GenerateInitialRequests(FBeginCookContext& BeginContex
 
 		const FName PackageFileFName = PackageDatas->GetFileNameByPackageName(PackageName);
 
+		UE::Cook::FInstigator& Instigator = FilesInPathInstigators.FindChecked(PackageName);
 		if (!PackageFileFName.IsNone())
 		{
-			UE::Cook::FInstigator& Instigator = FilesInPathInstigators.FindChecked(PackageName);
 			ExternalRequests->EnqueueUnique(UE::Cook::FFilePlatformRequest(PackageFileFName, MoveTemp(Instigator), TargetPlatforms));
 		}
 		else if (!FLinkerLoad::IsKnownMissingPackage(PackageName))
 		{
-			LogCookerMessage(FString::Printf(TEXT("Unable to find package for cooking %s"), *PackageName.ToString()),
+			LogCookerMessage(FString::Printf(TEXT("Unable to find package for cooking %s. Instigator: { %s }."),
+				*PackageName.ToString(), *Instigator.ToString()),
 				EMessageSeverity::Warning);
 		}
 	}
