Editor Context: replace revert icon in viewport editor context widget when an active folder is set with close.

#rnx

diff --git a/Engine/Source/Editor/UnrealEd/Private/SActorEditorContext.cpp b/Engine/Source/Editor/UnrealEd/Private/SActorEditorContext.cpp
--- a/Engine/Source/Editor/UnrealEd/Private/SActorEditorContext.cpp
+++ b/Engine/Source/Editor/UnrealEd/Private/SActorEditorContext.cpp
@@ -170,7 +170,7 @@ void SActorEditorContext::Rebuild()
 							.Content()
 							[
 								SNew(SImage)
-								.Image(FAppStyle::GetBrush("PropertyWindow.DiffersFromDefault"))
+								.Image(FAppStyle::GetBrush("ExpandableButton.CloseButton"))
 								.ColorAndOpacity(FSlateColor::UseForeground())
 							]
 						]
