Change icon to Icons.X from previous CL 20689520.

diff --git a/Engine/Source/Editor/UnrealEd/Private/SActorEditorContext.cpp b/Engine/Source/Editor/UnrealEd/Private/SActorEditorContext.cpp
--- a/Engine/Source/Editor/UnrealEd/Private/SActorEditorContext.cpp
+++ b/Engine/Source/Editor/UnrealEd/Private/SActorEditorContext.cpp
@@ -170,7 +170,7 @@ void SActorEditorContext::Rebuild()
 							.Content()
 							[
 								SNew(SImage)
-								.Image(FAppStyle::GetBrush("ExpandableButton.CloseButton"))
+								.Image(FAppStyle::GetBrush("Icons.X"))
 								.ColorAndOpacity(FSlateColor::UseForeground())
 							]
 						]
