Allow placing non-actor class assets that implement an actor factory that supports them

#rnx

diff --git a/Engine/Source/Editor/UnrealEd/Private/LevelEditorViewport.cpp b/Engine/Source/Editor/UnrealEd/Private/LevelEditorViewport.cpp
--- a/Engine/Source/Editor/UnrealEd/Private/LevelEditorViewport.cpp
+++ b/Engine/Source/Editor/UnrealEd/Private/LevelEditorViewport.cpp
@@ -1287,9 +1287,12 @@ FDropQuery FLevelEditorViewportClient::CanDropObjectsAtCoordinates(int32 MouseX,
 	UObject* AssetObj = AssetData.GetAsset();
 	UClass* ClassObj = Cast<UClass>(AssetObj);
 
+	// Check if the asset has an actor factory
+	bool bHasActorFactory = FActorFactoryAssetProxy::GetFactoryForAsset(AssetData) != nullptr;
+
 	if (ClassObj)
 	{
-		if (!ObjectTools::IsClassValidForPlacing(ClassObj))
+		if (!bHasActorFactory && !ObjectTools::IsClassValidForPlacing(ClassObj))
 		{
 			Result.bCanDrop = false;
 			Result.HintText = FText::Format(LOCTEXT("DragAndDrop_CannotDropAssetClassFmt", "The class '{0}' cannot be placed in a level"), FText::FromString(ClassObj->GetName()));
@@ -1301,9 +1304,6 @@ FDropQuery FLevelEditorViewportClient::CanDropObjectsAtCoordinates(int32 MouseX,
 
 	if (ensureMsgf(AssetObj != NULL, TEXT("AssetObj was null (%s)"), *AssetData.GetFullName()))
 	{
-		// Check if the asset has an actor factory
-		bool bHasActorFactory = FActorFactoryAssetProxy::GetFactoryForAsset(AssetData) != nullptr;
-
 		if (AssetObj->IsA(AActor::StaticClass()) || bHasActorFactory)
 		{
 			Result.bCanDrop = true;
