Fix thumbnail generation for Texture2DArray.

diff --git a/Engine/Source/Editor/UnrealEd/Private/ThumbnailRendering/TextureThumbnailRenderer.cpp b/Engine/Source/Editor/UnrealEd/Private/ThumbnailRendering/TextureThumbnailRenderer.cpp
--- a/Engine/Source/Editor/UnrealEd/Private/ThumbnailRendering/TextureThumbnailRenderer.cpp
+++ b/Engine/Source/Editor/UnrealEd/Private/ThumbnailRendering/TextureThumbnailRenderer.cpp
@@ -85,7 +85,7 @@ void UTextureThumbnailRenderer::Draw(UObject* Object, int32 X, int32 Y, uint32 W
 		{
 			bool bIsNormalMap = Texture2DArray->IsNormalMap();
 			bool bIsSingleChannel = true;
-			BatchedElementParameters = new FBatchedElementTexture2DPreviewParameters(0, 0, 0, bIsNormalMap, bIsSingleChannel, false, false, true);
+			BatchedElementParameters = new FBatchedElementTexture2DPreviewParameters((float)0, (float)0, (float)-1, bIsNormalMap, bIsSingleChannel, false, false, true);
 		}
 		else if (TextureLightProfile)
 		{
