Add -pie command line option to start PIE as soon as the editor finishes loading to help with benchmark and testing

#rnx

diff --git a/Engine/Source/Editor/UnrealEd/Private/UnrealEdGlobals.cpp b/Engine/Source/Editor/UnrealEd/Private/UnrealEdGlobals.cpp
--- a/Engine/Source/Editor/UnrealEd/Private/UnrealEdGlobals.cpp
+++ b/Engine/Source/Editor/UnrealEd/Private/UnrealEdGlobals.cpp
@@ -156,6 +156,7 @@ int32 EditorInit( IEngineLoop& EngineLoop )
 
 	// Are we in immersive mode?
 	const bool bIsImmersive = FPaths::IsProjectFilePathSet() && FParse::Param( FCommandLine::Get(), TEXT( "immersive" ) );
+	const bool bIsPlayInEditorRequested = FPaths::IsProjectFilePathSet() && FParse::Param(FCommandLine::Get(), TEXT("pie"));
 
 	// Do final set up on the editor frame and show it
 	{
@@ -164,7 +165,7 @@ int32 EditorInit( IEngineLoop& EngineLoop )
 		// Startup Slate main frame and other editor windows
 		{
 			const bool bStartImmersive = bIsImmersive;
-			const bool bStartPIE = bIsImmersive;
+			const bool bStartPIE = bIsImmersive || bIsPlayInEditorRequested;
 
 			IMainFrameModule& MainFrameModule = FModuleManager::LoadModuleChecked<IMainFrameModule>(TEXT("MainFrame"));
 			MainFrameModule.CreateDefaultMainFrame( bStartImmersive, bStartPIE );
