Fixed World Partition builder to work with LWC

diff --git a/Engine/Source/Editor/UnrealEd/Private/WorldPartition/WorldPartitionBuilder.cpp b/Engine/Source/Editor/UnrealEd/Private/WorldPartition/WorldPartitionBuilder.cpp
--- a/Engine/Source/Editor/UnrealEd/Private/WorldPartition/WorldPartitionBuilder.cpp
+++ b/Engine/Source/Editor/UnrealEd/Private/WorldPartition/WorldPartitionBuilder.cpp
@@ -142,19 +142,19 @@ bool UWorldPartitionBuilder::RunBuilder(UWorld* World)
 	return bResult;
 }
 
-static FIntVector GetCellCoord(const FVector& InPos, const int32 InCellSize)
+static FWorldBuilderCellCoord GetCellCoord(const FVector& InPos, const int32 InCellSize)
 {
-	return FIntVector(
+	return FWorldBuilderCellCoord(
 		FMath::FloorToInt(InPos.X / InCellSize),
 		FMath::FloorToInt(InPos.Y / InCellSize),
 		FMath::FloorToInt(InPos.Z / InCellSize)
 	);
 }
 
-static FIntVector GetCellCount(const FBox& InBounds, const int32 InCellSize)
+static FWorldBuilderCellCoord GetCellCount(const FBox& InBounds, const int32 InCellSize)
 {
-	const FIntVector MinCellCoords = GetCellCoord(InBounds.Min, InCellSize);
-	const FIntVector MaxCellCoords = FIntVector(FMath::CeilToInt(InBounds.Max.X / InCellSize),
+	const FWorldBuilderCellCoord MinCellCoords = GetCellCoord(InBounds.Min, InCellSize);
+	const FWorldBuilderCellCoord MaxCellCoords = FWorldBuilderCellCoord(FMath::CeilToInt(InBounds.Max.X / InCellSize),
 												FMath::CeilToInt(InBounds.Max.Y / InCellSize),
 												FMath::CeilToInt(InBounds.Max.Z / InCellSize));
 	return MaxCellCoords - MinCellCoords;
@@ -210,10 +210,10 @@ bool UWorldPartitionBuilder::Run(UWorld* World, FPackageSourceControlHelper& Pac
 	if ((LoadingMode == ELoadingMode::IterativeCells) || (LoadingMode == ELoadingMode::IterativeCells2D))
 	{
 		// do partial loading loop that calls RunInternal
-		const FIntVector MinCellCoords = GetCellCoord(CellInfo.EditorBounds.Min, IterativeCellSize);
-		const FIntVector NumCellsIterations = GetCellCount(CellInfo.EditorBounds, IterativeCellSize);
-		const FIntVector BeginCellCoords = MinCellCoords;
-		const FIntVector EndCellCoords = BeginCellCoords + NumCellsIterations;
+		const FWorldBuilderCellCoord MinCellCoords = GetCellCoord(CellInfo.EditorBounds.Min, IterativeCellSize);
+		const FWorldBuilderCellCoord NumCellsIterations = GetCellCount(CellInfo.EditorBounds, IterativeCellSize);
+		const FWorldBuilderCellCoord BeginCellCoords = MinCellCoords;
+		const FWorldBuilderCellCoord EndCellCoords = BeginCellCoords + NumCellsIterations;
 
 		auto CanIterateZ = [&BeginCellCoords, &EndCellCoords, LoadingMode](const bool bInResult, const int32 InZ) -> bool
 		{
@@ -236,11 +236,11 @@ bool UWorldPartitionBuilder::Run(UWorld* World, FPackageSourceControlHelper& Pac
 		
 		TArray<TUniquePtr<IWorldPartitionActorLoaderInterface::ILoaderAdapter>> LoaderAdapters;
 
-		for (int32 z = BeginCellCoords.Z; CanIterateZ(bResult, z); z++)
+		for (int64 z = BeginCellCoords.Z; CanIterateZ(bResult, z); z++)
 		{
-			for (int32 y = BeginCellCoords.Y; bResult && (y < EndCellCoords.Y); y++)
+			for (int64 y = BeginCellCoords.Y; bResult && (y < EndCellCoords.Y); y++)
 			{
-				for (int32 x = BeginCellCoords.X; bResult && (x < EndCellCoords.X); x++)
+				for (int64 x = BeginCellCoords.X; bResult && (x < EndCellCoords.X); x++)
 				{
 					IterationIndex++;
 					UE_LOG(LogWorldPartitionBuilder, Display, TEXT("[%d / %d] Processing cells..."), IterationIndex, IterationCount);
@@ -257,7 +257,7 @@ bool UWorldPartitionBuilder::Run(UWorld* World, FPackageSourceControlHelper& Pac
 					FBox BoundsToLoad(Min, Max);
 					BoundsToLoad = BoundsToLoad.ExpandBy(IterativeCellOverlapSize);
 
-					CellInfo.Location = FIntVector(x, y, z);
+					CellInfo.Location = FWorldBuilderCellCoord(x, y, z);
 					CellInfo.Bounds = BoundsToLoad;
 
 					UE_LOG(LogWorldPartitionBuilder, Verbose, TEXT("Loading Bounds: Min %s, Max %s"), *BoundsToLoad.Min.ToString(), *BoundsToLoad.Max.ToString());
diff --git a/Engine/Source/Editor/UnrealEd/Public/WorldPartition/WorldPartitionBuilder.h b/Engine/Source/Editor/UnrealEd/Public/WorldPartition/WorldPartitionBuilder.h
--- a/Engine/Source/Editor/UnrealEd/Public/WorldPartition/WorldPartitionBuilder.h
+++ b/Engine/Source/Editor/UnrealEd/Public/WorldPartition/WorldPartitionBuilder.h
@@ -5,33 +5,28 @@
 #include "PackageSourceControlHelper.h"
 #include "WorldPartitionBuilder.generated.h"
 
+typedef UE::Math::TIntVector3<int64> FWorldBuilderCellCoord;
+
 /**
  * Structure containing information about a World Partition Builder cell
  */
-USTRUCT()
 struct FCellInfo
 {
-	GENERATED_BODY()
-
 	FCellInfo();
 
 	/**
 	 * Location of the cell, expressed inside World Partition Builder space
 	 * (floor(Coordinate) / IterativeCellSize)
 	 */
-	UPROPERTY(VisibleAnywhere, Category = WorldPartitionBuilder)
-	FIntVector Location;
+	FWorldBuilderCellCoord Location;
 
 	/** Bounds of the cell */
-	UPROPERTY(VisibleAnywhere, Category = WorldPartitionBuilder)
 	FBox Bounds;
 
 	/** Whole space */
-	UPROPERTY(VisibleAnywhere, Category = WorldPartitionBuilder)
 	FBox EditorBounds;
 
 	/** The size of a cell used by the World Partition Builder */
-	UPROPERTY(VisibleAnywhere, Category = WorldPartitionBuilder)
 	int32 IterativeCellSize;
 };
 
