PR #9305: Using SelectionColor also for text in measure tool (Contributed by [REDACTED_PERSON])

diff --git a/Engine/Source/Editor/UnrealEd/Private/DragTool_Measure.cpp b/Engine/Source/Editor/UnrealEd/Private/DragTool_Measure.cpp
--- a/Engine/Source/Editor/UnrealEd/Private/DragTool_Measure.cpp
+++ b/Engine/Source/Editor/UnrealEd/Private/DragTool_Measure.cpp
@@ -77,15 +77,16 @@ void FDragTool_Measure::Render(const FSceneView* View, FCanvas* Canvas)
 
 	if (View != nullptr && Canvas != nullptr && Length >= 1.f)
 	{
+		const FLinearColor SelectionColor = GetDefault<UEditorStyleSettings>()->SelectionColor;
 		FCanvasLineItem LineItem( PixelStart, PixelEnd );
-		LineItem.SetColor(GetDefault<UEditorStyleSettings>()->SelectionColor);
+		LineItem.SetColor( SelectionColor );
 		Canvas->DrawItem( LineItem );
 
 		const FVector2D PixelMid = FVector2D(PixelStart + ((PixelEnd - PixelStart) / 2));
 
 		// Calculate number of decimal places to display, based on the current viewport zoom
 		float Divisor = 1.0f;
-		int DecimalPlaces = 0;
+		int32 DecimalPlaces = 0;
 		const float OrderOfMagnitude = FMath::LogX(10.0f, OrthoUnitsPerPixel);
 
 		switch (GetDefault<ULevelEditorViewportSettings>()->MeasuringToolUnits)
@@ -111,7 +112,7 @@ void FDragTool_Measure::Render(const FSceneView* View, FCanvas* Canvas)
 		const FText LengthStr = FText::AsNumber( Length / Divisor, &Options );
 
 
-		FCanvasTextItem TextItem( FVector2D( FMath::FloorToFloat(PixelMid.X), FMath::FloorToFloat(PixelMid.Y) ), LengthStr, GEngine->GetSmallFont(), FLinearColor::White );
+		FCanvasTextItem TextItem( FVector2D( FMath::FloorToFloat(PixelMid.X), FMath::FloorToFloat(PixelMid.Y) ), LengthStr, GEngine->GetSmallFont(), SelectionColor );
 		TextItem.bCentreX = true;
 		Canvas->DrawItem( TextItem );
 	}
