Revisit world partition cooking to fix problems caused by objects being moved from their original package prior to BeginCacheForCookedPlatformData being completed.
- Populate generator/generated package loads objects, change actor's level, registers actors and returns the loaded objects
- Finalize generator/generated package changes the package of actors and objects inside OFPA packages

diff --git a/Engine/Source/Editor/UnrealEd/Private/CookOnTheFlyServer.cpp b/Engine/Source/Editor/UnrealEd/Private/CookOnTheFlyServer.cpp
--- a/Engine/Source/Editor/UnrealEd/Private/CookOnTheFlyServer.cpp
+++ b/Engine/Source/Editor/UnrealEd/Private/CookOnTheFlyServer.cpp
@@ -2933,7 +2933,7 @@ static void ValidateAndExtendObjectsToMove(TArray<UObject*>& InObjects, UE::Cook
 		{
 			UE_LOG(LogCook, Warning, TEXT("CookPackageSplitter found non-valid object %s returned from %s on Splitter %s%s. Ignoring it."),
 				OriginalObject ? *OriginalObject->GetFullName() : TEXT("<null>"),
-				Info.IsGenerator() ? TEXT("GetObjectsToMoveIntoGeneratorPackage") : TEXT("GetObjectsToMoveIntoGeneratedPackage"),
+				Info.IsGenerator() ? TEXT("PopulateGeneratorPackage") : TEXT("PopulateGeneratedPackage"),
 				*Generator.GetSplitDataObjectName().ToString(),
 				Info.IsGenerator() ? TEXT("") : *FString::Printf(TEXT(", Package %s"), *Info.PackageData->GetPackageName().ToString()));
 			--OriginalNum;
@@ -2972,17 +2972,28 @@ UE::Cook::EPollStatus UCookOnTheFlyServer::BeginCacheObjectsToMove(UE::Cook::FGe
 
 	if (Info.GetSaveState() <= FCookGenerationInfo::ESaveState::CallObjectsToMove)
 	{
+		bool bPopulateSucceeded = false;
 		TArray<UObject*> ObjectsToMove;
 		if (Info.IsGenerator())
 		{
 			ConstructGeneratedPackagesForPresave(PackageData, Generator, GeneratedPackagesForPresave);
-			Splitter->GetObjectsToMoveIntoGeneratorPackage(Package, SplitDataObject, GeneratedPackagesForPresave, ObjectsToMove);
+			bPopulateSucceeded = Splitter->PopulateGeneratorPackage(Package, SplitDataObject, GeneratedPackagesForPresave, ObjectsToMove);
 		}
 		else
 		{
 			ICookPackageSplitter::FGeneratedPackageForPopulate SplitterInfo{ Info.RelativePath, Package, Info.IsCreateAsMap() };
-			Splitter->GetObjectsToMoveIntoGeneratedPackage(Package, SplitDataObject, SplitterInfo, ObjectsToMove);
+			bPopulateSucceeded = Splitter->PopulateGeneratedPackage(Package, SplitDataObject, SplitterInfo, ObjectsToMove);
 		}
+
+		if (!bPopulateSucceeded)
+		{
+			UE_LOG(LogCook, Error, TEXT("CookPackageSplitter returned false from %s. Splitter=%s%s"),
+				Info.IsGenerator() ? TEXT("PopulateGeneratorPackage") : TEXT("PopulateGeneratedPackage"),
+				*Generator.GetSplitDataObjectName().ToString(), 
+				Info.IsGenerator() ? TEXT("") : *FString::Printf(TEXT("\nGeneratedPackage: %s"), *PackageData.GetPackageName().ToString()));
+			return EPollStatus::Error;
+		}
+
 		ValidateAndExtendObjectsToMove(ObjectsToMove, Generator, Info);
 		Info.TakeOverCachedObjectsAndAddNew(PackageData.GetCachedObjectsInOuter(), ObjectsToMove);
 		Info.SetSaveStateComplete(FCookGenerationInfo::ESaveState::CallObjectsToMove);
@@ -3015,7 +3026,12 @@ UE::Cook::EPollStatus UCookOnTheFlyServer::PreSaveGeneratorPackage(UE::Cook::FPa
 	}
 
 	ConstructGeneratedPackagesForPresave(PackageData, Generator, GeneratedPackagesForPresave);
-	Splitter->PreSaveGeneratorPackage(Package, SplitDataObject, GeneratedPackagesForPresave);
+	if (!Splitter->PreSaveGeneratorPackage(Package, SplitDataObject, GeneratedPackagesForPresave))
+	{
+		UE_LOG(LogCook, Error, TEXT("PackageSplitter returned false from PreSaveGeneratorPackage. Splitter=%s"),
+			*Generator.GetSplitDataObjectName().ToString());
+		return EPollStatus::Error;
+	}
 	return EPollStatus::Success;
 }
 
@@ -3098,8 +3114,8 @@ UE::Cook::EPollStatus UCookOnTheFlyServer::BeginCachePostMove(UE::Cook::FGenerat
 			TEXT("Change the splitter's %s to construct new objects and declare existing objects that will be moved from other packages.\n")
 			TEXT("SplitterObject: %s%s\n")
 			TEXT("NumPendingObjects: %d, FirstPendingObject: %s"),
-			Info.IsGenerator() ? TEXT("PreSaveGeneratorPackage") : TEXT("TryPopulatePackage"),
-			Info.IsGenerator() ? TEXT("GetObjectsToMoveIntoGeneratorPackage") : TEXT("GetObjectsToMoveIntoGeneratedPackage"),
+			Info.IsGenerator() ? TEXT("PreSaveGeneratorPackage") : TEXT("PreSaveGeneratedPackage"),
+			Info.IsGenerator() ? TEXT("PopulateGeneratorPackage") : TEXT("PopulateGeneratedPackage"),
 			*SplitDataObject->GetFullName(),
 			Info.IsGenerator() ? TEXT("") : *FString::Printf(TEXT("\nGeneratedPackage: %s"), *PackageData.GetPackageName().ToString()),
 			PackageData.GetNumPendingCookedPlatformData(),
@@ -3180,16 +3196,16 @@ UE::Cook::EPollStatus UCookOnTheFlyServer::TryPopulateGeneratedPackage(UE::Cook:
 	PopulateData.RelativePath = GeneratedInfo.RelativePath;
 	PopulateData.Package = GeneratedPackage;
 	PopulateData.bCreatedAsMap = GeneratedInfo.IsCreateAsMap();
-	if (!Splitter->TryPopulatePackage(OwnerPackage, OwnerObject, PopulateData))
+	if (!Splitter->PreSaveGeneratedPackage(OwnerPackage, OwnerObject, PopulateData))
 	{
-		UE_LOG(LogCook, Error, TEXT("PackageSplitter returned false from TryPopulatePackage. Splitter=%s, Generated=%s."),
+		UE_LOG(LogCook, Error, TEXT("PackageSplitter returned false from PreSaveGeneratedPackage. Splitter=%s, Generated=%s."),
 			*Generator.GetSplitDataObjectName().ToString(), *GeneratedPackageName);
 		return EPollStatus::Error;
 	}
 	bool bPackageIsMap = GeneratedPackage->ContainsMap();
 	if (bPackageIsMap != GeneratedInfo.IsCreateAsMap())
 	{
-		UE_LOG(LogCook, Error, TEXT("PackageSplitter specified generated package is %s in GetGenerateList results, but then in TryPopulatePackage created it as %s. Splitter=%s, Generated=%s."),
+		UE_LOG(LogCook, Error, TEXT("PackageSplitter specified generated package is %s in GetGenerateList results, but then in PreSaveGeneratedPackage created it as %s. Splitter=%s, Generated=%s."),
 			(GeneratedInfo.IsCreateAsMap() ? TEXT("map") : TEXT("uasset")), (bPackageIsMap ? TEXT("map") : TEXT("uasset")),
 			*Generator.GetSplitDataObjectName().ToString(), *GeneratedPackageName);
 		return EPollStatus::Error;
diff --git a/Engine/Source/Editor/UnrealEd/Private/Cooker/CookPackageData.cpp b/Engine/Source/Editor/UnrealEd/Private/Cooker/CookPackageData.cpp
--- a/Engine/Source/Editor/UnrealEd/Private/Cooker/CookPackageData.cpp
+++ b/Engine/Source/Editor/UnrealEd/Private/Cooker/CookPackageData.cpp
@@ -1507,7 +1507,7 @@ void FGeneratorPackage::UpdateSaveAfterGarbageCollect(const FPackageData& Packag
 				// Log a cook error but remove the invalidated object.
 				UE_LOG(LogCook, Error, TEXT("PackageSplitter found an object returned from %s that was removed from memory during garbage collection. This will cause errors during save of the package.")
 					TEXT("\n\tSplitter=%s%s."),
-					Info->IsGenerator() ? TEXT("GetObjectsToMoveIntoGeneratorPackage") : TEXT("GetObjectsToMoveIntoGeneratedPackage"),
+					Info->IsGenerator() ? TEXT("PopulateGeneratorPackage") : TEXT("PopulateGeneratedPackage"),
 					*GetSplitDataObjectName().ToString(),
 					Info->IsGenerator() ? TEXT("") : *FString::Printf(TEXT(", Generated=%s."), *Info->PackageData->GetPackageName().ToString()));
 				Iter.RemoveCurrent();
diff --git a/Engine/Source/Editor/UnrealEd/Public/CookPackageSplitter.h b/Engine/Source/Editor/UnrealEd/Public/CookPackageSplitter.h
--- a/Engine/Source/Editor/UnrealEd/Public/CookPackageSplitter.h
+++ b/Engine/Source/Editor/UnrealEd/Public/CookPackageSplitter.h
@@ -68,27 +68,35 @@ public:
 		bool bCreatedAsMap = false;
 	};
 
-	virtual void GetObjectsToMoveIntoGeneratedPackage(UPackage* OwnerPackage, UObject* OwnerObject,
-		const FGeneratedPackageForPopulate& GeneratedPackage, TArray<UObject*>& OutObjectsToMove) {}
 	/**
 	 * Try to populate a generated package.
 	 *
 	 * Receive an empty UPackage generated from an element in GetGenerateList and populate it
+	 * Return a list of all the objects that will be moved into the Generated package during its save, so the cooker
+	 * can call BeginCacheForCookedPlatformData on them before the move
 	 * After returning, the given package will be queued for saving into the TargetDomain
-	 * Note that TryPopulatePackage will not be called for every package on
-	 * an iterative cook; it will only be called for the packages with changed dependencies. 
+	 * Note that PreSaveGeneratedPackage will not be called for every package on
+	 * an iterative cook; it will only be called for the packages with changed dependencies.
 	 *
 	 * @param OwnerPackage				The parent package being split
 	 * @param OwnerObject				The SplitDataClass instance that this CookPackageSplitter instance was created for
 	 * @param GeneratedPackage			Pointer and information about the package to populate
+	 * @param OutObjectsToMove			List of all the objects that will be moved into the Generated package during its save
 	 * @return							True if successfully populates,  false on error (this will cause a cook error).
 	 */
-	virtual bool TryPopulatePackage(UPackage* OwnerPackage, UObject* OwnerObject,
+	virtual bool PopulateGeneratedPackage(UPackage* OwnerPackage, UObject* OwnerObject,
+		const FGeneratedPackageForPopulate& GeneratedPackage, TArray<UObject*>& OutObjectsToMove) { return false; }
+	
+	/**
+	 * Called before saving the parent generator package, which itself occurs before PreSaveGeneratedPackage is called on the generated packages.
+	 * Make any required adjustments to the parent package before it is saved into the target domain.
+	 */
+	virtual bool PreSaveGeneratedPackage(UPackage* OwnerPackage, UObject* OwnerObject,
 		const FGeneratedPackageForPopulate& GeneratedPackage) = 0;
 	/**
 	 * Called after saving the generated package. Undo any required adjustments to the parent package that
-	 * were made in TryPopulatePackage, so that the package is once again ready for use in the editor or in
-	 * future GetGenerateList or TryPopulatePackage calls
+	 * were made in PreSaveGeneratedPackage, so that the package is once again ready for use in the editor or in
+	 * future GetGenerateList or PreSaveGeneratedPackage calls
 	 */
 	virtual void PostSaveGeneratedPackage(UPackage* OwnerPackage, UObject* OwnerObject,
 		const FGeneratedPackageForPopulate& GeneratedPackage) {}
@@ -110,19 +118,19 @@ public:
 	 * Return a list of all the objects that will be moved into the Generator package during its save, so the cooker
 	 * can call BeginCacheForCookedPlatformData on them before the move
 	 */
-	virtual void GetObjectsToMoveIntoGeneratorPackage(UPackage* OwnerPackage, UObject* OwnerObject,
-		const TArray<ICookPackageSplitter::FGeneratedPackageForPreSave>& GeneratedPackages, TArray<UObject*>& OutObjectsToMove) {}
+	virtual bool PopulateGeneratorPackage(UPackage* OwnerPackage, UObject* OwnerObject,
+		const TArray<ICookPackageSplitter::FGeneratedPackageForPreSave>& GeneratedPackages, TArray<UObject*>& OutObjectsToMove) { return false; }
 
 	/**
-	 * Called before saving the parent generator package, which itself occurs before TryPopulatePackage is called on the generated packages.
+	 * Called before saving the parent generator package, which itself occurs before PreSaveGeneratedPackage is called on the generated packages.
 	 * Make any required adjustments to the parent package before it is saved into the target domain.
 	 */
-	virtual void PreSaveGeneratorPackage(UPackage* OwnerPackage, UObject* OwnerObject, const TArray<FGeneratedPackageForPreSave>& PlaceholderPackages) {}
+	virtual bool PreSaveGeneratorPackage(UPackage* OwnerPackage, UObject* OwnerObject, const TArray<FGeneratedPackageForPreSave>& PlaceholderPackages) { return false; }
 
 	/**
 	 * Called after saving the parent generator package. Undo any required adjustments to the parent package that
 	 * were made in PreSaveGeneratorPackage, so that the package is once again ready for use in the editor or in
-	 * future GetGenerateList or TryPopulatePackage calls
+	 * future GetGenerateList or PreSaveGeneratedPackage calls
 	 */
 	virtual void PostSaveGeneratorPackage(UPackage* OwnerPackage, UObject* OwnerObject) {}
 
diff --git a/Engine/Source/Runtime/Engine/Private/WorldPartition/WorldPartition.cpp b/Engine/Source/Runtime/Engine/Private/WorldPartition/WorldPartition.cpp
--- a/Engine/Source/Runtime/Engine/Private/WorldPartition/WorldPartition.cpp
+++ b/Engine/Source/Runtime/Engine/Private/WorldPartition/WorldPartition.cpp
@@ -1256,22 +1256,22 @@ void UWorldPartition::DrawRuntimeHashPreview()
 	RuntimeHash->DrawPreview();
 }
 
-bool UWorldPartition::LoadGeneratorPackageObjectsForCook(TArray<UObject*>& OutLoadedObjects)
+bool UWorldPartition::PopulateGeneratorPackageForCook(TArray<UObject*>& OutLoadedObjects)
 {
 	check(RuntimeHash);
-	return RuntimeHash->LoadGeneratorPackageObjectsForCook(OutLoadedObjects);
+	return RuntimeHash->PopulateGeneratorPackageForCook(OutLoadedObjects);
 }
 
-bool UWorldPartition::LoadGeneratedPackageObjectsForCook(const FString& InPackageRelativePath, TArray<UObject*>& OutLoadedObjects)
+bool UWorldPartition::PopulateGeneratedPackageForCook(UPackage* InPackage, const FString& InPackageRelativePath, TArray<UObject*>& OutLoadedObjects)
 {
 	check(RuntimeHash);
-	return RuntimeHash->LoadGeneratedPackageObjectsForCook(InPackageRelativePath, OutLoadedObjects);
+	return RuntimeHash->PopulateGeneratedPackageForCook(InPackage, InPackageRelativePath, OutLoadedObjects);
 }
 
-bool UWorldPartition::PopulateGeneratedPackageForCook(UPackage* InPackage, const FString& InPackageRelativePath)
+bool UWorldPartition::FinalizeGeneratedPackageForCook(const FString& InPackageRelativePath)
 {
 	check(RuntimeHash);
-	return RuntimeHash->PopulateGeneratedPackageForCook(InPackage, InPackageRelativePath);
+	return RuntimeHash->FinalizeGeneratedPackageForCook(InPackageRelativePath);
 }
 
 bool UWorldPartition::FinalizeGeneratorPackageForCook(const TArray<ICookPackageSplitter::FGeneratedPackageForPreSave>& InGeneratedPackages)
diff --git a/Engine/Source/Runtime/Engine/Private/WorldPartition/WorldPartitionCookPackageSplitter.cpp b/Engine/Source/Runtime/Engine/Private/WorldPartition/WorldPartitionCookPackageSplitter.cpp
--- a/Engine/Source/Runtime/Engine/Private/WorldPartition/WorldPartitionCookPackageSplitter.cpp
+++ b/Engine/Source/Runtime/Engine/Private/WorldPartition/WorldPartitionCookPackageSplitter.cpp
@@ -87,7 +87,7 @@ TArray<ICookPackageSplitter::FGeneratedPackage> FWorldPartitionCookPackageSplitt
 	UWorld* PartitionedWorld = const_cast<UWorld*>(ConstPartitionedWorld);
 
 	// Store the World pointer to declare it to GarbageCollection; we do not want to allow the World to be Garbage Collected
-	// until we have finished all of our TryPopulatePackage calls, because we store information on the World 
+	// until we have finished all of our PreSaveGeneratedPackage calls, because we store information on the World 
 	// that is necessary for populate 
 	ReferencedWorld = PartitionedWorld;
 
@@ -117,36 +117,36 @@ TArray<ICookPackageSplitter::FGeneratedPackage> FWorldPartitionCookPackageSplitt
 	return PackagesToGenerate;
 }
 
-void FWorldPartitionCookPackageSplitter::GetObjectsToMoveIntoGeneratedPackage(UPackage* OwnerPackage, UObject* OwnerObject,
+bool FWorldPartitionCookPackageSplitter::PopulateGeneratedPackage(UPackage* OwnerPackage, UObject* OwnerObject,
 	const FGeneratedPackageForPopulate& GeneratedPackage, TArray<UObject*>& OutObjectsToMove)
 {
 	UWorld* PartitionedWorld = ValidateDataObject(OwnerObject);
 	UWorldPartition* WorldPartition = PartitionedWorld->PersistentLevel->GetWorldPartition();
-	WorldPartition->LoadGeneratedPackageObjectsForCook(GeneratedPackage.RelativePath, OutObjectsToMove);
+	return WorldPartition->PopulateGeneratedPackageForCook(GeneratedPackage.Package, GeneratedPackage.RelativePath, OutObjectsToMove);
 }
 
-bool FWorldPartitionCookPackageSplitter::TryPopulatePackage(UPackage* OwnerPackage, UObject* OwnerObject,
+bool FWorldPartitionCookPackageSplitter::PreSaveGeneratedPackage(UPackage* OwnerPackage, UObject* OwnerObject,
 	const FGeneratedPackageForPopulate& GeneratedPackage)
 {
 	UWorld* PartitionedWorld = ValidateDataObject(OwnerObject);
 	UWorldPartition* WorldPartition = PartitionedWorld->PersistentLevel->GetWorldPartition();
-	return WorldPartition->PopulateGeneratedPackageForCook(GeneratedPackage.Package, GeneratedPackage.RelativePath);
+	return WorldPartition->FinalizeGeneratedPackageForCook(GeneratedPackage.RelativePath);
 }
 
-void FWorldPartitionCookPackageSplitter::GetObjectsToMoveIntoGeneratorPackage(UPackage* OwnerPackage, UObject* OwnerObject,
+bool FWorldPartitionCookPackageSplitter::PopulateGeneratorPackage(UPackage* OwnerPackage, UObject* OwnerObject,
 	const TArray<ICookPackageSplitter::FGeneratedPackageForPreSave>& GeneratedPackages, TArray<UObject*>& OutObjectsToMove)
 {
 	UWorld* PartitionedWorld = ValidateDataObject(OwnerObject);
 	UWorldPartition* WorldPartition = PartitionedWorld->PersistentLevel->GetWorldPartition();
-	WorldPartition->LoadGeneratorPackageObjectsForCook(OutObjectsToMove);
+	return WorldPartition->PopulateGeneratorPackageForCook(OutObjectsToMove);
 }
 
-void FWorldPartitionCookPackageSplitter::PreSaveGeneratorPackage(UPackage* OwnerPackage, UObject* OwnerObject,
+bool FWorldPartitionCookPackageSplitter::PreSaveGeneratorPackage(UPackage* OwnerPackage, UObject* OwnerObject,
 	const TArray<FGeneratedPackageForPreSave>& GeneratedPackages)
 {
 	UWorld* PartitionedWorld = ValidateDataObject(OwnerObject);
 	UWorldPartition* WorldPartition = PartitionedWorld->PersistentLevel->GetWorldPartition();
-	WorldPartition->FinalizeGeneratorPackageForCook(GeneratedPackages);
+	return WorldPartition->FinalizeGeneratorPackageForCook(GeneratedPackages);
 }
 
 void FWorldPartitionCookPackageSplitter::OnOwnerReloaded(UPackage* OwnerPackage, UObject* OwnerObject)
diff --git a/Engine/Source/Runtime/Engine/Private/WorldPartition/WorldPartitionCookPackageSplitter.h b/Engine/Source/Runtime/Engine/Private/WorldPartition/WorldPartitionCookPackageSplitter.h
--- a/Engine/Source/Runtime/Engine/Private/WorldPartition/WorldPartitionCookPackageSplitter.h
+++ b/Engine/Source/Runtime/Engine/Private/WorldPartition/WorldPartitionCookPackageSplitter.h
@@ -20,13 +20,13 @@ public:
 	virtual void Teardown(ETeardown Status) override;
 	virtual bool UseInternalReferenceToAvoidGarbageCollect() override { return true; }
 	virtual TArray<ICookPackageSplitter::FGeneratedPackage> GetGenerateList(const UPackage* OwnerPackage, const UObject* OwnerObject) override;
-	virtual void GetObjectsToMoveIntoGeneratedPackage(UPackage* OwnerPackage, UObject* OwnerObject,
+	virtual bool PopulateGeneratedPackage(UPackage* OwnerPackage, UObject* OwnerObject,
 		const FGeneratedPackageForPopulate& GeneratedPackage, TArray<UObject*>& OutObjectsToMove) override;
-	virtual bool TryPopulatePackage(UPackage* OwnerPackage, UObject* OwnerObject,
+	virtual bool PreSaveGeneratedPackage(UPackage* OwnerPackage, UObject* OwnerObject,
 		const ICookPackageSplitter::FGeneratedPackageForPopulate& GeneratedPackage) override;
-	virtual void GetObjectsToMoveIntoGeneratorPackage(UPackage* OwnerPackage, UObject* OwnerObject,
+	virtual bool PopulateGeneratorPackage(UPackage* OwnerPackage, UObject* OwnerObject,
 		const TArray<ICookPackageSplitter::FGeneratedPackageForPreSave>& GeneratedPackages, TArray<UObject*>& OutObjectsToMove) override;
-	virtual void PreSaveGeneratorPackage(UPackage* OwnerPackage, UObject* OwnerObject,
+	virtual bool PreSaveGeneratorPackage(UPackage* OwnerPackage, UObject* OwnerObject,
 		const TArray<ICookPackageSplitter::FGeneratedPackageForPreSave>& GeneratedPackages) override;
 	virtual void OnOwnerReloaded(UPackage* OwnerPackage, UObject* OwnerObject) override;
 	//~ End of ICookPackageSplitter
diff --git a/Engine/Source/Runtime/Engine/Private/WorldPartition/WorldPartitionLevelHelper.cpp b/Engine/Source/Runtime/Engine/Private/WorldPartition/WorldPartitionLevelHelper.cpp
--- a/Engine/Source/Runtime/Engine/Private/WorldPartition/WorldPartitionLevelHelper.cpp
+++ b/Engine/Source/Runtime/Engine/Private/WorldPartition/WorldPartitionLevelHelper.cpp
@@ -17,7 +17,6 @@
 #include "WorldPartition/WorldPartitionActorDesc.h"
 #include "WorldPartition/WorldPartition.h"
 #include "WorldPartition/WorldPartitionPackageHelper.h"
-#include "StaticMeshCompiler.h"
 #include "LevelUtils.h"
 #include "Templates/SharedPointer.h"
 #include "ActorFolder.h"
@@ -104,11 +103,6 @@ void FWorldPartitionLevelHelper::MoveExternalActorsToLevel(const TArray<FWorldPa
 {
 	TRACE_CPUPROFILER_EVENT_SCOPE(FWorldPartitionLevelHelper::MoveExternalActorsToLevel);
 
-	// We can't have async compilation still going on while we move actors as this is going to ResetLoaders which will move bulkdata around that
-	// might still be used by async compilation. 
-	// #TODO_DC Revisit once virtualbulkdata are enabled
-	FStaticMeshCompilingManager::Get().FinishAllCompilation();
-
 	check(InLevel);
 	UPackage* LevelPackage = InLevel->GetPackage();
 	
@@ -128,7 +122,6 @@ void FWorldPartitionLevelHelper::MoveExternalActorsToLevel(const TArray<FWorldPa
 			check(ActorExternalPackage);
 
 			const bool bSameOuter = (InLevel == Actor->GetOuter());
-			Actor->SetPackageExternal(false, false);
 						
 			// Avoid calling Rename on the actor if it's already outered to InLevel as this will cause it's name to be changed. 
 			// (UObject::Rename doesn't check if Rename is being called with existing outer and assigns new name)
@@ -137,23 +130,10 @@ void FWorldPartitionLevelHelper::MoveExternalActorsToLevel(const TArray<FWorldPa
 				Actor->Rename(nullptr, InLevel, REN_ForceNoResetLoaders);
 			}
 			
-			check(Actor->GetPackage() == LevelPackage);
 			if (bSameOuter && !InLevel->Actors.Contains(Actor))
 			{
 				InLevel->AddLoadedActor(Actor);
 			}
-
-			// Include objects found in the source actor package in the destination level package
-			TArray<UObject*> Objects;
-			const bool bIncludeNestedSubobjects = false;
-			GetObjectsWithOuter(ActorExternalPackage, Objects, bIncludeNestedSubobjects);
-			for (UObject* Object : Objects)
-			{
-				if (Object->GetFName() != NAME_PackageMetaData)
-				{
-					Object->Rename(nullptr, LevelPackage, REN_ForceNoResetLoaders);
-				}
-			}
 		}
 		else
 		{
diff --git a/Engine/Source/Runtime/Engine/Private/WorldPartition/WorldPartitionRuntimeLevelStreamingCell.cpp b/Engine/Source/Runtime/Engine/Private/WorldPartition/WorldPartitionRuntimeLevelStreamingCell.cpp
--- a/Engine/Source/Runtime/Engine/Private/WorldPartition/WorldPartitionRuntimeLevelStreamingCell.cpp
+++ b/Engine/Source/Runtime/Engine/Private/WorldPartition/WorldPartitionRuntimeLevelStreamingCell.cpp
@@ -14,6 +14,7 @@
 #include "WorldPartition/WorldPartitionLevelHelper.h"
 #include "Engine/LevelStreamingAlwaysLoaded.h"
 #include "Algo/ForEach.h"
+#include "AssetCompilingManager.h"
 #endif
 
 UWorldPartitionRuntimeLevelStreamingCell::UWorldPartitionRuntimeLevelStreamingCell(const FObjectInitializer& ObjectInitializer)
@@ -172,25 +173,22 @@ bool UWorldPartitionRuntimeLevelStreamingCell::LoadCellObjectsForCook(TArray<UOb
 	return true;
 }
 
-bool UWorldPartitionRuntimeLevelStreamingCell::MoveAlwaysLoadedContentToPersistentLevel()
+bool UWorldPartitionRuntimeLevelStreamingCell::PopulateGeneratorPackageForCook(TArray<UObject*>& OutLoadedObjects)
 {
 	check(IsAlwaysLoaded());
+
+	if (!LoadCellObjectsForCook(OutLoadedObjects))
+	{
+		return false;
+	}
+
 	check(bCookHasLoadedCellObjects);
 	if (GetActorCount() > 0)
 	{
-		// At this point actors should be loaded
-		Algo::ForEach(CookLoadedObjects, [](const TWeakObjectPtr<UObject>& Object) { check(Object.Get()); });
-
 		UWorld* OuterWorld = GetOuterUWorldPartition()->GetTypedOuter<UWorld>();
 		FWorldPartitionLevelHelper::MoveExternalActorsToLevel(Packages, OuterWorld->PersistentLevel);
-
-		// Empty cell's package list (this ensures that no one can rely on cell's content).
-		Packages.Empty();
 	}
 
-	CookLoadedObjects.Empty();
-	CookPackageReferencer.RemoveReferences();
-	bCookHasLoadedCellObjects  = false;
 	return true;
 }
 
@@ -207,17 +205,22 @@ bool UWorldPartitionRuntimeLevelStreamingCell::PrepareCellForCook(UPackage* InPa
 
 		LevelStreaming = CreateLevelStreaming(InPackage->GetName());
 	}
-
 	return true;
 }
 
-bool UWorldPartitionRuntimeLevelStreamingCell::PopulateGeneratedPackageForCook(UPackage* InPackage)
+bool UWorldPartitionRuntimeLevelStreamingCell::PopulateGeneratedPackageForCook(UPackage* InPackage, TArray<UObject*>& OutLoadedObjects)
 {
 	check(!IsAlwaysLoaded());
 	if (!InPackage)
 	{
 		return false;
 	}
+
+	if (!LoadCellObjectsForCook(OutLoadedObjects))
+	{
+		return false;
+	}
+
 	check(bCookHasLoadedCellObjects);
 
 	if (GetActorCount() > 0)
@@ -228,9 +231,6 @@ bool UWorldPartitionRuntimeLevelStreamingCell::PopulateGeneratedPackageForCook(U
 			return false;
 		}
 
-		// At this point actors should be loaded
-		Algo::ForEach(CookLoadedObjects, [](const TWeakObjectPtr<UObject>& Object) { check(Object.Get()); });
-
 		// Create a level and move loaded actors in it
 		UWorldPartition* WorldPartition = GetOuterUWorldPartition();
 		UWorld* OuterWorld = WorldPartition->GetTypedOuter<UWorld>();
@@ -242,9 +242,58 @@ bool UWorldPartitionRuntimeLevelStreamingCell::PopulateGeneratedPackageForCook(U
 		FWorldPartitionLevelHelper::RemapLevelSoftObjectPaths(NewLevel, WorldPartition);
 	}
 
+	return true;
+}
+
+bool UWorldPartitionRuntimeLevelStreamingCell::FinalizePackageForCook()
+{
+	check(bCookHasLoadedCellObjects);
+
+	// Taken from original code introduced in CL 15771404 :
+	//     We can't have async compilation still going on while we move actors as this is going to ResetLoaders which will move bulkdata around that
+	//     might still be used by async compilation. 
+	// This code was modified to use FAssetCompilingManager instead of FStaticMeshCompilingManager.
+	//
+	// TODO: Remove once the core team fixes a remaining bug in the COTFS that requires this code to wait for async compilation.
+	FAssetCompilingManager::Get().FinishAllCompilation();
+
+	UPackage* LevelPackage = nullptr;
+	for (const TWeakObjectPtr<UObject>& LoadedObject : CookLoadedObjects)
+	{
+		check(LoadedObject.Get());
+		AActor* Actor = Cast<AActor>(LoadedObject);
+		if (UPackage* ActorExternalPackage = Actor ? Actor->GetExternalPackage() : nullptr)
+		{
+			UPackage* ActorLevelPackage = Actor->GetLevel()->GetPackage();
+			check(!LevelPackage || (LevelPackage == ActorLevelPackage));
+			LevelPackage = ActorLevelPackage;
+
+			// Remove package external on actor
+			Actor->SetPackageExternal(false, false);
+
+			// Make sure to also include objects found in the source actor package into the level package
+			TArray<UObject*> Objects;
+			const bool bIncludeNestedSubobjects = false;
+			GetObjectsWithOuter(ActorExternalPackage, Objects, bIncludeNestedSubobjects);
+			for (UObject* Object : Objects)
+			{
+				if (Object->GetFName() != NAME_PackageMetaData)
+				{
+					Object->Rename(nullptr, LevelPackage, REN_ForceNoResetLoaders);
+				}
+			}
+		}
+	}
+
+	if (IsAlwaysLoaded())
+	{
+		// Empty cell's package list (this ensures that no one can rely on cell's content).
+		Packages.Empty();
+	}
+
 	CookLoadedObjects.Empty();
 	CookPackageReferencer.RemoveReferences();
-	bCookHasLoadedCellObjects  = false;
+	bCookHasLoadedCellObjects = false;
 
 	return true;
 }
diff --git a/Engine/Source/Runtime/Engine/Private/WorldPartition/WorldPartitionRuntimeSpatialHash.cpp b/Engine/Source/Runtime/Engine/Private/WorldPartition/WorldPartitionRuntimeSpatialHash.cpp
--- a/Engine/Source/Runtime/Engine/Private/WorldPartition/WorldPartitionRuntimeSpatialHash.cpp
+++ b/Engine/Source/Runtime/Engine/Private/WorldPartition/WorldPartitionRuntimeSpatialHash.cpp
@@ -1161,7 +1161,7 @@ bool UWorldPartitionRuntimeSpatialHash::CreateStreamingGrid(const FSpatialHashRu
 	return true;
 }
 
-bool UWorldPartitionRuntimeSpatialHash::LoadGeneratedPackageObjectsForCook(const FString& InPackageRelativePath, TArray<UObject*>& OutLoadedObjects)
+bool UWorldPartitionRuntimeSpatialHash::PopulateGeneratedPackageForCook(UPackage* InPackage, const FString& InPackageRelativePath, TArray<UObject*>& OutLoadedObjects)
 {
 	OutLoadedObjects.Reset();
 	if (UWorldPartitionRuntimeCell** MatchingCell = PackagesToGenerateForCook.Find(InPackageRelativePath))
@@ -1169,20 +1169,20 @@ bool UWorldPartitionRuntimeSpatialHash::LoadGeneratedPackageObjectsForCook(const
 		UWorldPartitionRuntimeCell* Cell = *MatchingCell;
 		if (ensure(Cell))
 		{
-			return Cell->LoadCellObjectsForCook(OutLoadedObjects);
+			return Cell->PopulateGeneratedPackageForCook(InPackage, OutLoadedObjects);
 		}
 	}
 	return false;
 }
 
-bool UWorldPartitionRuntimeSpatialHash::PopulateGeneratedPackageForCook(UPackage* InPackage, const FString& InPackageRelativePath)
+bool UWorldPartitionRuntimeSpatialHash::FinalizeGeneratedPackageForCook(const FString& InPackageRelativePath)
 {
 	if (UWorldPartitionRuntimeCell** MatchingCell = PackagesToGenerateForCook.Find(InPackageRelativePath))
 	{
 		UWorldPartitionRuntimeCell* Cell = *MatchingCell;
 		if (ensure(Cell))
 		{
-			return Cell->PopulateGeneratedPackageForCook(InPackage);
+			return Cell->FinalizePackageForCook();
 		}
 	}
 	return false;
@@ -1203,13 +1203,13 @@ TArray<UWorldPartitionRuntimeCell*> UWorldPartitionRuntimeSpatialHash::GetAlways
 	return Result;
 }
 
-bool UWorldPartitionRuntimeSpatialHash::LoadGeneratorPackageObjectsForCook(TArray<UObject*>& OutLoadedObjects)
+bool UWorldPartitionRuntimeSpatialHash::PopulateGeneratorPackageForCook(TArray<UObject*>& OutLoadedObjects)
 {
 	OutLoadedObjects.Reset();
 	for (UWorldPartitionRuntimeCell* Cell : GetAlwaysLoadedCells())
 	{
 		check(Cell->IsAlwaysLoaded());
-		if (!Cell->LoadCellObjectsForCook(OutLoadedObjects))
+		if (!Cell->PopulateGeneratorPackageForCook(OutLoadedObjects))
 		{
 			return false;
 		}
@@ -1220,12 +1220,11 @@ bool UWorldPartitionRuntimeSpatialHash::LoadGeneratorPackageObjectsForCook(TArra
 bool UWorldPartitionRuntimeSpatialHash::FinalizeGeneratorPackageForCook(const TArray<ICookPackageSplitter::FGeneratedPackageForPreSave>& InGeneratedPackages)
 {
 	check(IsRunningCookCommandlet());
-	
-	// Move always loaded content in Persistent Level
+
 	for (UWorldPartitionRuntimeCell* Cell : GetAlwaysLoadedCells())
 	{
 		check(Cell->IsAlwaysLoaded());
-		if (!Cell->MoveAlwaysLoadedContentToPersistentLevel())
+		if (!Cell->FinalizePackageForCook())
 		{
 			return false;
 		}
diff --git a/Engine/Source/Runtime/Engine/Public/WorldPartition/WorldPartition.h b/Engine/Source/Runtime/Engine/Public/WorldPartition/WorldPartition.h
--- a/Engine/Source/Runtime/Engine/Public/WorldPartition/WorldPartition.h
+++ b/Engine/Source/Runtime/Engine/Public/WorldPartition/WorldPartition.h
@@ -151,9 +151,9 @@ public:
 	void RemapSoftObjectPath(FSoftObjectPath& ObjectPath);
 
 	// Cooking
-	bool LoadGeneratorPackageObjectsForCook(TArray<UObject*>& OutLoadedObjects);
-	bool LoadGeneratedPackageObjectsForCook(const FString& InPackageRelativePath, TArray<UObject*>& OutLoadedObjects);
-	bool PopulateGeneratedPackageForCook(UPackage* InPackage, const FString& InPackageRelativePath);
+	bool PopulateGeneratorPackageForCook(TArray<UObject*>& OutLoadedObjects);
+	bool PopulateGeneratedPackageForCook(UPackage* InPackage, const FString& InPackageRelativePath, TArray<UObject*>& OutLoadedObjects);
+	bool FinalizeGeneratedPackageForCook(const FString& InPackageRelativePath);
 	bool FinalizeGeneratorPackageForCook(const TArray<ICookPackageSplitter::FGeneratedPackageForPreSave>& InGeneratedPackages);
 
 	UE_DEPRECATED(5.1, "GetWorldBounds is deprecated, use GetEditorWorldBounds or GetRuntimeWorldBounds instead.")
diff --git a/Engine/Source/Runtime/Engine/Public/WorldPartition/WorldPartitionRuntimeCell.h b/Engine/Source/Runtime/Engine/Public/WorldPartition/WorldPartitionRuntimeCell.h
--- a/Engine/Source/Runtime/Engine/Public/WorldPartition/WorldPartitionRuntimeCell.h
+++ b/Engine/Source/Runtime/Engine/Public/WorldPartition/WorldPartitionRuntimeCell.h
@@ -202,9 +202,9 @@ class UWorldPartitionRuntimeCell : public UObject, public IWorldPartitionCell
 
 	// Cook methods
 	virtual bool PrepareCellForCook(UPackage* InPackage) { return false; }
-	virtual bool LoadCellObjectsForCook(TArray<UObject*>& OutLoadedObjects) PURE_VIRTUAL(UWorldPartitionRuntimeCell::LoadCellObjectsForCook, return false;);
-	virtual bool PopulateGeneratedPackageForCook(UPackage* InPackage) PURE_VIRTUAL(UWorldPartitionRuntimeCell::PopulateGeneratedPackageForCook, return false;);
-	virtual bool MoveAlwaysLoadedContentToPersistentLevel() PURE_VIRTUAL(UWorldPartitionRuntimeCell::MoveAlwaysLoadedContentToPersistentLevel, return false;);
+	virtual bool PopulateGeneratorPackageForCook(TArray<UObject*>& OutLoadedObjects) PURE_VIRTUAL(UWorldPartitionRuntimeCell::PopulateGeneratorPackageForCook, return false;);
+	virtual bool PopulateGeneratedPackageForCook(UPackage* InPackage, TArray<UObject*>& OutLoadedObjects) PURE_VIRTUAL(UWorldPartitionRuntimeCell::PopulateGeneratedPackageForCook, return false;);
+	virtual bool FinalizePackageForCook() PURE_VIRTUAL(UWorldPartitionRuntimeCell::FinalizePackageForCook, return false;);
 	virtual FString GetPackageNameToCreate() const PURE_VIRTUAL(UWorldPartitionRuntimeCell::GetPackageNameToCreate, return FString(""););
 
 	void SetIsHLOD(bool bInIsHLOD) { bIsHLOD = bInIsHLOD; }
diff --git a/Engine/Source/Runtime/Engine/Public/WorldPartition/WorldPartitionRuntimeHash.h b/Engine/Source/Runtime/Engine/Public/WorldPartition/WorldPartitionRuntimeHash.h
--- a/Engine/Source/Runtime/Engine/Public/WorldPartition/WorldPartitionRuntimeHash.h
+++ b/Engine/Source/Runtime/Engine/Public/WorldPartition/WorldPartitionRuntimeHash.h
@@ -34,10 +34,10 @@ class ENGINE_API UWorldPartitionRuntimeHash : public UObject
 
 #if WITH_EDITOR
 	virtual void SetDefaultValues() {}
-	virtual bool LoadGeneratorPackageObjectsForCook(TArray<UObject*>& OutLoadedObjects) { return false; }
-	virtual bool LoadGeneratedPackageObjectsForCook(const FString& InPackageRelativePath, TArray<UObject*>& OutLoadedObjects) { return false; }
-	virtual bool PopulateGeneratedPackageForCook(UPackage* InPackage, const FString& InPackageRelativePath) { return false; }
+	virtual bool PopulateGeneratorPackageForCook(TArray<UObject*>& OutLoadedObjects) { return false; }
 	virtual bool FinalizeGeneratorPackageForCook(const TArray<ICookPackageSplitter::FGeneratedPackageForPreSave>& InGeneratedPackages) { return false; }
+	virtual bool PopulateGeneratedPackageForCook(UPackage* InPackage, const FString& InPackageRelativePath, TArray<UObject*>& OutLoadedObjects) { return false; }
+	virtual bool FinalizeGeneratedPackageForCook(const FString& InPackageRelativePath) { return false; }
 	virtual bool GenerateStreaming(class UWorldPartitionStreamingPolicy* StreamingPolicy, const FActorClusterContext& ActorClusterContext, TArray<FString>* OutPackagesToGenerate) { return false; }
 	virtual void FlushStreaming() {}
 	virtual bool GenerateHLOD(ISourceControlHelper* SourceControlHelper, FActorClusterContext& ActorClusterContext, bool bCreateActorsOnly) { return false; }
diff --git a/Engine/Source/Runtime/Engine/Public/WorldPartition/WorldPartitionRuntimeLevelStreamingCell.h b/Engine/Source/Runtime/Engine/Public/WorldPartition/WorldPartitionRuntimeLevelStreamingCell.h
--- a/Engine/Source/Runtime/Engine/Public/WorldPartition/WorldPartitionRuntimeLevelStreamingCell.h
+++ b/Engine/Source/Runtime/Engine/Public/WorldPartition/WorldPartitionRuntimeLevelStreamingCell.h
@@ -40,9 +40,9 @@ class UWorldPartitionRuntimeLevelStreamingCell : public UWorldPartitionRuntimeSp
 	virtual void DumpStateLog(FHierarchicalLogArchive& Ar) override;
 	// Cook methods
 	virtual bool PrepareCellForCook(UPackage* InPackage) override;
-	virtual bool LoadCellObjectsForCook(TArray<UObject*>& OutLoadedObjects) override;
-	virtual bool PopulateGeneratedPackageForCook(UPackage* InPackage) override;
-	virtual bool MoveAlwaysLoadedContentToPersistentLevel() override;
+	virtual bool PopulateGeneratorPackageForCook(TArray<UObject*>& OutLoadedObjects) override;
+	virtual bool PopulateGeneratedPackageForCook(UPackage* InPackage, TArray<UObject*>& OutLoadedObjects) override;
+	virtual bool FinalizePackageForCook() override;
 	virtual FString GetPackageNameToCreate() const override;
 	//~End UWorldPartitionRuntimeCell Interface
 
@@ -59,6 +59,7 @@ private:
 	class UWorldPartitionLevelStreamingDynamic* GetOrCreateLevelStreaming() const;
 
 #if WITH_EDITOR
+	bool LoadCellObjectsForCook(TArray<UObject*>& OutLoadedObjects);
 	class UWorldPartitionLevelStreamingDynamic* CreateLevelStreaming(const FString& InPackageName = FString()) const;
 
 	FWorldPartitionLevelHelper::FPackageReferencer CookPackageReferencer;
diff --git a/Engine/Source/Runtime/Engine/Public/WorldPartition/WorldPartitionRuntimeSpatialHash.h b/Engine/Source/Runtime/Engine/Public/WorldPartition/WorldPartitionRuntimeSpatialHash.h
--- a/Engine/Source/Runtime/Engine/Public/WorldPartition/WorldPartitionRuntimeSpatialHash.h
+++ b/Engine/Source/Runtime/Engine/Public/WorldPartition/WorldPartitionRuntimeSpatialHash.h
@@ -232,10 +232,10 @@ public:
 
 #if WITH_EDITOR
 	virtual void SetDefaultValues() override;
-	virtual bool LoadGeneratorPackageObjectsForCook(TArray<UObject*>& OutLoadedObjects) override;
-	virtual bool LoadGeneratedPackageObjectsForCook(const FString& InPackageRelativePath, TArray<UObject*>& OutLoadedObjects) override;
-	virtual bool PopulateGeneratedPackageForCook(UPackage* InPackage, const FString& InPackageRelativePath) override;
+	virtual bool PopulateGeneratorPackageForCook(TArray<UObject*>& OutLoadedObjects) override;
 	virtual bool FinalizeGeneratorPackageForCook(const TArray<ICookPackageSplitter::FGeneratedPackageForPreSave>& InGeneratedPackages) override;
+	virtual bool PopulateGeneratedPackageForCook(UPackage* InPackage, const FString& InPackageRelativePath, TArray<UObject*>& OutLoadedObjects) override;
+	virtual bool FinalizeGeneratedPackageForCook(const FString& InPackageRelativePath) override;
 	virtual void FlushStreaming() override;
 	virtual bool GenerateHLOD(ISourceControlHelper* SourceControlHelper, FActorClusterContext& ActorClusterContext, bool bCreateActorsOnly) override;
 	virtual void DrawPreview() const override;
