Fixed WorldPartition MinimapBuilder to work with LWC.
Added CheckForErrors to detect  cases where minimap builder will either fail to generate or will generate a low quality mini-map.

diff --git a/Engine/Source/Editor/UnrealEd/Private/WorldPartition/WorldPartitionMiniMapBuilder.cpp b/Engine/Source/Editor/UnrealEd/Private/WorldPartition/WorldPartitionMiniMapBuilder.cpp
--- a/Engine/Source/Editor/UnrealEd/Private/WorldPartition/WorldPartitionMiniMapBuilder.cpp
+++ b/Engine/Source/Editor/UnrealEd/Private/WorldPartition/WorldPartitionMiniMapBuilder.cpp
@@ -55,52 +55,25 @@ bool UWorldPartitionMiniMapBuilder::PreRun(UWorld* World, FPackageSourceControlH
 	}
 
 	// World bounds to process
-	FBox WorldBounds(ForceInit);
-	{
-		// Override the minimap bounds it a world partiion minimap volume exists
-		for (TActorIterator<AWorldPartitionMiniMapVolume> It(World); It; ++It)
-		{
-			if (AWorldPartitionMiniMapVolume* WorldPartitionMiniMapVolume = *It)
-			{
-				WorldBounds += WorldPartitionMiniMapVolume->GetBounds().GetBox();
-			}
-		}
-
-		if (!WorldBounds.IsValid)
-		{
-			WorldBounds = World->GetWorldPartition()->GetRuntimeWorldBounds();
-		}
+	IterativeWorldBounds = WorldMiniMap->GetMiniMapWorldBounds();
 
-		IterativeWorldBounds = WorldBounds;
-	}
-	
-	// Compute minimap image size
-	{
-		MinimapImageSizeX = WorldBounds.GetSize().X / WorldMiniMap->WorldUnitsPerPixel;
-		MinimapImageSizeY = WorldBounds.GetSize().Y / WorldMiniMap->WorldUnitsPerPixel;
-
-		// For now, let's clamp to the maximum supported texture size
-		MinimapImageSizeX = FMath::Min(MinimapImageSizeX, UTexture::GetMaximumDimensionOfNonVT());
-		MinimapImageSizeY = FMath::Min(MinimapImageSizeY, UTexture::GetMaximumDimensionOfNonVT());
-		WorldUnitsPerPixel = FMath::CeilToInt(FMath::Max(WorldBounds.GetSize().X / MinimapImageSizeX, WorldBounds.GetSize().Y / MinimapImageSizeY));
-		MinimapImageSizeX = WorldBounds.GetSize().X / WorldUnitsPerPixel;
-		MinimapImageSizeY = WorldBounds.GetSize().Y / WorldUnitsPerPixel;
-	}
+	// Compute minimap resolution
+	WorldMiniMap->GetMiniMapResolution(MinimapImageSizeX, MinimapImageSizeY, WorldUnitsPerPixel);
 
 	// Create minimap texture
 	{
 		TStrongObjectPtr<UTextureFactory> Factory(NewObject<UTextureFactory>());
 		WorldMiniMap->MiniMapTexture = Factory->CreateTexture2D(WorldMiniMap, TEXT("MinimapTexture"), RF_NoFlags);
 		WorldMiniMap->MiniMapTexture->Source.Init(MinimapImageSizeX, MinimapImageSizeY, 1, 1, TSF_BGRA8);
-		WorldMiniMap->MiniMapWorldBounds = WorldBounds;
+		WorldMiniMap->MiniMapWorldBounds = IterativeWorldBounds;
 		MiniMapSourcePtr = WorldMiniMap->MiniMapTexture->Source.LockMip(0);
 	}
 
 	// Compute world to minimap transform
 	{
-		WorldToMinimap = FReversedZOrthoMatrix(WorldBounds.Min.X, WorldBounds.Max.X, WorldBounds.Min.Y, WorldBounds.Max.Y, 1.0f, 0.0f);
+		WorldToMinimap = FReversedZOrthoMatrix(IterativeWorldBounds.Min.X, IterativeWorldBounds.Max.X, IterativeWorldBounds.Min.Y, IterativeWorldBounds.Max.Y, 1.0f, 0.0f);
 
-		FVector3d Translation(WorldBounds.Max.X / WorldBounds.GetSize().X, WorldBounds.Max.Y / WorldBounds.GetSize().Y, 0);
+		FVector3d Translation(IterativeWorldBounds.Max.X / IterativeWorldBounds.GetSize().X, IterativeWorldBounds.Max.Y / IterativeWorldBounds.GetSize().Y, 0);
 		FVector3d Scaling(MinimapImageSizeX, MinimapImageSizeY, 1);
 
 		WorldToMinimap *= FTranslationMatrix(Translation);
@@ -134,8 +107,9 @@ bool UWorldPartitionMiniMapBuilder::RunInternal(UWorld* World, const FCellInfo&
 	// World X,Y to minimap X,Y
 	const FVector3d MinimapMin = WorldToMinimap.TransformPosition(ClampedBounds.Min);
 	const FVector3d MinimapMax = WorldToMinimap.TransformPosition(ClampedBounds.Max);
-	const FIntVector2 DstMin(FMath::Floor(MinimapMin.X), FMath::Floor(MinimapMin.Y));
-	const FIntVector2 DstMax(FMath::Floor(MinimapMax.X), FMath::Floor(MinimapMax.Y));
+	const FIntVector2 DstMin(FMath::Max(FMath::Floor(MinimapMin.X), 0), FMath::Max(FMath::Floor(MinimapMin.Y), 0));
+	const FIntVector2 DstMax(FMath::Max(FMath::Floor(MinimapMax.X), 0), FMath::Max(FMath::Floor(MinimapMax.Y), 0));
+
 	const uint32 CaptureWidthPixels = DstMax.X - DstMin.X;
 	const uint32 CaptureHeightPixels = DstMax.Y - DstMin.Y;
 
diff --git a/Engine/Source/Editor/WorldPartitionEditor/Private/WorldPartitionEditorModule.cpp b/Engine/Source/Editor/WorldPartitionEditor/Private/WorldPartitionEditorModule.cpp
--- a/Engine/Source/Editor/WorldPartitionEditor/Private/WorldPartitionEditorModule.cpp
+++ b/Engine/Source/Editor/WorldPartitionEditor/Private/WorldPartitionEditorModule.cpp
@@ -220,6 +220,11 @@ int32 FWorldPartitionEditorModule::GetInstancedFoliageGridSize() const
 	return GetDefault<UWorldPartitionEditorSettings>()->InstancedFoliageGridSize;
 }
 
+int32 FWorldPartitionEditorModule::GetMinimapLowQualityWorldUnitsPerPixelThreshold() const
+{
+	return GetDefault<UWorldPartitionEditorSettings>()->MinimapLowQualityWorldUnitsPerPixelThreshold;
+}
+
 void FWorldPartitionEditorModule::OnConvertMap()
 {
 	IContentBrowserSingleton& ContentBrowserSingleton = FModuleManager::LoadModuleChecked<FContentBrowserModule>("ContentBrowser").Get();
@@ -614,6 +619,7 @@ UWorldPartitionEditorSettings::UWorldPartitionEditorSettings()
 {
 	CommandletClass = UWorldPartitionConvertCommandlet::StaticClass();
 	InstancedFoliageGridSize = 25600;
+	MinimapLowQualityWorldUnitsPerPixelThreshold = 12800;
 }
 
 FString UWorldPartitionConvertOptions::ToCommandletArgs() const
diff --git a/Engine/Source/Editor/WorldPartitionEditor/Public/WorldPartition/WorldPartitionEditorSettings.h b/Engine/Source/Editor/WorldPartitionEditor/Public/WorldPartition/WorldPartitionEditorSettings.h
--- a/Engine/Source/Editor/WorldPartitionEditor/Public/WorldPartition/WorldPartitionEditorSettings.h
+++ b/Engine/Source/Editor/WorldPartitionEditor/Public/WorldPartition/WorldPartitionEditorSettings.h
@@ -20,4 +20,7 @@ public:
 
 	UPROPERTY(config, EditAnywhere, Category = Foliage, meta = (ClampMin=3200, ToolTip= "Editor grid size used for instance foliage actors in World Partition worlds"))
 	int32 InstancedFoliageGridSize;
+
+	UPROPERTY(config, EditAnywhere, Category = MiniMap, meta = (ClampMin = 100, ToolTip = "Threshold from which minimap generates a warning if its WorldUnitsPerPixel is above this value"))
+	int32 MinimapLowQualityWorldUnitsPerPixelThreshold;
 };
\ No newline at end of file
diff --git a/Engine/Source/Editor/WorldPartitionEditor/Public/WorldPartitionEditorModule.h b/Engine/Source/Editor/WorldPartitionEditor/Public/WorldPartitionEditorModule.h
--- a/Engine/Source/Editor/WorldPartitionEditor/Public/WorldPartitionEditorModule.h
+++ b/Engine/Source/Editor/WorldPartitionEditor/Public/WorldPartitionEditorModule.h
@@ -34,12 +34,17 @@ public:
 	/**
 	 * Returns placement grid size setting that should be assigned to new AWorldSettings actors.
 	 */
-	virtual int32 GetPlacementGridSize() const;
+	virtual int32 GetPlacementGridSize() const override;
 
 	/**
 	 * Returns foliage grid size setting that should be assigned to new AWorldSettings actors.
 	 */
-	virtual int32 GetInstancedFoliageGridSize() const;
+	virtual int32 GetInstancedFoliageGridSize() const override;
+
+	/**
+	 * Returns the threshold from which minimap generates a warning if its WorldUnitsPerPixel is above this value.
+	 */
+	virtual int32 GetMinimapLowQualityWorldUnitsPerPixelThreshold() const override;
 
 	/**
 	 * Convert the specified map to a world partition map.
diff --git a/Engine/Source/Runtime/Engine/Private/WorldPartition/WorldPartitionMiniMap.cpp b/Engine/Source/Runtime/Engine/Private/WorldPartition/WorldPartitionMiniMap.cpp
--- a/Engine/Source/Runtime/Engine/Private/WorldPartition/WorldPartitionMiniMap.cpp
+++ b/Engine/Source/Runtime/Engine/Private/WorldPartition/WorldPartitionMiniMap.cpp
@@ -1,10 +1,16 @@
 // Copyright Epic Games, Inc. All Rights Reserved.
 
 #include "WorldPartition/WorldPartitionMiniMap.h"
+#include "WorldPartition/WorldPartitionMiniMapVolume.h"
+#include "WorldPartition/WorldPartitionMiniMapHelper.h"
+#include "WorldPartition/WorldPartition.h"
+#include "WorldPartition/IWorldPartitionEditorModule.h"
+#include "Modules/ModuleManager.h"
 #include "Engine/Texture2D.h"
 #include "Logging/MessageLog.h"
 #include "Misc/MapErrors.h"
 #include "Misc/UObjectToken.h"
+#include "EngineUtils.h"
 #include "RenderUtils.h"
 #include "RHI.h"
 
@@ -39,6 +45,12 @@ void AWorldPartitionMiniMap::CheckForErrors()
 {
 	Super::CheckForErrors();
 
+	// We skip actors not part of main world (level instances)
+	if (FWorldPartitionMiniMapHelper::GetWorldPartitionMiniMap(GetWorld()) != this)
+	{
+		return;
+	}
+
 	if (MiniMapTexture != nullptr)
 	{
 		const int32 MaxTextureDimension = GetMax2DTextureDimension();
@@ -53,6 +65,62 @@ void AWorldPartitionMiniMap::CheckForErrors()
 				->AddToken(FMapErrorToken::Create("MinimapTextureSize"));
 		}
 	}
+
+	{
+		int32 EffectiveMinimapImageSizeX = 0;
+		int32 EffectiveMinimapImageSizeY = 0;
+		int32 EffectiveWorldUnitsPerPixel = 0;
+		GetMiniMapResolution(EffectiveMinimapImageSizeX, EffectiveMinimapImageSizeY, EffectiveWorldUnitsPerPixel);
+
+		IWorldPartitionEditorModule& WorldPartitionEditorModule = FModuleManager::LoadModuleChecked<IWorldPartitionEditorModule>("WorldPartitionEditor");
+		if (EffectiveWorldUnitsPerPixel > WorldPartitionEditorModule.GetMinimapLowQualityWorldUnitsPerPixelThreshold())
+		{
+			FMessageLog("MapCheck").Warning()
+				->AddToken(FUObjectToken::Create(this))
+				->AddToken(FTextToken::Create(FText::Format(LOCTEXT("MapCheck_Message_MinimapWorldUnitsPerPixel", "{0} : Expect a low quality minimap : effective world units per pixel is high ({1})."), FText::FromString(GetName()), EffectiveWorldUnitsPerPixel)))
+				->AddToken(FMapErrorToken::Create("MinimapWorldUnitsPerPixel"));
+		}
+	}
+}
+
+FBox AWorldPartitionMiniMap::GetMiniMapWorldBounds() const
+{
+	FBox WorldBounds(ForceInit);
+
+	UWorld* World = GetWorld();
+	check(World);
+
+	// Override the minimap bounds if world partition minimap volumes exists
+	for (TActorIterator<AWorldPartitionMiniMapVolume> It(World); It; ++It)
+	{
+		if (AWorldPartitionMiniMapVolume* WorldPartitionMiniMapVolume = *It)
+		{
+			WorldBounds += WorldPartitionMiniMapVolume->GetBounds().GetBox();
+		}
+	}
+
+	if (!WorldBounds.IsValid)
+	{
+		WorldBounds = World->GetWorldPartition()->GetRuntimeWorldBounds();
+	}
+
+	return WorldBounds;
+}
+
+void AWorldPartitionMiniMap::GetMiniMapResolution(int32& OutMinimapImageSizeX, int32& OutMinimapImageSizeY, int32& OutWorldUnitsPerPixel) const
+{
+	FBox WorldBounds = GetMiniMapWorldBounds();
+
+	// Compute minimap image size
+	OutMinimapImageSizeX = WorldBounds.GetSize().X / WorldUnitsPerPixel;
+	OutMinimapImageSizeY = WorldBounds.GetSize().Y / WorldUnitsPerPixel;
+
+	// For now, let's clamp to the maximum supported texture size
+	OutMinimapImageSizeX = FMath::Min(OutMinimapImageSizeX, UTexture::GetMaximumDimensionOfNonVT());
+	OutMinimapImageSizeY = FMath::Min(OutMinimapImageSizeY, UTexture::GetMaximumDimensionOfNonVT());
+	OutWorldUnitsPerPixel = FMath::CeilToInt(FMath::Max(WorldBounds.GetSize().X / OutMinimapImageSizeX, WorldBounds.GetSize().Y / OutMinimapImageSizeY));
+	OutMinimapImageSizeX = WorldBounds.GetSize().X / OutWorldUnitsPerPixel;
+	OutMinimapImageSizeY = WorldBounds.GetSize().Y / OutWorldUnitsPerPixel;
 }
 
 #endif
diff --git a/Engine/Source/Runtime/Engine/Private/WorldPartition/WorldPartitionMiniMapHelper.cpp b/Engine/Source/Runtime/Engine/Private/WorldPartition/WorldPartitionMiniMapHelper.cpp
--- a/Engine/Source/Runtime/Engine/Private/WorldPartition/WorldPartitionMiniMapHelper.cpp
+++ b/Engine/Source/Runtime/Engine/Private/WorldPartition/WorldPartitionMiniMapHelper.cpp
@@ -138,7 +138,7 @@ void FWorldPartitionMiniMapHelper::CalTopViewOfWorld(FMatrix& OutProjectionMatri
 	const float OrthoWidth = FMath::Max(1.f, ViewportWidth * OrthoZoom);
 	const float OrthoHeight = FMath::Max(1.f, ViewportHeight * OrthoZoom);
 
-	const float ZOffset = HALF_WORLD_MAX;
+	const double ZOffset = HALF_WORLD_MAX;
 	OutProjectionMatrix = FReversedZOrthoMatrix(
 		OrthoWidth,
 		OrthoHeight,
diff --git a/Engine/Source/Runtime/Engine/Public/WorldPartition/IWorldPartitionEditorModule.h b/Engine/Source/Runtime/Engine/Public/WorldPartition/IWorldPartitionEditorModule.h
--- a/Engine/Source/Runtime/Engine/Public/WorldPartition/IWorldPartitionEditorModule.h
+++ b/Engine/Source/Runtime/Engine/Public/WorldPartition/IWorldPartitionEditorModule.h
@@ -21,6 +21,7 @@ public:
 
 	virtual int32 GetPlacementGridSize() const = 0;
 	virtual int32 GetInstancedFoliageGridSize() const = 0;
+	virtual int32 GetMinimapLowQualityWorldUnitsPerPixelThreshold() const = 0;
 
 	/** Triggered when a world is added. */
 	DECLARE_EVENT_OneParam(IWorldPartitionEditorModule, FWorldPartitionCreated, UWorld*);
diff --git a/Engine/Source/Runtime/Engine/Public/WorldPartition/WorldPartitionMiniMap.h b/Engine/Source/Runtime/Engine/Public/WorldPartition/WorldPartitionMiniMap.h
--- a/Engine/Source/Runtime/Engine/Public/WorldPartition/WorldPartitionMiniMap.h
+++ b/Engine/Source/Runtime/Engine/Public/WorldPartition/WorldPartitionMiniMap.h
@@ -29,6 +29,8 @@ public:
 	virtual bool IsLockLocation() const { return true; }
 	virtual bool IsUserManaged() const final { return false; }
 	virtual void CheckForErrors() override;
+	FBox GetMiniMapWorldBounds() const;
+	void GetMiniMapResolution(int32& OutMinimapImageSizeX, int32& OutMinimapImageSizeY, int32& OutWorldUnitsPerPixel) const;
 #endif
 
 	/* WorldBounds for MinMapTexture */
diff --git a/Engine/Source/Runtime/Renderer/Private/ShadowSetup.cpp b/Engine/Source/Runtime/Renderer/Private/ShadowSetup.cpp
--- a/Engine/Source/Runtime/Renderer/Private/ShadowSetup.cpp
+++ b/Engine/Source/Runtime/Renderer/Private/ShadowSetup.cpp
@@ -989,7 +989,7 @@ void FProjectedShadowInfo::SetupWholeSceneProjection(
 	GetViewFrustumBounds(CasterOuterFrustum, FMatrix(TranslatedWorldToClipOuterMatrix), true);
 	GetViewFrustumBounds(ReceiverInnerFrustum, ReceiverInnerMatrix, true);
 	
-	InvReceiverInnerMatrix = FMatrix44f(ReceiverInnerMatrix.InverseFast());
+	InvReceiverInnerMatrix = FMatrix44f(ReceiverInnerMatrix.Inverse());
 
 	UpdateShaderDepthBias();
 
diff --git a/Engine/Source/Runtime/Renderer/Private/VT/VirtualTextureSystem.cpp b/Engine/Source/Runtime/Renderer/Private/VT/VirtualTextureSystem.cpp
--- a/Engine/Source/Runtime/Renderer/Private/VT/VirtualTextureSystem.cpp
+++ b/Engine/Source/Runtime/Renderer/Private/VT/VirtualTextureSystem.cpp
@@ -1252,8 +1252,8 @@
 	// Determine the screen-space coordinates of the texture we're trying to display
 	const float SrcPositionX0 = FMath::Max(InViewportPosition.X, 0.0f);
 	const float SrcPositionY0 = FMath::Max(InViewportPosition.Y, 0.0f);
-	const float SrcPositionX1 = FMath::Min(InViewportPosition.X + InViewportSize.X, ScreenSpaceSizeX);
-	const float SrcPositionY1 = FMath::Min(InViewportPosition.Y + InViewportSize.Y, ScreenSpaceSizeY);
+	const float SrcPositionX1 = FMath::Max(FMath::Min(InViewportPosition.X + InViewportSize.X, ScreenSpaceSizeX), 0.0f);
+	const float SrcPositionY1 = FMath::Max(FMath::Min(InViewportPosition.Y + InViewportSize.Y, ScreenSpaceSizeY), 0.0f);
 
 	const int32 WidthInBlocks = AllocatedVT->GetWidthInBlocks();
 	const int32 HeightInBlocks = AllocatedVT->GetHeightInBlocks();
