Fix AssetRegistryGenerator to write pakchunkN.txt files even for empty chunks, if the pakchunkN.txt file is referenced by the PakFileRules. We need to provide the options list for those to-be-populated pak files to the staging step. And currently the staging step will refuse to populate them at all if the Generator does not write out an empty pakchunkN.txt file for them.
#rnx

diff --git a/Engine/Source/Editor/UnrealEd/Private/Commandlets/AssetRegistryGenerator.cpp b/Engine/Source/Editor/UnrealEd/Private/Commandlets/AssetRegistryGenerator.cpp
--- a/Engine/Source/Editor/UnrealEd/Private/Commandlets/AssetRegistryGenerator.cpp
+++ b/Engine/Source/Editor/UnrealEd/Private/Commandlets/AssetRegistryGenerator.cpp
@@ -39,6 +39,7 @@
 #include "Serialization/JsonWriter.h"
 #include "Serialization/JsonReader.h"
 #include "Serialization/JsonSerializer.h"
+#include "String/ParseTokens.h"
 
 DEFINE_LOG_CATEGORY_STATIC(LogAssetRegistryGenerator, Log, All);
 
@@ -84,6 +85,72 @@ FName GetPackageNameFromDependencyPackageName(const FName RawPackageFName)
 	return PackageFName;
 }
 
+class FDefaultPakFileRules
+{
+public:
+	void InitializeFromConfig(const ITargetPlatform* TargetPlatform)
+	{
+		if (bInitialized)
+		{
+			return;
+		}
+
+		FConfigFile ConfigFile;
+		if (!FConfigCacheIni::LoadLocalIniFile(ConfigFile, TEXT("PakFileRules"), true /* bIsBaseIniName */))
+		{
+			return;
+		}
+		// Schema is defined in Engine\Config\BasePakFileRules.ini, see also GetPakFileRules in CopyBuildToStaging.Automation.cs
+
+		FString IniPlatformName = TargetPlatform->IniPlatformName();
+		for (TPair<FString, FConfigSection>& Pair : ConfigFile)
+		{
+			const FString& SectionName = Pair.Key;
+			bool bMatchesAllPlatforms = true;
+			bool bMatchesPlatform = false;
+			FString ApplyToPlatformsValue;
+			if (ConfigFile.GetString(*SectionName, TEXT("Platforms"), ApplyToPlatformsValue))
+			{
+				UE::String::ParseTokens(ApplyToPlatformsValue, ',',
+					[&bMatchesAllPlatforms, &bMatchesPlatform, &IniPlatformName](FStringView Token)
+					{
+						bMatchesAllPlatforms = false;
+						if (IniPlatformName == Token)
+						{
+							bMatchesPlatform = true;
+						}
+					}, UE::String::EParseTokensOptions::Trim | UE::String::EParseTokensOptions::SkipEmpty);
+			}
+			if (!bMatchesPlatform && !bMatchesAllPlatforms)
+			{
+				continue;
+			}
+
+			FString OverridePaksValue;
+			if (ConfigFile.GetString(*SectionName, TEXT("OverridePaks"), OverridePaksValue))
+			{
+				UE::String::ParseTokens(OverridePaksValue, ',', [this](FStringView Token)
+					{
+						ReferencedPaks.Add(FString(Token));
+					}, UE::String::EParseTokensOptions::Trim | UE::String::EParseTokensOptions::SkipEmpty);
+			}
+		}
+		bInitialized = true;
+	}
+
+	bool IsChunkReferenced(int32 PakchunkIndex)
+	{
+		TStringBuilder<64> ChunkFileName;
+		ChunkFileName.Appendf(TEXT("pakchunk%d"), PakchunkIndex);
+		FStringView ChunkFileNameView(ChunkFileName);
+		return ReferencedPaks.ContainsByHash(GetTypeHash(ChunkFileNameView), ChunkFileNameView);
+	}
+
+private:
+	bool bInitialized = false;
+	TSet<FString> ReferencedPaks;
+
+};
 
 //////////////////////////////////////////////////////////////////////////
 // FAssetRegistryGenerator
@@ -543,6 +610,7 @@ bool FAssetRegistryGenerator::GenerateStreamingInstallManifest(int64 InOverrideC
 	}
 
 	// generate per-chunk pak list files
+	FDefaultPakFileRules DefaultPakFileRules;
 	bool bSucceeded = true;
 	for (int32 PakchunkIndex = 0; PakchunkIndex < FinalChunkManifests.Num() && bSucceeded; ++PakchunkIndex)
 	{
@@ -568,10 +636,14 @@ bool FAssetRegistryGenerator::GenerateStreamingInstallManifest(int64 InOverrideC
 		TArray<FString> ChunkFilenames;
 		Manifest->GenerateValueArray(ChunkFilenames);
 
-		// Do not create any files if the chunk is empty
+		// Do not create any files if the chunk is empty and is not referenced by rules applied during staging
 		if (ChunkFilenames.IsEmpty())
 		{
-			continue;
+			DefaultPakFileRules.InitializeFromConfig(TargetPlatform);
+			if (!DefaultPakFileRules.IsChunkReferenced(PakchunkIndex))
+			{
+				continue;
+			}
 		}
 
 		bool bFinishedAllFiles = false;
