Cook memory GC triggers warn if MemoryMaxUsed config cannot be hit due to low system memory
adjust defaults so  that GC will happen when running out of available memory for assetcompilingmanager
help with crashes on client machines with low total memory

diff --git a/Engine/Config/BaseEditor.ini b/Engine/Config/BaseEditor.ini
--- a/Engine/Config/BaseEditor.ini
+++ b/Engine/Config/BaseEditor.ini
@@ -319,8 +319,10 @@ BulkDataRegistryEnabled=true
 HybridIterativeEnabled=false
 PackagesPerGC=0
 IdleTimeToGC=20
-MemoryMaxUsedPhysical=16384
-MemoryMinFreePhysical=1024
+; 2 GB memory is need to cook large textures
+;  set MinFreePhysical and MinFreeVirtual both to the same number
+MemoryMinFreePhysical=2048
+MemoryMinFreeVirtual=2048
 MaxConcurrentShaderJobs=8192
 MaxDiffsToLog=5
 IgnoreHeaderDiffs=True
@@ -609,4 +611,4 @@ bCanModifyPluginsFromBrowser=true
 
 [/Script/GameFeaturesEditor.GameFeaturesEditorSettings]
 +PluginTemplates=(Path=(Path="../../Plugins/Experimental/GameFeatures/Templates/GameFeaturePluginContentOnly"),Label=NSLOCTEXT("[/Script/GameFeaturesEditor]", "E17DB7E94A2C0F5968A52E89959864C0", "Game Feature (Content Only)"),Description=NSLOCTEXT("[/Script/GameFeaturesEditor]", "CB5136B64AFAFBA5EDEA078FCAC7BCC9", "Create a new Game Feature Plugin."),DefaultGameFeatureDataClass=None,DefaultGameFeatureDataName="")
-+PluginTemplates=(Path=(Path="../../Plugins/Experimental/GameFeatures/Templates/GameFeaturePluginWithCode"),Label=NSLOCTEXT("[/Script/GameFeaturesEditor]", "90D6FD2A455176CEC42DD79A23AEC3B0", "Game Feature (with C++)"),Description=NSLOCTEXT("[/Script/GameFeaturesEditor]", "2B00A0E7415C8CEAA6DA609A141237F0", "Create a new Game Feature Plugin with a minimal amount of code."),DefaultGameFeatureDataClass=None,DefaultGameFeatureDataName="")
\ No newline at end of file
++PluginTemplates=(Path=(Path="../../Plugins/Experimental/GameFeatures/Templates/GameFeaturePluginWithCode"),Label=NSLOCTEXT("[/Script/GameFeaturesEditor]", "90D6FD2A455176CEC42DD79A23AEC3B0", "Game Feature (with C++)"),Description=NSLOCTEXT("[/Script/GameFeaturesEditor]", "2B00A0E7415C8CEAA6DA609A141237F0", "Create a new Game Feature Plugin with a minimal amount of code."),DefaultGameFeatureDataClass=None,DefaultGameFeatureDataName="")
diff --git a/Engine/Source/Editor/UnrealEd/Private/CookOnTheFlyServer.cpp b/Engine/Source/Editor/UnrealEd/Private/CookOnTheFlyServer.cpp
--- a/Engine/Source/Editor/UnrealEd/Private/CookOnTheFlyServer.cpp
+++ b/Engine/Source/Editor/UnrealEd/Private/CookOnTheFlyServer.cpp
@@ -4221,11 +4221,21 @@ bool UCookOnTheFlyServer::HasExceededMaxMemory() const
 	int FiredTriggers = 0;
 
 	TStringBuilder<256> TriggerMessages;
+
 	const FPlatformMemoryStats MemStats = FPlatformMemory::GetStats();
+	
 	if (MemoryMinFreeVirtual > 0 || MemoryMinFreePhysical > 0)
 	{
 		++ActiveTriggers;
 		bool bFired = false;
+		
+		// trigger GC if we have less than MemoryMinFreeVirtual OR MemoryMinFreePhysical
+		// the check done in AssetCompilingManager is against the min of the two :
+		//uint64 AvailableMemory = FMath::Min(MemStats.AvailablePhysical, MemStats.AvailableVirtual);
+		// so for consistency the same check should be done here
+		// you can get that by setting the MemoryMinFreeVirtual and MemoryMinFreePhysical config to be the same
+
+		// AvailableVirtual is actually ullAvailPageFile (commit charge available)
 		if (MemoryMinFreeVirtual > 0 && MemStats.AvailableVirtual < MemoryMinFreeVirtual)
 		{
 			TriggerMessages.Appendf(TEXT("\n  CookSettings.MemoryMinFreeVirtual: Available virtual memory %dMiB is less than %dMiB."),
@@ -4244,8 +4254,24 @@ bool UCookOnTheFlyServer::HasExceededMaxMemory() const
 		}
 	}
 
+	// if MemoryMaxUsed is set, we won't GC until at least that much mem is used
+	// this can be useful if you demand that amount of memory as your min spec
 	if (MemoryMaxUsedVirtual > 0 || MemoryMaxUsedPhysical > 0)
 	{
+		// check validity of trigger :
+		if ( FiredTriggers > 0 )
+		{
+			// if the MaxUsed config exceeds the system memory, it can never be triggered and will prevent any GC :
+			uint64 MaxMaxUsed = FMath::Max(MemoryMaxUsedVirtual,MemoryMaxUsedPhysical);
+			if ( MaxMaxUsed >= MemStats.TotalPhysical )
+			{
+				UE_CALL_ONCE( [&](){
+					UE_LOG(LogCook, Warning, TEXT("Warning MemoryMaxUsed condition is larger than total memory (%dMiB >= %dMiB).  System does not have enough memory to cook this project."),
+				static_cast<uint32>(MaxMaxUsed / 1024 / 1024), static_cast<uint32>(MemStats.TotalPhysical / 1024 / 1024));			
+				} );
+			}
+		}
+
 		++ActiveTriggers;
 		bool bFired = false;
 		if (MemoryMaxUsedVirtual > 0 && MemStats.UsedVirtual >= MemoryMaxUsedVirtual)
