CookOnTheFly: Fix incorrect function used to check whether session platforms is empty; calling the wrong one gives an assertion.
#rnx

diff --git a/Engine/Source/Editor/UnrealEd/Private/CookOnTheFlyServer.cpp b/Engine/Source/Editor/UnrealEd/Private/CookOnTheFlyServer.cpp
--- a/Engine/Source/Editor/UnrealEd/Private/CookOnTheFlyServer.cpp
+++ b/Engine/Source/Editor/UnrealEd/Private/CookOnTheFlyServer.cpp
@@ -720,7 +720,7 @@ void UCookOnTheFlyServer::AddCookOnTheFlyPlatformFromGameThread(ITargetPlatform*
 
 void UCookOnTheFlyServer::StartCookOnTheFlySessionFromGameThread(ITargetPlatform* TargetPlatform)
 {
-	if (PlatformManager->GetSessionPlatforms().Num() == 0)
+	if (PlatformManager->GetNumSessionPlatforms() == 0)
 	{
 		InitializeSession();
 	}
diff --git a/Engine/Source/Editor/UnrealEd/Private/Cooker/CookPlatformManager.cpp b/Engine/Source/Editor/UnrealEd/Private/Cooker/CookPlatformManager.cpp
--- a/Engine/Source/Editor/UnrealEd/Private/Cooker/CookPlatformManager.cpp
+++ b/Engine/Source/Editor/UnrealEd/Private/Cooker/CookPlatformManager.cpp
@@ -101,6 +101,12 @@ namespace Cook
 		return SessionPlatforms;
 	}
 
+	int32 FPlatformManager::GetNumSessionPlatforms() const
+	{
+		checkf(IsSchedulerThread() || IsInPlatformsLock(), TEXT("Access to SessionPlatforms is only legal on non-scheduler threads when inside a ReadLockPlatforms scope."));
+		return SessionPlatforms.Num();
+	}
+
 	bool FPlatformManager::HasSessionPlatform(FPlatformId TargetPlatform) const
 	{
 		const bool bIsSchedulerThread = IsSchedulerThread();
diff --git a/Engine/Source/Editor/UnrealEd/Private/Cooker/CookPlatformManager.h b/Engine/Source/Editor/UnrealEd/Private/Cooker/CookPlatformManager.h
--- a/Engine/Source/Editor/UnrealEd/Private/Cooker/CookPlatformManager.h
+++ b/Engine/Source/Editor/UnrealEd/Private/Cooker/CookPlatformManager.h
@@ -74,6 +74,7 @@ namespace Cook
 
 		/** Returns the set of TargetPlatforms that is active for the CurrentCookByTheBook session or CookOnTheFly request. */
 		const TArray<const ITargetPlatform*>& GetSessionPlatforms() const;
+		int32 GetNumSessionPlatforms() const;
 
 		/** Return whether the given platform is already in the list of platforms for the current session. */
 		bool HasSessionPlatform(FPlatformId TargetPlatform) const;
