EditorPerf: Move statistics into a class FIncrementalVariance

### Notes
Math to handle merging incremental Welford's variance accumulation by [REDACTED_PERSON].
An existing bug was not handling First/Last Start/Finish times correctly. That is now fixed, and that is the only stats/telemetry difference other than a more accurate merge of variance/stddev.
Change to TNumericLimits instead of ballpark "big number" for min/max seed.

### Testing
Ran an EditorPerf pass and data is more correct after spot checked.
All telemetry pushed data is identical.

#rnx

diff --git a/Engine/Source/Editor/UnrealEd/Private/Commandlets/SummarizeTraceCommandlet.cpp b/Engine/Source/Editor/UnrealEd/Private/Commandlets/SummarizeTraceCommandlet.cpp
--- a/Engine/Source/Editor/UnrealEd/Private/Commandlets/SummarizeTraceCommandlet.cpp
+++ b/Engine/Source/Editor/UnrealEd/Private/Commandlets/SummarizeTraceCommandlet.cpp
@@ -9,6 +9,7 @@
 #include "Containers/StringConv.h"
 #include "GenericPlatform/GenericPlatformFile.h"
 #include "HAL/PlatformFileManager.h"
+#include "Math/NumericLimits.h"
 #include "Misc/Crc.h"
 #include "Misc/FileHelper.h"
 #include "Misc/Paths.h"
@@ -684,48 +685,121 @@ public:
 * 
 */
 
-/**
-* Perform an increment of work for Welford's variance, from which we can compute the standard deviation on the last step
-* 
-* @param InSample					The new sample value to operate on
-* @param InCount					The number of samples, inclusive
-* @param InOutMean					The mean of all samples, which will be updated in this function
-* @param InOutVarianceAccumulator	The accumulator for Welford's method
-*/
-static void WelfordsIncrement(const double InSample, const uint64 InCount, double& InOutMean, double& InOutVarianceAccumulator)
+class FIncrementalVariance
 {
-	check(InCount);
-	if (InCount)
+public:
+	FIncrementalVariance()
+		: Count(0)
+		, Mean(0.0)
+		, VarianceAccumulator(0.0)
 	{
-		const double OldMean = InOutMean;
-		InOutMean += ((InSample - InOutMean) / double(InCount));
-		InOutVarianceAccumulator += ((InSample - InOutMean) * (InSample - OldMean));
+
 	}
-}
 
-/**
-* Compute the standard deviation given Welford's accumulator and the overall count
-* 
-* @param InVarianceAccumulator	The accumulator for Welford's method
-* @param InCount				The number of samples, inclusive
-* 
-* @return The standard deviation in sample units
-*/
-static double WelfordsDeviation(const double InVarianceAccumulator, const uint64 InCount)
-{
-	if (InCount > 1)
+	uint64 GetCount() const
 	{
-		// Welford's final step, dependent on sample count
-		double VarianceSecondsSquared = InVarianceAccumulator / double(InCount - 1);
+		return Count;
+	}
 
-		// stddev is sqrt of variance, to restore to units of seconds (vs. seconds squared)
-		return sqrt(VarianceSecondsSquared);
+	double GetMean() const
+	{
+		return Mean;
 	}
-	else
+
+	/**
+	* Compute the variance given Welford's accumulator and the overall count
+	*
+	* @return The variance in sample units squared
+	*/
+	double Variance() const
+	{
+		double Result = 0.0;
+
+		if (Count > 1)
+		{
+			// Welford's final step, dependent on sample count
+			Result = VarianceAccumulator / double(Count - 1);
+		}
+
+		return Result;
+	}
+
+	/**
+	* Compute the standard deviation given Welford's accumulator and the overall count
+	*
+	* @return The standard deviation in sample units
+	*/
+	double Deviation() const
 	{
-		return 0.0;
+		double Result = 0.0;
+
+		if (Count > 1)
+		{
+			// Welford's final step, dependent on sample count
+			double DeviationSqrd = VarianceAccumulator / double(Count - 1);
+
+			// stddev is sqrt of variance, to restore to units (vs. units squared)
+			Result = sqrt(DeviationSqrd);
+		}
+
+		return Result;
 	}
-}
+
+	/**
+	* Perform an increment of work for Welford's variance, from which we can compute variation and standard deviation
+	*
+	* @param InSample	The new sample value to operate on
+	*/
+	void Increment(const double InSample)
+	{
+		Count++;
+		const double OldMean = Mean;
+		Mean += ((InSample - Mean) / double(Count));
+		VarianceAccumulator += ((InSample - Mean) * (InSample - OldMean));
+	}
+
+	/**
+	* Merge with another IncrementalVariance series in progress
+	*
+	* @param Other	The other variance incremented from another mutually exclusive population of analogous data.
+	*/
+	void Merge(const FIncrementalVariance& Other)
+	{
+		// empty other, nothing to do
+		if (Other.Count == 0)
+		{
+			return;
+		}
+
+		// empty this, just copy other
+		if (Count == 0)
+		{
+			Count = Other.Count;
+			Mean = Other.Mean;
+			VarianceAccumulator = Other.VarianceAccumulator;
+			return;
+		}
+
+		const double TotalPopulation = Count + Other.Count;
+		const double MeanDifference = Mean - Other.Mean;
+		const double A = ((Count - 1) * Variance()) + ((Other.Count - 1) * Other.Variance());
+		const double B = (MeanDifference) * (MeanDifference) * (Count * Other.Count / TotalPopulation);
+		const double MergedVariance = (A + B) / (TotalPopulation - 1);
+
+		const uint64 NewCount = Count + Other.Count;
+		const double NewMean = ((Mean * double(Count)) + (Other.Mean * double(Other.Count))) / double(NewCount);
+		const double NewVarianceAccumulator = MergedVariance * (NewCount - 1);
+
+		Count = NewCount;
+		Mean = NewMean;
+		VarianceAccumulator = NewVarianceAccumulator;
+	}
+
+private:
+	uint64 Count;
+	double Mean;
+	double VarianceAccumulator;
+};
 
 //
 // FSummarizeCpuAnalyzer - Helpers & Analyzers for Cpu channel events
@@ -734,7 +808,7 @@ static double WelfordsDeviation(const double InVarianceAccumulator, const uint64
 struct FSummarizeScope
 {
 	FString Name;
-	uint64 Count = 0;
+	FIncrementalVariance DurationVariance;
 	double TotalDurationSeconds = 0.0;
 
 	double FirstStartSeconds = 0.0;
@@ -745,20 +819,18 @@ struct FSummarizeScope
 	double LastFinishSeconds = 0.0;
 	double LastDurationSeconds = 0.0;
 
-	double MinDurationSeconds = 1e10;
-	double MaxDurationSeconds = -1e10;
-	double MeanDurationSeconds = 0.0;
-	double VarianceAccumulator = 0.0; // Accumulator for Welford's
+	double MinDurationSeconds = TNumericLimits<double>::Max();
+	double MaxDurationSeconds = TNumericLimits<double>::Min();
 
 	void AddDuration(double StartSeconds, double FinishSeconds)
 	{
-		Count += 1;
-
 		// compute the duration
 		double DurationSeconds = FinishSeconds - StartSeconds;
 
+		DurationVariance.Increment(DurationSeconds);
+
 		// only set first for the first sample
-		if (Count == 1)
+		if (DurationVariance.GetCount() == 1)
 		{
 			FirstStartSeconds = StartSeconds;
 			FirstFinishSeconds = FinishSeconds;
@@ -773,19 +845,27 @@ struct FSummarizeScope
 		TotalDurationSeconds += DurationSeconds;
 		MinDurationSeconds = FMath::Min(MinDurationSeconds, DurationSeconds);
 		MaxDurationSeconds = FMath::Max(MaxDurationSeconds, DurationSeconds);
+	}
+	
+	uint64 GetCount() const
+	{
+		return DurationVariance.GetCount();
+	}
 
-		// update variance for stddev later
-		WelfordsIncrement(DurationSeconds, Count, MeanDurationSeconds, VarianceAccumulator);
+	double GetMeanDurationSeconds() const
+	{
+		return DurationVariance.GetMean();
 	}
 
 	double GetDeviationDurationSeconds() const
 	{
-		return WelfordsDeviation(VarianceAccumulator, Count);
+		return DurationVariance.Deviation();
 	}
 
 	double GetCountPerSecond() const
 	{
 		double CountPerSecond = 0.0;
+		const uint64 Count = DurationVariance.GetCount();
 		if (Count)
 		{
 			CountPerSecond = Count / (LastFinishSeconds - FirstStartSeconds);
@@ -793,25 +873,28 @@ struct FSummarizeScope
 		return CountPerSecond;
 	}
 
-	void Merge(const FSummarizeScope& Scope)
+	void Merge(const FSummarizeScope& Other)
 	{
-		check(Name == Scope.Name);
-		TotalDurationSeconds += Scope.TotalDurationSeconds;
-		MinDurationSeconds = FMath::Min(MinDurationSeconds, Scope.MinDurationSeconds);
-		MaxDurationSeconds = FMath::Max(MaxDurationSeconds, Scope.MaxDurationSeconds);
+		check(Name == Other.Name);
+		DurationVariance.Merge(Other.DurationVariance);
 
-		const uint64 NewCount = Count + Scope.Count;
-		if (NewCount)
+		if (FirstStartSeconds > Other.FirstStartSeconds)
 		{
-			// perform weighted average, weighted on each sub-population's count
-			MeanDurationSeconds = ((MeanDurationSeconds * double(Count)) + (Scope.MeanDurationSeconds * double(Scope.Count))) / double(NewCount);
-
-			// perform weighted average, weighted on each sub-population's count
-			//  this is likely not the perfect union of two Welford variance accumuations due to them being squares of variance offsets
-			VarianceAccumulator = ((VarianceAccumulator * double(Count)) + (Scope.VarianceAccumulator * double(Scope.Count))) / double(NewCount);
+			FirstStartSeconds = Other.FirstStartSeconds;
+			FirstFinishSeconds = Other.FirstFinishSeconds;
+			FirstDurationSeconds = Other.FirstDurationSeconds;
 		}
-
-		Count = NewCount;
+		
+		if (LastStartSeconds < Other.LastStartSeconds)
+		{
+			LastStartSeconds = Other.LastStartSeconds;
+			LastFinishSeconds = Other.LastFinishSeconds;
+			LastDurationSeconds = Other.LastDurationSeconds;
+		}
+		
+		TotalDurationSeconds += Other.TotalDurationSeconds;
+		MinDurationSeconds = FMath::Min(MinDurationSeconds, Other.MinDurationSeconds);
+		MaxDurationSeconds = FMath::Max(MaxDurationSeconds, Other.MaxDurationSeconds);
 	}
 
 	FString GetValue(const FStringView& Statistic) const
@@ -822,7 +905,7 @@ struct FSummarizeScope
 		}
 		else if (Statistic == TEXT("Count"))
 		{
-			return FString::Printf(TEXT("%llu"), Count);
+			return FString::Printf(TEXT("%llu"), DurationVariance.GetCount());
 		}
 		else if (Statistic == TEXT("TotalDurationSeconds"))
 		{
@@ -862,11 +945,11 @@ struct FSummarizeScope
 		}
 		else if (Statistic == TEXT("MeanDurationSeconds"))
 		{
-			return FString::Printf(TEXT("%f"), MeanDurationSeconds);
+			return FString::Printf(TEXT("%f"), DurationVariance.GetMean());
 		}
 		else if (Statistic == TEXT("DeviationDurationSeconds"))
 		{
-			return FString::Printf(TEXT("%f"), GetDeviationDurationSeconds());
+			return FString::Printf(TEXT("%f"), DurationVariance.Deviation());
 		}
 		else if (Statistic == TEXT("CountPerSecond"))
 		{
@@ -1133,9 +1216,7 @@ struct FSummarizeCounter
 	FValue Minimum;
 	FValue Maximum;
 
-	uint64 Count = 0;
-	double Mean = 0.0;
-	double VarianceAccumulator = 0.0; // Accumulator for Welford's
+	FIncrementalVariance Variance;
 
 	double FirstSeconds = 0.0;
 	double LastSeconds = 0.0;
@@ -1167,9 +1248,9 @@ struct FSummarizeCounter
 		ensure(Type == TraceCounterType_Int);
 		if (Type == TraceCounterType_Int)
 		{
-			Count += 1;
+			Variance.Increment(double(InValue));
 
-			if (Count == 1)
+			if (Variance.GetCount() == 1)
 			{
 				First.Int = InValue;
 				FirstSeconds = InTimestamp;
@@ -1180,9 +1261,6 @@ struct FSummarizeCounter
 
 			Minimum.Int = FMath::Min(Minimum.Int, InValue);
 			Maximum.Int = FMath::Max(Maximum.Int, InValue);
-
-			// update variance for stddev later
-			WelfordsIncrement(double(InValue), Count, Mean, VarianceAccumulator);
 		}
 	}
 
@@ -1191,9 +1269,9 @@ struct FSummarizeCounter
 		ensure(Type == TraceCounterType_Float);
 		if (Type == TraceCounterType_Float)
 		{
-			Count += 1;
+			Variance.Increment(InValue);
 
-			if (Count == 1)
+			if (Variance.GetCount() == 1)
 			{
 				First.Float = InValue;
 				FirstSeconds = InTimestamp;
@@ -1204,20 +1282,28 @@ struct FSummarizeCounter
 
 			Minimum.Float = FMath::Min(Minimum.Float, InValue);
 			Maximum.Float = FMath::Max(Maximum.Float, InValue);
-
-			// update variance for stddev later
-			WelfordsIncrement(InValue, Count, Mean, VarianceAccumulator);
 		}
 	}
+	
+	uint64 GetCount() const
+	{
+		return Variance.GetCount();
+	}
+
+	double GetMean() const
+	{
+		return Variance.GetMean();
+	}
 
 	double GetDeviation() const
 	{
-		return WelfordsDeviation(VarianceAccumulator, Count);
+		return Variance.Deviation();
 	}
 
 	double GetCountPerSecond() const
 	{
 		double CountPerSecond = 0.0;
+		const uint64 Count = Variance.GetCount();
 		if (Count)
 		{
 			if (Count == 1)
@@ -1255,7 +1341,7 @@ struct FSummarizeCounter
 		}
 		else if (Statistic == TEXT("Count"))
 		{
-			return FString::Printf(TEXT("%llu"), Count);
+			return FString::Printf(TEXT("%llu"), Variance.GetCount());
 		}
 		else if (Statistic == TEXT("First"))
 		{
@@ -1275,19 +1361,19 @@ struct FSummarizeCounter
 		}
 		else if (Statistic == TEXT("Minimum"))
 		{
-			return PrintValue(Count ? Minimum : Zero);
+			return PrintValue(Variance.GetCount() ? Minimum : Zero);
 		}
 		else if (Statistic == TEXT("Maximum"))
 		{
-			return PrintValue(Count ? Maximum : Zero);
+			return PrintValue(Variance.GetCount() ? Maximum : Zero);
 		}
 		else if (Statistic == TEXT("Mean"))
 		{
-			return FString::Printf(TEXT("%f"), Mean);
+			return FString::Printf(TEXT("%f"), Variance.GetMean());
 		}
 		else if (Statistic == TEXT("Deviation"))
 		{
-			return FString::Printf(TEXT("%f"), GetDeviation());
+			return FString::Printf(TEXT("%f"), Variance.Deviation());
 		}
 		else if (Statistic == TEXT("CountPerSecond"))
 		{
@@ -2034,7 +2120,7 @@ int32 USummarizeTraceCommandlet::Main(const FString& CmdLineParams)
 			return;
 		}
 
-		if (Scope.Count == 0)
+		if (Scope.GetCount() == 0)
 		{
 			return;
 		}
@@ -2100,12 +2186,12 @@ int32 USummarizeTraceCommandlet::Main(const FString& CmdLineParams)
 
 			// note newline is at the front of every data line to prevent final extraneous newline, per customary for csv
 			WriteUTF8(CsvHandle, FString::Printf(TEXT("\n%s,%llu,%f,%f,%f,%f,%f,%f,%f,%f,%f,%f,%f,%f,"),
-				*Scope.Name, Scope.Count, Scope.GetCountPerSecond(),
+				*Scope.Name, Scope.GetCount(), Scope.GetCountPerSecond(),
 				Scope.TotalDurationSeconds,
 				Scope.FirstStartSeconds, Scope.FirstFinishSeconds, Scope.FirstDurationSeconds,
 				Scope.LastStartSeconds, Scope.LastFinishSeconds, Scope.LastDurationSeconds,
 				Scope.MinDurationSeconds, Scope.MaxDurationSeconds,
-				Scope.MeanDurationSeconds, Scope.GetDeviationDurationSeconds()));
+				Scope.GetMeanDurationSeconds(), Scope.GetDeviationDurationSeconds()));
 		}
 		CsvHandle->Flush();
 		delete CsvHandle;
@@ -2134,11 +2220,11 @@ int32 USummarizeTraceCommandlet::Main(const FString& CmdLineParams)
 
 			// note newline is at the front of every data line to prevent final extraneous newline, per customary for csv
 			WriteUTF8(CsvHandle, FString::Printf(TEXT("\n%s,%llu,%f,%s,%f,%s,%f,%s,%s,%f,%f,"),
-				*Counter.Value.Name, Counter.Value.Count, Counter.Value.GetCountPerSecond(),
+				*Counter.Value.Name, Counter.Value.GetCount(), Counter.Value.GetCountPerSecond(),
 				*Counter.Value.GetValue(TEXT("First")), Counter.Value.FirstSeconds,
 				*Counter.Value.GetValue(TEXT("Last")), Counter.Value.LastSeconds,
 				*Counter.Value.GetValue(TEXT("Minimum")), *Counter.Value.GetValue(TEXT("Maximum")),
-				Counter.Value.Mean, Counter.Value.GetDeviation()));
+				Counter.Value.GetMean(), Counter.Value.GetDeviation()));
 		}
 		CsvHandle->Flush();
 		delete CsvHandle;
