Fix unity build

diff --git a/Engine/Source/Editor/UnrealEd/Private/HierarchicalLOD.cpp b/Engine/Source/Editor/UnrealEd/Private/HierarchicalLOD.cpp
--- a/Engine/Source/Editor/UnrealEd/Private/HierarchicalLOD.cpp
+++ b/Engine/Source/Editor/UnrealEd/Private/HierarchicalLOD.cpp
@@ -50,6 +50,8 @@
 #include "Engine/LevelStreaming.h"
 #include "AssetCompilingManager.h"
 
+#include "Materials/Material.h"
+
 DEFINE_LOG_CATEGORY_STATIC(LogLODGenerator, Log, All);
 
 #define LOCTEXT_NAMESPACE "HierarchicalLOD"
