Deprecate the old InputKey method on the EditorViewportClient in favor of the much newer version that takes FInputKeyEventArgs struct. This was supposed to be deprecated in 4.21, but was commented out in a different stream with 4539769. For now, I've left the deprecation commented out in the Unreal Client.h so that we can fix up all the call sites of this function. A jira for enabling this deprecation can be found at UE-156743

Also update the InputAxis function to have the new version with an FInputDeviceId struct instead of an int32 Controllerid

diff --git a/Engine/Source/Editor/UnrealEd/Private/EditorViewportClient.cpp b/Engine/Source/Editor/UnrealEd/Private/EditorViewportClient.cpp
--- a/Engine/Source/Editor/UnrealEd/Private/EditorViewportClient.cpp
+++ b/Engine/Source/Editor/UnrealEd/Private/EditorViewportClient.cpp
@@ -62,7 +62,7 @@
 #include "Engine/World.h"
 #include "ProfilingDebugging/MiscTrace.h"
 #include "ProfilingDebugging/TraceScreenshot.h"
-
+#include "GenericPlatform/GenericPlatformInputDeviceMapper.h"
 #include "CustomEditorStaticScreenPercentage.h"
 #include "IImageWrapperModule.h"
 
@@ -2827,12 +2827,31 @@ bool FEditorViewportClient::ShouldScaleCameraSpeedByDistance() const
 	return GetDefault<ULevelEditorViewportSettings>()->bUseDistanceScaledCameraSpeed;
 }
 
-bool FEditorViewportClient::InputKey(FViewport* InViewport, int32 ControllerId, FKey Key, EInputEvent Event, float/*AmountDepressed*/, bool/*Gamepad*/)
+PRAGMA_DISABLE_DEPRECATION_WARNINGS
+bool FEditorViewportClient::InputKey(FViewport* InViewport, int32 ControllerId, FKey Key, EInputEvent Event, float AmountDepressed, bool bGamepad)
+{
+	FInputKeyEventArgs Args(InViewport, ControllerId, Key, Event);
+	Args.AmountDepressed = AmountDepressed;
+
+	return Internal_InputKey(Args);
+}
+PRAGMA_ENABLE_DEPRECATION_WARNINGS
+
+bool FEditorViewportClient::InputKey(const FInputKeyEventArgs& EventArgs)
+{
+	return Internal_InputKey(EventArgs);
+}
+
+bool FEditorViewportClient::Internal_InputKey(const FInputKeyEventArgs& EventArgs)
 {
 	if (bDisableInput)
 	{
 		return true;
 	}
+	
+	FKey Key = EventArgs.Key;
+	EInputEvent Event = EventArgs.Event;
+	FViewport* InViewport = EventArgs.Viewport;
 
 	// Let the current mode have a look at the input before reacting to it. 
 	// Note that bIsTracking tells us whether the viewport client is capturing mouse behavior to fly around,
@@ -2847,8 +2866,8 @@ bool FEditorViewportClient::InputKey(FViewport* InViewport, int32 ControllerId,
 		return true;
 	}
 
-	FInputEventState InputState(InViewport, Key, Event);
-
+	FInputEventState InputState(EventArgs.Viewport, EventArgs.Key, EventArgs.Event);
+	
 	bool bHandled = false;
 
 	if ((IsOrtho() || InputState.IsAltButtonPressed()) && (Key == EKeys::Left || Key == EKeys::Right || Key == EKeys::Up || Key == EKeys::Down))
@@ -2882,7 +2901,7 @@ bool FEditorViewportClient::InputKey(FViewport* InViewport, int32 ControllerId,
 	const int32	HitX = InViewport->GetMouseX();
 	const int32	HitY = InViewport->GetMouseY();
 
-	FCachedJoystickState* JoystickState = GetJoystickState(ControllerId);
+	FCachedJoystickState* JoystickState = GetJoystickState(EventArgs.InputDevice.GetId());
 	if (JoystickState)
 	{
 		JoystickState->KeyEventValues.Add(Key, Event);
@@ -4384,15 +4403,33 @@ FVector FEditorViewportClient::TranslateDelta( FKey InKey, float InDelta, bool I
 	return vec;
 }
 
+PRAGMA_DISABLE_DEPRECATION_WARNINGS
 bool FEditorViewportClient::InputAxis(FViewport* InViewport, int32 ControllerId, FKey Key, float Delta, float DeltaTime, int32 NumSamples, bool bGamepad)
+{
+	FInputDeviceId DeviceID = INPUTDEVICEID_NONE;
+	FPlatformUserId UserId = FGenericPlatformMisc::GetPlatformUserForUserIndex(ControllerId);
+	IPlatformInputDeviceMapper::Get().RemapControllerIdToPlatformUserAndDevice(ControllerId, UserId, DeviceID);
+	
+	return Internal_InputAxis(InViewport, DeviceID, Key, Delta, DeltaTime, NumSamples, bGamepad);
+}
+PRAGMA_ENABLE_DEPRECATION_WARNINGS
+
+bool FEditorViewportClient::InputAxis(FViewport* InViewport, FInputDeviceId DeviceID, FKey Key, float Delta, float DeltaTime, int32 NumSamples, bool bGamepad)
+{
+	return Internal_InputAxis(InViewport, DeviceID, Key, Delta, DeltaTime, NumSamples, bGamepad);
+}
+
+bool FEditorViewportClient::Internal_InputAxis(FViewport* InViewport, FInputDeviceId DeviceID, FKey Key, float Delta, float DeltaTime, int32 NumSamples, bool bGamepad)
 {
 	if (bDisableInput)
 	{
 		return true;
 	}
 
+	const FPlatformUserId UserId = IPlatformInputDeviceMapper::Get().GetUserForInputDevice(DeviceID);
+	
 	// Let the current mode have a look at the input before reacting to it.
-	if (ModeTools->InputAxis(this, Viewport, ControllerId, Key, Delta, DeltaTime))
+	if (ModeTools->InputAxis(this, Viewport, FGenericPlatformMisc::GetUserIndexForPlatformUser(UserId), Key, Delta, DeltaTime))
 	{
 		return true;
 	}
@@ -4423,7 +4460,7 @@ bool FEditorViewportClient::InputAxis(FViewport* InViewport, int32 ControllerId,
 	else
 	{
 		/**Save off axis commands for future camera work*/
-		FCachedJoystickState* JoystickState = GetJoystickState(ControllerId);
+		FCachedJoystickState* JoystickState = GetJoystickState(DeviceID.GetId());
 		if (JoystickState)
 		{
 			JoystickState->AxisDeltaValues.Add(Key, Delta);
diff --git a/Engine/Source/Editor/UnrealEd/Private/LevelEditorViewport.cpp b/Engine/Source/Editor/UnrealEd/Private/LevelEditorViewport.cpp
--- a/Engine/Source/Editor/UnrealEd/Private/LevelEditorViewport.cpp
+++ b/Engine/Source/Editor/UnrealEd/Private/LevelEditorViewport.cpp
@@ -95,6 +95,7 @@
 #include "LevelEditorDragDropHandler.h"
 #include "UnrealWidget.h"
 #include "EdModeInteractiveToolsContext.h"
+#include "GenericPlatform/GenericPlatformInputDeviceMapper.h"
 
 DEFINE_LOG_CATEGORY(LogEditorViewport);
 
@@ -2811,7 +2812,7 @@ void FLevelEditorViewportClient::SetLastKeyViewport()
 	}
 }
 
-bool FLevelEditorViewportClient::InputKey(FViewport* InViewport, int32 ControllerId, FKey Key, EInputEvent Event, float AmountDepressed, bool bGamepad)
+bool FLevelEditorViewportClient::InputKey(const FInputKeyEventArgs& InEventArgs)
 {
 	if (bDisableInput)
 	{
@@ -2819,16 +2820,16 @@ bool FLevelEditorViewportClient::InputKey(FViewport* InViewport, int32 Controlle
 	}
 
 	
-	const int32	HitX = InViewport->GetMouseX();
-	const int32	HitY = InViewport->GetMouseY();
+	const int32	HitX = InEventArgs.Viewport->GetMouseX();
+	const int32	HitY = InEventArgs.Viewport->GetMouseY();
 
-	FInputEventState InputState( InViewport, Key, Event );
+	FInputEventState InputState( InEventArgs.Viewport, InEventArgs.Key, InEventArgs.Event );
 
 	SetLastKeyViewport();
 
 	// Compute a view.
 	FSceneViewFamilyContext ViewFamily(FSceneViewFamily::ConstructionValues(
-		InViewport,
+		InEventArgs.Viewport,
 		GetScene(),
 		EngineShowFlags )
 		.SetRealtimeUpdate( IsRealtime() ) );
@@ -2847,7 +2848,7 @@ bool FLevelEditorViewportClient::InputKey(FViewport* InViewport, int32 Controlle
 		FSnappingUtils::SnapPointToGrid(GEditor->ClickLocation, FVector::ZeroVector);
 	}
 
-	if (GUnrealEd->ComponentVisManager.HandleInputKey(this, InViewport, Key, Event))
+	if (GUnrealEd->ComponentVisManager.HandleInputKey(this, InEventArgs.Viewport, InEventArgs.Key, InEventArgs.Event))
 	{
 		return true;
 	}
@@ -2870,7 +2871,7 @@ bool FLevelEditorViewportClient::InputKey(FViewport* InViewport, int32 Controlle
 	};
 
 
-	bool bCmdCtrlLPressed = (InputState.IsCommandButtonPressed() || InputState.IsCtrlButtonPressed()) && Key == EKeys::L;
+	bool bCmdCtrlLPressed = (InputState.IsCommandButtonPressed() || InputState.IsCtrlButtonPressed()) && InEventArgs.Key == EKeys::L;
 	if (bCmdCtrlLPressed && InputState.IsShiftButtonPressed())
 	{
 		ProcessAtmosphericLightShortcut(1, bUserIsControllingAtmosphericLight1);
@@ -2895,29 +2896,29 @@ bool FLevelEditorViewportClient::InputKey(FViewport* InViewport, int32 Controlle
 	bUserIsControllingAtmosphericLight1 = false;
 
 	UEditorWorldExtensionCollection& EditorWorldExtensionCollection = *GEditor->GetEditorWorldExtensionsManager()->GetEditorWorldExtensions(GetWorld());
-	if (EditorWorldExtensionCollection.InputKey(this, Viewport, Key, Event))
+	if (EditorWorldExtensionCollection.InputKey(this, Viewport, InEventArgs.Key, InEventArgs.Event))
 	{
 		return true;
 	}
 
-	bool bHandled = FEditorViewportClient::InputKey(InViewport,ControllerId,Key,Event,AmountDepressed,bGamepad);
+	bool bHandled = FEditorViewportClient::InputKey(InEventArgs);
 
 	// Handle input for the player height preview mode. 
-	if (!InputState.IsMouseButtonEvent() && CommandAcceptsInput(*this, Key, GetLevelViewportCommands().EnablePreviewMesh))
+	if (!InputState.IsMouseButtonEvent() && CommandAcceptsInput(*this, InEventArgs.Key, GetLevelViewportCommands().EnablePreviewMesh))
 	{
 		// Holding down the backslash buttons turns on the mode. 
-		if (Event == IE_Pressed)
+		if (InEventArgs.Event == IE_Pressed)
 		{
 			GEditor->SetPreviewMeshMode(true);
 
 			// If shift down, cycle between the preview meshes
-			if (CommandAcceptsInput(*this, Key, GetLevelViewportCommands().CyclePreviewMesh))
+			if (CommandAcceptsInput(*this, InEventArgs.Key, GetLevelViewportCommands().CyclePreviewMesh))
 			{
 				GEditor->CyclePreviewMesh();
 			}
 		}
 		// Releasing backslash turns off the mode. 
-		else if (Event == IE_Released)
+		else if (InEventArgs.Event == IE_Released)
 		{
 			GEditor->SetPreviewMeshMode(false);
 		}
@@ -3338,7 +3339,7 @@ private:
 	FLevelEditorViewportClient* PrevCurrentLevelEditingViewportClient;
 };
 
-bool FLevelEditorViewportClient::InputAxis(FViewport* InViewport, int32 ControllerId, FKey Key, float Delta, float DeltaTime, int32 NumSamples, bool bGamepad)
+bool FLevelEditorViewportClient::InputAxis(FViewport* InViewport, FInputDeviceId DeviceId, FKey Key, float Delta, float DeltaTime, int32 NumSamples, bool bGamepad)
 {
 	if (bDisableInput)
 	{
@@ -3349,7 +3350,7 @@ bool FLevelEditorViewportClient::InputAxis(FViewport* InViewport, int32 Controll
 
 	FScopedSetCurrentViewportClient( this );
 
-	return FEditorViewportClient::InputAxis(InViewport, ControllerId, Key, Delta, DeltaTime, NumSamples, bGamepad);
+	return FEditorViewportClient::InputAxis(InViewport, DeviceId, Key, Delta, DeltaTime, NumSamples, bGamepad);
 }
 
 
diff --git a/Engine/Source/Editor/UnrealEd/Public/EditorViewportClient.h b/Engine/Source/Editor/UnrealEd/Public/EditorViewportClient.h
--- a/Engine/Source/Editor/UnrealEd/Public/EditorViewportClient.h
+++ b/Engine/Source/Editor/UnrealEd/Public/EditorViewportClient.h
@@ -561,8 +561,12 @@ public:
 	virtual bool ProcessScreenShots(FViewport* Viewport) override;
 	virtual void RedrawRequested(FViewport* Viewport) override;
 	virtual void RequestInvalidateHitProxy(FViewport* Viewport) override;
-	virtual bool InputKey(FViewport* Viewport, int32 ControllerId, FKey Key, EInputEvent Event, float AmountDepressed = 1.f, bool bGamepad=false) override;
-	virtual bool InputAxis(FViewport* Viewport, int32 ControllerId, FKey Key, float Delta, float DeltaTime, int32 NumSamples=1, bool bGamepad=false) override;
+	//UE_DEPRECATED(5.1, "This version of InputKey is deprecated. Please use the version that takes FInputKeyEventArgs instead.")
+	virtual bool InputKey(FViewport* InViewport, int32 ControllerId, FKey Key, EInputEvent Event, float AmountDepressed = 1.f, bool bGamepad=false) override;
+	virtual bool InputKey(const FInputKeyEventArgs& EventArgs) override;
+	//UE_DEPRECATED(5.1, "This version of InputAxis is deprecated. Please use the version that takes a DeviceId instead.")
+	virtual bool InputAxis(FViewport* InViewport, int32 ControllerId, FKey Key, float Delta, float DeltaTime, int32 NumSamples=1, bool bGamepad=false) override;
+	virtual bool InputAxis(FViewport* Viewport, FInputDeviceId DeviceID, FKey Key, float Delta, float DeltaTime, int32 NumSamples=1, bool bGamepad=false) override;
 	virtual bool InputGesture(FViewport* Viewport, EGestureEvent GestureType, const FVector2D& GestureDelta, bool bIsDirectionInvertedFromDevice) override;
 	virtual void ReceivedFocus(FViewport* Viewport) override;
 	virtual void MouseEnter(FViewport* Viewport,int32 x, int32 y) override;
@@ -579,6 +583,13 @@ public:
 	virtual void SetEnabledStats(const TArray<FString>& InEnabledStats) override;
 	virtual bool IsStatEnabled(const FString& InName) const override;
 
+protected:
+
+	virtual bool Internal_InputKey(const FInputKeyEventArgs& EventArgs);
+	virtual bool Internal_InputAxis(FViewport* Viewport, FInputDeviceId DeviceID, FKey Key, float Delta, float DeltaTime, int32 NumSamples=1, bool bGamepad=false);
+
+public:
+	
 	/** FGCObject interface */
 	virtual void AddReferencedObjects( FReferenceCollector& Collector ) override;
 	virtual FString GetReferencerName() const override;
diff --git a/Engine/Source/Editor/UnrealEd/Public/LevelEditorViewport.h b/Engine/Source/Editor/UnrealEd/Public/LevelEditorViewport.h
--- a/Engine/Source/Editor/UnrealEd/Public/LevelEditorViewport.h
+++ b/Engine/Source/Editor/UnrealEd/Public/LevelEditorViewport.h
@@ -195,8 +195,8 @@ public:
 	////////////////////////////
 	// FEditorViewportClient interface
 	virtual void DrawCanvas( FViewport& InViewport, FSceneView& View, FCanvas& Canvas ) override;
-	virtual bool InputKey(FViewport* Viewport, int32 ControllerId, FKey Key, EInputEvent Event, float AmountDepressed = 1.f, bool bGamepad=false) override;
-	virtual bool InputAxis(FViewport* Viewport, int32 ControllerId, FKey Key, float Delta, float DeltaTime, int32 NumSamples=1, bool bGamepad=false) override;
+	virtual bool InputKey(const FInputKeyEventArgs& InEventArgs) override;
+	virtual bool InputAxis(FViewport* Viewport, FInputDeviceId DeviceId, FKey Key, float Delta, float DeltaTime, int32 NumSamples=1, bool bGamepad=false) override;
 	virtual EMouseCursor::Type GetCursor(FViewport* Viewport,int32 X,int32 Y) override;
 	virtual void CapturedMouseMove(FViewport* InViewport, int32 InMouseX, int32 InMouseY) override;
 	virtual void MouseMove(FViewport* InViewport, int32 x, int32 y) override;
