Cooker: Call BeginCacheForCookedPlatformData on objects created by earlier calls to BeginCacheForCookedPlatformData.
#rnx

diff --git a/Engine/Source/Editor/UnrealEd/Private/CookOnTheFlyServer.cpp b/Engine/Source/Editor/UnrealEd/Private/CookOnTheFlyServer.cpp
--- a/Engine/Source/Editor/UnrealEd/Private/CookOnTheFlyServer.cpp
+++ b/Engine/Source/Editor/UnrealEd/Private/CookOnTheFlyServer.cpp
@@ -2868,6 +2868,19 @@ UE::Cook::EPollStatus UCookOnTheFlyServer::PrepareSaveGeneratedPackage(UE::Cook:
 		{
 			return EPollStatus::Incomplete;
 		}
+		bool bFoundNewObjects;
+		EPollStatus Result = Info.RefreshPackageObjects(Generator, PackageData.GetPackage(), bFoundNewObjects,
+			FCookGenerationInfo::ESaveState::BeginCacheObjectsToMove);
+		if (Result != EPollStatus::Success)
+		{
+			return Result;
+		}
+		if (bFoundNewObjects)
+		{
+			// Call this function recursively to reexecute CallBeginCacheOnObjects in BeginCacheObjectsToMove.
+			// Note that RefreshPackageObjects checked for too many recursive calls and ErrorExited if so.
+			return PrepareSaveGeneratedPackage(Generator, PackageData, Timer, bPrecaching);
+		}
 		Info.SetSaveStateComplete(FCookGenerationInfo::ESaveState::FinishCacheObjectsToMove);
 	}
 
@@ -2913,42 +2926,25 @@ UE::Cook::EPollStatus UCookOnTheFlyServer::PrepareSaveGeneratedPackage(UE::Cook:
 		{
 			return EPollStatus::Incomplete;
 		}
-		Info.SetSaveStateComplete(FCookGenerationInfo::ESaveState::FinishCachePostMove);
-	}
-	check(Info.GetSaveState() == FCookGenerationInfo::ESaveState::ReadyForSave);
-
-	return EPollStatus::Success;
-}
-
-/** Find all subobjects of all objects in RootObjects and add them to the end of InObjects */
-static void ValidateAndExtendObjectsToMove(TArray<UObject*>& InObjects, UE::Cook::FGeneratorPackage& Generator, UE::Cook::FCookGenerationInfo& Info)
-{
-	TArray<UObject*> ObjectsInOuter;
-	int32 OriginalNum = InObjects.Num();
-	for (int32 Index = 0; Index < OriginalNum; ++Index)
-	{
-		ObjectsInOuter.Reset();
-		UObject* OriginalObject = InObjects[Index];
-		if (!IsValid(OriginalObject))
+		bool bFoundNewObjects;
+		EPollStatus Result = Info.RefreshPackageObjects(Generator, PackageData.GetPackage(), bFoundNewObjects,
+			FCookGenerationInfo::ESaveState::BeginCachePostMove);
+		if (Result != EPollStatus::Success)
 		{
-			UE_LOG(LogCook, Warning, TEXT("CookPackageSplitter found non-valid object %s returned from %s on Splitter %s%s. Ignoring it."),
-				OriginalObject ? *OriginalObject->GetFullName() : TEXT("<null>"),
-				Info.IsGenerator() ? TEXT("PopulateGeneratorPackage") : TEXT("PopulateGeneratedPackage"),
-				*Generator.GetSplitDataObjectName().ToString(),
-				Info.IsGenerator() ? TEXT("") : *FString::Printf(TEXT(", Package %s"), *Info.PackageData->GetPackageName().ToString()));
-			--OriginalNum;
-			--Index;
-			InObjects.RemoveAt(Index);
-			continue;
+			return Result;
 		}
-		GetObjectsWithOuter(OriginalObject, ObjectsInOuter, true /* bIncludeNestedObjects */, RF_NoFlags, EInternalObjectFlags::Garbage);
-		InObjects.Reserve(InObjects.Num() + ObjectsInOuter.Num());
-		for (UObject* Object : ObjectsInOuter)
+		if (bFoundNewObjects)
 		{
-			check(IsValid(Object));
-			InObjects.Add(Object);
+			// Call this function recursively to reexecute CallBeginCacheOnObjects in BeginCachePostMove
+			// Note that RefreshPackageObjects checked for too many recursive calls and ErrorExited if so.
+			return PrepareSaveGeneratedPackage(Generator, PackageData, Timer, bPrecaching);
 		}
+
+		Info.SetSaveStateComplete(FCookGenerationInfo::ESaveState::FinishCachePostMove);
 	}
+	check(Info.GetSaveState() == FCookGenerationInfo::ESaveState::ReadyForSave);
+
+	return EPollStatus::Success;
 }
 
 UE::Cook::EPollStatus UCookOnTheFlyServer::BeginCacheObjectsToMove(UE::Cook::FGeneratorPackage& Generator,
@@ -2994,8 +2990,7 @@ UE::Cook::EPollStatus UCookOnTheFlyServer::BeginCacheObjectsToMove(UE::Cook::FGe
 			return EPollStatus::Error;
 		}
 
-		ValidateAndExtendObjectsToMove(ObjectsToMove, Generator, Info);
-		Info.TakeOverCachedObjectsAndAddNew(PackageData.GetCachedObjectsInOuter(), ObjectsToMove);
+		Info.TakeOverCachedObjectsAndAddMoved(Generator, PackageData.GetCachedObjectsInOuter(), ObjectsToMove);
 		Info.SetSaveStateComplete(FCookGenerationInfo::ESaveState::CallObjectsToMove);
 	}
 
@@ -3088,7 +3083,13 @@ UE::Cook::EPollStatus UCookOnTheFlyServer::BeginCachePostMove(UE::Cook::FGenerat
 
 	if (Info.GetSaveState() <= FCookGenerationInfo::ESaveState::CallGetPostMoveObjects)
 	{
-		Info.RefreshPackageObjects(Package);
+		bool bFoundNewObjects;
+		EPollStatus Result = Info.RefreshPackageObjects(Generator, Package, bFoundNewObjects,
+			FCookGenerationInfo::ESaveState::Last);
+		if (Result != EPollStatus::Success)
+		{
+			return Result;
+		}
 		Info.SetSaveStateComplete(FCookGenerationInfo::ESaveState::CallGetPostMoveObjects);
 	}
 
@@ -3344,6 +3345,18 @@ UE::Cook::EPollStatus UCookOnTheFlyServer::PrepareSaveInternal(UE::Cook::FPackag
 		{
 			return EPollStatus::Incomplete;
 		}
+		bool bFoundNewObjects;
+		EPollStatus Result = PackageData.RefreshObjectCache(bFoundNewObjects);
+		if (Result != EPollStatus::Success)
+		{
+			return Result;
+		}
+		if (bFoundNewObjects)
+		{
+			// Call this function recursively to reexecute CallBeginCacheOnObjects.
+			// Note that RefreshObjectCache checked for too many recursive calls and ErrorExited if so.
+			return PrepareSaveInternal(PackageData, Timer, bPrecaching);
+		}
 	}
 
 	check(PackageData.GetNumPendingCookedPlatformData() == 0);
diff --git a/Engine/Source/Editor/UnrealEd/Private/Cooker/CookPackageData.cpp b/Engine/Source/Editor/UnrealEd/Private/Cooker/CookPackageData.cpp
--- a/Engine/Source/Editor/UnrealEd/Private/Cooker/CookPackageData.cpp
+++ b/Engine/Source/Editor/UnrealEd/Private/Cooker/CookPackageData.cpp
@@ -5,6 +5,7 @@
 #include "Algo/AnyOf.h"
 #include "Algo/Count.h"
 #include "Algo/Find.h"
+#include "Algo/Sort.h"
 #include "AssetCompilingManager.h"
 #include "AssetRegistry/IAssetRegistry.h"
 #include "Async/ParallelFor.h"
@@ -969,17 +970,14 @@ void FPackageData::CreateObjectCache()
 	{
 		PackageName = LocalPackage->GetFName();
 		TArray<UObject*> ObjectsInOuter;
-		GetObjectsWithOuter(LocalPackage, ObjectsInOuter);
+		// ignore RF_Garbage objects; they will not be serialized out so we don't need to call
+		// BeginCacheForCookedPlatformData on them
+		GetObjectsWithOuter(LocalPackage, ObjectsInOuter, true /* bIncludeNestedObjects */, RF_NoFlags, EInternalObjectFlags::Garbage);
 		CachedObjectsInOuter.Reset(ObjectsInOuter.Num());
 		for (UObject* Object : ObjectsInOuter)
 		{
 			FWeakObjectPtr ObjectWeakPointer(Object);
-			// ignore pending kill objects; they will not be serialized out so we don't need to call
-			// BeginCacheForCookedPlatformData on them
-			if (!ObjectWeakPointer.Get())
-			{
-				continue;
-			}
+			check(ObjectWeakPointer.Get()); // GetObjectsWithOuter with Garbage filtered out should only return valid-for-weakptr objects
 			CachedObjectsInOuter.Emplace(MoveTemp(ObjectWeakPointer));
 		}
 		SetHasSaveCache(true);
@@ -990,6 +988,94 @@ void FPackageData::CreateObjectCache()
 	}
 }
 
+static TArray<UObject*> SetDifference(TArray<UObject*>& A, TArray<UObject*>& B)
+{
+	Algo::Sort(A); // Don't use TArray.Sort, it sorts pointers as references and we want to sort them as pointers
+	Algo::Sort(B);
+	int32 ANum = A.Num();
+	int32 BNum = B.Num();
+	UObject** AData = A.GetData();
+	UObject** BData = B.GetData();
+
+	// Always move to the smallest next element from the two remaining lists and if it's in one set and not the
+	// other add it to the output if in A or skip it if in B.
+	int32 AIndex = 0;
+	int32 BIndex = 0;
+	TArray<UObject*> AMinusB;
+	while (AIndex < ANum && BIndex < BNum)
+	{
+		if (AData[AIndex] == BData[BIndex])
+		{
+			++AIndex;
+			++BIndex;
+			continue;
+		}
+		if (AData[AIndex] < BData[BIndex])
+		{
+			AMinusB.Add(AData[AIndex++]);
+		}
+		else
+		{
+			++BIndex;
+		}
+	}
+
+	// When we reach the end of B, all remaining elements of A are not in B.
+	while (AIndex < ANum)
+	{
+		AMinusB.Add(AData[AIndex++]);
+	}
+	return AMinusB;
+}
+
+EPollStatus FPackageData::RefreshObjectCache(bool& bOutFoundNewObjects)
+{
+	// Caller will only call this function after CallBeginCacheOnObjects finished successfully
+	int32 NumPlatforms = GetNumRequestedPlatforms();
+	check(NumPlatforms > 0);
+	check(GetCookedPlatformDataNextIndex()/NumPlatforms == GetCachedObjectsInOuter().Num());
+	check(Package.Get() != nullptr);
+
+	TArray<UObject*> OldObjects;
+	OldObjects.Reserve(CachedObjectsInOuter.Num());
+	for (FWeakObjectPtr& Object : CachedObjectsInOuter)
+	{
+		UObject* ObjectPtr = Object.Get();
+		if (ObjectPtr)
+		{
+			OldObjects.Add(ObjectPtr);
+		}
+	}
+	TArray<UObject*> CurrentObjects;
+	GetObjectsWithOuter(Package.Get(), CurrentObjects, true /* bIncludeNestedObjects */, RF_NoFlags, EInternalObjectFlags::Garbage);
+
+	TArray<UObject*> NewObjects = SetDifference(CurrentObjects, OldObjects);
+	bOutFoundNewObjects = NewObjects.Num() > 0;
+	if (bOutFoundNewObjects)
+	{
+		CachedObjectsInOuter.Reserve(CachedObjectsInOuter.Num() + NewObjects.Num());
+		for (UObject* Object : NewObjects)
+		{
+			FWeakObjectPtr ObjectWeakPointer(Object);
+			check(ObjectWeakPointer.Get()); // GetObjectsWithOuter with Garbage filtered out should only return valid-for-weakptr objects
+			CachedObjectsInOuter.Emplace(MoveTemp(ObjectWeakPointer));
+		}
+		// GetCookedPlatformDataNextIndex is already where it should be, pointing at the first of the objects we have added
+		// Change our state back so we know we need to CallBeginCacheOnObjects again 
+		SetCookedPlatformDataCalled(false);
+
+		if (++GetNumRetriesBeginCacheOnObjects() > FPackageData::GetMaxNumRetriesBeginCacheOnObjects())
+		{
+			UE_LOG(LogCook, Error, TEXT("Cooker has repeatedly tried to call BeginCacheForCookedPlatformData on all objects in the package, but keeps finding new objects.\n")
+				TEXT("Aborting the save of the package; programmer needs to debug why objects keep getting added to the package.\n")
+				TEXT("Package: %s. Most recent created object: %s."),
+				*GetPackageName().ToString(), *NewObjects[0]->GetFullName());
+			return EPollStatus::Error;
+		}
+	}
+	return EPollStatus::Success;
+}
+
 void FPackageData::ClearObjectCache()
 {
 	CachedObjectsInOuter.Empty();
@@ -1016,6 +1102,16 @@ int32& FPackageData::GetCookedPlatformDataNextIndex()
 	return CookedPlatformDataNextIndex;
 }
 
+int32& FPackageData::GetNumRetriesBeginCacheOnObjects()
+{
+	return NumRetriesBeginCacheOnObject;
+}
+
+int32 FPackageData::GetMaxNumRetriesBeginCacheOnObjects()
+{
+	return 10;
+}
+
 void FPackageData::CheckCookedPlatformDataEmpty() const
 {
 	check(GetCookedPlatformDataNextIndex() == 0);
@@ -1034,6 +1130,7 @@ void FPackageData::CheckCookedPlatformDataEmpty() const
 void FPackageData::ClearCookedPlatformData()
 {
 	CookedPlatformDataNextIndex = 0;
+	NumRetriesBeginCacheOnObject = 0;
 	// Note that GetNumPendingCookedPlatformData is not cleared; it persists across Saves and CookSessions
 	SetCookedPlatformDataStarted(false);
 	SetCookedPlatformDataCalled(false);
@@ -1537,64 +1634,125 @@ void FCookGenerationInfo::SetSaveStateComplete(ESaveState CompletedState)
 	}
 }
 
-void FCookGenerationInfo::TakeOverCachedObjectsAndAddNew(TArray<FWeakObjectPtr>& CachedObjectsInOuter, TArray<UObject*>& NewObjects)
+void FCookGenerationInfo::TakeOverCachedObjectsAndAddMoved(FGeneratorPackage& Generator,
+	TArray<FWeakObjectPtr>& CachedObjectsInOuter, TArray<UObject*>& MovedObjects)
 {
-	TMap<UObject*, TOptional<bool>> ObjectSet;
+	BeginCacheObjects.Objects.Reset();
+	TSet<UObject*> ObjectSet;
 	for (FWeakObjectPtr& ObjectInOuter : CachedObjectsInOuter)
 	{
 		UObject* Object = ObjectInOuter.Get();
 		if (Object)
 		{
-			TOptional<bool>& HasFinishedRound = ObjectSet.FindOrAdd(Object);
-			HasFinishedRound = true;
+			bool bAlreadyExists;
+			ObjectSet.Add(Object, &bAlreadyExists);
+			if (!bAlreadyExists)
+			{
+				FBeginCacheObject& Added = BeginCacheObjects.Objects.Add_GetRef(FBeginCacheObject{ Object });
+				Added.bHasFinishedRound = true;
+			}
+		}
+	}
+
+	TArray<UObject*> ChildrenOfMovedObjects;
+	for (UObject* Object : MovedObjects)
+	{
+		if (!IsValid(Object))
+		{
+			UE_LOG(LogCook, Warning, TEXT("CookPackageSplitter found non-valid object %s returned from %s on Splitter %s%s. Ignoring it."),
+				Object ? *Object->GetFullName() : TEXT("<null>"),
+				IsGenerator() ? TEXT("PopulateGeneratorPackage") : TEXT("PopulateGeneratedPackage"),
+				*Generator.GetSplitDataObjectName().ToString(),
+				IsGenerator() ? TEXT("") : *FString::Printf(TEXT(", Package %s"), *PackageData->GetPackageName().ToString()));
+			continue;
+		}
+
+		bool bAlreadyExists;
+		ObjectSet.Add(Object, &bAlreadyExists);
+		if (!bAlreadyExists)
+		{
+			FBeginCacheObject& Added = BeginCacheObjects.Objects.Add_GetRef(FBeginCacheObject{ Object });
+			Added.bHasFinishedRound = false;
+			Added.bIsRootMovedObject = true;
+			GetObjectsWithOuter(Object, ChildrenOfMovedObjects, true /* bIncludeNestedObjects */, RF_NoFlags, EInternalObjectFlags::Garbage);
 		}
 	}
-	for (UObject* Object : NewObjects)
+
+	for (UObject* Object : ChildrenOfMovedObjects)
 	{
-		TOptional<bool>& HasFinishedRound = ObjectSet.FindOrAdd(Object);
-		if (!HasFinishedRound.IsSet())
+		check(IsValid(Object));
+		bool bAlreadyExists;
+		ObjectSet.Add(Object, &bAlreadyExists);
+		if (!bAlreadyExists)
 		{
-			HasFinishedRound = false;
+			FBeginCacheObject& Added = BeginCacheObjects.Objects.Add_GetRef(FBeginCacheObject{ Object });
+			Added.bHasFinishedRound = false;
 		}
 	}
 
-	BeginCacheObjects.SetObjects(ObjectSet);
 	BeginCacheObjects.StartRound();
 
 	CachedObjectsInOuter.Reset();
 	SetHasTakenOverCachedCookedPlatformData(true);
 }
 
-void FCookGenerationInfo::RefreshPackageObjects(UPackage* Package)
+EPollStatus FCookGenerationInfo::RefreshPackageObjects(FGeneratorPackage& Generator, UPackage* Package,
+	bool& bOutFoundNewObjects, ESaveState DemotionState)
 {
-	TArray<UObject*> ObjectsInOuter;
-	GetObjectsWithOuter(Package, ObjectsInOuter);
+	bOutFoundNewObjects = false;
+	TArray<UObject*> CurrentObjectsInOuter;
+	GetObjectsWithOuter(Package, CurrentObjectsInOuter, true /* bIncludeNestedObjects */, RF_NoFlags, EInternalObjectFlags::Garbage);
 
-	TMap<UObject*, TOptional<bool>> ObjectSet;
+	TSet<UObject*> ObjectSet;
 	ObjectSet.Reserve(BeginCacheObjects.Objects.Num());
 	for (FBeginCacheObject& ExistingObject : BeginCacheObjects.Objects)
 	{
 		UObject* Object = ExistingObject.Object.Get();
 		if (Object)
 		{
-			TOptional<bool>& HasFinishedRound = ObjectSet.FindOrAdd(Object);
-			if (!HasFinishedRound.IsSet() || ExistingObject.bHasFinishedRound)
+			bool bAlreadyExists;
+			ObjectSet.Add(Object, &bAlreadyExists);
+			check(!bAlreadyExists); // Objects in BeginCacheObjects.Objects are guaranteed unique and we haven't added any others yet
+			if (ExistingObject.bIsRootMovedObject)
 			{
-				HasFinishedRound = ExistingObject.bHasFinishedRound;
+				GetObjectsWithOuter(Object, CurrentObjectsInOuter, true /* bIncludeNestedObjects */, RF_NoFlags, EInternalObjectFlags::Garbage);
 			}
 		}
 	}
-	for (UObject* Object : ObjectsInOuter)
+	UObject* FirstNewObject = nullptr;
+	for (UObject* Object : CurrentObjectsInOuter)
 	{
-		TOptional<bool>& HasFinishedRound = ObjectSet.FindOrAdd(Object);
-		if (!HasFinishedRound.IsSet())
+		bool bAlreadyExists;
+		ObjectSet.Add(Object, &bAlreadyExists);
+		if (!bAlreadyExists)
 		{
-			HasFinishedRound = false;
+			FBeginCacheObject& Added = BeginCacheObjects.Objects.Add_GetRef(FBeginCacheObject{ Object });
+			Added.bHasFinishedRound = false;
+			if (!FirstNewObject )
+			{
+				FirstNewObject  = Object;
+			}
 		}
 	}
+	bOutFoundNewObjects = FirstNewObject != nullptr;
 
-	BeginCacheObjects.SetObjects(ObjectSet);
 	BeginCacheObjects.StartRound();
+
+	if (FirstNewObject != nullptr && DemotionState != ESaveState::Last)
+	{
+		SetSaveState(DemotionState);
+		if (++PackageData->GetNumRetriesBeginCacheOnObjects() > FPackageData::GetMaxNumRetriesBeginCacheOnObjects())
+		{
+			UE_LOG(LogCook, Error, TEXT("Cooker has repeatedly tried to call BeginCacheForCookedPlatformData on all objects in a generated package, but keeps finding new objects.\n")
+				TEXT("Aborting the save of the package; programmer needs to debug why objects keep getting added to the package.\n")
+				TEXT("Splitter: %s%s. Most recent created object: %s."),
+				*Generator.GetSplitDataObjectName().ToString(),
+				IsGenerator() ? TEXT("") : *FString::Printf(TEXT(", Package: %s"), *PackageData->GetPackageName().ToString()),
+				*FirstNewObject->GetFullName());
+			return EPollStatus::Error;
+		}
+	}
+	return EPollStatus::Success;
 }
 
 void FBeginCacheObjects::Reset()
@@ -1604,17 +1762,6 @@ void FBeginCacheObjects::Reset()
 	NextIndexInRound = 0;
 }
 
-void FBeginCacheObjects::SetObjects(TMap<UObject*, TOptional<bool>>& InObjectSet)
-{
-	Objects.Reset(InObjectSet.Num());
-	for (TPair<UObject*, TOptional<bool>>& Pair : InObjectSet)
-	{
-		FBeginCacheObject& BeginCacheObject = Objects.Emplace_GetRef();
-		BeginCacheObject.bHasFinishedRound = Pair.Value.IsSet() && *Pair.Value;
-		BeginCacheObject.Object = Pair.Key;
-	}
-}
-
 void FBeginCacheObjects::StartRound()
 {
 	ObjectsInRound.Reset(Objects.Num());
diff --git a/Engine/Source/Editor/UnrealEd/Private/Cooker/CookPackageData.h b/Engine/Source/Editor/UnrealEd/Private/Cooker/CookPackageData.h
--- a/Engine/Source/Editor/UnrealEd/Private/Cooker/CookPackageData.h
+++ b/Engine/Source/Editor/UnrealEd/Private/Cooker/CookPackageData.h
@@ -358,6 +358,12 @@ public:
 	void CheckObjectCacheEmpty() const;
 	/** Populate CachedObjectsInOuter if not already populated. Invalid to call except when in the save state. */
 	void CreateObjectCache();
+	/**
+	 * Look for new Objects that were created during BeginCacheForCookedPlatformData calls, and if found add
+	 * them to the ObjectCache and set state so that we call BeginCacheForCookedPlatformData on the new objects.
+	 * ErrorExits if this creation of new objects happens too many times.
+	 */
+	EPollStatus RefreshObjectCache(bool& bOutFoundNewObjects);
 	/** Clear the CachedObjectsInOuter list, when e.g. leaving the save state. */
 	void ClearObjectCache();
 
@@ -365,6 +371,8 @@ public:
 	int32& GetNumPendingCookedPlatformData();
 	const int32& GetCookedPlatformDataNextIndex() const;
 	int32& GetCookedPlatformDataNextIndex();
+	int32& GetNumRetriesBeginCacheOnObjects();
+	static int32 GetMaxNumRetriesBeginCacheOnObjects();
 
 	/** Get/Set the flag for whether CachedObjectsInOuter is populated. Always false except during save state. */
 	bool GetHasSaveCache() const { return static_cast<bool>(bHasSaveCache); }
@@ -550,6 +558,7 @@ private:
 	FTrackedPreloadableFilePtr PreloadableFile;
 	int32 NumPendingCookedPlatformData = 0;
 	int32 CookedPlatformDataNextIndex = 0;
+	int32 NumRetriesBeginCacheOnObject = 0;
 	FOpenPackageResult PreloadableFileOpenResult;
 	FInstigator Instigator;
 
@@ -576,6 +585,7 @@ struct FBeginCacheObject
 {
 	FWeakObjectPtr Object;
 	bool bHasFinishedRound = false;
+	bool bIsRootMovedObject = false;
 };
 
 /**
@@ -589,7 +599,6 @@ struct FBeginCacheObjects
 	int32 NextIndexInRound = 0;
 
 	void Reset();
-	void SetObjects(TMap<UObject*, TOptional<bool>>& InObjectSet);
 
 	void StartRound();
 	void EndRound(int32 NumPlatforms);
@@ -649,9 +658,15 @@ public:
 	 * Steal the list of cached objects to call BeginCacheForCookedPlatformData on from the PackageData,
 	 * and add them to this. Also add the list of NewObjects reported by the splitter that will be moved into the package.
 	 */
-	void TakeOverCachedObjectsAndAddNew(TArray<FWeakObjectPtr>& CachedObjectsInOuter, TArray<UObject*>& NewObjects);
-	/** Fetch all the objects currently in the package and add them the list of objects that need BeginCacheForCookedPlatformData. */
-	void RefreshPackageObjects(UPackage* Package);
+	void TakeOverCachedObjectsAndAddMoved(FGeneratorPackage& Generator, 
+		TArray<FWeakObjectPtr>& CachedObjectsInOuter, TArray<UObject*>& MovedObjects);
+	/**
+	 * Fetch all the objects currently in the package and add them the list of objects that need BeginCacheForCookedPlatformData.
+	 * Reports whether new objects were found. If DemotionState is not ESaveState::Last, will SetState back to DemotionState
+	 * if new objects were found, and will error exit if this demotion has happened too many times.
+	 */
+	EPollStatus RefreshPackageObjects(FGeneratorPackage& Generator, UPackage* Package, bool& bOutFoundNewObjects,
+		ESaveState DemotionState);
 
 public:
 	FBeginCacheObjects BeginCacheObjects;
