this change allows DLC plugins to be cooked using the content worker in an uncooked environment. This is meant purely for CIS and local iteration where cloud workers might be invalid due to local serialization changes.

-- Testing ---
Change UncookedContentWorker to no longer grab the latest build data by default it should be uncessary and just adds to the running time.
change UncookedContentWorker tests to unmount the plugin after cooking starts as the newly mounted plugin might have a Primary Asset Id Collision
change tests back to using the basic set for gauntlet testing.

-- Cooking --
Add new commandline option for the cooker "AllowUncookedAssetReferences" when passed on the commandline all uncooked assets besides the requested DLC plugin will be considered cooked.

diff --git a/Engine/Source/Editor/UnrealEd/Private/CookOnTheFlyServer.cpp b/Engine/Source/Editor/UnrealEd/Private/CookOnTheFlyServer.cpp
--- a/Engine/Source/Editor/UnrealEd/Private/CookOnTheFlyServer.cpp
+++ b/Engine/Source/Editor/UnrealEd/Private/CookOnTheFlyServer.cpp
@@ -8867,6 +8867,7 @@ FBeginCookContext UCookOnTheFlyServer::CreateBeginCookByTheBookContext(const FCo
 	CookByTheBookOptions->bCookAgainstFixedBase = !!(CookOptions & ECookByTheBookOptions::CookAgainstFixedBase);
 	CookByTheBookOptions->bDlcLoadMainAssetRegistry = !!(CookOptions & ECookByTheBookOptions::DlcLoadMainAssetRegistry);
 	CookByTheBookOptions->bErrorOnEngineContentUse = StartupOptions.bErrorOnEngineContentUse;
+	CookByTheBookOptions->bAllowUncookedAssetReferences = FParse::Param(FCommandLine::Get(), TEXT("AllowUncookedAssetReferences"));
 	CookByTheBookOptions->DlcName = StartupOptions.DLCName;
 	if (CookByTheBookOptions->bSkipHardReferences && !CookByTheBookOptions->bSkipSoftReferences)
 	{
@@ -9131,7 +9132,7 @@ void UCookOnTheFlyServer::RecordDLCPackagesFromBaseGame(FBeginCookContext& Begin
 		if (!bUsingDevRegistryOverride)
 		{
 			bool bReadSucceeded = ReadDevelopmentAssetRegistry(PackageList, PlatformNameString);
-			if (!bReadSucceeded)
+			if (!bReadSucceeded && !CookByTheBookOptions->bAllowUncookedAssetReferences)
 			{
 				UE_LOG(LogCook, Fatal, TEXT("Could not find based-on AssetRegistry file %s for platform %s. ")
 					TEXT("When cooking DLC, %s is expected to exist Release/<platform> for each platform being cooked. (Or use DevelopmentAssetRegistryPlatformOverride=<PlatformName> to specify an override platform that all platforms should use to find the %s file). Terminating the cook."),
diff --git a/Engine/Source/Editor/UnrealEd/Private/Cooker/CookRequestCluster.cpp b/Engine/Source/Editor/UnrealEd/Private/Cooker/CookRequestCluster.cpp
--- a/Engine/Source/Editor/UnrealEd/Private/Cooker/CookRequestCluster.cpp
+++ b/Engine/Source/Editor/UnrealEd/Private/Cooker/CookRequestCluster.cpp
@@ -129,9 +129,11 @@ void FRequestCluster::AddClusters(UCookOnTheFlyServer& COTFS, FPackageDataSet& U
 	};
 
 	bool bErrorOnEngineContentUse = false;
+	bool bAllowUncookedAssetReferences = false;
 	UE::Cook::FCookByTheBookOptions& Options = *COTFS.CookByTheBookOptions;
 	FString DLCPath;
 	bErrorOnEngineContentUse = Options.bErrorOnEngineContentUse;
+	bAllowUncookedAssetReferences = Options.bAllowUncookedAssetReferences;
 	if (bErrorOnEngineContentUse)
 	{
 		DLCPath = FPaths::Combine(*COTFS.GetBaseDirectoryForDLC(), TEXT("Content"));
@@ -153,7 +155,7 @@ void FRequestCluster::AddClusters(UCookOnTheFlyServer& COTFS, FPackageDataSet& U
 		// a cluster
 		PackageData->GetRequestedPlatforms(RequestedPlatforms);
 		if (!IsRequestCookable(PackageData->GetPackageName(), PackageData, *COTFS.PackageDatas, *COTFS.PackageTracker,
-			DLCPath, bErrorOnEngineContentUse, RequestedPlatforms))
+			DLCPath, bErrorOnEngineContentUse, bAllowUncookedAssetReferences, RequestedPlatforms))
 		{
 			PackageData->SendToState(EPackageState::Idle, ESendFlags::QueueAdd);
 			continue;
@@ -192,6 +194,7 @@ void FRequestCluster::Initialize(UCookOnTheFlyServer& COTFS)
 		bAllowHardDependencies = !Options.bSkipHardReferences;
 		bAllowSoftDependencies = !Options.bSkipSoftReferences;
 		bErrorOnEngineContentUse = Options.bErrorOnEngineContentUse;
+		bAllowUncookedAssetReferences = Options.bAllowUncookedAssetReferences;
 	}
 	else
 	{
@@ -1093,12 +1096,12 @@ TMap<FPackageData*, TArray<FPackageData*>>& FRequestCluster::FGraphSearch::GetGr
 bool FRequestCluster::IsRequestCookable(FName PackageName, FPackageData*& InOutPackageData)
 {
 	return IsRequestCookable(PackageName, InOutPackageData, PackageDatas, PackageTracker,
-		DLCPath, bErrorOnEngineContentUse, GetPlatforms());
+		DLCPath, bErrorOnEngineContentUse, bAllowUncookedAssetReferences, GetPlatforms());
 }
 
 bool FRequestCluster::IsRequestCookable(FName PackageName, FPackageData*& InOutPackageData,
 	FPackageDatas& InPackageDatas, FPackageTracker& InPackageTracker,
-	FStringView InDLCPath, bool bInErrorOnEngineContentUse, TConstArrayView<const ITargetPlatform*> RequestPlatforms)
+	FStringView InDLCPath, bool bInErrorOnEngineContentUse, bool bInAllowUncookedAssetReferences, TConstArrayView<const ITargetPlatform*> RequestPlatforms)
 {
 	TStringBuilder<256> NameBuffer;
 	// We need to reject packagenames from adding themselves or their transitive dependencies using all the same rules that
@@ -1134,7 +1137,12 @@ bool FRequestCluster::IsRequestCookable(FName PackageName, FPackageData*& InOutP
 		{
 			if (!InOutPackageData->HasAllCookedPlatforms(RequestPlatforms, true /* bIncludeFailed */))
 			{
-				UE_LOG(LogCook, Error, TEXT("Uncooked Engine or Game content %s is being referenced by DLC!"), *NameBuffer);
+				// AllowUncookedAssetReferences should only be used when the DLC plugin to cook is going to be mounted where uncooked packages are available.
+				// This will allow a DLC plugin to be recooked continually and mounted in an uncooked editor which is useful for CI.
+				if (!bInAllowUncookedAssetReferences)
+				{
+					UE_LOG(LogCook, Error, TEXT("Uncooked Engine or Game content %s is being referenced by DLC!"), *NameBuffer);
+				}
 			}
 			return false;
 		}
diff --git a/Engine/Source/Editor/UnrealEd/Private/Cooker/CookRequestCluster.h b/Engine/Source/Editor/UnrealEd/Private/Cooker/CookRequestCluster.h
--- a/Engine/Source/Editor/UnrealEd/Private/Cooker/CookRequestCluster.h
+++ b/Engine/Source/Editor/UnrealEd/Private/Cooker/CookRequestCluster.h
@@ -298,7 +298,7 @@ private:
 	bool IsRequestCookable(FName PackageName, FPackageData*& InOutPackageData);
 	static bool IsRequestCookable(FName PackageName, FPackageData*& InOutPackageData,
 		FPackageDatas& InPackageDatas, FPackageTracker& InPackageTracker,
-		FStringView InDLCPath, bool bInErrorOnEngineContentUse, TConstArrayView<const ITargetPlatform*> RequestPlatforms);
+		FStringView InDLCPath, bool bInErrorOnEngineContentUse, bool bInAllowUncookedAssetReferences, TConstArrayView<const ITargetPlatform*> RequestPlatforms);
 
 	TArray<FFileNameRequest> InRequests;
 	TArray<FPackageData*> Requests;
@@ -317,6 +317,7 @@ private:
 	bool bAllowSoftDependencies = true;
 	bool bHybridIterativeEnabled = true;
 	bool bErrorOnEngineContentUse = false;
+	bool bAllowUncookedAssetReferences = false;
 	bool bPackageNamesComplete = false;
 	bool bDependenciesComplete = false;
 	bool bStartAsyncComplete = false;
diff --git a/Engine/Source/Editor/UnrealEd/Private/Cooker/CookTypes.h b/Engine/Source/Editor/UnrealEd/Private/Cooker/CookTypes.h
--- a/Engine/Source/Editor/UnrealEd/Private/Cooker/CookTypes.h
+++ b/Engine/Source/Editor/UnrealEd/Private/Cooker/CookTypes.h
@@ -290,6 +290,7 @@ public:
 
 	/** error when detecting engine content being used in this cook */
 	bool							bErrorOnEngineContentUse = false;
+	bool							bAllowUncookedAssetReferences = false; // this is a flag for dlc, will allow DLC to be cook when the fixed base might be missing references.
 	bool							bSkipHardReferences = false;
 	bool							bSkipSoftReferences = false;
 	bool							bFullLoadAndSave = false;
