crash fix for Launch On when there is no source code / build targets available (e.g. artist machine)

diff --git a/Engine/Source/Editor/UnrealEd/Private/PlayLevelViaLauncher.cpp b/Engine/Source/Editor/UnrealEd/Private/PlayLevelViaLauncher.cpp
--- a/Engine/Source/Editor/UnrealEd/Private/PlayLevelViaLauncher.cpp
+++ b/Engine/Source/Editor/UnrealEd/Private/PlayLevelViaLauncher.cpp
@@ -279,7 +279,11 @@ void UEditorEngine::StartPlayUsingLauncherSession(FRequestPlaySessionParams& InR
 	// only set the BuildTarget in code-based projects
 	if (LauncherSessionInfo->bPlayUsingLauncherHasCode)
 	{
-		LauncherProfile->SetBuildTarget(GetDefault<UProjectPackagingSettings>()->GetLaunchOnTargetInfo()->Name);
+		const FTargetInfo* TargetInfo = GetDefault<UProjectPackagingSettings>()->GetLaunchOnTargetInfo();
+		if (TargetInfo != nullptr)
+		{
+			LauncherProfile->SetBuildTarget(TargetInfo->Name);
+		}
 	}
 	LauncherProfile->SetCookMode(CurrentLauncherCookMode);
 	LauncherProfile->SetUnversionedCooking(!bIncrementalCooking); // Unversioned cooking is not allowed with incremental cooking
