Add OnNewActorsPlaced Delegate

diff --git a/Engine/Source/Editor/UnrealEd/Private/Editor.cpp b/Engine/Source/Editor/UnrealEd/Private/Editor.cpp
--- a/Engine/Source/Editor/UnrealEd/Private/Editor.cpp
+++ b/Engine/Source/Editor/UnrealEd/Private/Editor.cpp
@@ -126,6 +126,7 @@ FEditorDelegates::FOnAssetPreImport						FEditorDelegates::OnAssetPreImport;
 FEditorDelegates::FOnAssetPostImport					FEditorDelegates::OnAssetPostImport;
 FEditorDelegates::FOnAssetReimport						FEditorDelegates::OnAssetReimport;
 FEditorDelegates::FOnNewActorsDropped					FEditorDelegates::OnNewActorsDropped;
+FEditorDelegates::FOnNewActorsPlaced					FEditorDelegates::OnNewActorsPlaced;
 FEditorDelegates::FOnGridSnappingChanged				FEditorDelegates::OnGridSnappingChanged;
 FSimpleMulticastDelegate								FEditorDelegates::OnLightingBuildStarted;
 FSimpleMulticastDelegate								FEditorDelegates::OnLightingBuildKept;
diff --git a/Engine/Source/Editor/UnrealEd/Private/LevelEditorViewport.cpp b/Engine/Source/Editor/UnrealEd/Private/LevelEditorViewport.cpp
--- a/Engine/Source/Editor/UnrealEd/Private/LevelEditorViewport.cpp
+++ b/Engine/Source/Editor/UnrealEd/Private/LevelEditorViewport.cpp
@@ -287,6 +287,11 @@ TArray<AActor*> FLevelEditorViewportClient::TryPlacingActorFromObject( ULevel* I
 		}
 	}
 
+	if (PlacedActors.Num() > 0)
+	{
+		FEditorDelegates::OnNewActorsPlaced.Broadcast(ObjToUse, PlacedActors);
+	}
+
 	return PlacedActors;
 }
 
diff --git a/Engine/Source/Editor/UnrealEd/Public/Editor.h b/Engine/Source/Editor/UnrealEd/Public/Editor.h
--- a/Engine/Source/Editor/UnrealEd/Public/Editor.h
+++ b/Engine/Source/Editor/UnrealEd/Public/Editor.h
@@ -94,8 +94,10 @@ struct UNREALED_API FEditorDelegates
 	DECLARE_MULTICAST_DELEGATE_OneParam(FOnAssetReimport, UObject*);
 	/** delegate type for finishing up construction of a new blueprint */
 	DECLARE_MULTICAST_DELEGATE_OneParam(FOnFinishPickingBlueprintClass, UClass*);
-	/** delegate type for triggering when new actors are dropped on to the viewport */
+	/** delegate type for triggering when new actors are dropped on to the viewport via drag and drop */
 	DECLARE_MULTICAST_DELEGATE_TwoParams(FOnNewActorsDropped, const TArray<UObject*>&, const TArray<AActor*>&);
+	/** delegate type for triggering when new actors are placed on to the viewport. Triggers before NewActorsDropped if placement is caused by a drop action */
+	DECLARE_MULTICAST_DELEGATE_TwoParams(FOnNewActorsPlaced, UObject*, const TArray<AActor*>&);
 	/** delegate type for when attempting to apply an object to an actor */
 	DECLARE_MULTICAST_DELEGATE_TwoParams(FOnApplyObjectToActor, UObject*, AActor*);
 	/** delegate type for triggering when grid snapping has changed */
@@ -252,6 +254,8 @@ struct UNREALED_API FEditorDelegates
 	static FOnAssetReimport OnAssetReimport;
 	/** Called when new actors are dropped on to the viewport */
 	static FOnNewActorsDropped OnNewActorsDropped;
+	/** Called when new actors are placed in the viewport */
+	static FOnNewActorsPlaced OnNewActorsPlaced;
 	/** Called when grid snapping is changed */
 	static FOnGridSnappingChanged OnGridSnappingChanged;
 	/** Called when a lighting build has started */
