Creating a new Level Instance will now put it in the current folder.

diff --git a/Engine/Source/Editor/UnrealEd/Private/EditorActorFolders.cpp b/Engine/Source/Editor/UnrealEd/Private/EditorActorFolders.cpp
--- a/Engine/Source/Editor/UnrealEd/Private/EditorActorFolders.cpp
+++ b/Engine/Source/Editor/UnrealEd/Private/EditorActorFolders.cpp
@@ -534,12 +534,12 @@ void FActorFolders::SetIsFolderExpanded(UWorld& InWorld, const FFolder& InFolder
 	GetOrCreateWorldFolders(InWorld).SetIsFolderExpanded(InFolder, bIsExpanded);
 }
 
-FFolder FActorFolders::GetActorEditorContextFolder(UWorld& InWorld) const
+FFolder FActorFolders::GetActorEditorContextFolder(UWorld& InWorld, bool bMustMatchCurrentLevel) const
 {
 	const UWorldFolders** Folders = (const UWorldFolders **)WorldFolders.Find(&InWorld);
 	if (Folders)
 	{
-		return (*Folders)->GetActorEditorContextFolder();
+		return (*Folders)->GetActorEditorContextFolder(bMustMatchCurrentLevel);
 	}
 	return FFolder::GetWorldRootFolder(&InWorld);
 }
@@ -559,8 +559,9 @@ void FActorFolders::OnExecuteActorEditorContextAction(UWorld* InWorld, const EAc
 		case EActorEditorContextAction::ApplyContext:
 			check(InActor && InActor->GetWorld() == InWorld);
 			{
-				FFolder Folder = GetActorEditorContextFolder(*InWorld);
-				if (!Folder.IsNone())
+				const bool bMustMatchCurrentLevel = false;
+				FFolder Folder = GetActorEditorContextFolder(*InWorld, bMustMatchCurrentLevel);
+				if (!Folder.IsNone() && (Folder.GetRootObjectAssociatedLevel() == InActor->GetLevel()))
 				{
 					// Currently not supported to change rootobject through this interface
 					if (Folder.GetRootObject() == InActor->GetFolderRootObject())
diff --git a/Engine/Source/Editor/UnrealEd/Private/WorldFolders.cpp b/Engine/Source/Editor/UnrealEd/Private/WorldFolders.cpp
--- a/Engine/Source/Editor/UnrealEd/Private/WorldFolders.cpp
+++ b/Engine/Source/Editor/UnrealEd/Private/WorldFolders.cpp
@@ -184,16 +184,22 @@ bool UWorldFolders::SetIsFolderExpanded(const FFolder& InFolder, bool bIsExpande
 	return false;
 }
 
-FFolder UWorldFolders::GetActorEditorContextFolder() const
+FFolder UWorldFolders::GetActorEditorContextFolder(bool bMustMatchCurrentLevel) const
 {
 	FFolder Folder = CurrentFolder.GetFolder();
 	if (ContainsFolder(Folder))
 	{
 		ULevelInstanceSubsystem* LevelInstanceSubsystem = GetWorld()->GetSubsystem<ULevelInstanceSubsystem>();
 		ILevelInstanceInterface* EditingLevelInstance = LevelInstanceSubsystem ? LevelInstanceSubsystem->GetEditingLevelInstance() : nullptr;
-		UObject* EditingLevelInstanceObject = EditingLevelInstance ? CastChecked<UObject>(EditingLevelInstance) : nullptr;
-		FFolder::FRootObject RootObject = FFolder::FRootObject(EditingLevelInstanceObject ? EditingLevelInstanceObject : GetWorld()->GetCurrentLevel());
-		if (Folder.GetRootObject() == RootObject)
+		if (UObject* EditingLevelInstanceObject = EditingLevelInstance ? CastChecked<UObject>(EditingLevelInstance) : nullptr)
+		{
+			FFolder::FRootObject RootObject = FFolder::FRootObject(EditingLevelInstanceObject);
+			if (Folder.GetRootObject() == RootObject)
+			{
+				return Folder;
+			}
+		}
+		else if (!bMustMatchCurrentLevel || (Folder.GetRootObject() == FFolder::FRootObject(GetWorld()->GetCurrentLevel())))
 		{
 			return Folder;
 		}
diff --git a/Engine/Source/Editor/UnrealEd/Public/EditorActorFolders.h b/Engine/Source/Editor/UnrealEd/Public/EditorActorFolders.h
--- a/Engine/Source/Editor/UnrealEd/Public/EditorActorFolders.h
+++ b/Engine/Source/Editor/UnrealEd/Public/EditorActorFolders.h
@@ -162,7 +162,7 @@ struct UNREALED_API FActorFolders : public FGCObject, public IActorEditorContext
 	virtual TSharedRef<SWidget> GetActorEditorContextWidget(UWorld* InWorld) const override;
 	virtual FOnActorEditorContextClientChanged& GetOnActorEditorContextClientChanged() override { return ActorEditorContextClientChanged; }
 	//~ End IActorEditorContextClient interface
-	FFolder GetActorEditorContextFolder(UWorld& InWorld) const;
+	FFolder GetActorEditorContextFolder(UWorld& InWorld, bool bMustMatchCurrentLevel = true) const;
 	void SetActorEditorContextFolder(UWorld& InWorld, const FFolder& InFolder);
 
 private:
diff --git a/Engine/Source/Editor/UnrealEd/Public/WorldFolders.h b/Engine/Source/Editor/UnrealEd/Public/WorldFolders.h
--- a/Engine/Source/Editor/UnrealEd/Public/WorldFolders.h
+++ b/Engine/Source/Editor/UnrealEd/Public/WorldFolders.h
@@ -80,7 +80,7 @@ public:
 	bool RenameFolder(const FFolder& InOldFolder, const FFolder& InNewFolder);
 	bool IsFolderExpanded(const FFolder& InFolder) const;
 	bool SetIsFolderExpanded(const FFolder& InFolder, bool bIsExpanded);
-	FFolder GetActorEditorContextFolder() const;
+	FFolder GetActorEditorContextFolder(bool bMustMatchCurrentLevel) const;
 	bool SetActorEditorContextFolder(const FFolder& InFolder);
 	void PushActorEditorContext();
 	void PopActorEditorContext();
