Cooker: Improve feedback when patch creation fails the ensure that GetAllPackageFilenamesFromAssetRegistry succeeded.
#rnx

diff --git a/Engine/Source/Editor/UnrealEd/Private/CookOnTheFlyServer.cpp b/Engine/Source/Editor/UnrealEd/Private/CookOnTheFlyServer.cpp
--- a/Engine/Source/Editor/UnrealEd/Private/CookOnTheFlyServer.cpp
+++ b/Engine/Source/Editor/UnrealEd/Private/CookOnTheFlyServer.cpp
@@ -9041,7 +9041,10 @@ void UCookOnTheFlyServer::GenerateInitialRequests(FBeginCookContext& BeginContex
 			FString OriginalAssetRegistryPath = GetBasedOnReleaseVersionAssetRegistryPath(BasedOnReleaseVersion, TargetPlatform->PlatformName()) / GetAssetRegistryFilename();
 
 			TArray<UE::Cook::FConstructPackageData> BasedOnReleaseDatas;
-			verify(GetAllPackageFilenamesFromAssetRegistry(OriginalAssetRegistryPath, true, false, BasedOnReleaseDatas));
+			bool bFoundAssetRegistry = GetAllPackageFilenamesFromAssetRegistry(OriginalAssetRegistryPath, true, false, BasedOnReleaseDatas);
+			ensureMsgf(bFoundAssetRegistry, TEXT("Unable to find AssetRegistry results from cook of previous version. Expected to find file %s.\n")
+				TEXT("This prevents us from running validation that all files cooked in the previous release are also added to the current release."),
+				*OriginalAssetRegistryPath);
 
 			TArray<const ITargetPlatform*, TInlineAllocator<1>> RequestPlatforms;
 			RequestPlatforms.Add(TargetPlatform);
