Cooker memory : temporary workaround: ignore the free virtual memory limit on Linux until we have fixed AvailableVirtual in FUnixPlatformMemory::GetStats.
#rnx

diff --git a/Engine/Source/Editor/UnrealEd/Private/CookOnTheFlyServer.cpp b/Engine/Source/Editor/UnrealEd/Private/CookOnTheFlyServer.cpp
--- a/Engine/Source/Editor/UnrealEd/Private/CookOnTheFlyServer.cpp
+++ b/Engine/Source/Editor/UnrealEd/Private/CookOnTheFlyServer.cpp
@@ -4261,12 +4261,15 @@ bool UCookOnTheFlyServer::HasExceededMaxMemory() const
 		// you can get that by setting the MemoryMinFreeVirtual and MemoryMinFreePhysical config to be the same
 
 		// AvailableVirtual is actually ullAvailPageFile (commit charge available)
+		// Temporarily ignore MemoryMinFreeVirtual on Linux until we have fixed the value of AvailableVirtual in FUnixPlatformMemory::GetStats
+#if !PLATFORM_LINUX
 		if (MemoryMinFreeVirtual > 0 && MemStats.AvailableVirtual < MemoryMinFreeVirtual)
 		{
 			TriggerMessages.Appendf(TEXT("\n  CookSettings.MemoryMinFreeVirtual: Available virtual memory %dMiB is less than %dMiB."),
 				static_cast<uint32>(MemStats.AvailableVirtual / 1024 / 1024), static_cast<uint32>(MemoryMinFreeVirtual / 1024 / 1024));
 			bFired = true;
 		}
+#endif
 		if (MemoryMinFreePhysical > 0 && MemStats.AvailablePhysical < MemoryMinFreePhysical)
 		{
 			TriggerMessages.Appendf(TEXT("\n  CookSettings.MemoryMinFreePhysical: Available physical memory %dMiB is less than %dMiB."),
