- Saving external actors also backups the world folders expansion state.
- Only actor folders initially collapsed will override the last loaded folders expansion state.

diff --git a/Engine/Source/Editor/UnrealEd/Private/EditorActorFolders.cpp b/Engine/Source/Editor/UnrealEd/Private/EditorActorFolders.cpp
--- a/Engine/Source/Editor/UnrealEd/Private/EditorActorFolders.cpp
+++ b/Engine/Source/Editor/UnrealEd/Private/EditorActorFolders.cpp
@@ -46,6 +46,7 @@ FActorFolders::FActorFolders()
 
 	FEditorDelegates::MapChange.AddRaw(this, &FActorFolders::OnMapChange);
 	FEditorDelegates::PostSaveWorldWithContext.AddRaw(this, &FActorFolders::OnWorldSaved);
+	FEditorDelegates::PostSaveExternalActors.AddRaw(this, &FActorFolders::SaveWorldFoldersState);
 }
 
 FActorFolders::~FActorFolders()
@@ -57,6 +58,7 @@ FActorFolders::~FActorFolders()
 
 	FEditorDelegates::MapChange.RemoveAll(this);
 	FEditorDelegates::PostSaveWorldWithContext.RemoveAll(this);
+	FEditorDelegates::PostSaveExternalActors.RemoveAll(this);
 }
 
 void FActorFolders::AddReferencedObjects(FReferenceCollector& Collector)
@@ -169,10 +171,14 @@ void FActorFolders::OnMapChange(uint32 MapChangeFlags)
 
 void FActorFolders::OnWorldSaved(UWorld* World, FObjectPostSaveContext ObjectSaveContext)
 {
-	// Attempt to save the folder state
+	SaveWorldFoldersState(World);
+}
+
+void FActorFolders::SaveWorldFoldersState(UWorld* World)
+{
 	if (UWorldFolders** Folders = WorldFolders.Find(World))
 	{
-		(*Folders)->OnWorldSaved();
+		(*Folders)->SaveState();
 	}
 }
 
@@ -451,8 +457,11 @@ void FActorFolders::OnActorFolderAdded(UActorFolder* InActorFolder)
 	check(Level->IsUsingActorFolders());
 
 	AddFolderToWorld(*Level->GetWorld(), InActorFolder->GetFolder());
-	// @todo_ow: This code completely overrides the expanded state that was read from the json file (see UWorldFolders::LoadState)
-	SetIsFolderExpanded(*Level->GetWorld(), InActorFolder->GetFolder(), InActorFolder->IsInitiallyExpanded());
+	// To avoid overriding the expanded state initialized by UWorldFolders::LoadState, only override value if folder is initially collapsed.
+	if (!InActorFolder->IsInitiallyExpanded())
+	{
+		SetIsFolderExpanded(*Level->GetWorld(), InActorFolder->GetFolder(), false);
+	}
 }
 
 void FActorFolders::OnFolderRootObjectRemoved(UWorld& InWorld, const FFolder::FRootObject& InFolderRootObject)
diff --git a/Engine/Source/Editor/UnrealEd/Private/WorldFolders.cpp b/Engine/Source/Editor/UnrealEd/Private/WorldFolders.cpp
--- a/Engine/Source/Editor/UnrealEd/Private/WorldFolders.cpp
+++ b/Engine/Source/Editor/UnrealEd/Private/WorldFolders.cpp
@@ -277,11 +277,6 @@ void UWorldFolders::Serialize(FArchive& Ar)
 	Super::Serialize(Ar);
 }
 
-void UWorldFolders::OnWorldSaved()
-{
-	SaveState();
-}
-
 FString UWorldFolders::GetWorldStateFilename() const
 {
 	UPackage* Package = World->GetOutermost();
diff --git a/Engine/Source/Editor/UnrealEd/Public/EditorActorFolders.h b/Engine/Source/Editor/UnrealEd/Public/EditorActorFolders.h
--- a/Engine/Source/Editor/UnrealEd/Public/EditorActorFolders.h
+++ b/Engine/Source/Editor/UnrealEd/Public/EditorActorFolders.h
@@ -203,6 +203,9 @@ private:
 	/** Called after a world has been saved */
 	void OnWorldSaved(UWorld* World, FObjectPostSaveContext ObjectSaveContext);
 
+	/** Attempt to save the folders state */
+	void SaveWorldFoldersState(UWorld* World);
+
 	/** Remove any references to folder arrays for dead worlds */
 	void Housekeeping();
 
diff --git a/Engine/Source/Editor/UnrealEd/Public/WorldFolders.h b/Engine/Source/Editor/UnrealEd/Public/WorldFolders.h
--- a/Engine/Source/Editor/UnrealEd/Public/WorldFolders.h
+++ b/Engine/Source/Editor/UnrealEd/Public/WorldFolders.h
@@ -87,7 +87,7 @@ public:
 	bool ContainsFolder(const FFolder& InFolder) const;
 	void ForEachFolder(TFunctionRef<bool(const FFolder&)> Operation);
 	void ForEachFolderWithRootObject(const FFolder::FRootObject& InFolderRootObject, TFunctionRef<bool(const FFolder&)> Operation);
-	void OnWorldSaved();
+	void SaveState();
 	UWorld* GetWorld() const;
 
 	//~ Begin UObject
@@ -109,7 +109,6 @@ private:
 
 	FString GetWorldStateFilename() const;
 	void LoadState();
-	void SaveState();
 
 	TUniquePtr<FWorldPersistentFolders> PersistentFolders;
 	TUniquePtr<FWorldTransientFolders> TransientFolders;
