Add API for native scene import.

The goal of this API is to allow using the importing pipeline with FBX scene objects you have created yourself, or have been handed from a different API.

I have tested skeletal mesh importing with this workflow and it seems to work!

diff --git a/Engine/Source/Editor/UnrealEd/Private/Fbx/FbxMainImport.cpp b/Engine/Source/Editor/UnrealEd/Private/Fbx/FbxMainImport.cpp
--- a/Engine/Source/Editor/UnrealEd/Private/Fbx/FbxMainImport.cpp
+++ b/Engine/Source/Editor/UnrealEd/Private/Fbx/FbxMainImport.cpp
@@ -533,8 +533,9 @@ FFbxImporter::FFbxImporter()
 	, ImportOptions(NULL)
 	, GeometryConverter(NULL)
 	, SdkManager(NULL)
-	, Importer( NULL )
+	, Importer(NULL)
 	, bFirstMesh(true)
+	, FbxCreator(UnFbx::EFbxCreator::Unknow)
 	, Logger(NULL)
 {
 	// Create the SdkManager
@@ -1798,6 +1799,19 @@ bool FFbxImporter::ImportFromFile(const FString& Filename, const FString& Type,
 	return Result;
 }
 
+void FFbxImporter::ImportFromScene(FbxScene* InScene)
+{
+	if (!InScene)
+	{
+		UE_LOG(LogFbx, Warning, TEXT("Attempted to import invalid scene."));
+		return;
+	}
+
+	ClearAllCaches();
+	Scene = InScene;
+	FbxCreator = EFbxCreator::Unknow;
+}
+
 FString FFbxImporter::MakeName(const ANSICHAR* Name)
 {
 	const TCHAR SpecialChars[] = {TEXT('.'), TEXT(','), TEXT('/'), TEXT('`'), TEXT('%')};
diff --git a/Engine/Source/Editor/UnrealEd/Public/FbxImporter.h b/Engine/Source/Editor/UnrealEd/Public/FbxImporter.h
--- a/Engine/Source/Editor/UnrealEd/Public/FbxImporter.h
+++ b/Engine/Source/Editor/UnrealEd/Public/FbxImporter.h
@@ -687,6 +687,14 @@ public:
 	 */
 	UNREALED_API bool ImportFromFile(const FString& Filename, const FString& Type, bool bPreventMaterialNameClash = false);
 
+	/**
+	 * Prime the importer with an already existing scene object.
+	 *
+	 * @param Scene - The scene we want to load.
+	 */
+	UE_DEPRECATED(5.1, "Do not use this function.")
+	UNREALED_API void ImportFromScene(FbxScene* InScene);
+
 	/**
 	 * Retrieve the FBX loader's error message explaining its failure to read a given FBX file.
 	 * Note that the message should be valid even if the parser is successful and may contain warnings.
