Editor: add function FModeToolkit::RebuildToolPalette() to force tool palette rebuild based on current tool palette names provided by the toolkit subclass. This allows an EdMode to modify the tool palette on-demand.

diff --git a/Engine/Source/Editor/UnrealEd/Private/Toolkits/BaseToolkit.cpp b/Engine/Source/Editor/UnrealEd/Private/Toolkits/BaseToolkit.cpp
--- a/Engine/Source/Editor/UnrealEd/Private/Toolkits/BaseToolkit.cpp
+++ b/Engine/Source/Editor/UnrealEd/Private/Toolkits/BaseToolkit.cpp
@@ -395,7 +395,7 @@ void FModeToolkit::InvokeUI()
 				// Also build the toolkit here 
 				TArray<FName> PaletteNames;
 				GetToolPaletteNames(PaletteNames);
-				for (auto Palette : PaletteNames)
+				for (const FName& Palette : PaletteNames)
 				{
 					TSharedRef<SWidget> PaletteWidget = CreatePaletteWidget(CommandList, ScriptableMode->GetModeInfo().ToolbarCustomizationName, Palette);
 					ActiveToolBarRows.Emplace(ScriptableMode->GetID(), Palette, GetToolPaletteDisplayName(Palette), PaletteWidget);
@@ -408,6 +408,37 @@ void FModeToolkit::InvokeUI()
 	}
 }
 
+
+
+void FModeToolkit::RebuildModeToolPalette()
+{
+	if (ModeUILayer.IsValid() && HasIntegratedToolPalettes())
+	{
+		if ( TSharedPtr<FAssetEditorModeUILayer> ModeUILayerPtr = ModeUILayer.Pin() )
+		{
+			TSharedPtr<FUICommandList> CommandList;
+			if (GetScriptableEditorMode().IsValid())
+			{
+				UEdMode* ScriptableMode = GetScriptableEditorMode().Get();
+				CommandList = GetToolkitCommands();
+				ActiveToolBarRows.Reset();
+
+				TArray<FName> PaletteNames;
+				GetToolPaletteNames(PaletteNames);
+				for (const FName& Palette : PaletteNames)
+				{
+					TSharedRef<SWidget> PaletteWidget = CreatePaletteWidget(CommandList, ScriptableMode->GetModeInfo().ToolbarCustomizationName, Palette);
+					ActiveToolBarRows.Emplace(ScriptableMode->GetID(), Palette, GetToolPaletteDisplayName(Palette), PaletteWidget);
+				}
+
+				RebuildModeToolBar();
+			}
+		}
+	}
+
+}
+
+
 TSharedRef<SDockTab> FModeToolkit::CreatePrimaryModePanel(const FSpawnTabArgs& Args)
 {
 	TSharedRef<SWidget> TabContent = SNew(SBorder)
diff --git a/Engine/Source/Editor/UnrealEd/Public/Toolkits/BaseToolkit.h b/Engine/Source/Editor/UnrealEd/Public/Toolkits/BaseToolkit.h
--- a/Engine/Source/Editor/UnrealEd/Public/Toolkits/BaseToolkit.h
+++ b/Engine/Source/Editor/UnrealEd/Public/Toolkits/BaseToolkit.h
@@ -163,6 +163,12 @@ public:
 	FName GetCurrentPalette() const;
 	void SetModeSettingsObject(UObject* InSettingsObject);
 	virtual void InvokeUI();
+
+	/**
+	 * Force the Mode Toolbar/Palette to be repopulated with the current ToolPaletteNames
+	 */
+	virtual void RebuildModeToolPalette();
+
 protected:
 	virtual void OnToolStarted(UInteractiveToolManager* Manager, UInteractiveTool* Tool);
 	virtual void OnToolEnded(UInteractiveToolManager* Manager, UInteractiveTool* Tool);
