WorldThumbnailRenderer: Create RF_Transient ALevelBounds

diff --git a/Engine/Source/Editor/UnrealEd/Private/ThumbnailRendering/WorldThumbnailRenderer.cpp b/Engine/Source/Editor/UnrealEd/Private/ThumbnailRendering/WorldThumbnailRenderer.cpp
--- a/Engine/Source/Editor/UnrealEd/Private/ThumbnailRendering/WorldThumbnailRenderer.cpp
+++ b/Engine/Source/Editor/UnrealEd/Private/ThumbnailRendering/WorldThumbnailRenderer.cpp
@@ -110,6 +110,7 @@ FSceneView* UWorldThumbnailRenderer::CreateView(UWorld* World, FSceneViewFamily*
 			{
 				// Ensure a Level Bounds Actor exists for future renders
 				FActorSpawnParameters SpawnParameters;
+				SpawnParameters.ObjectFlags |= RF_Transient;
 				SpawnParameters.OverrideLevel = Level;
 				LevelBounds = World->SpawnActor<ALevelBounds>(SpawnParameters);
 				LevelBounds->UpdateLevelBoundsImmediately();
