Implement proper support for togglable stationary direct lighting on mobile and delete 'bake stationary as static' path.
This change makes HQ and LQ lightmaps to have unified lighting content and now they differ only in encoding.

diff --git a/Engine/Shaders/Private/DeferredLightingCommon.ush b/Engine/Shaders/Private/DeferredLightingCommon.ush
--- a/Engine/Shaders/Private/DeferredLightingCommon.ush
+++ b/Engine/Shaders/Private/DeferredLightingCommon.ush
@@ -49,7 +49,7 @@ struct FDeferredLightData
 	/** Intensity of non-shadow-casting contact shadows */
 	float ContactShadowNonShadowCastingIntensity;
 	float2 DistanceFadeMAD;
-	float4 ShadowMapChannelMask;
+	half4 ShadowMapChannelMask;
 	/** Whether ContactShadowLength is in World Space or in Screen Space. */
 	bool ContactShadowLengthInWS;
 	/** Whether to use inverse squared falloff. */
@@ -182,7 +182,7 @@ float ShadowRayCast(
 #error "Must set SUPPORT_CONTACT_SHADOWS"
 #endif
 
-void GetShadowTerms(float SceneDepth, float4 PrecomputedShadowFactors, uint ShadingModelID, float ContactShadowOpacity, FDeferredLightData LightData, float3 TranslatedWorldPosition, half3 L, half4 LightAttenuation, float Dither, inout FShadowTerms Shadow)
+void GetShadowTerms(float SceneDepth, half4 PrecomputedShadowFactors, uint ShadingModelID, float ContactShadowOpacity, FDeferredLightData LightData, float3 TranslatedWorldPosition, half3 L, half4 LightAttenuation, float Dither, inout FShadowTerms Shadow)
 {
 	float ContactShadowLength = 0.0f;
 	const float ContactShadowLengthScreenScale = View.ClipToView[1][1] * SceneDepth;
@@ -195,10 +195,14 @@ void GetShadowTerms(float SceneDepth, float4 PrecomputedShadowFactors, uint Shad
 		// LightAttenuation: Light function + per-object shadows in z, per-object SSS shadowing in w, 
 		// Whole scene directional light shadows in x, whole scene directional light SSS shadows in y
 		// Get static shadowing from the appropriate GBuffer channel
-		float UsesStaticShadowMap = dot(LightData.ShadowMapChannelMask, float4(1, 1, 1, 1));
-		float StaticShadowing = lerp(1, dot(PrecomputedShadowFactors, LightData.ShadowMapChannelMask), UsesStaticShadowMap);
+#if ALLOW_STATIC_LIGHTING
+		half UsesStaticShadowMap = dot(LightData.ShadowMapChannelMask, half4(1, 1, 1, 1));
+		half StaticShadowing = lerp(1, dot(PrecomputedShadowFactors, LightData.ShadowMapChannelMask), UsesStaticShadowMap);
+#else
+		half StaticShadowing = 1.0f;
+#endif
 
-		if (LightData.bRadialLight)
+		if (LightData.bRadialLight || SHADING_PATH_MOBILE)
 		{
 			// Remapping the light attenuation buffer (see ShadowRendering.cpp)
 
diff --git a/Engine/Shaders/Private/DeferredShadingCommon.ush b/Engine/Shaders/Private/DeferredShadingCommon.ush
--- a/Engine/Shaders/Private/DeferredShadingCommon.ush
+++ b/Engine/Shaders/Private/DeferredShadingCommon.ush
@@ -649,7 +649,7 @@ FGBufferData MobileDecodeGBuffer(half4 InGBufferA, half4 InGBufferB, half4 InGBu
 	GBuffer.BaseColor = DecodeBaseColor(InGBufferC.rgb);
 #if ALLOW_STATIC_LIGHTING
 	GBuffer.GBufferAO = 1;
-	GBuffer.PrecomputedShadowFactors = InGBufferC.a;
+	GBuffer.PrecomputedShadowFactors = half4(InGBufferC.a, 1, 1, 1);
 #else
 	GBuffer.GBufferAO = InGBufferC.a;
 	GBuffer.PrecomputedShadowFactors = 1.0;
diff --git a/Engine/Shaders/Private/LightmapCommon.ush b/Engine/Shaders/Private/LightmapCommon.ush
--- a/Engine/Shaders/Private/LightmapCommon.ush
+++ b/Engine/Shaders/Private/LightmapCommon.ush
@@ -292,7 +292,7 @@ half4 GetPrecomputedShadowMasks(VTPageTableResult LightmapVTPageTableResult, FVe
 }
 
 // Used by mobile renderer only
-half GetPrimaryPrecomputedShadowMask(VTPageTableResult LightmapVTPageTableResult, FVertexFactoryInterpolantsVSToPS Interpolants, FMaterialPixelParameters MaterialParameters)
+half4 GetPrimaryPrecomputedShadowMask(VTPageTableResult LightmapVTPageTableResult, FVertexFactoryInterpolantsVSToPS Interpolants, FMaterialPixelParameters MaterialParameters)
 {
 	#if STATICLIGHTING_TEXTUREMASK && STATICLIGHTING_SIGNEDDISTANCEFIELD
 
@@ -302,15 +302,15 @@ half GetPrimaryPrecomputedShadowMask(VTPageTableResult LightmapVTPageTableResult
 
 		// Fetch the distance field data
 		#if LIGHTMAP_VT_ENABLED
-			half DistanceField = SampleLightmapVT(LightmapVTPageTableResult, 2u, LightmapDataIndex, LightmapResourceCluster.VTStaticShadowTexture, LightmapResourceCluster.StaticShadowTextureSampler);
+			half4 DistanceField = SampleLightmapVT(LightmapVTPageTableResult, 2u, LightmapDataIndex, LightmapResourceCluster.VTStaticShadowTexture, LightmapResourceCluster.StaticShadowTextureSampler);
 		#else
-			half DistanceField = Texture2DSample(LightmapResourceCluster.StaticShadowTexture, LightmapResourceCluster.StaticShadowTextureSampler, ShadowMapCoordinate).r;
+			half4 DistanceField = Texture2DSample(LightmapResourceCluster.StaticShadowTexture, LightmapResourceCluster.StaticShadowTextureSampler, ShadowMapCoordinate);
 		#endif
 		float4 InvUniformPenumbraSizes = GetLightmapData(LightmapDataIndex).InvUniformPenumbraSizes;
-		float DistanceFieldBias = -.5f * InvUniformPenumbraSizes.x + .5f;
+		float4 DistanceFieldBias = -.5f * InvUniformPenumbraSizes + .5f;
 		// Compute shadow factors by scaling and biasing the distance
-		half ShadowFactor = saturate( DistanceField * InvUniformPenumbraSizes.x + DistanceFieldBias );
-		return GetLightmapData(LightmapDataIndex).StaticShadowMapMasks.r * ShadowFactor * ShadowFactor;
+		half4 ShadowFactor = saturate( DistanceField * InvUniformPenumbraSizes + DistanceFieldBias );
+		return GetLightmapData(LightmapDataIndex).StaticShadowMapMasks * ShadowFactor * ShadowFactor;
 	#elif MOVABLE_DIRECTIONAL_LIGHT
 		// Do this before checking for lightmaps as we might have a lightmap + movable directional light
 		return 1.0f;
@@ -323,7 +323,7 @@ half GetPrimaryPrecomputedShadowMask(VTPageTableResult LightmapVTPageTableResult
 			// output per-primitive directional light shadowing if requested
 			if ((GetPrimitiveData(MaterialParameters).Flags & PRIMITIVE_SCENE_DATA_FLAG_USE_SINGLE_SAMPLE_SHADOW_SL) != 0 && ResolvedView.IndirectLightingCacheShowFlag > 0.0f)
 			{
-				return IndirectLightingCache.DirectionalLightShadowing;
+				return half4(IndirectLightingCache.DirectionalLightShadowing, 1, 1, 1);
 			}
 		#endif
 	#endif
diff --git a/Engine/Shaders/Private/MobileBasePassPixelShader.usf b/Engine/Shaders/Private/MobileBasePassPixelShader.usf
--- a/Engine/Shaders/Private/MobileBasePassPixelShader.usf
+++ b/Engine/Shaders/Private/MobileBasePassPixelShader.usf
@@ -552,11 +552,10 @@ void Main(
 	FLightAccumulator DirectLighting = (FLightAccumulator)0;
 	LightAccumulator_AddSplit(DirectLighting, DiffuseColor, 0.0f, DiffuseColor, 1.0f, false);
 
-	half Shadow = GetPrimaryPrecomputedShadowMask(LightmapVTPageTableResult, Interpolants, MaterialParameters).r;
+	GBuffer.PrecomputedShadowFactors = GetPrimaryPrecomputedShadowMask(LightmapVTPageTableResult, Interpolants, MaterialParameters);
 
 #if DEFERRED_SHADING_PATH
 	GBuffer.IndirectIrradiance = IndirectIrradiance;
-	GBuffer.PrecomputedShadowFactors.r = Shadow;
 	MobileEncodeGBuffer(GBuffer, OutGBufferA, OutGBufferB, OutGBufferC, OutGBufferD);
 #else
 
@@ -565,7 +564,7 @@ void Main(
 	half4 DynamicShadowFactors = 1.0f;
 
 	// Directional light
-	AccumulateDirectionalLighting(GBuffer, MaterialParameters.WorldPosition_CamRelative, CameraVector, MaterialParameters.ScreenPosition, SvPosition, Shadow, DynamicShadowFactors, DirectLighting);
+	AccumulateDirectionalLighting(GBuffer, MaterialParameters.WorldPosition_CamRelative, CameraVector, MaterialParameters.ScreenPosition, SvPosition, DynamicShadowFactors, DirectLighting);
 
 	float2 LocalPosition = SvPosition.xy - ResolvedView.ViewRectMin.xy;
 	uint GridIndex = ComputeLightGridCellIndex(uint2(LocalPosition.x, LocalPosition.y), SvPosition.w);
diff --git a/Engine/Shaders/Private/MobileDeferredShading.usf b/Engine/Shaders/Private/MobileDeferredShading.usf
--- a/Engine/Shaders/Private/MobileDeferredShading.usf
+++ b/Engine/Shaders/Private/MobileDeferredShading.usf
@@ -152,8 +152,6 @@ void MobileDirectionalLightPS(
 
 	FLightAccumulator DirectLighting = (FLightAccumulator)0;
 
-	// Check movable light param to determine if we should be using precomputed shadows
-	half Shadow = LightFunctionParameters2.z > 0.0f ? 1.0f : GBuffer.PrecomputedShadowFactors.r;
 	float4 ScreenPosition = SvPositionToScreenPosition(float4(SvPosition.xyz, GBuffer.Depth));
 	
 	float2 LocalPosition = SvPosition.xy - View.ViewRectMin.xy;
@@ -170,7 +168,7 @@ void MobileDirectionalLightPS(
 
 	// Directional light
 	half4 DynamicShadowFactors = 1.0f;
-	AccumulateDirectionalLighting(GBuffer, TranslatedWorldPosition, CameraVector, ScreenPosition, SvPosition, Shadow, DynamicShadowFactors, DirectLighting);
+	AccumulateDirectionalLighting(GBuffer, TranslatedWorldPosition, CameraVector, ScreenPosition, SvPosition, DynamicShadowFactors, DirectLighting);
 
 #if ENABLE_CLUSTERED_REFLECTION || ENABLE_SKY_LIGHT || ENABLE_PLANAR_REFLECTION
 	// If we have a single directional light, apply relfection and sky contrubution here
@@ -220,7 +218,7 @@ FDeferredLightData SetupLightDataForStandardDeferred()
 //	LightData.ContactShadowLengthInWS = DeferredLightUniforms.ContactShadowLength < 0.0f;
 //	LightData.ContactShadowNonShadowCastingIntensity = DeferredLightUniforms.ContactShadowNonShadowCastingIntensity;
 //	LightData.DistanceFadeMAD = DeferredLightUniforms.DistanceFadeMAD;
-//	LightData.ShadowMapChannelMask = DeferredLightUniforms.ShadowMapChannelMask;
+	LightData.ShadowMapChannelMask = DeferredLightUniforms.ShadowMapChannelMask;
 	LightData.ShadowedBits = DeferredLightUniforms.ShadowedBits;
 	LightData.bInverseSquared = DeferredLightUniforms.FalloffExponent == 0; // Directional lights don't use 'inverse squared attenuation'
 	LightData.bRadialLight = true;
diff --git a/Engine/Shaders/Private/MobileLightingCommon.ush b/Engine/Shaders/Private/MobileLightingCommon.ush
--- a/Engine/Shaders/Private/MobileLightingCommon.ush
+++ b/Engine/Shaders/Private/MobileLightingCommon.ush
@@ -182,33 +182,30 @@ half MobileDirectionalLightCSM(float2 ScreenPosition, float SceneDepth, inout fl
 	Mobile Directional Light.
 ------------------------------------------------------------------------------*/
 
-void AccumulateDirectionalLighting(FGBufferData GBuffer, float3 TranslatedWorldPosition, half3 CameraVector, float4 ScreenPosition, float4 SvPosition, half Shadow, inout half4 DynamicShadowFactors, inout FLightAccumulator DirectLighting)
+void AccumulateDirectionalLighting(FGBufferData GBuffer, float3 TranslatedWorldPosition, half3 CameraVector, float4 ScreenPosition, float4 SvPosition, inout half4 DynamicShadowFactors, inout FLightAccumulator DirectLighting)
 {
 #if USE_SHADOWMASKTEXTURE
 	half4 PreviewShadowMapChannelMask = half4(1.0f, 0.0f, 0.0f, 0.0f);
 	DynamicShadowFactors = Texture2DSample(MobileBasePass.ScreenSpaceShadowMaskTexture, MobileBasePass.ScreenSpaceShadowMaskSampler, SvPositionToBufferUV(SvPosition));
 	DynamicShadowFactors = DecodeLightAttenuation(DynamicShadowFactors);
 		
-	#if ENABLE_CLUSTERED_LIGHTS
-		const FDirectionalLightData DirectionalLightData = GetDirectionalLightData(0);
-		PreviewShadowMapChannelMask = UnpackShadowMapChannelMask(DirectionalLightData.DirectionalLightShadowMapChannelMask >> 4);
-	#endif
+	PreviewShadowMapChannelMask = UnpackShadowMapChannelMask(MobileDirectionalLight.DirectionalLightShadowMapChannelMask >> 4);
 
-	half ShadowMap = dot(PreviewShadowMapChannelMask, DynamicShadowFactors);
-	Shadow = ApplyPrecomputedShadowMask(ShadowMap, Shadow);
+	half DynamicShadowing = dot(PreviewShadowMapChannelMask, DynamicShadowFactors);
 #else
 
 	float  ShadowPositionZ = 0;
 	half ShadowMap = MobileDirectionalLightCSM(ScreenPosition.xy, ScreenPosition.w, ShadowPositionZ);
+	half DynamicShadowing = 1.0f;
 
 	#if !MOBILE_DEFERRED_SHADING && DIRECTIONAL_LIGHT_CSM && !MATERIAL_SHADINGMODEL_SINGLELAYERWATER
 		// Cascaded Shadow Map
 		if (IsCSMEnabled())
 		{		
-			Shadow = ApplyPrecomputedShadowMask(ShadowMap, Shadow);
+			DynamicShadowing = ShadowMap;
 		}
 	#elif MOBILE_DEFERRED_SHADING
-		Shadow = min(ShadowMap, Shadow);
+		DynamicShadowing = ShadowMap;
 	#endif
 #endif
 
@@ -220,8 +217,9 @@ void AccumulateDirectionalLighting(FGBufferData GBuffer, float3 TranslatedWorldP
 	LightData.SpecularScale = MobileDirectionalLight.DirectionalLightDistanceFadeMADAndSpecularScale.z;
 	LightData.ShadowedBits = 1;
 	LightData.HairTransmittance = InitHairTransmittanceData();
+	LightData.ShadowMapChannelMask = UnpackShadowMapChannelMask(MobileDirectionalLight.DirectionalLightShadowMapChannelMask);
 
-	half4 LightAttenuation = half4(Shadow, Shadow, 1, 1);
+	half4 LightAttenuation = half4(1, 1, DynamicShadowing, DynamicShadowing);
 
 	float SurfaceShadow = 0;
 	FLightAccumulator NewLighting = AccumulateDynamicLighting(TranslatedWorldPosition, CameraVector, GBuffer, 1, GBuffer.ShadingModelID, LightData, LightAttenuation, 0, uint2(0, 0), SurfaceShadow);
@@ -443,7 +441,7 @@ FDeferredLightData ConvertLocalToDeferredLight(const FLocalLightData In)
 	// bits [7:0]
 //	Out.ShadowedBits = (LightTypeAndPackedShadowMapChannelMask & 0xFF) != 0 ? 1 : 0;
 	// Static shadowing uses ShadowMapChannel, dynamic shadows are packed into light attenuation using PreviewShadowMapChannel
-//	Out.ShadowMapChannelMask = UnpackShadowMapChannelMask(LightTypeAndPackedShadowMapChannelMask);
+	Out.ShadowMapChannelMask = UnpackShadowMapChannelMask(LightTypeAndPackedShadowMapChannelMask);
 	Out.ShadowedBits = 1;
 	Out.HairTransmittance = InitHairTransmittanceData();
 	return Out;
@@ -499,18 +497,18 @@ void AccumulateLightGridLocalLighting(const FCulledLightsGridData InLightGridDat
 		// bits [17:16] really
 		uint LightType = LightTypeAndPackedShadowMapChannelMask >> 16;
 
-		half Shadow = 1.0f;
+		half DynamicShadowing = 1.0f;
 #if USE_SHADOWMASKTEXTURE
 		if (LightType == LIGHT_TYPE_SPOT)
 		{
 			half4 PreviewShadowMapChannelMask = UnpackShadowMapChannelMask(LightTypeAndPackedShadowMapChannelMask >> 4);
-			Shadow = dot(PreviewShadowMapChannelMask, DynamicShadowFactors);
+			DynamicShadowing = dot(PreviewShadowMapChannelMask, DynamicShadowFactors);
 		}
 #endif
 
 		FDeferredLightData LightData = ConvertLocalToDeferredLight(LocalLight);
 
-		half4 LightAttenuation = half4(1, 1, Shadow, Shadow);
+		half4 LightAttenuation = half4(1, 1, DynamicShadowing, DynamicShadowing);
 
 		float SurfaceShadow = 0;
 		FLightAccumulator NewLighting = AccumulateDynamicLighting(TranslatedWorldPosition, CameraVector, GBuffer, 1, GBuffer.ShadingModelID, LightData, LightAttenuation, 0, uint2(0, 0), SurfaceShadow);
diff --git a/Engine/Source/Developer/Android/AndroidTargetPlatform/Private/AndroidTargetPlatform.cpp b/Engine/Source/Developer/Android/AndroidTargetPlatform/Private/AndroidTargetPlatform.cpp
--- a/Engine/Source/Developer/Android/AndroidTargetPlatform/Private/AndroidTargetPlatform.cpp
+++ b/Engine/Source/Developer/Android/AndroidTargetPlatform/Private/AndroidTargetPlatform.cpp
@@ -502,12 +502,7 @@ void FAndroidTargetPlatform::GetTextureFormats( const UTexture* Texture, TArray<
 
 	for (FName& TextureFormatName : LayerFormats)
 	{
-		if (Texture->LODGroup == TEXTUREGROUP_Shadowmap)
-		{
-			// forward rendering only needs one channel for shadow maps
-			TextureFormatName = FName(TEXT("G8"));
-		}
-		else if (Texture->LODGroup == TEXTUREGROUP_Terrain_Weightmap && bCompressLandscapeWeightMaps)
+		if (Texture->LODGroup == TEXTUREGROUP_Terrain_Weightmap && bCompressLandscapeWeightMaps)
 		{
 			TextureFormatName = AndroidTexFormat::NameAutoDXT;
 		}
diff --git a/Engine/Source/Developer/IOS/IOSTargetPlatform/Private/IOSTargetPlatform.cpp b/Engine/Source/Developer/IOS/IOSTargetPlatform/Private/IOSTargetPlatform.cpp
--- a/Engine/Source/Developer/IOS/IOSTargetPlatform/Private/IOSTargetPlatform.cpp
+++ b/Engine/Source/Developer/IOS/IOSTargetPlatform/Private/IOSTargetPlatform.cpp
@@ -597,12 +597,6 @@ void FIOSTargetPlatform::GetTextureFormats( const UTexture* Texture, TArray< TAr
 	TArray<FName>& TextureFormatNames = OutFormats.AddDefaulted_GetRef();
 	TextureFormatNames.Reserve(NumLayers);
 
-	// forward rendering only needs one channel for shadow maps
-	if (Texture->LODGroup == TEXTUREGROUP_Shadowmap && !SupportsMetalMRT())
-	{
-		TextureFormatNames.Init(NameG8, NumLayers);
-	}
-
 	// optionaly compress landscape weightmaps for a mobile rendering
 	static const auto CompressLandscapeWeightMapsVar = IConsoleManager::Get().FindTConsoleVariableDataInt(TEXT("r.Mobile.CompressLandscapeWeightMaps"));
 	static const bool bCompressLandscapeWeightMaps = (CompressLandscapeWeightMapsVar && CompressLandscapeWeightMapsVar->GetValueOnAnyThread() != 0);
diff --git a/Engine/Source/Editor/UnrealEd/Private/Lightmass/ImportVolumetricLightmap.cpp b/Engine/Source/Editor/UnrealEd/Private/Lightmass/ImportVolumetricLightmap.cpp
--- a/Engine/Source/Editor/UnrealEd/Private/Lightmass/ImportVolumetricLightmap.cpp
+++ b/Engine/Source/Editor/UnrealEd/Private/Lightmass/ImportVolumetricLightmap.cpp
@@ -89,8 +89,6 @@ struct FImportedVolumetricLightmapBrick
 	float AverageClosestGeometryDistance;
 	TArray<FFloat3Packed> AmbientVector;
 	TArray<FColor> SHCoefficients[6];
-	TArray<FFloat3Packed> LQLightColor;
-	TArray<FColor> LQLightDirection;
 	TArray<FColor> SkyBentNormal;
 	TArray<uint8> DirectionalLightShadowing;
 	TArray<Lightmass::FIrradianceVoxelImportProcessingData> TaskVoxelImportProcessingData;
@@ -148,9 +146,6 @@ bool CopyFromBrickmapTexel(
 			*(FColor*)&BrickData.SHCoefficients[i].Data[LinearDestCellIndex * sizeof(FColor)] = FilteredVolumeLookupReconverted<FColor>(BrickTextureCoordinate, CurrentLevelData.BrickDataDimensions, (const FColor*)BrickData.SHCoefficients[i].Data.GetData());
 		}
 
-		*(FFloat3Packed*)&BrickData.LQLightColor.Data[LinearDestCellIndex * sizeof(FFloat3Packed)] = FilteredVolumeLookupReconverted<FFloat3Packed>(BrickTextureCoordinate, CurrentLevelData.BrickDataDimensions, (const FFloat3Packed*)BrickData.LQLightColor.Data.GetData());
-		*(FColor*)&BrickData.LQLightDirection.Data[LinearDestCellIndex * sizeof(FColor)] = FilteredVolumeLookupReconverted<FColor>(BrickTextureCoordinate, CurrentLevelData.BrickDataDimensions, (const FColor*)BrickData.LQLightDirection.Data.GetData());
-
 		if (BrickData.SkyBentNormal.Data.Num() > 0)
 		{
 			*(FColor*)&BrickData.SkyBentNormal.Data[LinearDestCellIndex * sizeof(FColor)] = FilteredVolumeLookupReconverted<FColor>(BrickTextureCoordinate, CurrentLevelData.BrickDataDimensions, (const FColor*)BrickData.SkyBentNormal.Data.GetData());
@@ -201,9 +196,6 @@ void FLightmassProcessor::ImportIrradianceTasks(bool& bGenerateSkyShadowing, TAr
 					ReadArray(Channel, NewBrick.SHCoefficients[i]);
 				}
 
-				ReadArray(Channel, NewBrick.LQLightColor);
-				ReadArray(Channel, NewBrick.LQLightDirection);
-
 				ReadArray(Channel, NewBrick.SkyBentNormal);
 				ReadArray(Channel, NewBrick.DirectionalLightShadowing);
 
@@ -1085,9 +1077,6 @@ void FLightmassProcessor::ImportVolumetricLightmap()
 		{
 			CurrentLevelData.BrickData.SHCoefficients[i].Format = PF_B8G8R8A8;
 		}
-
-		CurrentLevelData.BrickData.LQLightColor.Format = PF_FloatR11G11B10;
-		CurrentLevelData.BrickData.LQLightDirection.Format = PF_B8G8R8A8;
 	}
 
 	const int32 NumBottomLevelBricksTrimmedByInterpolationError = TrimBricksByInterpolationError(BricksByDepth, VolumetricLightmapSettings);
@@ -1147,10 +1136,6 @@ void FLightmassProcessor::ImportVolumetricLightmap()
 			CurrentLevelData.BrickData.SHCoefficients[i].Resize(TotalBrickDataSize * Stride);
 		}
 
-		CurrentLevelData.BrickData.LQLightColor.Resize(TotalBrickDataSize * GPixelFormats[CurrentLevelData.BrickData.LQLightColor.Format].BlockBytes);
-
-		CurrentLevelData.BrickData.LQLightDirection.Resize(TotalBrickDataSize * GPixelFormats[CurrentLevelData.BrickData.LQLightDirection.Format].BlockBytes);
-
 		VoxelImportProcessingData.Empty(TotalBrickDataSize);
 		VoxelImportProcessingData.AddZeroed(TotalBrickDataSize);
 	}
@@ -1185,22 +1170,6 @@ void FLightmassProcessor::ImportVolumetricLightmap()
 					CurrentLevelData.BrickData.SHCoefficients[i].Data.GetData());
 			}
 
-			CopyBrickToAtlasVolumeTexture(
-				GPixelFormats[CurrentLevelData.BrickData.LQLightColor.Format].BlockBytes,
-				CurrentLevelData.BrickDataDimensions,
-				BrickLayoutPosition,
-				FIntVector(BrickSize),
-				(const uint8*)Brick.LQLightColor.GetData(),
-				CurrentLevelData.BrickData.LQLightColor.Data.GetData());
-
-			CopyBrickToAtlasVolumeTexture(
-				GPixelFormats[CurrentLevelData.BrickData.LQLightDirection.Format].BlockBytes,
-				CurrentLevelData.BrickDataDimensions,
-				BrickLayoutPosition,
-				FIntVector(BrickSize),
-				(const uint8*)Brick.LQLightDirection.GetData(),
-				CurrentLevelData.BrickData.LQLightDirection.Data.GetData());
-
 			if (bGenerateSkyShadowing)
 			{
 				CopyBrickToAtlasVolumeTexture(
@@ -1379,9 +1348,6 @@ void FLightmassProcessor::ImportVolumetricLightmap()
 					{
 						SubLevelData.BrickData.SHCoefficients[i].Format = PF_B8G8R8A8;
 					}
-
-					SubLevelData.BrickData.LQLightColor.Format = PF_FloatR11G11B10;
-					SubLevelData.BrickData.LQLightDirection.Format = PF_B8G8R8A8;
 				}
 
 				SubLevelData.BrickDataDimensions = SubLevelBrickLayoutDimension * PaddedBrickSize;
@@ -1401,10 +1367,6 @@ void FLightmassProcessor::ImportVolumetricLightmap()
 					const int32 Stride = GPixelFormats[SubLevelData.BrickData.SHCoefficients[i].Format].BlockBytes;
 					SubLevelData.BrickData.SHCoefficients[i].Resize(TotalBrickDataSize * Stride);
 				}
-
-				SubLevelData.BrickData.LQLightColor.Resize(TotalBrickDataSize * GPixelFormats[SubLevelData.BrickData.LQLightColor.Format].BlockBytes);
-
-				SubLevelData.BrickData.LQLightDirection.Resize(TotalBrickDataSize * GPixelFormats[SubLevelData.BrickData.LQLightDirection.Format].BlockBytes);
 			}
 		}
 
@@ -1460,26 +1422,6 @@ void FLightmassProcessor::ImportVolumetricLightmap()
 						SubLevelData.BrickData.SHCoefficients[i].Data);
 				}
 
-				CopyBetweenAtlasVolumeTextures(
-					GPixelFormats[CurrentLevelData.BrickData.LQLightColor.Format].BlockBytes,
-					FIntVector(PaddedBrickSize),
-					CurrentLevelData.BrickDataDimensions,
-					BrickLayoutPosition,
-					CurrentLevelData.BrickData.LQLightColor.Data,
-					SubLevelBrickLayoutDimensions[LevelGuid] * PaddedBrickSize,
-					SubLevelBrickLayoutPosition,
-					SubLevelData.BrickData.LQLightColor.Data);
-
-				CopyBetweenAtlasVolumeTextures(
-					GPixelFormats[CurrentLevelData.BrickData.LQLightDirection.Format].BlockBytes,
-					FIntVector(PaddedBrickSize),
-					CurrentLevelData.BrickDataDimensions,
-					BrickLayoutPosition,
-					CurrentLevelData.BrickData.LQLightDirection.Data,
-					SubLevelBrickLayoutDimensions[LevelGuid] * PaddedBrickSize,
-					SubLevelBrickLayoutPosition,
-					SubLevelData.BrickData.LQLightDirection.Data);
-
 				if (bGenerateSkyShadowing)
 				{
 					CopyBetweenAtlasVolumeTextures(
@@ -1583,8 +1525,6 @@ void FLightmassProcessor::ImportVolumetricLightmap()
 				{
 					ConvertBGRA8ToRGBA8ForLayer(SubLevelData.BrickData.SHCoefficients[i]);
 				}
-
-				ConvertBGRA8ToRGBA8ForLayer(SubLevelData.BrickData.LQLightDirection);
 			}
 		}
 	}
@@ -1624,8 +1564,6 @@ void FLightmassProcessor::ImportVolumetricLightmap()
 				SubLevelData.BrickData.SHCoefficients[5].Data.Num() + 
 				SubLevelData.BrickData.SkyBentNormal.Data.Num() + 
 				SubLevelData.BrickData.DirectionalLightShadowing.Data.Num() +
-				SubLevelData.BrickData.LQLightColor.Data.Num() +
-				SubLevelData.BrickData.LQLightDirection.Data.Num() +
 				SubLevelData.SubLevelBrickPositions.Num() * SubLevelData.SubLevelBrickPositions.GetTypeSize() +
 				SubLevelData.IndirectionTextureOriginalValues.Num() * SubLevelData.IndirectionTextureOriginalValues.GetTypeSize();
 		}
diff --git a/Engine/Source/Programs/UnrealLightmass/Private/ImportExport/Exporter.cpp b/Engine/Source/Programs/UnrealLightmass/Private/ImportExport/Exporter.cpp
--- a/Engine/Source/Programs/UnrealLightmass/Private/ImportExport/Exporter.cpp
+++ b/Engine/Source/Programs/UnrealLightmass/Private/ImportExport/Exporter.cpp
@@ -502,9 +502,6 @@ namespace Lightmass
 					WriteArray(TaskData.BrickData[BrickIndex].SHCoefficients[i]);
 				}
 
-				WriteArray(TaskData.BrickData[BrickIndex].LQLightColor);
-				WriteArray(TaskData.BrickData[BrickIndex].LQLightDirection);
-
 				WriteArray(TaskData.BrickData[BrickIndex].SkyBentNormal);
 				WriteArray(TaskData.BrickData[BrickIndex].DirectionalLightShadowing);
 				WriteArray(TaskData.BrickData[BrickIndex].VoxelImportProcessingData);
diff --git a/Engine/Source/Programs/UnrealLightmass/Private/Lighting/AdaptiveVolumetricLightmap.cpp b/Engine/Source/Programs/UnrealLightmass/Private/Lighting/AdaptiveVolumetricLightmap.cpp
--- a/Engine/Source/Programs/UnrealLightmass/Private/Lighting/AdaptiveVolumetricLightmap.cpp
+++ b/Engine/Source/Programs/UnrealLightmass/Private/Lighting/AdaptiveVolumetricLightmap.cpp
@@ -571,11 +571,6 @@ void FStaticLightingSystem::ProcessVolumetricLightmapBrickTask(FVolumetricLightm
 	BrickData.AmbientVector.Empty(TotalBrickSize);
 	BrickData.AmbientVector.AddDefaulted(TotalBrickSize);
 
-	BrickData.LQLightColor.Empty(TotalBrickSize);
-	BrickData.LQLightColor.AddDefaulted(TotalBrickSize);
-	BrickData.LQLightDirection.Empty(TotalBrickSize);
-	BrickData.LQLightDirection.AddDefaulted(TotalBrickSize);
-
 	BrickData.VoxelImportProcessingData.Empty(TotalBrickSize);
 	BrickData.VoxelImportProcessingData.AddDefaulted(TotalBrickSize);
 
@@ -943,36 +938,6 @@ void FIrradianceBrickData::SetFromVolumeLightingSample(int32 Index, const FVolum
 
 	DirectionalLightShadowing[Index] = (uint8)FMath::Clamp<int32>(FMath::RoundToInt(Sample.DirectionalLightShadowing * MAX_uint8), 0, MAX_uint8);
 
-	{
-		FSHVectorRGB3 HQ;
-		FSHVectorRGB3 LQ;
-		for (int32 CoefficientIndex = 0; CoefficientIndex < LM_NUM_SH_COEFFICIENTS; CoefficientIndex++)
-		{
-			HQ.R.V[CoefficientIndex] = Sample.HighQualityCoefficients[CoefficientIndex][0];
-			HQ.G.V[CoefficientIndex] = Sample.HighQualityCoefficients[CoefficientIndex][1];
-			HQ.B.V[CoefficientIndex] = Sample.HighQualityCoefficients[CoefficientIndex][2];
-
-			LQ.R.V[CoefficientIndex] = Sample.LowQualityCoefficients[CoefficientIndex][0];
-			LQ.G.V[CoefficientIndex] = Sample.LowQualityCoefficients[CoefficientIndex][1];
-			LQ.B.V[CoefficientIndex] = Sample.LowQualityCoefficients[CoefficientIndex][2];
-		}
-		FSHVectorRGB3 DirectLight = LQ - HQ;
-		FVector3f MaxLightDir( DirectLight.GetLuminance().GetMaximumDirection() );
-
-		// Set the max direction
-		FLinearColor MaxLightDirAsColor(MaxLightDir);
-		LQLightDirection[Index] = (MaxLightDirAsColor * FLinearColor(.5f, .5f, .5f, .5f) + FLinearColor(.5f, .5f, .5f, .5f)).QuantizeRound();
-		
-		// Set color along the max direction
-		FSHVector3 BrigthestDiffuseTransferSH = FSHVector3::CalcDiffuseTransfer(FVector(MaxLightDir));
-		FLinearColor BrightestLighting = Dot(DirectLight, BrigthestDiffuseTransferSH);
-		BrightestLighting.R = FMath::Max(BrightestLighting.R, 0.0f);
-		BrightestLighting.G = FMath::Max(BrightestLighting.G, 0.0f);
-		BrightestLighting.B = FMath::Max(BrightestLighting.B, 0.0f);
-
-		LQLightColor[Index] = FFloat3Packed(BrightestLighting);
-	}
-
 	FIrradianceVoxelImportProcessingData NewImportData;
 	NewImportData.bInsideGeometry = bInsideGeometry;
 	NewImportData.bBorderVoxel = bBorderVoxel;
diff --git a/Engine/Source/Programs/UnrealLightmass/Private/Lighting/FinalGather.cpp b/Engine/Source/Programs/UnrealLightmass/Private/Lighting/FinalGather.cpp
--- a/Engine/Source/Programs/UnrealLightmass/Private/Lighting/FinalGather.cpp
+++ b/Engine/Source/Programs/UnrealLightmass/Private/Lighting/FinalGather.cpp
@@ -120,11 +120,9 @@ void FStaticLightingSystem::CalculateVolumeSampleIncidentRadiance(
 	RepresentativeVertex.GenerateTriangleTangents();
 
 	FGatheredLightSample3 UpperStaticDirectLighting;
-	// Stationary point and spot light direct contribution
-	FGatheredLightSample3 UpperToggleableDirectLighting;
 	float UpperToggleableDirectionalLightShadowing = 1;
 
-	CalculateApproximateDirectLighting(RepresentativeVertex, LightingSample.GetRadius(), VertexOffsets, .1f, false, false, bDebugThisSample, MappingContext, UpperStaticDirectLighting, UpperToggleableDirectLighting, UpperToggleableDirectionalLightShadowing);
+	CalculateApproximateDirectLighting(RepresentativeVertex, LightingSample.GetRadius(), VertexOffsets, .1f, false, false, bDebugThisSample, MappingContext, UpperStaticDirectLighting, UpperToggleableDirectionalLightShadowing);
 	
 	const double EndUpperDirectLightingTime = FPlatformTime::Seconds();
 	MappingContext.Stats.VolumetricLightmapDirectLightingTime += EndUpperDirectLightingTime - EndGatherTime;
@@ -164,11 +162,9 @@ void FStaticLightingSystem::CalculateVolumeSampleIncidentRadiance(
 
 	FLightingCacheGatherInfo LowerGatherInfo;
 	FGatheredLightSample3 LowerStaticDirectLighting;
-	// Stationary point and spot light direct contribution
-	FGatheredLightSample3 LowerToggleableDirectLighting;
 	float LowerToggleableDirectionalLightShadowing = 1;
 
-	CalculateApproximateDirectLighting(RepresentativeVertex, LightingSample.GetRadius(), VertexOffsets, .1f, false, false, bDebugThisSample, MappingContext, LowerStaticDirectLighting, LowerToggleableDirectLighting, LowerToggleableDirectionalLightShadowing);
+	CalculateApproximateDirectLighting(RepresentativeVertex, LightingSample.GetRadius(), VertexOffsets, .1f, false, false, bDebugThisSample, MappingContext, LowerStaticDirectLighting, LowerToggleableDirectionalLightShadowing);
 
 	const double EndLowerDirectLightingTime = FPlatformTime::Seconds();
 	MappingContext.Stats.VolumetricLightmapDirectLightingTime += EndLowerDirectLightingTime - EndUpperFinalGatherTime;
@@ -196,11 +192,6 @@ void FStaticLightingSystem::CalculateVolumeSampleIncidentRadiance(
 
 	const FGatheredLightSample3 CombinedIndirectLighting = UpperHemisphereSample + LowerHemisphereSample;
 	const FGatheredLightSample3 CombinedHighQualitySample = UpperStaticDirectLighting + LowerStaticDirectLighting + CombinedIndirectLighting;
-	 
-	// Composite point and spot stationary direct lighting into the low quality volume samples, since we won't be applying them dynamically
-	FGatheredLightSample3 CombinedLowQualitySample = UpperStaticDirectLighting + UpperToggleableDirectLighting + LowerStaticDirectLighting + LowerToggleableDirectLighting + CombinedIndirectLighting;
-	// Composite stationary sky light contribution to the low quality volume samples, since we won't be applying it dynamically
-	CombinedLowQualitySample = CombinedLowQualitySample + UpperHemisphereSample.StationarySkyLighting + LowerHemisphereSample.StationarySkyLighting;
 
 	for (int32 CoefficientIndex = 0; CoefficientIndex < LM_NUM_SH_COEFFICIENTS; CoefficientIndex++)
 	{
@@ -208,9 +199,10 @@ void FStaticLightingSystem::CalculateVolumeSampleIncidentRadiance(
 		LightingSample.HighQualityCoefficients[CoefficientIndex][1] = CombinedHighQualitySample.SHVector.G.V[CoefficientIndex];
 		LightingSample.HighQualityCoefficients[CoefficientIndex][2] = CombinedHighQualitySample.SHVector.B.V[CoefficientIndex];
 
-		LightingSample.LowQualityCoefficients[CoefficientIndex][0] = CombinedLowQualitySample.SHVector.R.V[CoefficientIndex];
-		LightingSample.LowQualityCoefficients[CoefficientIndex][1] = CombinedLowQualitySample.SHVector.G.V[CoefficientIndex];
-		LightingSample.LowQualityCoefficients[CoefficientIndex][2] = CombinedLowQualitySample.SHVector.B.V[CoefficientIndex];
+		// Copy to LQ coefficients
+		LightingSample.LowQualityCoefficients[CoefficientIndex][0] = CombinedHighQualitySample.SHVector.R.V[CoefficientIndex];
+		LightingSample.LowQualityCoefficients[CoefficientIndex][1] = CombinedHighQualitySample.SHVector.G.V[CoefficientIndex];
+		LightingSample.LowQualityCoefficients[CoefficientIndex][2] = CombinedHighQualitySample.SHVector.B.V[CoefficientIndex];
 	}
 
 	LightingSample.DirectionalLightShadowing = FMath::Max(UpperToggleableDirectionalLightShadowing, LowerToggleableDirectionalLightShadowing);
@@ -1324,7 +1316,6 @@ SampleType FStaticLightingSystem::IncomingRadianceAdaptive(
 			CombinedSkyUnoccludedDirection += FilteredLighting.UnoccludedSkyVector;
 
 			IncomingRadiance.AddIncomingRadiance(Radiance, SampleWeight, TangentPathDirection, WorldPathDirection);
-			IncomingRadiance.AddIncomingStationarySkyLight(FilteredLighting.StationarySkyLighting, SampleWeight, TangentPathDirection, WorldPathDirection);
 			checkSlow(IncomingRadiance.AreFloatsValid());
 			NumSamplesOccluded += FilteredLighting.NumSamplesOccluded;
 		}
diff --git a/Engine/Source/Programs/UnrealLightmass/Private/Lighting/GatheredLightingSample.h b/Engine/Source/Programs/UnrealLightmass/Private/Lighting/GatheredLightingSample.h
--- a/Engine/Source/Programs/UnrealLightmass/Private/Lighting/GatheredLightingSample.h
+++ b/Engine/Source/Programs/UnrealLightmass/Private/Lighting/GatheredLightingSample.h
@@ -210,11 +210,6 @@ public:
 		AddWeighted(FGatheredLightSampleUtil::PointLightWorldSpace<SHOrder>(IncomingRadiance, TangentSpaceDirection, WorldSpaceDirection), Weight);
 	}
 
-	inline void AddIncomingStationarySkyLight(const FLinearColor& IncomingSkyLight, float Weight, const FVector4f& TangentSpaceDirection, const FVector4f& WorldSpaceDirection)
-	{
-		StationarySkyLighting.AddWeighted(FGatheredLightSampleUtil::PointLightWorldSpace<SHOrder>(IncomingSkyLight, TangentSpaceDirection, WorldSpaceDirection), Weight);
-	}
-
 	bool AreFloatsValid() const;
 
 	TFinalGatherSample operator*(float Scalar) const
@@ -222,7 +217,6 @@ public:
 		TFinalGatherSample Result;
 		(TGatheredLightSample<SHOrder>&)Result = (const TGatheredLightSample<SHOrder>&)(*this) * Scalar;
 		Result.Occlusion = Occlusion * Scalar;
-		Result.StationarySkyLighting = StationarySkyLighting * Scalar;
 		return Result;
 	}
 
@@ -231,7 +225,6 @@ public:
 		TFinalGatherSample Result;
 		(TGatheredLightSample<SHOrder>&)Result = (const TGatheredLightSample<SHOrder>&)(*this) + (const TGatheredLightSample<SHOrder>&)SampleB;
 		Result.Occlusion = Occlusion + SampleB.Occlusion;
-		Result.StationarySkyLighting = StationarySkyLighting + SampleB.StationarySkyLighting;
 		return Result;
 	}
 };
diff --git a/Engine/Source/Programs/UnrealLightmass/Private/Lighting/LightingSystem.h b/Engine/Source/Programs/UnrealLightmass/Private/Lighting/LightingSystem.h
--- a/Engine/Source/Programs/UnrealLightmass/Private/Lighting/LightingSystem.h
+++ b/Engine/Source/Programs/UnrealLightmass/Private/Lighting/LightingSystem.h
@@ -1356,10 +1356,6 @@ struct FIrradianceBrickData
 	TArray<FColor> SkyBentNormal;
 	TArray<uint8> DirectionalLightShadowing;
 
-	/** direct lighting from stationary lights for low quality light maps. */
-	TArray<FFloat3Packed> LQLightColor;
-	TArray<FColor> LQLightDirection;
-
 	TArray<FIrradianceVoxelImportProcessingData> VoxelImportProcessingData;
 
 	void SetFromVolumeLightingSample(int32 Index, const FVolumeLightingSample& Sample, bool bInsideGeometry, float MinDistanceToSurface, bool bBorderVoxel);
@@ -2176,7 +2172,6 @@ private:
 		bool bDebugThisSample,
 		FStaticLightingMappingContext& MappingContext,
 		TGatheredLightSample<SHOrder>& OutStaticDirectLighting,
-		TGatheredLightSample<SHOrder>& OutToggleableDirectLighting,
 		float& OutToggleableDirectionalLightShadowing) const;
 
 	void GatherVolumeImportancePhotonDirections(
@@ -2431,8 +2426,7 @@ private:
 		FShadowMapData2D*& ShadowMapData,
 		const FTexelToVertexMap& TexelToVertexMap, 
 		bool bDebugThisMapping,
-		const FLight* Light,
-		const bool bLowQualityLightMapsOnly) const;
+		const FLight* Light) const;
 
 	/** 
 	 * Calculate signed distance field shadowing from a single light,  
diff --git a/Engine/Source/Programs/UnrealLightmass/Private/Lighting/LightingSystem.inl b/Engine/Source/Programs/UnrealLightmass/Private/Lighting/LightingSystem.inl
--- a/Engine/Source/Programs/UnrealLightmass/Private/Lighting/LightingSystem.inl
+++ b/Engine/Source/Programs/UnrealLightmass/Private/Lighting/LightingSystem.inl
@@ -26,7 +26,7 @@ TGatheredLightSample<SHOrder> FGatheredLightSampleUtil::PointLightWorldSpace(con
 
 	if (TangentDirection.Z >= 0.0f)
 	{
-		Result.SHVector.AddIncomingRadiance(Color, 1, (const FVector4&)WorldDirection);
+		Result.SHVector.AddIncomingRadiance(Color, 1, FVector4(WorldDirection));
 
 		FSHVector2 SH = FSHVector2::SHBasisFunction(FVector4(TangentDirection));
 		// Evaluate lighting along the smoothed vertex normal direction, so that later we can guarantee an SH intensity of 1 along the normal
@@ -100,7 +100,6 @@ void FStaticLightingSystem::CalculateApproximateDirectLighting(
 	bool bDebugThisSample,
 	FStaticLightingMappingContext& MappingContext,
 	TGatheredLightSample<SHOrder>& OutStaticDirectLighting,
-	TGatheredLightSample<SHOrder>& OutToggleableDirectLighting,
 	float& OutToggleableDirectionalLightShadowing) const
 {
 	check(VertexOffsets.Num() > 0);
@@ -215,10 +214,6 @@ void FStaticLightingSystem::CalculateApproximateDirectLighting(
 				{
 					OutToggleableDirectionalLightShadowing = Transmission.GetLuminance();
 				}
-				else
-				{
-					OutToggleableDirectLighting = OutToggleableDirectLighting + Lighting;
-				}
 			}
 		}
 	}
diff --git a/Engine/Source/Programs/UnrealLightmass/Private/Lighting/LightmapData.cpp b/Engine/Source/Programs/UnrealLightmass/Private/Lighting/LightmapData.cpp
--- a/Engine/Source/Programs/UnrealLightmass/Private/Lighting/LightmapData.cpp
+++ b/Engine/Source/Programs/UnrealLightmass/Private/Lighting/LightmapData.cpp
@@ -292,9 +292,9 @@ namespace Lightmass
 
 				// U, V, W, LogL
 				// UVW stored in gamma space
-				DestCoefficients.Coefficients[0][0] = (uint8)FMath::Clamp<int32>( FMath::RoundToInt( FMath::Pow( U, 1.0f / 2.2f ) * 255.0f ), 0, 255 );
-				DestCoefficients.Coefficients[0][1] = (uint8)FMath::Clamp<int32>( FMath::RoundToInt( FMath::Pow( V, 1.0f / 2.2f ) * 255.0f ), 0, 255 );
-				DestCoefficients.Coefficients[0][2] = (uint8)FMath::Clamp<int32>( FMath::RoundToInt( FMath::Pow( W, 1.0f / 2.2f ) * 255.0f ), 0, 255 );
+				DestCoefficients.Coefficients[0][0] = (uint8)FMath::Clamp<int32>( FMath::RoundToInt( FMath::Sqrt(U) * 255.0f ), 0, 255 );
+				DestCoefficients.Coefficients[0][1] = (uint8)FMath::Clamp<int32>( FMath::RoundToInt( FMath::Sqrt(V) * 255.0f ), 0, 255 );
+				DestCoefficients.Coefficients[0][2] = (uint8)FMath::Clamp<int32>( FMath::RoundToInt( FMath::Sqrt(W) * 255.0f ), 0, 255 );
 				DestCoefficients.Coefficients[0][3] = (uint8)FMath::Clamp<int32>( FMath::RoundToInt( LogL * 255.0f ), 0, 255 );
 				
 				float Dx = SourceSample.Coefficients[1][0] * CoefficientMultiply[1][0] + CoefficientAdd[1][0];
diff --git a/Engine/Source/Programs/UnrealLightmass/Private/Lighting/Radiosity.cpp b/Engine/Source/Programs/UnrealLightmass/Private/Lighting/Radiosity.cpp
--- a/Engine/Source/Programs/UnrealLightmass/Private/Lighting/Radiosity.cpp
+++ b/Engine/Source/Programs/UnrealLightmass/Private/Lighting/Radiosity.cpp
@@ -418,12 +418,11 @@ void FStaticLightingSystem::RadiositySetupTextureMapping(FStaticLightingTextureM
 				if (ImportanceTracingSettings.bUseRadiositySolverForLightMultibounce)
 				{
 					FGatheredLightSample DirectLighting;
-					FGatheredLightSample Unused;
 					float Unused2;
 					TArray<FVector3f, TInlineAllocator<1>> VertexOffsets;
 					VertexOffsets.Add(FVector3f(0, 0, 0));
 
-					CalculateApproximateDirectLighting(CurrentVertex, TexelToVertex.TexelRadius, VertexOffsets, .1f, true, true, bDebugThisTexel, MappingContext, DirectLighting, Unused, Unused2);
+					CalculateApproximateDirectLighting(CurrentVertex, TexelToVertex.TexelRadius, VertexOffsets, .1f, true, true, bDebugThisTexel, MappingContext, DirectLighting, Unused2);
 
 					IncidentLightingForRadiosity += DirectLighting.IncidentLighting;
 				}
diff --git a/Engine/Source/Programs/UnrealLightmass/Private/Lighting/TextureMapping.cpp b/Engine/Source/Programs/UnrealLightmass/Private/Lighting/TextureMapping.cpp
--- a/Engine/Source/Programs/UnrealLightmass/Private/Lighting/TextureMapping.cpp
+++ b/Engine/Source/Programs/UnrealLightmass/Private/Lighting/TextureMapping.cpp
@@ -665,12 +665,11 @@ void FStaticLightingSystem::FinalizeSurfaceCacheTextureMapping(FStaticLightingTe
 				if (GeneralSettings.ViewSingleBounceNumber < 0 || GeneralSettings.ViewSingleBounceNumber == 1)
 				{
 					FGatheredLightSample DirectLighting;
-					FGatheredLightSample Unused;
 					float Unused2;
 					TArray<FVector3f, TInlineAllocator<1>> VertexOffsets;
 					VertexOffsets.Add(FVector3f(0, 0, 0));
 
-					CalculateApproximateDirectLighting(CurrentVertex, TexelToVertex.TexelRadius, VertexOffsets, .1f, true, true, bDebugThisTexel && PhotonMappingSettings.bVisualizeCachedApproximateDirectLighting, MappingContext, DirectLighting, Unused, Unused2);
+					CalculateApproximateDirectLighting(CurrentVertex, TexelToVertex.TexelRadius, VertexOffsets, .1f, true, true, bDebugThisTexel && PhotonMappingSettings.bVisualizeCachedApproximateDirectLighting, MappingContext, DirectLighting, Unused2);
 					FinalIncidentLighting += DirectLighting.IncidentLighting;
 				}
 
@@ -812,7 +811,7 @@ void FStaticLightingSystem::ProcessTextureMapping(FStaticLightingTextureMapping*
 						if (Light->LightFlags & GI_LIGHT_USE_AREA_SHADOWS_FOR_SEPARATE_SHADOW_FACTOR)
 						{
 							FShadowMapData2D* ShadowMapData = new FShadowMapData2D(TextureMapping->CachedSizeX,TextureMapping->CachedSizeY);
-							CalculateDirectAreaLightingTextureMapping(TextureMapping, MappingContext, LightMapData, ShadowMapData, TexelToVertexMap, bDebugThisMapping, Light, false);
+							CalculateDirectAreaLightingTextureMapping(TextureMapping, MappingContext, LightMapData, ShadowMapData, TexelToVertexMap, bDebugThisMapping, Light);
 
 							if (ShadowMapData)
 							{
@@ -850,17 +849,6 @@ void FStaticLightingSystem::ProcessTextureMapping(FStaticLightingTextureMapping*
 								CalculateDirectSignedDistanceFieldLightingTextureMappingLightSpace(TextureMapping, MappingContext, LightMapData, SignedDistanceFieldShadowMaps, TexelToVertexMap, TexelToCornersMap, bDebugThisMapping, Light);
 							}
 						}
-
-						// Stationary directional light is never put into the lightmap, even with low quality lightmaps
-						if (!Light->GetDirectionalLight())
-						{
-							// Also calculate static lighting for simple light maps.  We'll force the shadows into simple light maps, but
-							// won't actually add the lights to the light guid list.  Instead, at runtime we'll check the shadow map guids
-							// for lights that are baked into light maps on platforms that don't support shadow mapping.
-							FShadowMapData2D* ShadowMapData = NULL;
-							const bool bLowQualityLightMapsOnly = true;
-							CalculateDirectAreaLightingTextureMapping(TextureMapping, MappingContext, LightMapData, ShadowMapData, TexelToVertexMap, bDebugThisMapping, Light, bLowQualityLightMapsOnly);
-						}
 					}
 					else if (Light->UseStaticLighting())
 					{
@@ -868,7 +856,7 @@ void FStaticLightingSystem::ProcessTextureMapping(FStaticLightingTextureMapping*
 
 						// Calculate direct lighting from area lights
 						// Shadow penumbras will be correctly shaped and will be softer for larger light sources and distant shadow casters.
-						CalculateDirectAreaLightingTextureMapping(TextureMapping, MappingContext, LightMapData, ShadowMapData, TexelToVertexMap, bDebugThisMapping, Light, false);
+						CalculateDirectAreaLightingTextureMapping(TextureMapping, MappingContext, LightMapData, ShadowMapData, TexelToVertexMap, bDebugThisMapping, Light);
 
 						if (Light->GetMeshAreaLight() == NULL)
 						{
@@ -1406,8 +1394,7 @@ void FStaticLightingSystem::CalculateDirectAreaLightingTextureMapping(
 	FShadowMapData2D*& ShadowMapData,
 	const FTexelToVertexMap& TexelToVertexMap, 
 	bool bDebugThisMapping,
-	const FLight* Light,
-	const bool bLowQualityLightMapsOnly) const
+	const FLight* Light) const
 {
 	LIGHTINGSTAT(FScopedRDTSCTimer AreaShadowsTimer(MappingContext.Stats.AreaShadowsThreadTime));
 
@@ -1731,14 +1718,7 @@ void FStaticLightingSystem::CalculateDirectAreaLightingTextureMapping(
 							const FGatheredLightSample DirectLighting = CalculatePointLighting(TextureMapping, CurrentVertex, TexelToVertex.ElementIndex, Light, LightIntensity, Transmission);
 
 							FGatheredLightMapSample& CurrentLightSample = LightMapData(X,Y);
-							if( bLowQualityLightMapsOnly )
-							{
-								CurrentLightSample.LowQuality.AddWeighted(DirectLighting, AdjustedShadowFactor);
-							}
-							else
-							{
-								CurrentLightSample.AddWeighted(DirectLighting, AdjustedShadowFactor);
-							}
+							CurrentLightSample.AddWeighted(DirectLighting, AdjustedShadowFactor);
 						}
 					}
 				}
@@ -2830,12 +2810,11 @@ void FStaticLightingSystem::CalculateDirectLightingTextureMappingPhotonMap(
 						NearestPhoton = FindNearestIrradiancePhoton(CurrentVertex, MappingContext, TempIrradiancePhotons, false, bDebugThisTexel);
 
 						FGatheredLightSample DirectLightingSample;
-						FGatheredLightSample Unused;
 						float Unused2;
 						TArray<FVector3f, TInlineAllocator<1>> VertexOffsets;
 						VertexOffsets.Add(FVector3f(0, 0, 0));
 
-						CalculateApproximateDirectLighting(CurrentVertex, TexelToVertex.TexelRadius, VertexOffsets, .1f, true, true, bDebugThisTexel, MappingContext, DirectLightingSample, Unused, Unused2);
+						CalculateApproximateDirectLighting(CurrentVertex, TexelToVertex.TexelRadius, VertexOffsets, .1f, true, true, bDebugThisTexel, MappingContext, DirectLightingSample, Unused2);
 
 						DirectLighting = DirectLightingSample.IncidentLighting;
 					}
@@ -3033,9 +3012,6 @@ void FStaticLightingSystem::ProcessInterpolateTask(FInterpolateIndirectTaskDescr
 				CurrentLightSample.AddWeighted(IndirectLighting, 1);
 				CurrentLightSample.HighQuality.AOMaterialMask = IndirectLighting.Occlusion;
 
-				// Stationary sky light contribution goes into low quality lightmap only, bent normal sky shadowing will be exported separately
-				CurrentLightSample.LowQuality.AddWeighted(IndirectLighting.StationarySkyLighting, 1);
-
 				if (AmbientOcclusionSettings.bUseAmbientOcclusion && AmbientOcclusionSettings.bVisualizeAmbientOcclusion)
 				{
 					//@todo - this will only be the correct intensity for simple lightmaps
diff --git a/Engine/Source/Runtime/Core/Public/UObject/UE5MainStreamObjectVersion.h b/Engine/Source/Runtime/Core/Public/UObject/UE5MainStreamObjectVersion.h
--- a/Engine/Source/Runtime/Core/Public/UObject/UE5MainStreamObjectVersion.h
+++ b/Engine/Source/Runtime/Core/Public/UObject/UE5MainStreamObjectVersion.h
@@ -250,6 +250,9 @@ struct CORE_API FUE5MainStreamObjectVersion
 		// Revert the IAnimationDataModel changes. Animation assets 
 		BackoutAnimationDataModelInterface,
 
+		// Made stationary local and skylights behave similar to SM5
+		MobileStationaryLocalLights,
+
 		// -----<new versions can be added above this line>-------------------------------------------------
 		VersionPlusOne,
 		LatestVersion = VersionPlusOne - 1
diff --git a/Engine/Source/Runtime/Engine/Private/MapBuildData.cpp b/Engine/Source/Runtime/Engine/Private/MapBuildData.cpp
--- a/Engine/Source/Runtime/Engine/Private/MapBuildData.cpp
+++ b/Engine/Source/Runtime/Engine/Private/MapBuildData.cpp
@@ -21,6 +21,7 @@ MapBuildData.cpp
 #include "UObject/MobileObjectVersion.h"
 #include "UObject/RenderingObjectVersion.h"
 #include "UObject/ReflectionCaptureObjectVersion.h"
+#include "UObject/UE5MainStreamObjectVersion.h"
 #include "ContentStreaming.h"
 #include "Components/ReflectionCaptureComponent.h"
 #include "Interfaces/ITargetPlatform.h"
@@ -459,6 +460,7 @@ void UMapBuildDataRegistry::Serialize(FArchive& Ar)
 	Ar.UsingCustomVersion(FMobileObjectVersion::GUID);
 	Ar.UsingCustomVersion(FReflectionCaptureObjectVersion::GUID);
 	Ar.UsingCustomVersion(FUE5ReleaseStreamObjectVersion::GUID);
+	Ar.UsingCustomVersion(FUE5MainStreamObjectVersion::GUID);
 
 	if (!StripFlags.IsDataStrippedForServer())
 	{
diff --git a/Engine/Source/Runtime/Engine/Private/PrecomputedVolumetricLightmap.cpp b/Engine/Source/Runtime/Engine/Private/PrecomputedVolumetricLightmap.cpp
--- a/Engine/Source/Runtime/Engine/Private/PrecomputedVolumetricLightmap.cpp
+++ b/Engine/Source/Runtime/Engine/Private/PrecomputedVolumetricLightmap.cpp
@@ -8,6 +8,7 @@
 #include "Stats/Stats.h"
 #include "EngineDefines.h"
 #include "UObject/RenderingObjectVersion.h"
+#include "UObject/UE5MainStreamObjectVersion.h"
 #include "SceneManagement.h"
 #include "UnrealEngine.h"
 #include "Engine/MapBuildDataRegistry.h"
@@ -122,18 +123,13 @@ FArchive& operator<<(FArchive& Ar,FPrecomputedVolumetricLightmapData& Volume)
 	
 	if (Ar.CustomVer(FMobileObjectVersion::GUID) >= FMobileObjectVersion::LQVolumetricLightmapLayers)
 	{
-		if (Ar.IsCooking() && !Ar.CookingTarget()->SupportsFeature(ETargetPlatformFeatures::LowQualityLightmaps))
+		if (Ar.CustomVer(FUE5MainStreamObjectVersion::GUID) < FUE5MainStreamObjectVersion::MobileStationaryLocalLights && Ar.IsLoading())
 		{
-			// Don't serialize cooked LQ data if the cook target does not want it.
+			// Don't serialize cooked LQ data
 			FVolumetricLightmapDataLayer Dummy;
 			Ar << Dummy;
 			Ar << Dummy;
 		}
-		else
-		{
-			Ar << Volume.BrickData.LQLightColor;
-			Ar << Volume.BrickData.LQLightDirection;
-		}
 	}
 
 	if (Ar.CustomVer(FRenderingObjectVersion::GUID) >= FRenderingObjectVersion::VolumetricLightmapStreaming)
@@ -144,12 +140,6 @@ FArchive& operator<<(FArchive& Ar,FPrecomputedVolumetricLightmapData& Volume)
 
 	if (Ar.IsLoading())
 	{
-		if (GMaxRHIFeatureLevel >= ERHIFeatureLevel::SM5 && !GIsEditor)
-		{
-			// drop LQ data for SM4+
-			Volume.BrickData.DiscardLowQualityLayers();
-		}
-
 		Volume.bTransient = false;
 
 		const SIZE_T VolumeBytes = Volume.GetAllocatedBytes();
diff --git a/Engine/Source/Runtime/Engine/Private/SceneManagement.cpp b/Engine/Source/Runtime/Engine/Private/SceneManagement.cpp
--- a/Engine/Source/Runtime/Engine/Private/SceneManagement.cpp
+++ b/Engine/Source/Runtime/Engine/Private/SceneManagement.cpp
@@ -705,6 +705,7 @@ FMobileDirectionalLightShaderParameters::FMobileDirectionalLightShaderParameters
 	// light, default to black
 	DirectionalLightColor = FLinearColor::Black;
 	DirectionalLightDirectionAndShadowTransition = FVector4f(EForceInit::ForceInitToZero);
+	DirectionalLightShadowMapChannelMask = 0xFF;
 
 	// white texture should act like a shadowmap cleared to the farplane.
 	DirectionalLightShadowTexture = GWhiteTexture->TextureRHI;
diff --git a/Engine/Source/Runtime/Engine/Public/PrecomputedVolumetricLightmap.h b/Engine/Source/Runtime/Engine/Public/PrecomputedVolumetricLightmap.h
--- a/Engine/Source/Runtime/Engine/Public/PrecomputedVolumetricLightmap.h
+++ b/Engine/Source/Runtime/Engine/Public/PrecomputedVolumetricLightmap.h
@@ -83,9 +83,6 @@ struct FVolumetricLightmapBasicBrickDataLayers
 class FVolumetricLightmapBrickData : public FVolumetricLightmapBasicBrickDataLayers
 {
 public:
-	// Mobile LQ layers:
-	FVolumetricLightmapDataLayer LQLightColor;
-	FVolumetricLightmapDataLayer LQLightDirection;
 
 	ENGINE_API int32 GetMinimumVoxelSize() const;
 
@@ -115,7 +112,6 @@ public:
 	SIZE_T GetAllocatedBytes() const
 	{
 		SIZE_T NumBytes = AmbientVector.DataSize + SkyBentNormal.DataSize + DirectionalLightShadowing.DataSize;
-		NumBytes += LQLightColor.Data.Num() + LQLightDirection.Data.Num();
 
 		for (int32 i = 0; i < UE_ARRAY_COUNT(SHCoefficients); i++)
 		{
@@ -137,13 +133,6 @@ public:
 		SkyBentNormal.bNeedsCPUAccess = InAccess;
 		DirectionalLightShadowing.bNeedsCPUAccess = InAccess;
 	}
-
-	// discard the layers used for low quality lightmap (LQ includes direct lighting from stationary lights).
-	void DiscardLowQualityLayers()
-	{
-		LQLightColor.Discard();
-		LQLightDirection.Discard();
-	}
 };
 
 /** 
diff --git a/Engine/Source/Runtime/Engine/Public/SceneView.h b/Engine/Source/Runtime/Engine/Public/SceneView.h
--- a/Engine/Source/Runtime/Engine/Public/SceneView.h
+++ b/Engine/Source/Runtime/Engine/Public/SceneView.h
@@ -572,6 +572,7 @@ BEGIN_GLOBAL_SHADER_PARAMETER_STRUCT_WITH_CONSTRUCTOR(FMobileDirectionalLightSha
 	SHADER_PARAMETER_EX(FVector4f, DirectionalLightDistanceFadeMADAndSpecularScale, EShaderPrecisionModifier::Half) // .z is used for SpecularScale, .w is not used atm
 	SHADER_PARAMETER_EX(FVector4f, DirectionalLightShadowDistances, EShaderPrecisionModifier::Half)
 	SHADER_PARAMETER_ARRAY(FMatrix44f, DirectionalLightScreenToShadow, [MAX_MOBILE_SHADOWCASCADES])
+	SHADER_PARAMETER(uint32, DirectionalLightShadowMapChannelMask)
 	SHADER_PARAMETER_TEXTURE(Texture2D, DirectionalLightShadowTexture)
 	SHADER_PARAMETER_SAMPLER(SamplerState, DirectionalLightShadowSampler)
 END_GLOBAL_SHADER_PARAMETER_STRUCT()
diff --git a/Engine/Source/Runtime/Renderer/Private/LightMapRendering.cpp b/Engine/Source/Runtime/Renderer/Private/LightMapRendering.cpp
--- a/Engine/Source/Runtime/Renderer/Private/LightMapRendering.cpp
+++ b/Engine/Source/Runtime/Renderer/Private/LightMapRendering.cpp
@@ -228,20 +228,6 @@ void InterpolateVolumetricLightmap(
 	LQSH.G = GetSHVector3(AmbientVector.Y, ReadSHCoefficient(2), ReadSHCoefficient(3));
 	LQSH.B = GetSHVector3(AmbientVector.Z, ReadSHCoefficient(4), ReadSHCoefficient(5));
 
-	if(VolumetricLightmapData.BrickData.LQLightColor.Data.Num())
-	{
-		// Add stationary direct lighting to match ILC's LQ behavior.
-		FLinearColor LQLightColour = FilteredVolumeLookup<FFloat3Packed>(BrickTextureCoordinate, VolumetricLightmapData.BrickDataDimensions, (const FFloat3Packed*)VolumetricLightmapData.BrickData.LQLightColor.Data.GetData());
-		FVector LQLightDirection = (FVector)FilteredVolumeLookup<FColor>(BrickTextureCoordinate, VolumetricLightmapData.BrickDataDimensions, (const FColor*)VolumetricLightmapData.BrickData.LQLightDirection.Data.GetData());
-		//Swap X and Z channel because it was swapped at ImportVolumetricLightmap for changing format from BGRA to RGBA
-		Swap(LQLightDirection.X, LQLightDirection.Z);
-
-		LQLightDirection = LQLightDirection*2.0f - FVector(1.0f, 1.0f, 1.0f);
-		LQLightDirection.Normalize();
-
-		LQSH.AddIncomingRadiance(LQLightColour, 1.0f, LQLightDirection);
-	}
-
 	// Pack the 3rd order SH as the shader expects
 	OutInterpolation.IndirectLightingSHCoefficients0[0] = FVector4f(LQSH.R.V[0], LQSH.R.V[1], LQSH.R.V[2], LQSH.R.V[3]) * INV_PI;
 	OutInterpolation.IndirectLightingSHCoefficients0[1] = FVector4f(LQSH.G.V[0], LQSH.G.V[1], LQSH.G.V[2], LQSH.G.V[3]) * INV_PI;
diff --git a/Engine/Source/Runtime/Renderer/Private/MobileBasePass.cpp b/Engine/Source/Runtime/Renderer/Private/MobileBasePass.cpp
--- a/Engine/Source/Runtime/Renderer/Private/MobileBasePass.cpp
+++ b/Engine/Source/Runtime/Renderer/Private/MobileBasePass.cpp
@@ -495,21 +495,6 @@ void MobileBasePass::SetTranslucentRenderState(FMeshPassProcessorRenderState& Dr
 	}
 }
 
-bool MobileBasePass::StationarySkyLightHasBeenApplied(const FScene* Scene, ELightMapPolicyType LightMapPolicyType)
-{
-	return Scene
-		&& Scene->SkyLight
-		&& Scene->SkyLight->bWantsStaticShadowing
-		&& (LightMapPolicyType == LMP_LQ_LIGHTMAP
-			|| LightMapPolicyType == LMP_MOBILE_DISTANCE_FIELD_SHADOWS_AND_LQ_LIGHTMAP
-			|| LightMapPolicyType == LMP_MOBILE_DISTANCE_FIELD_SHADOWS_LIGHTMAP_AND_CSM
-			|| LightMapPolicyType == LMP_MOBILE_DIRECTIONAL_LIGHT_CSM_AND_LIGHTMAP
-			|| LightMapPolicyType == LMP_MOBILE_DIRECTIONAL_LIGHT_AND_SH_INDIRECT
-			|| LightMapPolicyType == LMP_MOBILE_DIRECTIONAL_LIGHT_CSM_AND_SH_INDIRECT
-			|| LightMapPolicyType == LMP_MOBILE_MOVABLE_DIRECTIONAL_LIGHT_WITH_LIGHTMAP
-			|| LightMapPolicyType == LMP_MOBILE_MOVABLE_DIRECTIONAL_LIGHT_CSM_WITH_LIGHTMAP);
-}
-
 static FMeshDrawCommandSortKey GetBasePassStaticSortKey(EBlendMode BlendMode, bool bBackground)
 {
 	FMeshDrawCommandSortKey SortKey;
@@ -674,16 +659,7 @@ bool FMobileBasePassMeshProcessor::Process(
 	
 	if (Scene && Scene->SkyLight)
 	{
-		// Clustered reflection is always enabled on mobile deferred translucent material
-		const bool bEnableClusteredReflections = MobileEnableClusteredReflections(Scene->GetShaderPlatform());
-
-		//The stationary skylight contribution has been added both to the LowQuality Lightmap and ILC on mobile, so we should skip the sky light spherical harmonic contribution for it.
-		bool bStationarySkyLightHasBeenApplied = MobileBasePass::StationarySkyLightHasBeenApplied(Scene, LightMapPolicyType) && !bEnableClusteredReflections;
-
-		//Two side material should enable sky light for the back face since only the front face has light map and it will be corrected in base pass shader.
-		bool bSkipStationarySkyLight = bStationarySkyLightHasBeenApplied && !MaterialResource.IsTwoSided();
-
-		bEnableSkyLight = ShadingModels.IsLit() && Scene->ShouldRenderSkylightInBasePass(BlendMode) && (!bSkipStationarySkyLight);
+		bEnableSkyLight = ShadingModels.IsLit() && Scene->ShouldRenderSkylightInBasePass(BlendMode);
 	}
 
 	bool bEnableLocalLights = false;
@@ -693,7 +669,7 @@ bool FMobileBasePassMeshProcessor::Process(
 		// Clustered lighting could only be enabled on mobile forward and always enabled on mobile deferred translucent material
 		if (!bUsesDeferredShading && (MobileForwardEnableLocalLights(Scene->GetShaderPlatform()) || bIsTranslucentMaterialOnDeferred))
 		{
-			bEnableLocalLights = PrimitiveSceneProxy->GetPrimitiveSceneInfo()->NumMobileMovableLocalLights > 0;
+			bEnableLocalLights = PrimitiveSceneProxy->GetPrimitiveSceneInfo()->NumMobileDynamicLocalLights > 0;
 		}
 	}
 
diff --git a/Engine/Source/Runtime/Renderer/Private/MobileBasePassRendering.cpp b/Engine/Source/Runtime/Renderer/Private/MobileBasePassRendering.cpp
--- a/Engine/Source/Runtime/Renderer/Private/MobileBasePassRendering.cpp
+++ b/Engine/Source/Runtime/Renderer/Private/MobileBasePassRendering.cpp
@@ -229,6 +229,9 @@ void SetupMobileDirectionalLightUniformParameters(
 	bool bDynamicShadows,
 	FMobileDirectionalLightShaderParameters& Params)
 {
+	static const auto AllowStaticLightingVar = IConsoleManager::Get().FindTConsoleVariableDataInt(TEXT("r.AllowStaticLighting"));
+	const bool bAllowStaticLighting = (!AllowStaticLightingVar || AllowStaticLightingVar->GetValueOnRenderThread() != 0);
+
 	ERHIFeatureLevel::Type FeatureLevel = Scene.GetFeatureLevel();
 	FLightSceneInfo* Light = Scene.MobileDirectionalLights[ChannelIdx];
 	if (Light)
@@ -241,6 +244,25 @@ void SetupMobileDirectionalLightUniformParameters(
 		Params.DirectionalLightDistanceFadeMADAndSpecularScale.Y = -FadeParams.X * FadeParams.Y;
 		Params.DirectionalLightDistanceFadeMADAndSpecularScale.Z = Light->Proxy->GetSpecularScale();
 
+		int32 ShadowMapChannel = Light->Proxy->GetShadowMapChannel();
+		int32 DynamicShadowMapChannel = Light->GetDynamicShadowMapChannel();
+
+		if (!bAllowStaticLighting)
+		{
+			ShadowMapChannel = INDEX_NONE;
+		}
+
+		// Static shadowing uses ShadowMapChannel, dynamic shadows are packed into light attenuation using DynamicShadowMapChannel
+		Params.DirectionalLightShadowMapChannelMask =
+			(ShadowMapChannel == 0 ? 1 : 0) |
+			(ShadowMapChannel == 1 ? 2 : 0) |
+			(ShadowMapChannel == 2 ? 4 : 0) |
+			(ShadowMapChannel == 3 ? 8 : 0) |
+			(DynamicShadowMapChannel == 0 ? 16 : 0) |
+			(DynamicShadowMapChannel == 1 ? 32 : 0) |
+			(DynamicShadowMapChannel == 2 ? 64 : 0) |
+			(DynamicShadowMapChannel == 3 ? 128 : 0);
+
 		if (bDynamicShadows && VisibleLightInfos.IsValidIndex(Light->Id) && VisibleLightInfos[Light->Id].AllProjectedShadows.Num() > 0)
 		{
 			const TArray<FProjectedShadowInfo*, SceneRenderingAllocator>& DirectionalLightShadowInfos = VisibleLightInfos[Light->Id].AllProjectedShadows;
diff --git a/Engine/Source/Runtime/Renderer/Private/MobileBasePassRendering.h b/Engine/Source/Runtime/Renderer/Private/MobileBasePassRendering.h
--- a/Engine/Source/Runtime/Renderer/Private/MobileBasePassRendering.h
+++ b/Engine/Source/Runtime/Renderer/Private/MobileBasePassRendering.h
@@ -336,8 +336,6 @@ namespace MobileBasePass
 
 	void SetOpaqueRenderState(FMeshPassProcessorRenderState& DrawRenderState, const FPrimitiveSceneProxy* PrimitiveSceneProxy, const FMaterial& Material, bool bEnableReceiveDecalOutput, bool bUsesDeferredShading);
 	void SetTranslucentRenderState(FMeshPassProcessorRenderState& DrawRenderState, const FMaterial& Material);
-
-	bool StationarySkyLightHasBeenApplied(const FScene* Scene, ELightMapPolicyType LightMapPolicyType);
 };
 
 
diff --git a/Engine/Source/Runtime/Renderer/Private/PrimitiveSceneInfo.cpp b/Engine/Source/Runtime/Renderer/Private/PrimitiveSceneInfo.cpp
--- a/Engine/Source/Runtime/Renderer/Private/PrimitiveSceneInfo.cpp
+++ b/Engine/Source/Runtime/Renderer/Private/PrimitiveSceneInfo.cpp
@@ -168,7 +168,7 @@ FPrimitiveSceneInfo::FPrimitiveSceneInfo(UPrimitiveComponent* InComponent,FScene
 	LightList(NULL),
 	LastRenderTime(-FLT_MAX),
 	Scene(InScene),
-	NumMobileMovableLocalLights(0),
+	NumMobileDynamicLocalLights(0),
 	bShouldRenderInMainPass(InComponent->SceneProxy->ShouldRenderInMainPass()),
 	bVisibleInRealTimeSkyCapture(InComponent->SceneProxy->IsVisibleInRealTimeSkyCaptures()),
 #if RHI_RAYTRACING
diff --git a/Engine/Source/Runtime/Renderer/Private/RendererScene.cpp b/Engine/Source/Runtime/Renderer/Private/RendererScene.cpp
--- a/Engine/Source/Runtime/Renderer/Private/RendererScene.cpp
+++ b/Engine/Source/Runtime/Renderer/Private/RendererScene.cpp
@@ -2869,20 +2869,10 @@ void FScene::AddPrecomputedVolumetricLightmap(const FPrecomputedVolumetricLightm
 	FScene* Scene = this;
 
 	ENQUEUE_RENDER_COMMAND(AddVolumeCommand)
-		([Scene, Volume, bIsPersistentLevel](FRHICommandListImmediate& RHICmdList)
-		{
-#if !(UE_BUILD_SHIPPING || UE_BUILD_TEST)
-			if (Volume && Scene->GetShadingPath() == EShadingPath::Mobile)
-			{
-				const FPrecomputedVolumetricLightmapData* VolumeData = Volume->Data;
-				if (VolumeData && VolumeData->BrickData.LQLightDirection.Data.Num() == 0)
-				{
-					FPlatformAtomics::InterlockedIncrement(&Scene->NumUncachedStaticLightingInteractions);
-				}
-			}
-#endif
+	([Scene, Volume, bIsPersistentLevel](FRHICommandListImmediate& RHICmdList)
+	{
 		Scene->VolumetricLightmapSceneData.AddLevelVolume(Volume, Scene->GetShadingPath(), bIsPersistentLevel);
-		});
+	});
 }
 
 void FScene::RemovePrecomputedVolumetricLightmap(const FPrecomputedVolumetricLightmap* Volume)
@@ -2890,20 +2880,10 @@ void FScene::RemovePrecomputedVolumetricLightmap(const FPrecomputedVolumetricLig
 	FScene* Scene = this; 
 
 	ENQUEUE_RENDER_COMMAND(RemoveVolumeCommand)
-		([Scene, Volume](FRHICommandListImmediate& RHICmdList) 
-		{
-#if !(UE_BUILD_SHIPPING || UE_BUILD_TEST)
-			if (Volume && Scene->GetShadingPath() == EShadingPath::Mobile)
-			{
-				const FPrecomputedVolumetricLightmapData* VolumeData = Volume->Data;
-				if (VolumeData && VolumeData->BrickData.LQLightDirection.Data.Num() == 0)
-				{
-					FPlatformAtomics::InterlockedDecrement(&Scene->NumUncachedStaticLightingInteractions);
-				}
-			}
-#endif
+	([Scene, Volume](FRHICommandListImmediate& RHICmdList) 
+	{
 		Scene->VolumetricLightmapSceneData.RemoveLevelVolume(Volume);
-		});
+	});
 }
 
 void FScene::AddRuntimeVirtualTexture(class URuntimeVirtualTextureComponent* Component)
diff --git a/Engine/Source/Runtime/Renderer/Private/SceneCore.cpp b/Engine/Source/Runtime/Renderer/Private/SceneCore.cpp
--- a/Engine/Source/Runtime/Renderer/Private/SceneCore.cpp
+++ b/Engine/Source/Runtime/Renderer/Private/SceneCore.cpp
@@ -126,6 +126,9 @@ void FLightPrimitiveInteraction::Create(FLightSceneInfo* LightSceneInfo,FPrimiti
 	check(PrimitiveSceneInfo->Proxy && LightSceneInfo->Proxy);
 	PrimitiveSceneInfo->Proxy->GetLightRelevance(LightSceneInfo->Proxy, bDynamic, bRelevant, bIsLightMapped, bShadowMapped);
 
+	// Mobile renders stationary and dynamic local lights as dynamic
+	bDynamic |= (PrimitiveSceneInfo->Scene->GetShadingPath() == EShadingPath::Mobile && bShadowMapped && LightSceneInfo->Proxy->IsLocalLight());
+
 	if (bRelevant && bDynamic
 		// Don't let lights with static shadowing or static lighting affect primitives that should use static lighting, but don't have valid settings (lightmap res 0, etc)
 		// This prevents those components with invalid lightmap settings from causing lighting to remain unbuilt after a build
@@ -233,7 +236,7 @@ FLightPrimitiveInteraction::FLightPrimitiveInteraction(
 		// Add the interaction to the light's interaction list.
 		PrevPrimitiveLink = PrimitiveSceneInfo->Proxy->IsMeshShapeOftenMoving() ? &LightSceneInfo->DynamicInteractionOftenMovingPrimitiveList : &LightSceneInfo->DynamicInteractionStaticPrimitiveList;
 
-		// mobile movable spotlights / point lights
+		// mobile local lights with dynamic lighting
 		if (PrimitiveSceneInfo->Scene->GetShadingPath() == EShadingPath::Mobile && LightSceneInfo->ShouldRenderLightViewIndependent())
 		{
 			const uint8 LightType = LightSceneInfo->Proxy->GetLightType();
@@ -246,8 +249,8 @@ FLightPrimitiveInteraction::FLightPrimitiveInteraction(
 			if( bIsValidLightType )
 			{
 				bMobileDynamicLocalLight = true;
-				PrimitiveSceneInfo->NumMobileMovableLocalLights++;
-				if (PrimitiveSceneInfo->NumMobileMovableLocalLights == 1)
+				PrimitiveSceneInfo->NumMobileDynamicLocalLights++;
+				if (PrimitiveSceneInfo->NumMobileDynamicLocalLights == 1)
 				{
 					// Update static meshes to choose the shader permutation with local lights.
 					PrimitiveSceneInfo->BeginDeferredUpdateStaticMeshes();
@@ -311,8 +314,8 @@ FLightPrimitiveInteraction::~FLightPrimitiveInteraction()
 	// Track mobile movable local light count
 	if (bMobileDynamicLocalLight)
 	{
-		PrimitiveSceneInfo->NumMobileMovableLocalLights--;
-		if (PrimitiveSceneInfo->NumMobileMovableLocalLights == 0)
+		PrimitiveSceneInfo->NumMobileDynamicLocalLights--;
+		if (PrimitiveSceneInfo->NumMobileDynamicLocalLights == 0)
 		{
 			// Update static meshes to choose the shader permutation without local lights.
 			PrimitiveSceneInfo->BeginDeferredUpdateStaticMeshes();
diff --git a/Engine/Source/Runtime/Renderer/Public/PrimitiveSceneInfo.h b/Engine/Source/Runtime/Renderer/Public/PrimitiveSceneInfo.h
--- a/Engine/Source/Runtime/Renderer/Public/PrimitiveSceneInfo.h
+++ b/Engine/Source/Runtime/Renderer/Public/PrimitiveSceneInfo.h
@@ -378,8 +378,8 @@ public:
 	/** The scene the primitive is in. */
 	FScene* Scene;
 
-	/** The number of movable local lights for mobile */
-	int32 NumMobileMovableLocalLights;
+	/** The number of local lights with dynamic lighting for mobile */
+	int32 NumMobileDynamicLocalLights;
 
 	/** Set to true for the primitive to be rendered in the main pass to be visible in a view. */
 	bool bShouldRenderInMainPass : 1;
