Multiprocess Cook - Add CookWorkerMode and IWorkerRequests interface.
#rn Cooking, Minor

diff --git a/Engine/Source/Developer/IoStoreUtilities/Private/ZenStoreWriter.cpp b/Engine/Source/Developer/IoStoreUtilities/Private/ZenStoreWriter.cpp
--- a/Engine/Source/Developer/IoStoreUtilities/Private/ZenStoreWriter.cpp
+++ b/Engine/Source/Developer/IoStoreUtilities/Private/ZenStoreWriter.cpp
@@ -369,7 +369,7 @@ void FZenStoreWriter::Initialize(const FCookInfo& Info)
 
 	if (!bInitialized)
 	{
-		if (Info.bFullBuild)
+		if (Info.bFullBuild && !Info.bWorkerOnSharedSandbox)
 		{
 			UE_LOG(LogZenStoreWriter, Display, TEXT("Deleting %s..."), *OutputPath);
 			const bool bRequireExists = false;
diff --git a/Engine/Source/Editor/UnrealEd/Classes/Commandlets/CookCommandlet.h b/Engine/Source/Editor/UnrealEd/Classes/Commandlets/CookCommandlet.h
--- a/Engine/Source/Editor/UnrealEd/Classes/Commandlets/CookCommandlet.h
+++ b/Engine/Source/Editor/UnrealEd/Classes/Commandlets/CookCommandlet.h
@@ -58,8 +58,6 @@ class UCookCommandlet
 	bool bVerboseCookerWarnings;
 	/** only clean up objects which are not in use by the cooker when we gc (false will enable full gc) */
 	bool bPartialGC;
-	/** Do not cook any shaders.  Shader maps will be empty */
-	bool bNoShaderCooking;
 	/** Ignore ini settings out of date. */
 	bool bIgnoreIniSettingsOutOfDate;
 	/** All commandline tokens */
@@ -84,9 +82,13 @@ class UCookCommandlet
 	/** Cooks for specified targets */
 	bool CookByTheBook(const TArray<ITargetPlatform*>& Platforms);
 
+	bool CookAsCookWorker();
+
 	/** Collect garbage if the cooker's TickResults requested it */
 	void ConditionalCollectGarbage(uint32 TickResults, UCookOnTheFlyServer& COTFS);
 
+	UE_DEPRECATED(5.1, "No longer used")
+	bool bNoShaderCooking;
 public:
 
 	//~ Begin UCommandlet Interface
diff --git a/Engine/Source/Editor/UnrealEd/Classes/CookOnTheSide/CookOnTheFlyServer.h b/Engine/Source/Editor/UnrealEd/Classes/CookOnTheSide/CookOnTheFlyServer.h
--- a/Engine/Source/Editor/UnrealEd/Classes/CookOnTheSide/CookOnTheFlyServer.h
+++ b/Engine/Source/Editor/UnrealEd/Classes/CookOnTheSide/CookOnTheFlyServer.h
@@ -108,6 +108,8 @@ namespace ECookMode
 		CookByTheBookFromTheEditor,
 		/** Cooking by the book (not in the editor). */
 		CookByTheBook,
+		/** Commandlet helper for a separate director process. Director might be in any of the other modes. */
+		CookWorker,
 	};
 }
 
@@ -123,13 +125,17 @@ ENUM_CLASS_FLAGS(ECookTickFlags);
 namespace UE::Cook
 {
 	class FBuildDefinitions;
-	class FExternalRequests;
+	class FCookDirector;
 	class FRequestCluster;
 	class FSaveCookedPackageContext;
+	class FWorkerRequestsLocal;
+	class FWorkerRequestsRemote;
 	class ICookOnTheFlyRequestManager;
 	class ICookOnTheFlyNetworkServer;
+	class IWorkerRequests;
 	enum class EPollStatus : uint8;
 	enum class EReleaseSaveReason : uint8;
+	struct FBeginCookConfigSettings;
 	struct FConstructPackageData;
 	struct FCookByTheBookOptions;
 	struct FCookOnTheFlyOptions;
@@ -137,6 +143,7 @@ namespace UE::Cook
 	struct FCookGenerationInfo;
 	struct FCookSavePackageContext;
 	struct FGeneratorPackage;
+	struct FInitializeConfigSettings;
 	struct FPackageData;
 	struct FPackageDatas;
 	struct FPackageTracker;
@@ -165,6 +172,9 @@ private:
 	int64 Maximum = 0;
 };
 
+// tmap of the Config name, Section name, Key name, to the value
+typedef TMap<FName, TMap<FName, TMap<FName, TArray<FString>>>> FIniSettingContainer;
+
 }
 UCLASS()
 class UNREALED_API UCookOnTheFlyServer : public UObject, public FTickableEditorObject, public FExec, public UE::Cook::ICookInfo
@@ -242,7 +252,9 @@ private:
 
 	/** Current cook mode the cook on the fly server is running in */
 	ECookMode::Type CurrentCookMode = ECookMode::CookOnTheFly;
-	/** Directory to output to instead of the default should be empty in the case of DLC cooking */ 
+	/** CookMode of the Cook Director, equal to CurrentCookMode if not a CookWorker, otherwise equal to Director's CurrentCookMode */
+	ECookMode::Type DirectorCookMode = ECookMode::CookOnTheFly;
+	/** Directory to output to instead of the default should be empty in the case of DLC cooking */
 	FString OutputDirectoryOverride;
 
 	TUniquePtr<UE::Cook::FCookByTheBookOptions> CookByTheBookOptions;
@@ -349,17 +361,14 @@ private:
 	// iterative ini settings checking
 	// growing list of ini settings which are accessed over the course of the cook
 
-	// tmap of the Config name, Section name, Key name, to the value
-	typedef TMap<FName, TMap<FName, TMap<FName, TArray<FString>>>> FIniSettingContainer;
-
 	mutable FCriticalSection ConfigFileCS;
-	mutable FIniSettingContainer AccessedIniStrings;
+	mutable UE::Cook::FIniSettingContainer AccessedIniStrings;
 	TArray<const FConfigFile*> OpenConfigFiles;
 	TArray<FString> ConfigSettingDenyList;
 	void OnFConfigDeleted(const FConfigFile* Config);
 	void OnFConfigCreated(const FConfigFile* Config);
 
-	void ProcessAccessedIniSettings(const FConfigFile* Config, FIniSettingContainer& AccessedIniStrings) const;
+	void ProcessAccessedIniSettings(const FConfigFile* Config, UE::Cook::FIniSettingContainer& AccessedIniStrings) const;
 
 	void OnRequestClusterCompleted(const UE::Cook::FRequestCluster& RequestCluster);
 
@@ -585,9 +594,16 @@ public:
 	*/
 	void CancelCookByTheBook();
 
-	bool IsCookByTheBookRunning() const;
-	/** Report whether the CookOnTheFlyServer is in a cook session, either CookByTheBook or CookOnTheFly. Used to restrict operations when cooking and reduce cputime when not cooking. */
+	/** Connect to the CookDirector host, Initialize this UCOTFS, and start the CookWorker session */
+	bool TryInitializeCookWorker();
+
+	/** Terminate the CookWorker session */
+	void ShutdownCookAsCookWorker();
+
+	/** Is the local CookOnTheFlyServer in a cook session in any CookMode? Used to restrict operations when cooking and reduce cputime when not cooking. */
 	bool IsInSession() const;
+	/** Is the local CookOnTheFlyServer in a cook session in CookByTheBook Mode? */
+	bool IsCookByTheBookRunning() const;
 
 	/**
 	* Get any packages which are in memory, these were probably required to be loaded because of the current package we are cooking, so we should probably cook them also
@@ -602,6 +618,8 @@ public:
 	uint32 TickCookByTheBook(const float TimeSlice, ECookTickFlags TickFlags = ECookTickFlags::None);
 	/** Tick COTF until it finishes or needs to yield. Should only be called when in CookOnTheFly Mode. */
 	uint32 TickCookOnTheFly(const float TimeSlice, ECookTickFlags TickFlags = ECookTickFlags::None);
+	/** Tick CookWorker until it finishes or needs to yield. Should only be called when in CookWorker Mode. */
+	uint32 TickCookWorker();
 
 	/**
 	 * Execute CookByTheBook as a load of all packages followed by a save of all packages. This mode is sometimes faster than normal CBTB.
@@ -609,7 +627,7 @@ public:
 	 */
 	uint32 CookFullLoadAndSave();
 
-	/** Return whether the configuration and commandline settings have selected FullLoadAndSave. */
+	/** Is the Local CookOnTheFLyServer using FullLoadAndSave? Used in SingleProcess cooks only, always false on CookWorkers or MPCook Directors. */
 	bool IsFullLoadAndSave() const;
 
 	/**
@@ -685,38 +703,30 @@ public:
 
 	uint32 NumConnections() const;
 
-	/**
-	* Is this cooker running in the editor
-	*
-	* @return true if we are running in the editor
-	*/
+	/** Is the local CookOnTheFlyServer running in the editor? */
 	bool IsCookingInEditor() const;
 
-	/**
-	* Is this cooker running in real time mode (where it needs to respect the timeslice) 
-	* 
-	* @return returns true if this cooker is running in realtime mode like in the editor
-	*/
+	/** Is the local CookOnTheFlyServer initialized to run in real time mode (where it needs to respect the timeslice), e.g. it is running in Editor? */
 	bool IsRealtimeMode() const;
 
-	/**
-	* Helper function returns if we are in any cook by the book mode
-	*
-	* @return if the cook mode is a cook by the book mode
-	*/
+	/** Is the local CookOnTheFlyServer initialized to run CookByTheBook (find all packages needed by the game and cook them)? */
 	bool IsCookByTheBookMode() const;
+	/** Is the CookDirector (local CookOnTheFlyServer for SPCook or MPDirector, or the remote MPDirector for a CookWorker) initialized to run CookByTheBook? */
+	bool IsDirectorCookByTheBook() const;
 
+	/** Is the local CookOnTheFlyServer in a mode that uses ShaderCodeLibraries rather than storing shaders in the Packages that use them? */
 	bool IsUsingShaderCodeLibrary() const;
 
+	/** Is the local CookOnTheFlyServer using ZenStore (cooked packages are stored in Zen's Cache storage rather than as loose files on disk)? */
 	bool IsUsingZenStore() const;
 
-	/**
-	* Helper function returns if we are in any cook on the fly mode
-	*
-	* @return if the cook mode is a cook on the fly mode
-	*/
+	/** Is the local CookOnTheFlyServer initialized to run CookOnTheFly (accept connections from game executables, cook what they ask for)? */
 	bool IsCookOnTheFlyMode() const;
+	/** Is the CookDirector (local CookOnTheFlyServer for SPCook or MPDirector, or the remote MPDirector for a CookWorker) initialized to run CookOnTheFly? */
+	bool IsDirectorCookOnTheFly() const;
 
+	/** Is the local CookOnTheFlyServer initialized to run as a CookWorker (connect to a Director cooker cooking CBTB or COTF, cook what it assigns to us)? */
+	bool IsCookWorkerMode() const;
 
 	virtual void BeginDestroy() override;
 
@@ -800,6 +810,7 @@ public:
 
 private:
 
+	/** Is the local CookOnTheFlyServer initialized to run in CookOnTheFly AND using the legacy scheduling method that predates COTF2? */
 	bool IsUsingLegacyCookOnTheFlyScheduling() const;
 
 	/** Update accumulators of editor data and consume editor changes since the last cook when starting cooks in the editor. */
@@ -842,7 +853,7 @@ private:
 	static void GetGameDefaultObjects(const TArray<ITargetPlatform*>& TargetPlatforms, TMap<FName, TArray<FName>>& GameDefaultObjectsOut);
 
 	/* Collect filespackages that should not be cooked from ini settings and commandline. Does not include checking UAssetManager, which has to be queried later */
-	TArray<FName> GetNeverCookPackageFileNames(TArrayView<const FString> ExtraNeverCookDirectories = TArrayView<const FString>());
+	TArray<FName> GetNeverCookPackageFileNames(TArrayView<const FString> ExtraNeverCookDirectories = TArrayView<const FString>()) const; // const because it is not always called and should avoid side effects
 
 
 	/**
@@ -931,6 +942,13 @@ private:
 	FBeginCookContext CreateAddPlatformContext(ITargetPlatform* TargetPlatform);
 	void GetCookOnTheFlyUnsolicitedFiles(const ITargetPlatform* TargetPlatform, const FString& PlatformName, TArray<FString>& UnsolicitedFiles, const FString& Filename, bool bIsCookable);
 
+	//////////////////////////////////////////////////////////////////////////
+	// CookWorker specific functions
+	/** Start a CookWorker Session */
+	void StartCookAsCookWorker();
+	/** Construct the on-stack data for StartCook functions, based on the CookWorker's configuration */
+	FBeginCookContext CreateCookWorkerContext();
+
 	//////////////////////////////////////////////////////////////////////////
 	// general functions
 
@@ -951,14 +969,19 @@ private:
 		UE::Cook::FPackageData* ReportingPackageData = nullptr);
 
 	/** Read information about the previous cook from disk and set whether the current cook is iterative based on previous cook, config, and commandline. */
-	void SetBeginCookIterativeFlags(FBeginCookContext& BeginContext);
+	void LoadBeginCookIterativeFlags(FBeginCookContext& BeginContext);
+	void LoadBeginCookIterativeFlagsLocal(FBeginCookContext& BeginContext) const; // const because it is not always called and should avoid sideeffects
 	/** Initialize the sandbox for a new cook session */
 	void BeginCookSandbox(FBeginCookContext& BeginContext);
 
+	/** Set parameters that rely on config settings from startup and don't depend on information from StartCook* functions. */
+	void LoadInitializeConfigSettings(const FString& InOutputDirectoryOverride);
+	void SetInitializeConfigSettings(UE::Cook::FInitializeConfigSettings&& Settings);
+
 	/** Set parameters that rely on config and CookByTheBook settings. */
-	void SetBeginCookConfigSettings(FBeginCookContext& BeginContext);
-	/** Initialize NeverCookPackageList from packaging settings and platform-specific sources. */
-	void SetNeverCookPackageConfigSettings(FBeginCookContext& BeginContext);
+	void LoadBeginCookConfigSettings(FBeginCookContext& BeginContext);
+	void SetBeginCookConfigSettings(FBeginCookContext& BeginContext, UE::Cook::FBeginCookConfigSettings&& Settings);
+	void SetNeverCookPackageConfigSettings(FBeginCookContext& BeginContext, UE::Cook::FBeginCookConfigSettings& Settings);
 
 	/**
 	* Finalize the package store
@@ -980,7 +1003,8 @@ private:
 	* Build up a map of unsupported packages per platform that can be checked before saving.
 	*/
 	void DiscoverPlatformSpecificNeverCookPackages(
-		const TArrayView<const ITargetPlatform* const>& TargetPlatforms, const TArray<FString>& UBTPlatformStrings);
+		const TArrayView<const ITargetPlatform* const>& TargetPlatforms, const TArray<FString>& UBTPlatformStrings,
+		UE::Cook::FBeginCookConfigSettings& Settings) const; // const because it is not always called and should avoid side effects
 
 	/**
 	* GetDependentPackages
@@ -1056,7 +1080,7 @@ private:
 	* @param IniVersionStrings return list of the important current ini version strings
 	* @return false if function fails (should assume all platforms are out of date)
 	*/
-	bool GetCurrentIniVersionStrings( const ITargetPlatform* TargetPlatform, FIniSettingContainer& IniVersionStrings ) const;
+	bool GetCurrentIniVersionStrings( const ITargetPlatform* TargetPlatform, UE::Cook::FIniSettingContainer& IniVersionStrings ) const;
 
 	/**
 	* GetCookedIniVersionStrings gets the ini version strings used in previous cook for specified target platform
@@ -1064,7 +1088,7 @@ private:
 	* @param IniVersionStrings return list of the previous cooks ini version strings
 	* @return false if function fails to find the ini version strings
 	*/
-	bool GetCookedIniVersionStrings( const ITargetPlatform* TargetPlatform, FIniSettingContainer& IniVersionStrings, TMap<FString, FString>& AdditionalStrings ) const;
+	bool GetCookedIniVersionStrings( const ITargetPlatform* TargetPlatform, UE::Cook::FIniSettingContainer& IniVersionStrings, TMap<FString, FString>& AdditionalStrings ) const;
 
 	/** Test the CurrentCookSettings against the previous cooksettings to decide if we have to wipe the previous cook even when running iteratively. */
 	bool ArePreviousCookSettingsCompatible(const TMap<FName, FString>& CurrentCookSettings, const ITargetPlatform* TargetPlatform) const;
@@ -1102,11 +1126,10 @@ private:
 	/* Create the delete-old-cooked-directory helper.*/
 	FAsyncIODelete& GetAsyncIODelete();
 
+	/** Is the local CookOnTheFlyServer cooking a DLC plugin rather than a Project+Engine+EmbeddedPlugins? */
 	bool IsCookingDLC() const;
 
-	/**
-	* Returns true if we're cooking against a fixed release version
-	*/
+	/** Is the local CookOnTheFlyServer cooking a Project+Engine+EmbeddedPlugin Patch based on a previous Release? */
 	bool IsCookingAgainstFixedBase() const;
 
 	/**
@@ -1123,6 +1146,7 @@ private:
 
 	FString GetContentDirectoryForDLC() const;
 
+	/** Is the local CookOnTheFlyServer cooking a Project+Engine+EmbeddedPlugin Release that can be used as a base for future DLC or Patches? */
 	bool IsCreatingReleaseVersion();
 
 	/**
@@ -1149,12 +1173,7 @@ private:
 
 
 
-	/**
-	* IsCookFlagSet
-	* 
-	* @param InCookFlag used to check against the current cook flags
-	* @return true if the cook flag is set false otherwise
-	*/
+	/** Does the local CookOnTheFlyServer have all flags in InCookFlags set to true? */
 	bool IsCookFlagSet( const ECookInitializationFlags& InCookFlags ) const 
 	{
 		return (CookFlags & InCookFlags) != ECookInitializationFlags::None;
@@ -1221,8 +1240,10 @@ private:
 	UE::Cook::EPollStatus TryPopulateGeneratedPackage(UE::Cook::FGeneratorPackage& Generator, UE::Cook::FCookGenerationInfo& GeneratedInfo);
 
 	ICookedPackageWriter& FindOrCreatePackageWriter(const ITargetPlatform* TargetPlatform);
+	const ICookedPackageWriter* FindPackageWriter(const ITargetPlatform* TargetPlatform) const;
 	void FindOrCreateSaveContexts(TConstArrayView<const ITargetPlatform*> TargetPlatforms);
 	UE::Cook::FCookSavePackageContext& FindOrCreateSaveContext(const ITargetPlatform* TargetPlatform);
+	const UE::Cook::FCookSavePackageContext* FindSaveContext(const ITargetPlatform* TargetPlatform) const;
 	/** Allocate a new FCookSavePackageContext and ICookedPackageWriter for the given platform. */
 	UE::Cook::FCookSavePackageContext* CreateSaveContext(const ITargetPlatform* TargetPlatform);
 
@@ -1271,8 +1292,9 @@ private:
 	// These helper structs are all TUniquePtr rather than inline members so that we can keep their headers private.  See class header comments for their purpose.
 	TUniquePtr<UE::Cook::FPackageTracker> PackageTracker;
 	TUniquePtr<UE::Cook::FPackageDatas> PackageDatas;
-	TUniquePtr<UE::Cook::FExternalRequests> ExternalRequests;
+	TUniquePtr<UE::Cook::IWorkerRequests> WorkerRequests;
 	TUniquePtr<UE::Cook::FBuildDefinitions> BuildDefinitions;
+	TUniquePtr<UE::Cook::FCookDirector> CookDirector;
 
 	TArray<UE::Cook::FCookSavePackageContext*> SavePackageContexts;
 	/** Objects that were collected during the single-threaded PreGarbageCollect callback and that should be reported as referenced in CookerAddReferencedObjects. */
@@ -1301,11 +1323,15 @@ private:
 	double IdleStatusStartTime = 0.0;
 	uint32 CookedPackageCountSinceLastGC = 0;
 	float WaitForAsyncSleepSeconds = 0.0f;
+	float DisplayUpdatePeriodSeconds = 0.0f;
 	EIdleStatus IdleStatus = EIdleStatus::Done;
 
+	friend UE::Cook::FBeginCookConfigSettings;
+	friend UE::Cook::FGeneratorPackage;
 	friend UE::Cook::FPackageData;
 	friend UE::Cook::FPendingCookedPlatformData;
 	friend UE::Cook::FPlatformManager;
 	friend UE::Cook::FRequestCluster;
-	friend UE::Cook::FGeneratorPackage;
+	friend UE::Cook::FWorkerRequestsLocal;
+	friend UE::Cook::FWorkerRequestsRemote;
 };
diff --git a/Engine/Source/Editor/UnrealEd/Private/Commandlets/CookCommandlet.cpp b/Engine/Source/Editor/UnrealEd/Private/Commandlets/CookCommandlet.cpp
--- a/Engine/Source/Editor/UnrealEd/Private/Commandlets/CookCommandlet.cpp
+++ b/Engine/Source/Editor/UnrealEd/Private/Commandlets/CookCommandlet.cpp
@@ -401,6 +401,25 @@ namespace DetailedCookStats
 }
 #endif
 
+namespace UE::Cook
+{
+
+struct FScopeRootObject
+{
+	UObject* Object;
+	FScopeRootObject(UObject* InObject) : Object(InObject)
+	{
+		Object->AddToRoot();
+	}
+
+	~FScopeRootObject()
+	{
+		Object->RemoveFromRoot();
+	}
+};
+
+}
+
 UCookCommandlet::UCookCommandlet( const FObjectInitializer& ObjectInitializer )
 	: Super(ObjectInitializer)
 {
@@ -412,22 +431,8 @@ bool UCookCommandlet::CookOnTheFly( FGuid InstanceId, int32 Timeout, bool bForce
 {
 	UCookOnTheFlyServer *CookOnTheFlyServer = NewObject<UCookOnTheFlyServer>();
 
-	struct FScopeRootObject
-	{
-		UObject *Object;
-		FScopeRootObject( UObject *InObject ) : Object( InObject )
-		{
-			Object->AddToRoot();
-		}
-
-		~FScopeRootObject()
-		{
-			Object->RemoveFromRoot();
-		}
-	};
-
 	// make sure that the cookonthefly server doesn't get cleaned up while we are garbage collecting below :)
-	FScopeRootObject S(CookOnTheFlyServer);
+	UE::Cook::FScopeRootObject S(CookOnTheFlyServer);
 
 	UCookerSettings const* CookerSettings = GetDefault<UCookerSettings>();
 	ECookInitializationFlags IterateFlags = ECookInitializationFlags::Iterative;
@@ -459,11 +464,6 @@ bool UCookCommandlet::CookOnTheFly( FGuid InstanceId, int32 Timeout, bool bForce
 		}
 	}
 
-	if (bNoShaderCooking)
-	{
-		GShaderCompilingManager->SkipShaderCompilation(true);
-	}
-
 	FDateTime LastConnectionTime = FDateTime::UtcNow();
 	bool bHadConnection = false;
 
@@ -507,11 +507,6 @@ bool UCookCommandlet::CookOnTheFly( FGuid InstanceId, int32 Timeout, bool bForce
 		}
 	}
 
-	if (bNoShaderCooking)
-	{
-		GShaderCompilingManager->SkipShaderCompilation(false);
-	}
-
 	CookOnTheFlyServer->ShutdownCookOnTheFly();
 	return true;
 }
@@ -541,7 +536,6 @@ int32 UCookCommandlet::Main(const FString& CmdLineParams)
 	bPartialGC = Switches.Contains(TEXT("Partialgc"));
 	ShowErrorCount = !Switches.Contains(TEXT("DIFFONLY"));
 	ShowProgress = !Switches.Contains(TEXT("DIFFONLY"));
-	bNoShaderCooking = bCookOnTheFly; // Do not cook any shaders into the shader maps. Always true if we are running w/ cook on the fly
 	bIgnoreIniSettingsOutOfDate = Switches.Contains(TEXT("IgnoreIniSettingsOutOfDate"));
 	bFastCook = Switches.Contains(TEXT("FastCook"));
 
@@ -590,6 +584,10 @@ int32 UCookCommandlet::Main(const FString& CmdLineParams)
 
 		CookOnTheFly( InstanceId, Timeout, bForceClose, TargetPlatforms);
 	}
+	else if (Switches.Contains(TEXT("COOKWORKER")))
+	{
+		CookAsCookWorker();
+	}
 	else
 	{
 		const TArray<ITargetPlatform*>& Platforms = TPM.GetActiveTargetPlatforms();
@@ -622,23 +620,8 @@ bool UCookCommandlet::CookByTheBook( const TArray<ITargetPlatform*>& Platforms)
 
 	COOK_STAT(FScopedDurationTimer CookByTheBookTimer(DetailedCookStats::CookByTheBookTimeSec));
 	UCookOnTheFlyServer *CookOnTheFlyServer = NewObject<UCookOnTheFlyServer>();
-
-	struct FScopeRootObject
-	{
-		UObject *Object;
-		FScopeRootObject( UObject *InObject ) : Object( InObject )
-		{
-			Object->AddToRoot();
-		}
-
-		~FScopeRootObject()
-		{
-			Object->RemoveFromRoot();
-		}
-	};
-
 	// make sure that the cookonthefly server doesn't get cleaned up while we are garbage collecting below :)
-	FScopeRootObject S(CookOnTheFlyServer);
+	UE::Cook::FScopeRootObject S(CookOnTheFlyServer);
 
 	UCookerSettings const* CookerSettings = GetDefault<UCookerSettings>();
 	const UProjectPackagingSettings* const PackagingSettings = GetDefault<UProjectPackagingSettings>();
@@ -672,9 +655,6 @@ bool UCookCommandlet::CookByTheBook( const TArray<ITargetPlatform*>& Platforms)
 	FString DLCName;
 	FParse::Value( *Params, TEXT("DLCNAME="), DLCName);
 
-	FString ChildCookFile;
-	FParse::Value(*Params, TEXT("cookchild="), ChildCookFile);
-
 	FString BasedOnReleaseVersion;
 	FParse::Value( *Params, TEXT("BasedOnReleaseVersion="), BasedOnReleaseVersion);
 
@@ -909,6 +889,34 @@ bool UCookCommandlet::CookByTheBook( const TArray<ITargetPlatform*>& Platforms)
 	return true;
 }
 
+bool UCookCommandlet::CookAsCookWorker()
+{
+#if OUTPUT_COOKTIMING
+	TRACE_CPUPROFILER_EVENT_SCOPE_ON_CHANNEL(CookAsCookWorker, CookChannel);
+#endif // OUTPUT_COOKTIMING
+
+	UCookOnTheFlyServer* CookOnTheFlyServer = NewObject<UCookOnTheFlyServer>();
+	// make sure that the cookonthefly server doesn't get cleaned up while we are garbage collecting below
+	UE::Cook::FScopeRootObject S(CookOnTheFlyServer);
+
+	if (!CookOnTheFlyServer->TryInitializeCookWorker())
+	{
+		return false;
+	}
+
+	{
+		DECLARE_SCOPE_CYCLE_COUNTER(TEXT("CookByTheBook.MainLoop"), STAT_CookByTheBook_MainLoop, STATGROUP_LoadTime);
+		while (CookOnTheFlyServer->IsInSession())
+		{
+			uint32 TickResults = CookOnTheFlyServer->TickCookWorker();
+			ConditionalCollectGarbage(TickResults, *CookOnTheFlyServer);
+		}
+	}
+	CookOnTheFlyServer->ShutdownCookAsCookWorker();
+
+	return true;
+}
+
 void UCookCommandlet::ConditionalCollectGarbage(uint32 TickResults, UCookOnTheFlyServer& COTFS)
 {
 	if ((TickResults & UCookOnTheFlyServer::COSR_RequiresGC) == 0)
@@ -968,7 +976,7 @@ void UCookCommandlet::ConditionalCollectGarbage(uint32 TickResults, UCookOnTheFl
 	}
 #if OUTPUT_COOKTIMING
 	TOptional<FScopedDurationTimer> CBTBScopedDurationTimer;
-	if (COTFS.IsCookByTheBookMode())
+	if (!COTFS.IsCookOnTheFlyMode())
 	{
 		CBTBScopedDurationTimer.Emplace(DetailedCookStats::TickLoopGCTimeSec);
 	}
@@ -993,7 +1001,7 @@ void UCookCommandlet::ConditionalCollectGarbage(uint32 TickResults, UCookOnTheFl
 		}
 	};
 
-	if (COTFS.IsCookByTheBookMode())
+	if (!COTFS.IsCookOnTheFlyMode())
 	{
 		DumpMemStats();
 	}
@@ -1002,7 +1010,7 @@ void UCookCommandlet::ConditionalCollectGarbage(uint32 TickResults, UCookOnTheFl
 
 	int32 NumObjectsAfterGC = GUObjectArray.GetObjectArrayNumMinusAvailable();
 	int32 NumObjectsAvailableAfterGC = GUObjectArray.GetObjectArrayEstimatedAvailable();
-	if (COTFS.IsCookByTheBookMode())
+	if (!COTFS.IsCookOnTheFlyMode())
 	{
 		UE_LOG(LogCookCommandlet, Display, TEXT("%s GC before %d available %d after %d available %d"),
 			(bPartialGC ? TEXT("Partial") : TEXT("Full")), NumObjectsBeforeGC, NumObjectsAvailableBeforeGC, NumObjectsAfterGC, NumObjectsAvailableAfterGC);
diff --git a/Engine/Source/Editor/UnrealEd/Private/CookOnTheFlyServer.cpp b/Engine/Source/Editor/UnrealEd/Private/CookOnTheFlyServer.cpp
--- a/Engine/Source/Editor/UnrealEd/Private/CookOnTheFlyServer.cpp
+++ b/Engine/Source/Editor/UnrealEd/Private/CookOnTheFlyServer.cpp
@@ -19,6 +19,7 @@
 #include "Commandlets/ShaderPipelineCacheToolsCommandlet.h"
 #include "Containers/RingBuffer.h"
 #include "Cooker/AsyncIODelete.h"
+#include "Cooker/CookDirector.h"
 #include "Cooker/CookedSavePackageValidator.h"
 #include "Cooker/CookOnTheFlyServerInterface.h"
 #include "Cooker/CookPackageData.h"
@@ -32,6 +33,8 @@
 #include "Cooker/LooseCookedPackageWriter.h"
 #include "Cooker/NetworkFileCookOnTheFlyRequestManager.h"
 #include "Cooker/PackageTracker.h"
+#include "Cooker/WorkerRequestsLocal.h"
+#include "Cooker/WorkerRequestsRemote.h"
 #include "CookerSettings.h"
 #include "CookPackageSplitter.h"
 #include "DerivedDataCacheInterface.h"
@@ -341,20 +344,15 @@ public:
 
 		FName StandardFileName(*FPaths::CreateStandardFilename(CookPackageRequest.Filename));
 		UE_LOG(LogCook, Verbose, TEXT("Enqueing cook request, Filename='%s', Platform='%s'"), *CookPackageRequest.Filename, *CookPackageRequest.PlatformName.ToString());
-		Cooker.ExternalRequests->EnqueueUnique(
-			FFilePlatformRequest(StandardFileName, EInstigator::CookOnTheFly, TargetPlatform, MoveTemp(CookPackageRequest.CompletionCallback)),
-			true);
-		if (Cooker.ExternalRequests->CookRequestEvent)
-		{
-			Cooker.ExternalRequests->CookRequestEvent->Trigger();
-		}
+		Cooker.WorkerRequests->AddCookOnTheFlyRequest(
+			FFilePlatformRequest(StandardFileName, EInstigator::CookOnTheFly, TargetPlatform, MoveTemp(CookPackageRequest.CompletionCallback)));
 
 		return true;
 	};
 
 	virtual void MarkPackageDirty(const FName& PackageName) override
 	{
-		Cooker.ExternalRequests->AddCallback([this, PackageName]()
+		Cooker.WorkerRequests->AddCookOnTheFlyCallback([this, PackageName]()
 			{
 				UE::Cook::FPackageData* PackageData = Cooker.PackageDatas->FindPackageDataByPackageName(PackageName);
 				if (!PackageData)
@@ -405,17 +403,12 @@ private:
 		}
 
 		FEventRef Event;
-		Cooker.ExternalRequests->AddCallback([this, &Event, &TargetPlatform]()
+		Cooker.WorkerRequests->AddCookOnTheFlyCallback([this, &Event, &TargetPlatform]()
 			{
 				Cooker.AddCookOnTheFlyPlatformFromGameThread(TargetPlatform);
 				Event->Trigger();
 			});
 
-		if (Cooker.ExternalRequests->CookRequestEvent)
-		{
-			Cooker.ExternalRequests->CookRequestEvent->Trigger();
-		}
-
 		Event->Wait();
 		return TargetPlatform;
 	}
@@ -429,18 +422,13 @@ private:
 
 		FEventRef Event;
 		ICookedPackageWriter* PackageWriter = nullptr;
-		Cooker.ExternalRequests->AddCallback([this, &Event, &TargetPlatform, &PackageWriter]()
+		Cooker.WorkerRequests->AddCookOnTheFlyCallback([this, &Event, &TargetPlatform, &PackageWriter]()
 			{
 				PackageWriter = &Cooker.FindOrCreatePackageWriter(TargetPlatform);
 				check(PackageWriter);
 				Event->Trigger();
 			});
 
-		if (Cooker.ExternalRequests->CookRequestEvent)
-		{
-			Cooker.ExternalRequests->CookRequestEvent->Trigger();
-		}
-
 		Event->Wait();
 		return *PackageWriter;
 	}
@@ -536,6 +524,11 @@ bool UCookOnTheFlyServer::StartCookOnTheFly(FCookOnTheFlyStartupOptions InCookOn
 {
 	using namespace UE::Cook;
 
+	if (!IsCookingInEditor())
+	{
+		GShaderCompilingManager->SkipShaderCompilation(true);
+	}
+
 	LLM_SCOPE_BYTAG(Cooker);
 #if WITH_COTF
 	check(IsCookOnTheFlyMode());
@@ -549,7 +542,7 @@ bool UCookOnTheFlyServer::StartCookOnTheFly(FCookOnTheFlyStartupOptions InCookOn
 	CreateSandboxFile(BeginContext);
 
 	CookOnTheFlyServerInterface = MakeUnique<UCookOnTheFlyServer::FCookOnTheFlyServerInterface>(*this);
-	ExternalRequests->CookRequestEvent = FPlatformProcess::GetSynchEventFromPool();
+	WorkerRequests->InitializeCookOnTheFly();
 
 	// Precreate the map of all possible target platforms so we can access the collection of existing platforms in a threadsafe manner
 	// Each PlatformData in the map will be uninitialized until we call AddCookOnTheFlyPlatform for the platform
@@ -560,7 +553,7 @@ bool UCookOnTheFlyServer::StartCookOnTheFly(FCookOnTheFlyStartupOptions InCookOn
 	}
 	PlatformManager->SetArePlatformsPrepopulated(true);
 
-	SetBeginCookConfigSettings(BeginContext);
+	LoadBeginCookConfigSettings(BeginContext);
 
 	GRedirectCollector.OnStartupPackageLoadComplete();
 
@@ -691,7 +684,7 @@ void UCookOnTheFlyServer::AddCookOnTheFlyPlatformFromGameThread(ITargetPlatform*
 	// Initialize systems and settings that the rest of AddCookOnTheFlyPlatformFromGameThread depends on
 	// Functions in this section are ordered and can depend on the functions before them
 	FindOrCreateSaveContexts(BeginContext.TargetPlatforms);
-	SetBeginCookIterativeFlags(BeginContext);
+	LoadBeginCookIterativeFlags(BeginContext);
 
 	// Initialize systems referenced by later stages or that need to start early for async performance
 	// Functions in this section must not need to read/write the SandboxDirectory or MemoryCookedPackages
@@ -740,7 +733,7 @@ void UCookOnTheFlyServer::OnTargetPlatformsInvalidated()
 
 	PackageDatas->RemapTargetPlatforms(Remap);
 	PackageTracker->RemapTargetPlatforms(Remap);
-	ExternalRequests->RemapTargetPlatforms(Remap);
+	WorkerRequests->RemapTargetPlatforms(Remap);
 
 	if (PlatformManager->GetArePlatformsPrepopulated())
 	{
@@ -1087,9 +1080,14 @@ bool UCookOnTheFlyServer::IsCookByTheBookMode() const
 	return CurrentCookMode == ECookMode::CookByTheBookFromTheEditor || CurrentCookMode == ECookMode::CookByTheBook;
 }
 
+bool UCookOnTheFlyServer::IsDirectorCookByTheBook() const
+{
+	return DirectorCookMode == ECookMode::CookByTheBookFromTheEditor || DirectorCookMode == ECookMode::CookByTheBook;
+}
+
 bool UCookOnTheFlyServer::IsUsingShaderCodeLibrary() const
 {
-	return IsCookByTheBookMode() && AllowShaderCompiling();
+	return IsDirectorCookByTheBook() && AllowShaderCompiling();
 }
 
 bool UCookOnTheFlyServer::IsUsingZenStore() const
@@ -1102,6 +1100,12 @@ bool UCookOnTheFlyServer::IsCookOnTheFlyMode() const
 	return CurrentCookMode == ECookMode::CookOnTheFly || CurrentCookMode == ECookMode::CookOnTheFlyFromTheEditor; 
 }
 
+bool UCookOnTheFlyServer::IsDirectorCookOnTheFly() const
+{
+	return DirectorCookMode == ECookMode::CookOnTheFly || DirectorCookMode == ECookMode::CookOnTheFlyFromTheEditor;
+}
+
+
 bool UCookOnTheFlyServer::IsUsingLegacyCookOnTheFlyScheduling() const
 {
 	return CookOnTheFlyRequestManager && CookOnTheFlyRequestManager->ShouldUseLegacyScheduling();
@@ -1192,22 +1196,20 @@ COREUOBJECT_API extern bool GOutputCookingWarnings;
 
 void UCookOnTheFlyServer::WaitForRequests(int TimeoutMs)
 {
-	if (ExternalRequests->CookRequestEvent)
-	{
-		ExternalRequests->CookRequestEvent->Wait(TimeoutMs, true);
-	}
+	WorkerRequests->WaitForCookOnTheFlyEvents(TimeoutMs);
 }
 
 bool UCookOnTheFlyServer::HasRemainingWork() const
 { 
-	return ExternalRequests->HasRequests() || PackageDatas->GetMonitor().GetNumInProgress() > 0;
+	return WorkerRequests->HasExternalRequests() ||
+		PackageDatas->GetMonitor().GetNumInProgress() > 0;
 }
 
 bool UCookOnTheFlyServer::RequestPackage(const FName& StandardFileName, const TArrayView<const ITargetPlatform* const>& TargetPlatforms, const bool bForceFrontOfQueue)
 {
 	using namespace UE::Cook;
 
-	if (!IsCookByTheBookMode())
+	if (IsCookOnTheFlyMode())
 	{
 		bCookOnTheFlyExternalRequests = true;
 		for (const ITargetPlatform* TargetPlatform : TargetPlatforms)
@@ -1217,7 +1219,7 @@ bool UCookOnTheFlyServer::RequestPackage(const FName& StandardFileName, const TA
 		}
 	}
 
-	ExternalRequests->EnqueueUnique(
+	WorkerRequests->AddPublicInterfaceRequest(
 		FFilePlatformRequest(StandardFileName, EInstigator::RequestPackageFunction, TargetPlatforms),
 		bForceFrontOfQueue);
 	return true;
@@ -1240,7 +1242,7 @@ bool UCookOnTheFlyServer::RequestPackage(const FName& StandardFileName, const TA
 
 bool UCookOnTheFlyServer::RequestPackage(const FName& StandardPackageFName, const bool bForceFrontOfQueue)
 {
-	check(IsCookByTheBookMode()); // Invalid to call RequestPackage without a list of TargetPlatforms unless we are in cook by the book mode
+	check(!IsCookOnTheFlyMode()); // Invalid to call RequestPackage without a list of TargetPlatforms if we are in CookOnTheFly
 	return RequestPackage(StandardPackageFName, PlatformManager->GetSessionPlatforms(), bForceFrontOfQueue);
 }
 
@@ -1282,6 +1284,19 @@ uint32 UCookOnTheFlyServer::TickCookOnTheFly(const float TimeSlice, ECookTickFla
 	return StackData.ResultFlags;
 }
 
+uint32 UCookOnTheFlyServer::TickCookWorker()
+{
+	check(IsCookWorkerMode());
+
+	LLM_SCOPE_BYTAG(Cooker);
+	COOK_STAT(FScopedDurationTimer TickTimer(DetailedCookStats::TickCookOnTheSideTimeSec));
+	UE::Cook::FTickStackData StackData(MAX_flt, ECookTickFlags::None);
+
+	TickMainCookLoop(StackData);
+
+	return StackData.ResultFlags;
+}
+
 void UCookOnTheFlyServer::TickMainCookLoop(UE::Cook::FTickStackData& StackData)
 {
 	UE_SCOPED_HIERARCHICAL_COOKTIMER(TickMainCookLoop);
@@ -1291,7 +1306,7 @@ void UCookOnTheFlyServer::TickMainCookLoop(UE::Cook::FTickStackData& StackData)
 	}
 
 	bool bContinueTick = true;
-	while (bContinueTick && (!IsEngineExitRequested() || CurrentCookMode == ECookMode::CookByTheBook))
+	while (bContinueTick && (!IsEngineExitRequested() || (IsCookByTheBookMode() && !IsCookingInEditor())))
 	{
 		TickCookStatus(StackData);
 
@@ -1585,8 +1600,13 @@ void UCookOnTheFlyServer::UpdateDisplay(UE::Cook::FTickStackData& StackData, boo
 
 	const double CurrentTime = StackData.LoopStartTime;
 	const float DeltaProgressDisplayTime = CurrentTime - LastProgressDisplayTime;
+	if (!bForceDisplay && DeltaProgressDisplayTime < DisplayUpdatePeriodSeconds)
+	{
+		return;
+	}
+
 	const int32 CookedPackagesCount = PackageDatas->GetNumCooked();
-	const int32 CookPendingCount = ExternalRequests->GetNumRequests() + PackageDatas->GetMonitor().GetNumInProgress();
+	const int32 CookPendingCount = WorkerRequests->GetNumExternalRequests() + PackageDatas->GetMonitor().GetNumInProgress();
 	if (bForceDisplay ||
 		(DeltaProgressDisplayTime >= GCookProgressUpdateTime && CookPendingCount != 0 &&
 			(LastCookedPackagesCount != CookedPackagesCount || LastCookPendingCount != CookPendingCount || DeltaProgressDisplayTime > GCookProgressRepeatTime)))
@@ -1971,7 +1991,7 @@ static void ProcessDeferredCommands(UCookOnTheFlyServer& COTFS)
 {
 #if OUTPUT_COOKTIMING
 	TOptional<FScopedDurationTimer> CBTBScopedDurationTimer;
-	if (COTFS.IsCookByTheBookMode())
+	if (!COTFS.IsCookOnTheFlyMode())
 	{
 		CBTBScopedDurationTimer.Emplace(DetailedCookStats::TickLoopProcessDeferredCommandsTimeSec);
 	}
@@ -2038,14 +2058,14 @@ void UCookOnTheFlyServer::InitializePollables()
 			Pollables.Emplace(MoveTemp(Pollable));
 		}
 	}
-	if (IsCookByTheBookMode())
+	if (!IsCookOnTheFlyMode())
 	{
 		if (!IsCookingInEditor())
 		{
 			Pollables.Emplace(new FPollable(TEXT("CommandletStats"), 60.f, 5.f, [](FTickStackData&) { CBTBTickCommandletStats(); }));
 		}
 	}
-	else if (IsCookOnTheFlyMode())
+	else
 	{
 		RecompileRequestsPollable = new FPollable(TEXT("RecompileShaderRequests"), FPollable::EManualTrigger(), [this](FTickStackData&) { TickRecompileShaderRequestsPrivate(); });
 		Pollables.Add(FPollableQueueKey(RecompileRequestsPollable));
@@ -2207,7 +2227,7 @@ UCookOnTheFlyServer::ECookAction UCookOnTheFlyServer::DecideNextCookAction(UE::C
 
 void UCookOnTheFlyServer::PumpExternalRequests(const UE::Cook::FCookerTimer& CookerTimer)
 {
-	if (!ExternalRequests->HasRequests())
+	if (!WorkerRequests->HasExternalRequests())
 	{
 		return;
 	}
@@ -2220,7 +2240,7 @@ void UCookOnTheFlyServer::PumpExternalRequests(const UE::Cook::FCookerTimer& Coo
 	{
 		BuildRequests.Reset();
 		SchedulerCallbacks.Reset();
-		RequestType = ExternalRequests->DequeueNextCluster(SchedulerCallbacks, BuildRequests);
+		RequestType = WorkerRequests->DequeueNextCluster(SchedulerCallbacks, BuildRequests);
 		if (RequestType == UE::Cook::EExternalRequestType::None)
 		{
 			// No more requests to process
@@ -2252,6 +2272,7 @@ void UCookOnTheFlyServer::PumpExternalRequests(const UE::Cook::FCookerTimer& Coo
 
 bool UCookOnTheFlyServer::TryCreateRequestCluster(UE::Cook::FPackageData& PackageData)
 {
+	check(PackageData.IsInProgress()); // This should only be called from Pump functions, and only on in-progress Packages
 	using namespace UE::Cook;
 	if (!PackageData.AreAllRequestedPlatformsExplored())
 	{
@@ -2468,8 +2489,7 @@ void UCookOnTheFlyServer::LoadPackageInQueue(UE::Cook::FPackageData& PackageData
 				*OtherPackageData.GetFileName().ToString());
 			TArray<const ITargetPlatform*, TInlineAllocator<ExpectedMaxNumPlatforms>> RequestedPlatforms;
 			PackageData.GetRequestedPlatforms(RequestedPlatforms);
-			OtherPackageData.UpdateRequestData(RequestedPlatforms, PackageData.GetIsUrgent(), FCompletionCallback(),
-				FInstigator(PackageData.GetInstigator()));
+			QueueDiscoveredPackageData(OtherPackageData, FInstigator(PackageData.GetInstigator()));
 
 			PackageData.SetPlatformsCooked(PlatformManager->GetSessionPlatforms(), true);
 			RejectPackageToLoad(PackageData, TEXT("is redirected to another filename"));
@@ -2589,6 +2609,8 @@ UE::Cook::FPackageData* UCookOnTheFlyServer::QueueDiscoveredPackage(UPackage* Pa
 void UCookOnTheFlyServer::QueueDiscoveredPackageData(UE::Cook::FPackageData& PackageData,
 	UE::Cook::FInstigator&& Instigator, bool bLoadReady)
 {
+	using namespace UE::Cook;
+
 	const TArray<const ITargetPlatform*>& TargetPlatforms = PlatformManager->GetSessionPlatforms();
 	if (PackageData.HasAllCookedPlatforms(TargetPlatforms, true /* bIncludeFailed */))
 	{
@@ -2596,6 +2618,11 @@ void UCookOnTheFlyServer::QueueDiscoveredPackageData(UE::Cook::FPackageData& Pac
 		return;
 	}
 
+	if (PackageData.IsInProgress())
+	{
+		return;
+	}
+
 	if (CookOnTheFlyRequestManager)
 	{
 		if (PackageData.IsGenerated())
@@ -2607,9 +2634,18 @@ void UCookOnTheFlyServer::QueueDiscoveredPackageData(UE::Cook::FPackageData& Pac
 			return;
 		}
 	}
-	
-	if (!PackageData.IsInProgress() &&
-		(PackageData.IsGenerated() || !CookByTheBookOptions->bSkipHardReferences))
+
+	bool bShouldAddToQueue;
+	WorkerRequests->AddDiscoveredPackage(PackageData, Instigator, bLoadReady, bShouldAddToQueue);
+	if (!bShouldAddToQueue)
+	{
+		return;
+	}
+
+	if ((Instigator.Category == EInstigator::GeneratedPackage) || 
+		(Instigator.Category == EInstigator::Unspecified) || 
+		(Instigator.Category == EInstigator::Unsolicited && !CookByTheBookOptions->bSkipHardReferences) ||
+		!CookByTheBookOptions->bSkipSoftReferences)
 	{
 		PackageData.SetRequestData(TargetPlatforms, /*bIsUrgent*/ false, UE::Cook::FCompletionCallback(),
 			MoveTemp(Instigator));
@@ -2654,7 +2690,7 @@ void UCookOnTheFlyServer::UpdatePackageFilter()
 void UCookOnTheFlyServer::OnRemoveSessionPlatform(const ITargetPlatform* TargetPlatform)
 {
 	PackageDatas->OnRemoveSessionPlatform(TargetPlatform);
-	ExternalRequests->OnRemoveSessionPlatform(TargetPlatform);
+	WorkerRequests->OnRemoveSessionPlatform(TargetPlatform);
 }
 
 void UCookOnTheFlyServer::TickNetwork()
@@ -2662,15 +2698,18 @@ void UCookOnTheFlyServer::TickNetwork()
 	// Only CookOnTheFly handles network requests
 	// It is not safe to call PruneUnreferencedSessionPlatforms in CookByTheBook because StartCookByTheBook does not AddRef its session platforms
 	check(IsCookOnTheFlyMode())
-	if (IsInSession() && !bCookOnTheFlyExternalRequests)
+	if (IsInSession())
 	{
-		PlatformManager->PruneUnreferencedSessionPlatforms(*this);
+		if (!bCookOnTheFlyExternalRequests)
+		{
+			PlatformManager->PruneUnreferencedSessionPlatforms(*this);
+		}
 	}
 	else
 	{
 		// Process callbacks in case there is a callback pending that needs to create a session
 		TArray<UE::Cook::FSchedulerCallback> Callbacks;
-		if (ExternalRequests->DequeueCallbacks(Callbacks))
+		if (WorkerRequests->DequeueSchedulerCallbacks(Callbacks))
 		{
 			for (UE::Cook::FSchedulerCallback& Callback : Callbacks)
 			{
@@ -3490,7 +3529,7 @@ void UCookOnTheFlyServer::ReleaseCookedPlatformData(UE::Cook::FPackageData& Pack
 					{
 						check(BeginCacheObject.bHasFinishedRound);
 						Object->ClearAllCachedCookedPlatformData();
-						if (IsCookByTheBookMode())
+						if (IsDirectorCookByTheBook() && !IsCookingInEditor())
 						{
 							Object->WillNeverCacheCookedPlatformDataAgain();
 						}
@@ -3505,7 +3544,7 @@ void UCookOnTheFlyServer::ReleaseCookedPlatformData(UE::Cook::FPackageData& Pack
 					if (Object)
 					{
 						Object->ClearAllCachedCookedPlatformData();
-						if (CurrentCookMode == ECookMode::CookByTheBook)
+						if (IsDirectorCookByTheBook() && !IsCookingInEditor())
 						{
 							Object->WillNeverCacheCookedPlatformDataAgain();
 						}
@@ -3641,7 +3680,7 @@ void UCookOnTheFlyServer::ReleaseCookedPlatformData(UE::Cook::FPackageData& Pack
 
 	if (ReleaseSaveReason != EReleaseSaveReason::RecreateObjectCache)
 	{
-		if (CurrentCookMode == ECookMode::CookByTheBook)
+		if (!IsCookOnTheFlyMode() && !IsCookingInEditor())
 		{
 			UPackage* Package = PackageData.GetPackage();
 			if (Package && Package->GetLinker())
@@ -3892,7 +3931,7 @@ void UCookOnTheFlyServer::PumpSaves(UE::Cook::FTickStackData& StackData, uint32
 		{
 			if (IsUsingLegacyCookOnTheFlyScheduling() && !PackageData.GetIsUrgent())
 			{
-				if (ExternalRequests->HasRequests() || PackageDatas->GetMonitor().GetNumUrgent() > 0)
+				if (WorkerRequests->HasExternalRequests() || PackageDatas->GetMonitor().GetNumUrgent() > 0)
 				{
 					bShouldFinishTick = true;
 				}
@@ -3921,7 +3960,6 @@ void UCookOnTheFlyServer::PumpSaves(UE::Cook::FTickStackData& StackData, uint32
 		}
 		else // !IsCookOnTheFlyMode
 		{
-			check(IsCookByTheBookMode());
 			if (StackData.Timer.IsTimeUp())
 			{
 				// our timeslice is up
@@ -4070,7 +4108,7 @@ void UCookOnTheFlyServer::PostLoadPackageFixup(UE::Cook::FPackageData& PackageDa
 	// Perform special processing for UWorld
 	World->PersistentLevel->HandleLegacyMapBuildData();
 
-	if (IsCookByTheBookMode() == false || CookByTheBookOptions->bSkipSoftReferences)
+	if (IsDirectorCookOnTheFly() || CookByTheBookOptions->bSkipSoftReferences)
 	{
 		return;
 	}
@@ -4099,13 +4137,10 @@ void UCookOnTheFlyServer::PostLoadPackageFixup(UE::Cook::FPackageData& PackageDa
 	FName OwnerName = Package->GetFName();
 	for (const FString& PackageName : NewPackagesToCook)
 	{
-		FName PackageFName(*PackageName);
-		UE::Cook::FPackageData* NewPackageData = PackageDatas->TryAddPackageDataByPackageName(PackageFName);
-		if (NewPackageData && !NewPackageData->IsInProgress())
+		UE::Cook::FPackageData* NewPackageData = PackageDatas->TryAddPackageDataByPackageName(FName(*PackageName));
+		if (NewPackageData)
 		{
-			NewPackageData->UpdateRequestData(PlatformManager->GetSessionPlatforms(), false /* IsUrgent */,
-				UE::Cook::FCompletionCallback(),
-				UE::Cook::FInstigator(UE::Cook::EInstigator::Dependency, OwnerName));
+			QueueDiscoveredPackageData(*NewPackageData, UE::Cook::FInstigator(UE::Cook::EInstigator::Dependency, OwnerName));
 		}
 	}
 }
@@ -4582,7 +4617,7 @@ void UCookOnTheFlyServer::MarkPackageDirtyForCooker(UPackage* Package, bool bAll
 
 	if (IsInSession() && !bAllowInSession)
 	{
-		ExternalRequests->AddCallback([this, PackageName]() { MarkPackageDirtyForCookerFromSchedulerThread(PackageName); });
+		WorkerRequests->AddEditorActionCallback([this, PackageName]() { MarkPackageDirtyForCookerFromSchedulerThread(PackageName); });
 	}
 	else
 	{
@@ -4614,8 +4649,7 @@ void UCookOnTheFlyServer::MarkPackageDirtyForCookerFromSchedulerThread(const FNa
 		}
 		else if (IsCookByTheBookMode() && IsInSession() && bHadCookedPlatforms)
 		{
-			PackageData->UpdateRequestData(PlatformManager->GetSessionPlatforms(), false /* bIsUrgent */, UE::Cook::FCompletionCallback(),
-				UE::Cook::FInstigator(UE::Cook::EInstigator::Unspecified, GInstigatorMarkPackageDirty));
+			QueueDiscoveredPackageData(*PackageData, UE::Cook::FInstigator(UE::Cook::EInstigator::Unspecified, GInstigatorMarkPackageDirty));
 		}
 
 		if ( IsCookOnTheFlyMode() && FileModifiedDelegate.IsBound())
@@ -4647,7 +4681,13 @@ void UCookOnTheFlyServer::ShutdownCookOnTheFly()
 		ClearHierarchyTimers();
 
 		ShutdownCookSession();
+
+		if (!IsCookingInEditor())
+		{
+			GShaderCompilingManager->SkipShaderCompilation(false);
+		}
 	}
+
 }
 
 uint32 UCookOnTheFlyServer::GetPackagesPerGC() const
@@ -4663,7 +4703,7 @@ uint32 UCookOnTheFlyServer::GetPackagesPerPartialGC() const
 
 double UCookOnTheFlyServer::GetIdleTimeToGC() const
 {
-	if (CurrentCookMode == ECookMode::CookOnTheFly)
+	if (IsCookOnTheFlyMode() && !IsCookingInEditor())
 	{
 		// For COTF outside of the editor we want to release open linker file handles promptly but still give some time for new requests to come in
 		return 0.5;
@@ -5178,7 +5218,7 @@ FSaveCookedPackageContext::FSaveCookedPackageContext(UCookOnTheFlyServer& InCOTF
 	, Package(PackageData.GetPackage())
 	, PackageName(Package ? Package->GetName() : FString())
 	, Filename(PackageData.GetFileName().ToString())
-	, ArchiveCookContext(Package, COTFS.IsCookByTheBookMode() ? FArchiveCookContext::ECookByTheBook : FArchiveCookContext::ECookOnTheFly)
+	, ArchiveCookContext(Package, InCOTFS.IsDirectorCookByTheBook() ? FArchiveCookContext::ECookByTheBook : FArchiveCookContext::ECookOnTheFly)
 {
 }
 
@@ -5349,7 +5389,7 @@ void FSaveCookedPackageContext::FinishPlatform()
 		// TODO: Reenable BuildDefinitionList once FCbPackage support for empty FCbObjects is in
 		//Info.Attachments.Add({ "BuildDefinitionList", BuildDefinitionList });
 		Info.WriteOptions = IPackageWriter::EWriteOptions::Write;
-		if (COTFS.IsCookByTheBookMode())
+		if (COTFS.IsDirectorCookByTheBook())
 		{
 			Info.WriteOptions |= IPackageWriter::EWriteOptions::ComputeHash;
 		}
@@ -5358,7 +5398,7 @@ void FSaveCookedPackageContext::FinishPlatform()
 	}
 
 	// Update asset registry
-	if (COTFS.IsCookByTheBookMode())
+	if (COTFS.IsDirectorCookByTheBook())
 	{
 		Generator.UpdateAssetRegistryPackageData(*Package, SavePackageResult, MoveTemp(*ArchiveCookContext.GetCookTagList()));
 	}
@@ -5414,7 +5454,7 @@ void FSaveCookedPackageContext::FinishPackage()
 	if (!COTFS.CookByTheBookOptions->bSkipSoftReferences)
 	{
 		// Also request any localized variants of this package
-		if (COTFS.IsCookByTheBookMode() && !FPackageName::IsLocalizedPackage(Package->GetName()))
+		if (!FPackageName::IsLocalizedPackage(Package->GetName()))
 		{
 			const TArray<FName>* LocalizedVariants = COTFS.CookByTheBookOptions->SourceToLocalizedPackageVariants.Find(PackageFName);
 			if (LocalizedVariants)
@@ -5424,9 +5464,7 @@ void FSaveCookedPackageContext::FinishPackage()
 					UE::Cook::FPackageData* LocalizedPackageData = COTFS.PackageDatas->TryAddPackageDataByPackageName(LocalizedPackageName);
 					if (LocalizedPackageData)
 					{
-						bool bIsUrgent = false;
-						LocalizedPackageData->UpdateRequestData(PlatformsForPackage, bIsUrgent, UE::Cook::FCompletionCallback(),
-							UE::Cook::FInstigator(UE::Cook::EInstigator::SoftDependency, PackageFName));
+						COTFS.QueueDiscoveredPackageData(*LocalizedPackageData, UE::Cook::FInstigator(UE::Cook::EInstigator::SoftDependency, PackageFName));
 					}
 				}
 			}
@@ -5448,22 +5486,18 @@ void FSaveCookedPackageContext::FinishPackage()
 				}
 			}
 
-			if (COTFS.IsCookByTheBookMode())
+			UE::Cook::FPackageData* SoftObjectPackageData = COTFS.PackageDatas->TryAddPackageDataByPackageName(SoftObjectPackage);
+			if (SoftObjectPackageData)
 			{
-				UE::Cook::FPackageData* SoftObjectPackageData = COTFS.PackageDatas->TryAddPackageDataByPackageName(SoftObjectPackage);
-				if (SoftObjectPackageData)
-				{
-					bool bIsUrgent = false;
-					SoftObjectPackageData->UpdateRequestData(PlatformsForPackage, bIsUrgent, UE::Cook::FCompletionCallback(),
-						UE::Cook::FInstigator(UE::Cook::EInstigator::SoftDependency, PackageFName));
-				}
+				COTFS.QueueDiscoveredPackageData(*SoftObjectPackageData, UE::Cook::FInstigator(UE::Cook::EInstigator::SoftDependency, PackageFName));
 			}
 		}
 	}
 
 	if (!bHasRetryErrorCode)
 	{
-		if ((COTFS.CurrentCookMode == ECookMode::CookOnTheFly) && !PackageData.GetIsUrgent())
+		if (COTFS.IsCookOnTheFlyMode() && !PackageData.GetIsUrgent() &&
+			(!COTFS.CookOnTheFlyRequestManager || COTFS.CookOnTheFlyRequestManager->ShouldUseLegacyScheduling()))
 		{
 			// this is an unsolicited package
 			if (FPaths::FileExists(Filename))
@@ -5490,28 +5524,99 @@ void UCookOnTheFlyServer::Initialize( ECookMode::Type DesiredCookMode, ECookInit
 {
 	TRACE_CPUPROFILER_EVENT_SCOPE(UCookOnTheFlyServer::Initialize);
 
+	CurrentCookMode = DesiredCookMode;
+	CookFlags = InCookFlags;
+
 	PackageDatas = MakeUnique<UE::Cook::FPackageDatas>(*this);
 	PlatformManager = MakeUnique<UE::Cook::FPlatformManager>();
-	ExternalRequests = MakeUnique<UE::Cook::FExternalRequests>();
 	PackageTracker = MakeUnique<UE::Cook::FPackageTracker>(*PackageDatas.Get());
 	DiffModeHelper = MakeUnique<FDiffModeCookServerUtils>();
 	BuildDefinitions = MakeUnique<UE::Cook::FBuildDefinitions>();
 	CookByTheBookOptions = MakeUnique<UE::Cook::FCookByTheBookOptions>();
 	CookOnTheFlyOptions = MakeUnique<UE::Cook::FCookOnTheFlyOptions>();
 	AssetRegistry = IAssetRegistry::Get();
+	if (!IsCookWorkerMode())
+	{
+		WorkerRequests.Reset(new UE::Cook::FWorkerRequestsLocal());
+	}
+	else
+	{
+		check(WorkerRequests); // Caller should have constructed
+	}
+	DirectorCookMode = WorkerRequests->GetDirectorCookMode(*this);
+	if (IsCookByTheBookMode() && !IsCookingInEditor())
+	{
+		bool bCookMultiProcess = false;;
+		GConfig->GetBool(TEXT("CookSettings"), TEXT("CookMultiProcess"), bCookMultiProcess, GEditorIni);
+		bCookMultiProcess |= FParse::Param(FCommandLine::Get(), TEXT("-CookMultiProcess"));
+		bCookMultiProcess &= FParse::Param(FCommandLine::Get(), TEXT("-CookSingleProcess"));
+		if (bCookMultiProcess)
+		{
+			CookDirector = MakeUnique<UE::Cook::FCookDirector>(*this);
+		}
+	}
 
 	UE::Cook::InitializeTls();
 	UE::Cook::FPlatformManager::InitializeTls();
 
-	OutputDirectoryOverride = InOutputDirectoryOverride;
-	CurrentCookMode = DesiredCookMode;
-	CookFlags = InCookFlags;
+	LoadInitializeConfigSettings(InOutputDirectoryOverride);
+
+	CookProgressRetryBusyPeriodSeconds = GCookProgressRetryBusyTime;
+	if (IsCookOnTheFlyMode() && !IsRealtimeMode())
+	{
+		// Remove sleeps when waiting on async operations and otherwise idle; busy wait instead to minimize latency
+		CookProgressRetryBusyPeriodSeconds = 0.0f;
+	}
+	DisplayUpdatePeriodSeconds = FMath::Min(GCookProgressRepeatTime, FMath::Min(GCookProgressUpdateTime, GCookProgressDiagnosticTime));
+
+	PollNextTimeSeconds = MAX_flt;
+	PollNextTimeIdleSeconds = MAX_flt;
+
+	CurrentAsyncCacheForType = MaxAsyncCacheForType;
+	if (IsCookByTheBookMode())
+	{
+		for (TObjectIterator<UPackage> It; It; ++It)
+		{
+			if ((*It) != GetTransientPackage())
+			{
+				CookByTheBookOptions->StartupPackages.Add(It->GetFName());
+				UE_LOG(LogCook, Verbose, TEXT("Cooker startup package %s"), *It->GetName());
+			}
+		}
+	}
+
+	IdleStatus = EIdleStatus::Done;
+	IdleStatusStartTime = FPlatformTime::Seconds();
+
+	if (!IsCookOnTheFlyMode() && !IsCookingInEditor() &&
+		FPlatformMisc::SupportsMultithreadedFileHandles() && // Preloading moves file handles between threads
+		!GAllowCookedDataInEditorBuilds // // Use of preloaded files is not yet implemented when GAllowCookedDataInEditorBuilds is on, see FLinkerLoad::CreateLoader
+		)
+	{
+		bPreloadingEnabled = true;
+		FLinkerLoad::SetPreloadingEnabled(true);
+	}
+
+	// Prepare a map of SplitDataClass to FRegisteredCookPackageSplitter* for TryGetRegisteredCookPackageSplitter to use
+	RegisteredSplitDataClasses.Reset();
+	UE::Cook::Private::FRegisteredCookPackageSplitter::ForEach([this](UE::Cook::Private::FRegisteredCookPackageSplitter* RegisteredCookPackageSplitter)
+	{
+		UClass* SplitDataClass = RegisteredCookPackageSplitter->GetSplitDataClass();
+		for (TObjectIterator<UClass> ClassIt; ClassIt; ++ClassIt)
+		{
+			if (ClassIt->IsChildOf(SplitDataClass) && !ClassIt->HasAnyClassFlags(CLASS_Abstract | CLASS_Deprecated | CLASS_NewerVersionExists))
+			{
+				RegisteredSplitDataClasses.Add(SplitDataClass, RegisteredCookPackageSplitter);
+			}
+		}
+	});
 
 	FCoreUObjectDelegates::GetPreGarbageCollectDelegate().AddUObject(this, &UCookOnTheFlyServer::PreGarbageCollect);
 	FCoreUObjectDelegates::GetPostGarbageCollect().AddUObject(this, &UCookOnTheFlyServer::PostGarbageCollect);
 
 	if (IsCookingInEditor())
 	{
+		check(!IsCookWorkerMode()); // To allow in-editor callbacks on CookWorker, FWorkerRequestsRemote::AddEditorActionCallback will need to be updated to allow editor operations
 		FCoreUObjectDelegates::OnObjectPropertyChanged.AddUObject(this, &UCookOnTheFlyServer::OnObjectPropertyChanged);
 		FCoreUObjectDelegates::OnObjectModified.AddUObject(this, &UCookOnTheFlyServer::OnObjectModified);
 		FCoreUObjectDelegates::OnObjectPreSave.AddUObject(this, &UCookOnTheFlyServer::OnObjectSaved);
@@ -5524,6 +5629,25 @@ void UCookOnTheFlyServer::Initialize( ECookMode::Type DesiredCookMode, ECookInit
 
 	GetTargetPlatformManager()->GetOnTargetPlatformsInvalidatedDelegate().AddUObject(this, &UCookOnTheFlyServer::OnTargetPlatformsInvalidated);
 
+}
+
+
+void UCookOnTheFlyServer::LoadInitializeConfigSettings(const FString& InOutputDirectoryOverride)
+{
+	UE::Cook::FInitializeConfigSettings Settings;
+	WorkerRequests->GetInitializeConfigSettings(*this, InOutputDirectoryOverride, Settings);
+	SetInitializeConfigSettings(MoveTemp(Settings));
+}
+
+namespace UE::Cook
+{
+
+void FInitializeConfigSettings::LoadLocal(const FString& InOutputDirectoryOverride)
+{
+	using namespace UE::Cook;
+
+	OutputDirectoryOverride = InOutputDirectoryOverride;
+
 	MaxPrecacheShaderJobs = FPlatformMisc::NumberOfCores() - 1; // number of cores -1 is a good default allows the editor to still be responsive to other shader requests and allows cooker to take advantage of multiple processors while the editor is running
 	GConfig->GetInt(TEXT("CookSettings"), TEXT("MaxPrecacheShaderJobs"), MaxPrecacheShaderJobs, GEditorIni);
 
@@ -5541,16 +5665,6 @@ void UCookOnTheFlyServer::Initialize( ECookMode::Type DesiredCookMode, ECookInit
 	IdleTimeToGC = 20.0;
 	GConfig->GetDouble( TEXT("CookSettings"), TEXT("IdleTimeToGC"), IdleTimeToGC, GEditorIni );
 
-	if (IsCookOnTheFlyMode() && !IsRealtimeMode())
-	{
-		// Remove sleeps when waiting on async operations and otherwise idle; busy wait instead to minimize latency
-		CookProgressRetryBusyPeriodSeconds = 0.0f;
-	}
-	else
-	{
-		CookProgressRetryBusyPeriodSeconds = GCookProgressRetryBusyTime;
-	}
-
 	auto ReadMemorySetting = [](const TCHAR* SettingName, uint64& TargetVariable)
 	{
 		int32 ValueInMB = 0;
@@ -5574,12 +5688,6 @@ void UCookOnTheFlyServer::Initialize( ECookMode::Type DesiredCookMode, ECookInit
 	GConfig->GetFloat(TEXT("CookSettings"), TEXT("MemoryExpectedFreedToSpreadRatio"),
 		MemoryExpectedFreedToSpreadRatio, GEditorIni);
 
-	MaxPreloadAllocated = 16;
-	DesiredSaveQueueLength = 8;
-	DesiredLoadQueueLength = 8;
-	LoadBatchSize = 16;
-	RequestBatchSize = 16;
-
 	MinFreeUObjectIndicesBeforeGC = 100000;
 	GConfig->GetInt(TEXT("CookSettings"), TEXT("MinFreeUObjectIndicesBeforeGC"), MinFreeUObjectIndicesBeforeGC, GEditorIni);
 	MinFreeUObjectIndicesBeforeGC = FMath::Max(MinFreeUObjectIndicesBeforeGC, 0);
@@ -5592,52 +5700,49 @@ void UCookOnTheFlyServer::Initialize( ECookMode::Type DesiredCookMode, ECookInit
 	UE_LOG(LogCook, Display, TEXT("CookSettings for Memory: MemoryMaxUsedVirtual %dMiB, MemoryMaxUsedPhysical %dMiB, MemoryMinFreeVirtual %dMiB, MemoryMinFreePhysical %dMiB"),
 		MemoryMaxUsedVirtual / 1024 / 1024, MemoryMaxUsedPhysical / 1024 / 1024, MemoryMinFreeVirtual / 1024 / 1024, MemoryMinFreePhysical / 1024 / 1024);
 
-	if (IsCookByTheBookMode() && !IsCookingInEditor() &&
-		FPlatformMisc::SupportsMultithreadedFileHandles() && // Preloading moves file handles between threads
-		!GAllowCookedDataInEditorBuilds // // Use of preloaded files is not yet implemented when GAllowCookedDataInEditorBuilds is on, see FLinkerLoad::CreateLoader
-		)
+	const FConfigSection* CacheSettings = GConfig->GetSectionPrivate(TEXT("CookPlatformDataCacheSettings"), false, true, GEditorIni);
+	if (CacheSettings)
 	{
-		bPreloadingEnabled = true;
-		FLinkerLoad::SetPreloadingEnabled(true);
+		for (const auto& CacheSetting : *CacheSettings)
+		{
+
+			const FString& ReadString = CacheSetting.Value.GetValue();
+			int32 ReadValue = FCString::Atoi(*ReadString);
+			int32 Count = FMath::Max(2, ReadValue);
+			MaxAsyncCacheForType.Add(CacheSetting.Key, Count);
+		}
 	}
 
-	PollNextTimeSeconds = MAX_flt;
-	PollNextTimeIdleSeconds = MAX_flt;
-	WaitForAsyncSleepSeconds = 1.0f;
+	bHybridIterativeDebug = FParse::Param(FCommandLine::Get(), TEXT("hybriditerativedebug"));
+}
 
-	IdleStatus = EIdleStatus::Done;
-	IdleStatusStartTime = FPlatformTime::Seconds();
+}
 
-	{
-		const FConfigSection* CacheSettings = GConfig->GetSectionPrivate(TEXT("CookPlatformDataCacheSettings"), false, true, GEditorIni);
-		if ( CacheSettings )
-		{
-			for ( const auto& CacheSetting : *CacheSettings )
-			{
-				
-				const FString& ReadString = CacheSetting.Value.GetValue();
-				int32 ReadValue = FCString::Atoi(*ReadString);
-				int32 Count = FMath::Max( 2,  ReadValue );
-				MaxAsyncCacheForType.Add( CacheSetting.Key,  Count );
-			}
-		}
-		CurrentAsyncCacheForType = MaxAsyncCacheForType;
-	}
+void UCookOnTheFlyServer::SetInitializeConfigSettings(UE::Cook::FInitializeConfigSettings&& Settings)
+{
+	MaxPrecacheShaderJobs = Settings.MaxPrecacheShaderJobs;
+	MaxConcurrentShaderJobs = Settings.MaxConcurrentShaderJobs;
+	PackagesPerGC = Settings.PackagesPerGC;
+	MemoryExpectedFreedToSpreadRatio = Settings.MemoryExpectedFreedToSpreadRatio;
+	IdleTimeToGC = Settings.IdleTimeToGC;
+	MemoryMaxUsedVirtual = Settings.MemoryMaxUsedVirtual;
+	MemoryMaxUsedPhysical = Settings.MemoryMaxUsedPhysical;
+	MemoryMinFreeVirtual = Settings.MemoryMinFreeVirtual;
+	MemoryMinFreePhysical = Settings.MemoryMinFreePhysical;
+	MinFreeUObjectIndicesBeforeGC = Settings.MinFreeUObjectIndicesBeforeGC;
+	MaxNumPackagesBeforePartialGC = Settings.MaxNumPackagesBeforePartialGC;
+	ConfigSettingDenyList = MoveTemp(Settings.ConfigSettingDenyList);
+	MaxAsyncCacheForType = MoveTemp(Settings.MaxAsyncCacheForType);
+	bHybridIterativeDebug = Settings.bHybridIterativeDebug;
+
+	MaxPreloadAllocated = 16;
+	DesiredSaveQueueLength = 8;
+	DesiredLoadQueueLength = 8;
+	LoadBatchSize = 16;
+	RequestBatchSize = 16;
+	WaitForAsyncSleepSeconds = 1.0f;
 
 
-	if (IsCookByTheBookMode())
-	{
-		for (TObjectIterator<UPackage> It; It; ++It)
-		{
-			if ((*It) != GetTransientPackage())
-			{
-				CookByTheBookOptions->StartupPackages.Add(It->GetFName());
-				UE_LOG(LogCook, Verbose, TEXT("Cooker startup package %s"), *It->GetName());
-			}
-		}
-		bHybridIterativeDebug = FParse::Param(FCommandLine::Get(), TEXT("hybriditerativedebug"));
-	}
-	
 	// See if there are any plugins that need to be remapped for the sandbox
 	const FProjectDescriptor* Project = IProjectManager::Get().GetCurrentProject();
 	if (Project != nullptr)
@@ -5663,20 +5768,28 @@ void UCookOnTheFlyServer::Initialize( ECookMode::Type DesiredCookMode, ECookInit
 			}
 		}
 	}
+}
 
-	// Prepare a map SplitDataClass to FRegisteredCookPackageSplitter* for TryGetRegisteredCookPackageSplitter to use
-	RegisteredSplitDataClasses.Reset();
-	UE::Cook::Private::FRegisteredCookPackageSplitter::ForEach([this](UE::Cook::Private::FRegisteredCookPackageSplitter* RegisteredCookPackageSplitter)
+bool UCookOnTheFlyServer::TryInitializeCookWorker()
+{
+	using namespace UE::Cook;
+
+	FString CookDirectorHost;
+	if (!FParse::Value(FCommandLine::Get(), TEXT("-CookDirectorHost="), CookDirectorHost))
 	{
-		UClass* SplitDataClass = RegisteredCookPackageSplitter->GetSplitDataClass();
-		for (TObjectIterator<UClass> ClassIt; ClassIt; ++ClassIt)
-		{
-			if (ClassIt->IsChildOf(SplitDataClass) && !ClassIt->HasAnyClassFlags(CLASS_Abstract | CLASS_Deprecated | CLASS_NewerVersionExists))
-			{
-				RegisteredSplitDataClasses.Add(SplitDataClass, RegisteredCookPackageSplitter);
-			}
-		}
-	});
+		UE_LOG(LogCook, Error, TEXT("CookWorker startup failed: no CookDirector specified on commandline."))
+		return false;
+	}
+	TUniquePtr<FWorkerRequestsRemote> RemoteTasks = MakeUnique<FWorkerRequestsRemote>(*this);
+	ECookInitializationFlags LocalCookFlags;
+	if (!RemoteTasks->TryConnect(FDirectorConnectionInfo{ CookDirectorHost }, LocalCookFlags))
+	{
+		return false;
+	}
+	WorkerRequests.Reset(RemoteTasks.Release());
+	Initialize(ECookMode::CookWorker, LocalCookFlags, FString());
+	StartCookAsCookWorker();
+	return true;
 }
 
 void UCookOnTheFlyServer::InitializeSession()
@@ -5855,7 +5968,7 @@ FString UCookOnTheFlyServer::GetOutputDirectoryOverride(FBeginCookContext& Begin
 	{
 		if ( IsCookingDLC() )
 		{
-			check( IsCookByTheBookMode() );
+			check( !IsDirectorCookOnTheFly() );
 			OutputDirectory = FPaths::Combine(*GetBaseDirectoryForDLC(), TEXT("Saved"), TEXT("Cooked"), TEXT("[Platform]"));
 		}
 		else if ( IsCookingInEditor() )
@@ -5874,7 +5987,7 @@ FString UCookOnTheFlyServer::GetOutputDirectoryOverride(FBeginCookContext& Begin
 	else if (!OutputDirectory.Contains(TEXT("[Platform]"), ESearchCase::IgnoreCase, ESearchDir::FromEnd) )
 	{
 		// Output directory needs to contain [Platform] token to be able to cook for multiple targets.
-		if ( IsCookByTheBookMode() )
+		if ( !IsDirectorCookOnTheFly() )
 		{
 			checkf(BeginContext.TargetPlatforms.Num() == 1,
 				TEXT("If OutputDirectoryOverride is provided when cooking multiple platforms, it must include [Platform] in the text, to be replaced with the name of each of the requested Platforms.") );
@@ -6058,7 +6171,7 @@ void GetAdditionalCurrentIniVersionStrings( const UCookOnTheFlyServer* CookOnThe
 
 
 
-bool UCookOnTheFlyServer::GetCurrentIniVersionStrings( const ITargetPlatform* TargetPlatform, FIniSettingContainer& IniVersionStrings ) const
+bool UCookOnTheFlyServer::GetCurrentIniVersionStrings( const ITargetPlatform* TargetPlatform, UE::Cook::FIniSettingContainer& IniVersionStrings ) const
 {
 	{
 		FScopeLock Lock(&ConfigFileCS);
@@ -6160,7 +6273,7 @@ bool UCookOnTheFlyServer::GetCurrentIniVersionStrings( const ITargetPlatform* Ta
 }
 
 
-bool UCookOnTheFlyServer::GetCookedIniVersionStrings(const ITargetPlatform* TargetPlatform, FIniSettingContainer& OutIniSettings, TMap<FString,FString>& OutAdditionalSettings) const
+bool UCookOnTheFlyServer::GetCookedIniVersionStrings(const ITargetPlatform* TargetPlatform, UE::Cook::FIniSettingContainer& OutIniSettings, TMap<FString,FString>& OutAdditionalSettings) const
 {
 	const FString EditorIni = FPaths::ProjectDir() / TEXT("Metadata") / TEXT("CookedIniVersion.txt");
 	const FString SandboxEditorIni = ConvertToFullSandboxPath(*EditorIni, true);
@@ -6255,11 +6368,11 @@ void UCookOnTheFlyServer::OnFConfigDeleted(const FConfigFile* Config)
 	}
 
 	FScopeLock Lock(&ConfigFileCS);
-	ProcessAccessedIniSettings(Config, AccessedIniStrings);
+	WorkerRequests->ReportAccessedIniSettings(*this, *Config);
 	OpenConfigFiles.Remove(Config);
 }
 
-void UCookOnTheFlyServer::ProcessAccessedIniSettings(const FConfigFile* Config, FIniSettingContainer& OutAccessedIniStrings) const
+void UCookOnTheFlyServer::ProcessAccessedIniSettings(const FConfigFile* Config, UE::Cook::FIniSettingContainer& OutAccessedIniStrings) const
 {	
 	if (Config->Name == NAME_None)
 	{
@@ -6399,14 +6512,14 @@ TMap<FName, FString> UCookOnTheFlyServer::CalculateCookSettingStrings() const
 	const FName NAME_CookMode(TEXT("CookMode"));
 
 	CookSettingStrings.Add(FName(TEXT("Version")), TEXT("C7C76F79"));
-	if (IsCookByTheBookMode())
+	if (IsDirectorCookByTheBook())
 	{
 		CookSettingStrings.Add(NAME_CookMode, TEXT("CookByTheBook"));
 		CookSettingStrings.Add(FName(TEXT("DLCName")), CookByTheBookOptions->DlcName);
 	}
 	else
 	{
-		check(IsCookOnTheFlyMode());
+		check(IsDirectorCookOnTheFly());
 		CookSettingStrings.Add(NAME_CookMode, TEXT("CookOnTheFly"));
 	}
 	return CookSettingStrings;
@@ -6457,7 +6570,7 @@ bool UCookOnTheFlyServer::IniSettingsOutOfDate(const ITargetPlatform* TargetPlat
 {
 	TGuardValue<bool> A(GSuppressProcessConfigSettings, true);
 
-	FIniSettingContainer OldIniSettings;
+	UE::Cook::FIniSettingContainer OldIniSettings;
 	TMap<FString, FString> OldAdditionalSettings;
 	if ( GetCookedIniVersionStrings(TargetPlatform, OldIniSettings, OldAdditionalSettings) == false)
 	{
@@ -6601,7 +6714,7 @@ bool UCookOnTheFlyServer::SaveCurrentIniSettings(const ITargetPlatform* TargetPl
 	TMap<FString, FString> AdditionalIniSettings;
 	GetAdditionalCurrentIniVersionStrings(this, TargetPlatform, AdditionalIniSettings);
 
-	FIniSettingContainer CurrentIniSettings;
+	UE::Cook::FIniSettingContainer CurrentIniSettings;
 	GetCurrentIniVersionStrings(TargetPlatform, CurrentIniSettings);
 
 	const FString EditorIni = FPaths::ProjectDir() / TEXT("Metadata") / TEXT("CookedIniVersion.txt");
@@ -6663,7 +6776,7 @@ bool UCookOnTheFlyServer::SaveCurrentIniSettings(const ITargetPlatform* TargetPl
 void UCookOnTheFlyServer::OnRequestClusterCompleted(const UE::Cook::FRequestCluster& RequestCluster)
 {
 	using namespace UE::Cook;
-	if (IsCookByTheBookMode())
+	if (IsDirectorCookByTheBook())
 	{
 		if (bHybridIterativeDebug)
 		{
@@ -6699,6 +6812,7 @@ void UCookOnTheFlyServer::PopulateCookedPackages(TArrayView<const ITargetPlatfor
 {
 	using namespace UE::Cook;
 	TRACE_CPUPROFILER_EVENT_SCOPE(UCookOnTheFlyServer::PopulateCookedPackages);
+	checkf(!IsCookWorkerMode(), TEXT("Calling PopulateCookedPackages should be impossible in a CookWorker."));
 
 	// TODO: NumPackagesIterativelySkipped is only counted for the first platform; to count all platforms we would
 	// have to check whether each one is already cooked.
@@ -6802,7 +6916,7 @@ void UCookOnTheFlyServer::PopulateCookedPackages(TArrayView<const ITargetPlatfor
 					FPackageData* PackageData = PackageDatas->TryAddPackageDataByPackageName(RemovePackageName);
 					if (PackageData)
 					{
-						ExternalRequests->EnqueueUnique(FFilePlatformRequest(PackageData->GetFileName(),
+						WorkerRequests->AddStartCookByTheBookRequest(FFilePlatformRequest(PackageData->GetFileName(),
 							EInstigator::IterativeCook, TConstArrayView<const ITargetPlatform*>{ TargetPlatform }));
 					}
 				}
@@ -7541,7 +7655,7 @@ const FString UCookOnTheFlyServer::GetSandboxAssetRegistryFilename()
 
 	if (IsCookingDLC())
 	{
-		check(IsCookByTheBookMode());
+		check(IsDirectorCookByTheBook());
 		const FString DLCRegistryFilename = FPaths::Combine(*GetBaseDirectoryForDLC(), GetAssetRegistryFilename());
 		return ConvertToFullSandboxPath(*DLCRegistryFilename, true);
 	}
@@ -7603,12 +7717,15 @@ void UCookOnTheFlyServer::BeginCookStartShaderCodeLibrary(FBeginCookContext& Beg
         }
     }
 
-	CleanShaderCodeLibraries();
+	if (!IsCookWorkerMode())
+	{
+		CleanShaderCodeLibraries();
+	}
 }
 
 void UCookOnTheFlyServer::BeginCookFinishShaderCodeLibrary(FBeginCookContext& BeginContext)
 {
-	check(BeginContext.StartupOptions && IsCookByTheBookMode()); // CookByTheBook only for now
+	check(BeginContext.StartupOptions && IsDirectorCookByTheBook()); // CookByTheBook only for now
 	// don't resave the global shader map files in dlc
 	if (!IsCookingDLC() && !EnumHasAnyFlags(BeginContext.StartupOptions->CookOptions, ECookByTheBookOptions::ForceDisableSaveGlobalShaders))
 	{
@@ -8156,8 +8273,6 @@ void UCookOnTheFlyServer::ShutdownCookSession()
 
 	if (IsCookByTheBookMode())
 	{
-		CookByTheBookOptions->ClearSessionData();
-
 		UnregisterCookByTheBookDelegates();
 
 		PrintFinishStats();
@@ -8165,6 +8280,7 @@ void UCookOnTheFlyServer::ShutdownCookSession()
 		OutputHierarchyTimers();
 		ClearHierarchyTimers();
 	}
+	CookByTheBookOptions->ClearSessionData();
 	PlatformManager->ClearSessionPlatforms(*this);
 }
 
@@ -8296,7 +8412,7 @@ void UCookOnTheFlyServer::CancelAllQueues()
 	// Discard the external build requests, but execute any pending SchedulerCallbacks since these might have important teardowns
 	TArray<UE::Cook::FSchedulerCallback> SchedulerCallbacks;
 	TArray<UE::Cook::FFilePlatformRequest> UnusedRequests;
-	ExternalRequests->DequeueAll(SchedulerCallbacks, UnusedRequests);
+	WorkerRequests->DequeueAllExternal(SchedulerCallbacks, UnusedRequests);
 	for (UE::Cook::FSchedulerCallback& SchedulerCallback : SchedulerCallbacks)
 	{
 		SchedulerCallback();
@@ -8464,28 +8580,26 @@ void UCookOnTheFlyServer::CreateSandboxFile(FBeginCookContext& BeginContext)
 
 }
 
-void UCookOnTheFlyServer::SetBeginCookConfigSettings(FBeginCookContext& BeginContext)
+void UCookOnTheFlyServer::LoadBeginCookConfigSettings(FBeginCookContext& BeginContext)
+{
+	UE::Cook::FBeginCookConfigSettings Settings;
+	WorkerRequests->GetBeginCookConfigSettings(*this, BeginContext, Settings);
+	SetBeginCookConfigSettings(BeginContext, MoveTemp(Settings));
+}
+
+namespace UE::Cook
+{
+
+void FBeginCookConfigSettings::LoadLocal(FBeginCookContext& BeginContext)
 {
 	GConfig->GetBool(TEXT("CookSettings"), TEXT("HybridIterativeEnabled"), bHybridIterativeEnabled, GEditorIni);
 	// TODO: HybridIterative is not yet implemented for DLC
-	bHybridIterativeEnabled &= !IsCookingDLC();
+	bHybridIterativeEnabled &= !BeginContext.COTFS.IsCookingDLC();
 	// HybridIterative uses TargetDomain storage of dependencies which is only implemented in ZenStore
-	bHybridIterativeEnabled &= IsUsingZenStore();
-
-	PackageDatas->SetBeginCookConfigSettings();
+	bHybridIterativeEnabled &= BeginContext.COTFS.IsUsingZenStore();
 
-	SetNeverCookPackageConfigSettings(BeginContext);
-	if (!bFirstCookInThisProcessInitialized)
-	{
-		// This is the first cook; set bFirstCookInThisProcess=true for the entire cook until SetBeginCookConfigSettings is called to mark the second cook
-		bFirstCookInThisProcessInitialized = true;
-		bFirstCookInThisProcess = true;
-	}
-	else
-	{
-		// We have cooked before; set bFirstCookInThisProcess=false
-		bFirstCookInThisProcess = false;
-	}
+	FParse::Value(FCommandLine::Get(), TEXT("-CookShowInstigator="), CookShowInstigator);
+	LoadNeverCookLocal(BeginContext);
 
 	for (const ITargetPlatform* TargetPlatform : BeginContext.TargetPlatforms)
 	{
@@ -8501,17 +8615,42 @@ void UCookOnTheFlyServer::SetBeginCookConfigSettings(FBeginCookContext& BeginCon
 	}
 }
 
-void UCookOnTheFlyServer::SetNeverCookPackageConfigSettings(FBeginCookContext& BeginContext)
+}
+
+void UCookOnTheFlyServer::SetBeginCookConfigSettings(FBeginCookContext& BeginContext, UE::Cook::FBeginCookConfigSettings&& Settings)
 {
-	PackageTracker->NeverCookPackageList.Empty();
+	bHybridIterativeEnabled = Settings.bHybridIterativeEnabled;
+	PackageDatas->SetBeginCookConfigSettings(Settings.CookShowInstigator);
+	SetNeverCookPackageConfigSettings(BeginContext, Settings);
+	if (!bFirstCookInThisProcessInitialized)
+	{
+		// This is the first cook; set bFirstCookInThisProcess=true for the entire cook until SetBeginCookConfigSettings is called to mark the second cook
+		bFirstCookInThisProcessInitialized = true;
+		bFirstCookInThisProcess = true;
+	}
+	else
+	{
+		// We have cooked before; set bFirstCookInThisProcess=false
+		bFirstCookInThisProcess = false;
+	}
+}
+
+namespace UE::Cook
+{
+
+void FBeginCookConfigSettings::LoadNeverCookLocal(FBeginCookContext& BeginContext)
+{
+	NeverCookPackageList.Reset();
+	PlatformSpecificNeverCookPackages.Reset();
+
 	TArrayView<const FString> ExtraNeverCookDirectories;
 	if (BeginContext.StartupOptions)
 	{
 		ExtraNeverCookDirectories = BeginContext.StartupOptions->NeverCookDirectories;
 	}
-	for (FName NeverCookPackage : GetNeverCookPackageFileNames(ExtraNeverCookDirectories))
+	for (FName NeverCookPackage : BeginContext.COTFS.GetNeverCookPackageFileNames(ExtraNeverCookDirectories))
 	{
-		PackageTracker->NeverCookPackageList.Add(NeverCookPackage);
+		NeverCookPackageList.Add(NeverCookPackage);
 	}
 
 	// use temp list of UBT platform strings to discover PlatformSpecificNeverCookPackages
@@ -8526,11 +8665,29 @@ void UCookOnTheFlyServer::SetNeverCookPackageConfigSettings(FBeginCookContext& B
 			UBTPlatformStrings.Emplace(MoveTemp(UBTPlatformName));
 		}
 
-		DiscoverPlatformSpecificNeverCookPackages(BeginContext.TargetPlatforms, UBTPlatformStrings);
+		BeginContext.COTFS.DiscoverPlatformSpecificNeverCookPackages(BeginContext.TargetPlatforms, UBTPlatformStrings, *this);
 	}
 }
 
-void UCookOnTheFlyServer::SetBeginCookIterativeFlags(FBeginCookContext& BeginContext)
+}
+
+void UCookOnTheFlyServer::SetNeverCookPackageConfigSettings(FBeginCookContext& BeginContext, UE::Cook::FBeginCookConfigSettings& Settings)
+{
+	UE::Cook::FThreadSafeSet<FName>& NeverCookPackageList = PackageTracker->NeverCookPackageList;
+	NeverCookPackageList.Empty();
+	for (FName FileName : Settings.NeverCookPackageList)
+	{
+		NeverCookPackageList.Add(FileName);
+	}
+	PackageTracker->PlatformSpecificNeverCookPackages = MoveTemp(Settings.PlatformSpecificNeverCookPackages);
+}
+
+void UCookOnTheFlyServer::LoadBeginCookIterativeFlags(FBeginCookContext& BeginContext)
+{
+	WorkerRequests->GetBeginCookIterativeFlags(*this, BeginContext);
+}
+
+void UCookOnTheFlyServer::LoadBeginCookIterativeFlagsLocal(FBeginCookContext& BeginContext) const
 {
 	const bool bIsDiffOnly = FParse::Param(FCommandLine::Get(), TEXT("DIFFONLY"));
 	const bool bIterative = !FParse::Param(FCommandLine::Get(), TEXT("fullcook")) && (bHybridIterativeEnabled || IsCookFlagSet(ECookInitializationFlags::Iterative));
@@ -8540,7 +8697,9 @@ void UCookOnTheFlyServer::SetBeginCookIterativeFlags(FBeginCookContext& BeginCon
 	{
 		const ITargetPlatform* TargetPlatform = PlatformContext.TargetPlatform;
 		UE::Cook::FPlatformData* PlatformData = PlatformContext.PlatformData;
-		ICookedPackageWriter& PackageWriter = FindOrCreatePackageWriter(TargetPlatform);
+		const ICookedPackageWriter* PackageWriterPtr = FindPackageWriter(TargetPlatform);
+		check(PackageWriterPtr); // PackageContexts should have been created by SelectSessoinPlatforms or by FindOrCreateSaveContexts in AddCookOnTheFlyPlatformFromGameThread
+		const ICookedPackageWriter& PackageWriter(*PackageWriterPtr);
 		bool bIterateSharedBuild = false;
 		if (bIterative && bIsSharedIterativeCook && !PlatformData->bIsSandboxInitialized)
 		{
@@ -8568,7 +8727,6 @@ void UCookOnTheFlyServer::SetBeginCookIterativeFlags(FBeginCookContext& BeginCon
 			}
 		}
 		PlatformContext.CurrentCookSettings = CalculateCookSettingStrings();
-		const bool bIsIniSettingsOutOfDate = IniSettingsOutOfDate(TargetPlatform); // needs to be executed for side effects even if non-iterative
 		PlatformContext.bHasMemoryResults = PlatformData->bIsSandboxInitialized;
 
 		if (bIsDiffOnly)
@@ -8593,24 +8751,28 @@ void UCookOnTheFlyServer::SetBeginCookIterativeFlags(FBeginCookContext& BeginCon
 				UE_LOG(LogCook, Display, TEXT("Cook invalidated for platform %s because cook DLC settings have changed, clearing all cooked content"), *TargetPlatform->PlatformName());
 				bIterativeAllowed = false;
 			}
-			else if (bIsIniSettingsOutOfDate)
+			else
 			{
-				if (!IsCookFlagSet(ECookInitializationFlags::IgnoreIniSettingsOutOfDate))
+				const bool bIsIniSettingsOutOfDate = IniSettingsOutOfDate(TargetPlatform);
+				if (bIsIniSettingsOutOfDate)
 				{
-					UE_LOG(LogCook, Display, TEXT("Cook invalidated for platform %s ini settings don't match from last cook, clearing all cooked content"), *TargetPlatform->PlatformName());
-					bIterativeAllowed = false;
+					if (!IsCookFlagSet(ECookInitializationFlags::IgnoreIniSettingsOutOfDate))
+					{
+						UE_LOG(LogCook, Display, TEXT("Cook invalidated for platform %s ini settings don't match from last cook, clearing all cooked content"), *TargetPlatform->PlatformName());
+						bIterativeAllowed = false;
+					}
+					else
+					{
+						UE_LOG(LogCook, Display, TEXT("Inisettings were out of date for platform %s but we are going with it anyway because IgnoreIniSettingsOutOfDate is set"), *TargetPlatform->PlatformName());
+						bIterativeAllowed = true;
+					}
 				}
 				else
 				{
-					UE_LOG(LogCook, Display, TEXT("Inisettings were out of date for platform %s but we are going with it anyway because IgnoreIniSettingsOutOfDate is set"), *TargetPlatform->PlatformName());
+					UE_LOG(LogCook, Display, TEXT("Keeping cooked content for platform %s and cooking iteratively"), *TargetPlatform->PlatformName());
 					bIterativeAllowed = true;
 				}
 			}
-			else
-			{
-				UE_LOG(LogCook, Display, TEXT("Keeping cooked content for platform %s and cooking iteratively"), *TargetPlatform->PlatformName());
-				bIterativeAllowed = true;
-			}
 
 			if (bIterativeAllowed)
 			{
@@ -8647,12 +8809,17 @@ void UCookOnTheFlyServer::BeginCookSandbox(FBeginCookContext& BeginContext)
 			UE::Cook::FPlatformData* PlatformData = PlatformContext.PlatformData;
 			ICookedPackageWriter& PackageWriter = FindOrCreatePackageWriter(TargetPlatform);
 			ICookedPackageWriter::FCookInfo CookInfo;
-			CookInfo.CookMode = IsCookOnTheFlyMode() ? ICookedPackageWriter::FCookInfo::CookOnTheFlyMode : ICookedPackageWriter::FCookInfo::CookByTheBookMode;
+			CookInfo.CookMode = IsDirectorCookOnTheFly() ? ICookedPackageWriter::FCookInfo::CookOnTheFlyMode : ICookedPackageWriter::FCookInfo::CookByTheBookMode;
 			CookInfo.bFullBuild = PlatformContext.bFullBuild;
 			CookInfo.bIterateSharedBuild = PlatformContext.bIterateSharedBuild;
+			CookInfo.bWorkerOnSharedSandbox = PlatformContext.bWorkerOnSharedSandbox;
 			PackageWriter.Initialize(CookInfo);
-			// Clean the Manifest directory even on iterative builds; it is written from scratch each time
-			PlatformData->RegistryGenerator->CleanManifestDirectories();
+			if (!PlatformContext.bWorkerOnSharedSandbox)
+			{
+				// Clean the Manifest directory even on iterative builds; it is written from scratch each time
+				// But only do this if we own the output directory
+				PlatformData->RegistryGenerator->CleanManifestDirectories();
+			}
 
 			if (PlatformContext.bPopulateMemoryResultsFromDiskResults)
 			{
@@ -8662,9 +8829,12 @@ void UCookOnTheFlyServer::BeginCookSandbox(FBeginCookContext& BeginContext)
 			{
 				AlreadyCookedPlatforms.Add(TargetPlatform);
 			}
-			bool bResultResults = PlatformContext.bHasMemoryResults && PlatformContext.bClearMemoryResults;
-			ResetPlatforms.Emplace(TargetPlatform, bResultResults);
-			SaveCookSettings(PlatformContext.CurrentCookSettings, TargetPlatform);
+			bool bResetResults = PlatformContext.bHasMemoryResults && PlatformContext.bClearMemoryResults;
+			ResetPlatforms.Emplace(TargetPlatform, bResetResults);
+			if (!PlatformContext.bWorkerOnSharedSandbox)
+			{
+				SaveCookSettings(PlatformContext.CurrentCookSettings, TargetPlatform);
+			}
 
 			PlatformData->bIsSandboxInitialized = true;
 		}
@@ -8705,7 +8875,7 @@ void UCookOnTheFlyServer::BeginCookSandbox(FBeginCookContext& BeginContext)
 UE::Cook::FCookSavePackageContext* UCookOnTheFlyServer::CreateSaveContext(const ITargetPlatform* TargetPlatform)
 {
 	using namespace UE::Cook;
-	checkf(SandboxFile, TEXT("SaveContexts cannot be created until after CreateSandboxFile has been called from StartCookByTheBook or StartCookOnTheFly"));
+	checkf(SandboxFile, TEXT("SaveContexts cannot be created until after CreateSandboxFile has been called from a StartCook function."));
 
 	const FString RootPathSandbox = ConvertToFullSandboxPath(FPaths::RootDir(), true);
 	FString MetadataPathSandbox;
@@ -8766,7 +8936,8 @@ void UCookOnTheFlyServer::ClearPackageStoreContexts()
 }
 
 void UCookOnTheFlyServer::DiscoverPlatformSpecificNeverCookPackages(
-	const TArrayView<const ITargetPlatform* const>& TargetPlatforms, const TArray<FString>& UBTPlatformStrings)
+	const TArrayView<const ITargetPlatform* const>& TargetPlatforms, const TArray<FString>& UBTPlatformStrings,
+	UE::Cook::FBeginCookConfigSettings& Settings) const
 {
 	TArray<const ITargetPlatform*> PluginUnsupportedTargetPlatforms;
 	TArray<FAssetData> PluginAssets;
@@ -8812,7 +8983,7 @@ void UCookOnTheFlyServer::DiscoverPlatformSpecificNeverCookPackages(
 
 			for (const ITargetPlatform* TargetPlatform: PluginUnsupportedTargetPlatforms)
 			{
-				TSet<FName>& NeverCookPackages = PackageTracker->PlatformSpecificNeverCookPackages.FindOrAdd(TargetPlatform);
+				TSet<FName>& NeverCookPackages = Settings.PlatformSpecificNeverCookPackages.FindOrAdd(TargetPlatform);
 				for (const FAssetData& Asset : PluginAssets)
 				{
 					NeverCookPackages.Add(Asset.PackageName);
@@ -8834,9 +9005,9 @@ void UCookOnTheFlyServer::StartCookByTheBook( const FCookByTheBookStartupOptions
 	FBeginCookContext BeginContext = CreateBeginCookByTheBookContext(CookByTheBookStartupOptions);
 	BlockOnAssetRegistry();
 	CreateSandboxFile(BeginContext);
-	SetBeginCookConfigSettings(BeginContext);
+	LoadBeginCookConfigSettings(BeginContext);
 	SelectSessionPlatforms(BeginContext);
-	SetBeginCookIterativeFlags(BeginContext);
+	LoadBeginCookIterativeFlags(BeginContext);
 
 	// Initialize systems referenced by later stages or that need to start early for async performance
 	// Functions in this section must not need to read/write the SandboxDirectory or MemoryCookedPackages
@@ -8866,7 +9037,7 @@ void UCookOnTheFlyServer::StartCookByTheBook( const FCookByTheBookStartupOptions
 
 FBeginCookContext UCookOnTheFlyServer::CreateBeginCookByTheBookContext(const FCookByTheBookStartupOptions& StartupOptions)
 {
-	FBeginCookContext BeginContext;
+	FBeginCookContext BeginContext(*this);
 
 	BeginContext.StartupOptions = &StartupOptions;
 	const ECookByTheBookOptions& CookOptions = StartupOptions.CookOptions;
@@ -8878,7 +9049,8 @@ FBeginCookContext UCookOnTheFlyServer::CreateBeginCookByTheBookContext(const FCo
 	CookByTheBookOptions->CreateReleaseVersion = StartupOptions.CreateReleaseVersion;
 	CookByTheBookOptions->bSkipHardReferences = !!(CookOptions & ECookByTheBookOptions::SkipHardReferences);
 	CookByTheBookOptions->bSkipSoftReferences = !!(CookOptions & ECookByTheBookOptions::SkipSoftReferences);
-	CookByTheBookOptions->bFullLoadAndSave = !!(CookOptions & ECookByTheBookOptions::FullLoadAndSave) && !IsCookingInEditor();
+	CookByTheBookOptions->bFullLoadAndSave = !!(CookOptions & ECookByTheBookOptions::FullLoadAndSave) &&
+		!IsCookingInEditor() && !CookDirector;
 	CookByTheBookOptions->bCookAgainstFixedBase = !!(CookOptions & ECookByTheBookOptions::CookAgainstFixedBase);
 	CookByTheBookOptions->bDlcLoadMainAssetRegistry = !!(CookOptions & ECookByTheBookOptions::DlcLoadMainAssetRegistry);
 	CookByTheBookOptions->bErrorOnEngineContentUse = StartupOptions.bErrorOnEngineContentUse;
@@ -8909,12 +9081,12 @@ FBeginCookContext UCookOnTheFlyServer::CreateBeginCookOnTheFlyContext(const FCoo
 	bZenStore = Options.bZenStore;
 	CookOnTheFlyOptions->bBindAnyPort = Options.bBindAnyPort;
 	CookOnTheFlyOptions->bPlatformProtocol = Options.bPlatformProtocol;
-	return FBeginCookContext();
+	return FBeginCookContext(*this);
 }
 
 FBeginCookContext UCookOnTheFlyServer::CreateAddPlatformContext(ITargetPlatform* TargetPlatform)
 {
-	FBeginCookContext BeginContext;
+	FBeginCookContext BeginContext(*this);
 
 	BeginContext.TargetPlatforms.Add(TargetPlatform);
 
@@ -8925,6 +9097,87 @@ FBeginCookContext UCookOnTheFlyServer::CreateAddPlatformContext(ITargetPlatform*
 	return BeginContext;
 }
 
+void UCookOnTheFlyServer::StartCookAsCookWorker()
+{
+	UE_SCOPED_COOKTIMER(StartCookWorker);
+	LLM_SCOPE_BYTAG(Cooker);
+	check(IsInGameThread());
+	check(IsCookWorkerMode());
+
+	// Initialize systems and settings that the rest of StartCookAsCookWorker depends on
+	// Functions in this section are ordered and can depend on the functions before them
+	FBeginCookContext BeginContext = CreateCookWorkerContext();
+	// MPCOOKTODO: Load serialized AssetRegistry from Director
+	UE::Cook::FPackageDatas::OnAssetRegistryGenerated(*AssetRegistry);
+	CreateSandboxFile(BeginContext);
+	LoadBeginCookConfigSettings(BeginContext);
+	SelectSessionPlatforms(BeginContext);
+	LoadBeginCookIterativeFlags(BeginContext);
+	if (IsDirectorCookOnTheFly())
+	{
+		GShaderCompilingManager->SkipShaderCompilation(true);
+	}
+
+	// Initialize systems referenced by later stages or that need to start early for async performance
+	// MPCOOKTODO: Need to send information from ShaderCodeLibrary to director
+	if (IsDirectorCookByTheBook())
+	{
+		BeginCookStartShaderCodeLibrary(BeginContext); // start shader code library cooking asynchronously; we block on it later
+	}
+	// MPCOOKTODO: Remove PlatformAssetRegistries, funnel it through API on WorkerRequests
+	RefreshPlatformAssetRegistries(BeginContext.TargetPlatforms); // Required by BeginCookSandbox stage
+
+	// Clear in-memory CookedPackages, or preserve them and cook iteratively in-process. We do not modify the CookedPackages on disk, because
+	// that was already done as necessary by the Director
+	BeginCookSandbox(BeginContext);
+
+	// Initialize systems that nothing in StartCookAsCookWorker references
+	// Functions in this section are not dependent upon each other and can be ordered arbitrarily or for async performance	
+	BeginCookEDLCookInfo(BeginContext); // MPCOOKTODO: Need to send information from EDLCookInfo to director
+	BeginCookPackageWriters(BeginContext);
+	InitializePollables();
+	if (IsDirectorCookByTheBook())
+	{
+		RegisterCookByTheBookDelegates();
+	}
+	if (IsDirectorCookByTheBook())
+	{
+		BeginCookFinishShaderCodeLibrary(BeginContext);
+	}
+}
+
+void UCookOnTheFlyServer::ShutdownCookAsCookWorker()
+{
+	if (IsDirectorCookByTheBook())
+	{
+		UnregisterCookByTheBookDelegates();
+	}
+	ShutdownCookSession();
+	FString LibraryName = GetProjectShaderLibraryName();
+	FString ActualLibraryName = GenerateShaderCodeLibraryName(LibraryName, IsCookFlagSet(ECookInitializationFlags::IterateSharedBuild));
+	FShaderLibraryCooker::EndCookingLibrary(ActualLibraryName);
+	FShaderLibraryCooker::Shutdown();
+	if (IsDirectorCookOnTheFly())
+	{
+		GShaderCompilingManager->SkipShaderCompilation(false);
+	}
+}
+
+FBeginCookContext UCookOnTheFlyServer::CreateCookWorkerContext()
+{
+	// MPCOOKTODO: Copy CookByTheBookOptions and CookOnTheFlyOptions from CookDirector
+	// MPCOOKTODO: Copy TargetPlatforms from CookDirector
+	// MPCOOKTODO: GenerateLocalizationReferences(BeginContext.StartupOptions->CookCultures);
+	// MPCOOKTODO: RecordDLCPackagesFromBaseGame(BeginContext);
+
+	return FBeginCookContext(*this);
+}
+
+bool UCookOnTheFlyServer::IsCookWorkerMode() const
+{
+	return CurrentCookMode == ECookMode::CookWorker;
+}
+
 void UCookOnTheFlyServer::GenerateInitialRequests(FBeginCookContext& BeginContext)
 {
 	TArray<ITargetPlatform*>& TargetPlatforms = BeginContext.TargetPlatforms;
@@ -9019,7 +9272,7 @@ void UCookOnTheFlyServer::GenerateInitialRequests(FBeginCookContext& BeginContex
 		UE::Cook::FInstigator& Instigator = FilesInPathInstigators.FindChecked(PackageName);
 		if (!PackageFileFName.IsNone())
 		{
-			ExternalRequests->EnqueueUnique(UE::Cook::FFilePlatformRequest(PackageFileFName, MoveTemp(Instigator), TargetPlatforms));
+			WorkerRequests->AddStartCookByTheBookRequest(UE::Cook::FFilePlatformRequest(PackageFileFName, MoveTemp(Instigator), TargetPlatforms));
 		}
 		else if (!FLinkerLoad::IsKnownMissingPackage(PackageName))
 		{
@@ -9053,7 +9306,7 @@ void UCookOnTheFlyServer::GenerateInitialRequests(FBeginCookContext& BeginContex
 			RequestPlatforms.Add(TargetPlatform);
 			for (const UE::Cook::FConstructPackageData& PackageData : BasedOnReleaseDatas)
 			{
-				ExternalRequests->EnqueueUnique(
+				WorkerRequests->AddStartCookByTheBookRequest(
 					UE::Cook::FFilePlatformRequest(PackageData.NormalizedFileName,
 						UE::Cook::EInstigator::PreviousAssetRegistry, RequestPlatforms));
 			}
@@ -9280,7 +9533,7 @@ void UCookOnTheFlyServer::UnregisterCookByTheBookDelegates()
 }
 
 
-TArray<FName> UCookOnTheFlyServer::GetNeverCookPackageFileNames(TArrayView<const FString> ExtraNeverCookDirectories)
+TArray<FName> UCookOnTheFlyServer::GetNeverCookPackageFileNames(TArrayView<const FString> ExtraNeverCookDirectories) const
 {
 	TArray<FString> NeverCookDirectories(ExtraNeverCookDirectories);
 
@@ -9302,7 +9555,7 @@ TArray<FName> UCookOnTheFlyServer::GetNeverCookPackageFileNames(TArrayView<const
 	};
 	const UProjectPackagingSettings* const PackagingSettings = GetDefault<UProjectPackagingSettings>();
 
-	if (IsCookByTheBookMode())
+	if (IsDirectorCookByTheBook())
 	{
 		// Respect the packaging settings nevercook directories for CookByTheBook
 		AddDirectoryPathArray(PackagingSettings->DirectoriesToNeverCook, TEXT("ProjectSettings -> Project -> Packaging -> Directories to never cook"));
@@ -9357,7 +9610,7 @@ void UCookOnTheFlyServer::MaybeMarkPackageAsAlreadyLoaded(UPackage *Package)
 {
 	// can't use this optimization while cooking in editor
 	check(IsCookingInEditor()==false);
-	check(IsCookByTheBookMode());
+	check(IsDirectorCookByTheBook());
 
 	// if the package is already fully loaded then we are not going to mark it up anyway
 	if ( Package->IsFullyLoaded() )
@@ -9374,6 +9627,7 @@ void UCookOnTheFlyServer::MaybeMarkPackageAsAlreadyLoaded(UPackage *Package)
 		return;
 	}
 	FName StandardName = PackageData->GetFileName();
+	// MPCOOKTODO: Mark it as fastload if its saving on another worker
 	if (PackageData->HasAnyCookedPlatform())
 	{
 		bShouldMarkAsAlreadyProcessed = PackageData->HasAllCookedPlatforms(PlatformManager->GetSessionPlatforms(), true /* bIncludeFailed */);
@@ -9624,11 +9878,11 @@ uint32 UCookOnTheFlyServer::CookFullLoadAndSave()
 	{
 		UE_LOG(LogCook, Display, TEXT("Loading requested packages..."));
 		UE_SCOPED_HIERARCHICAL_COOKTIMER(FullLoadAndSave_RequestedLoads);
-		while (ExternalRequests->HasRequests())
+		while (WorkerRequests->HasExternalRequests())
 		{
 			TArray<UE::Cook::FFilePlatformRequest> BuildRequests;
 			TArray<UE::Cook::FSchedulerCallback> SchedulerCallbacks;
-			UE::Cook::EExternalRequestType RequestType = ExternalRequests->DequeueNextCluster(SchedulerCallbacks, BuildRequests);
+			UE::Cook::EExternalRequestType RequestType = WorkerRequests->DequeueNextCluster(SchedulerCallbacks, BuildRequests);
 			if (RequestType == UE::Cook::EExternalRequestType::Callback)
 			{
 				for (UE::Cook::FSchedulerCallback& SchedulerCallback : SchedulerCallbacks)
@@ -10194,6 +10448,12 @@ ICookedPackageWriter& UCookOnTheFlyServer::FindOrCreatePackageWriter(const ITarg
 	return *FindOrCreateSaveContext(TargetPlatform).PackageWriter;
 }
 
+const ICookedPackageWriter* UCookOnTheFlyServer::FindPackageWriter(const ITargetPlatform* TargetPlatform) const
+{
+	const UE::Cook::FCookSavePackageContext* Context = FindSaveContext(TargetPlatform);
+	return Context ? Context->PackageWriter : nullptr;
+}
+
 void UCookOnTheFlyServer::FindOrCreateSaveContexts(TConstArrayView<const ITargetPlatform*> TargetPlatforms)
 {
 	for (const ITargetPlatform* TargetPlatform : TargetPlatforms)
@@ -10214,6 +10474,18 @@ UE::Cook::FCookSavePackageContext& UCookOnTheFlyServer::FindOrCreateSaveContext(
 	return *SavePackageContexts.Add_GetRef(CreateSaveContext(TargetPlatform));
 }
 
+const UE::Cook::FCookSavePackageContext* UCookOnTheFlyServer::FindSaveContext(const ITargetPlatform* TargetPlatform) const
+{
+	for (const UE::Cook::FCookSavePackageContext* Context : SavePackageContexts)
+	{
+		if (Context->SaveContext.TargetPlatform == TargetPlatform)
+		{
+			return Context;
+		}
+	}
+	return nullptr;
+}
+
 void UCookOnTheFlyServer::ConditionalInstallImportBehaviorCallback()
 {
 	if (!bImportBehaviorCallbackInstalled && !IsCookingInEditor() && FEditorDomain::Get())
diff --git a/Engine/Source/Editor/UnrealEd/Private/Cooker/CookDirector.h b/Engine/Source/Editor/UnrealEd/Private/Cooker/CookDirector.h
new file mode 100644
--- /dev/null
+++ b/Engine/Source/Editor/UnrealEd/Private/Cooker/CookDirector.h
@@ -0,0 +1,17 @@
+// Copyright Epic Games, Inc. All Rights Reserved.
+
+#pragma once
+
+class UCookOnTheFlyServer;
+
+namespace UE::Cook
+{
+
+class FCookDirector
+{
+public:
+	FCookDirector(UCookOnTheFlyServer& COTFS)
+	{}
+};
+
+}
\ No newline at end of file
diff --git a/Engine/Source/Editor/UnrealEd/Private/Cooker/CookPackageData.cpp b/Engine/Source/Editor/UnrealEd/Private/Cooker/CookPackageData.cpp
--- a/Engine/Source/Editor/UnrealEd/Private/Cooker/CookPackageData.cpp
+++ b/Engine/Source/Editor/UnrealEd/Private/Cooker/CookPackageData.cpp
@@ -2010,18 +2010,17 @@ FPackageDatas::FPackageDatas(UCookOnTheFlyServer& InCookOnTheFlyServer)
 {
 }
 
-void FPackageDatas::SetBeginCookConfigSettings()
+void FPackageDatas::SetBeginCookConfigSettings(FStringView CookShowInstigator)
 {
-	FString FileOrPackageName;
 	ShowInstigatorPackageData = nullptr;
-	if (FParse::Value(FCommandLine::Get(), TEXT("-CookShowInstigator="), FileOrPackageName))
+	if (!CookShowInstigator.IsEmpty())
 	{
 		FString LocalPath;
 		FString PackageName;
-		if (!FPackageName::TryConvertToMountedPath(FileOrPackageName, &LocalPath, &PackageName, nullptr, nullptr, nullptr))
+		if (!FPackageName::TryConvertToMountedPath(CookShowInstigator, &LocalPath, &PackageName, nullptr, nullptr, nullptr))
 		{
-			UE_LOG(LogCook, Fatal, TEXT("-CookShowInstigator argument %s is not a mounted filename or packagename"),
-				*FileOrPackageName);
+			UE_LOG(LogCook, Fatal, TEXT("-CookShowInstigator argument %.*s is not a mounted filename or packagename"),
+				CookShowInstigator.Len(), CookShowInstigator.GetData());
 		}
 		else
 		{
@@ -2029,8 +2028,8 @@ void FPackageDatas::SetBeginCookConfigSettings()
 			ShowInstigatorPackageData = TryAddPackageDataByPackageName(PackageFName);
 			if (!ShowInstigatorPackageData)
 			{
-				UE_LOG(LogCook, Fatal, TEXT("-CookShowInstigator argument %s could not be found on disk"),
-					*FileOrPackageName);
+				UE_LOG(LogCook, Fatal, TEXT("-CookShowInstigator argument %.*s could not be found on disk"),
+					CookShowInstigator.Len(), CookShowInstigator.GetData());
 			}
 		}
 	}
diff --git a/Engine/Source/Editor/UnrealEd/Private/Cooker/CookPackageData.h b/Engine/Source/Editor/UnrealEd/Private/Cooker/CookPackageData.h
--- a/Engine/Source/Editor/UnrealEd/Private/Cooker/CookPackageData.h
+++ b/Engine/Source/Editor/UnrealEd/Private/Cooker/CookPackageData.h
@@ -962,7 +962,7 @@ public:
 	static void OnAssetRegistryGenerated(IAssetRegistry& InAssetRegistry);
 
 	/** Called each time BeginCook is called, to initialize settings from config */
-	void SetBeginCookConfigSettings();
+	void SetBeginCookConfigSettings(FStringView CookShowInstigator);
 
 	/** FGCObject interface function - return a debug name describing this FGCObject. */
 	virtual FString GetReferencerName() const override;
diff --git a/Engine/Source/Editor/UnrealEd/Private/Cooker/CookPlatformManager.cpp b/Engine/Source/Editor/UnrealEd/Private/Cooker/CookPlatformManager.cpp
--- a/Engine/Source/Editor/UnrealEd/Private/Cooker/CookPlatformManager.cpp
+++ b/Engine/Source/Editor/UnrealEd/Private/Cooker/CookPlatformManager.cpp
@@ -4,9 +4,9 @@
 
 #include "Commandlets/AssetRegistryGenerator.h"
 #include "CookOnTheSide/CookOnTheFlyServer.h"
-#include "CookRequests.h"
 #include "Interfaces/ITargetPlatform.h"
 #include "Interfaces/ITargetPlatformManagerModule.h"
+#include "IWorkerRequests.h"
 #include "Misc/ScopeRWLock.h"
 
 namespace UE
@@ -285,7 +285,7 @@ namespace Cook
 
 		if (!HasSessionPlatform(PlatformData->TargetPlatform))
 		{
-			CookOnTheFlyServer.ExternalRequests->AddCallback([PlatformName, LocalCookOnTheFlyServer = &CookOnTheFlyServer]()
+			CookOnTheFlyServer.WorkerRequests->AddCookOnTheFlyCallback([PlatformName, LocalCookOnTheFlyServer = &CookOnTheFlyServer]()
 				{
 					ITargetPlatform* TargetPlatform = GetTargetPlatformManager()->FindTargetPlatform(PlatformName.ToString());
 					if (TargetPlatform)
diff --git a/Engine/Source/Editor/UnrealEd/Private/Cooker/CookRequests.cpp b/Engine/Source/Editor/UnrealEd/Private/Cooker/CookRequests.cpp
--- a/Engine/Source/Editor/UnrealEd/Private/Cooker/CookRequests.cpp
+++ b/Engine/Source/Editor/UnrealEd/Private/Cooker/CookRequests.cpp
@@ -162,12 +162,12 @@ void FFilePlatformRequest::RemapTargetPlatforms(const TMap<ITargetPlatform*, ITa
 
 int32 FExternalRequests::GetNumRequests() const
 {
-	return RequestCount;
+	return RequestCount.load(std::memory_order_relaxed);
 }
 
 bool FExternalRequests::HasRequests() const
 {
-	return RequestCount > 0;
+	return GetNumRequests() > 0;
 }
 
 void FExternalRequests::AddCallback(FSchedulerCallback&& Callback)
@@ -175,7 +175,7 @@ void FExternalRequests::AddCallback(FSchedulerCallback&& Callback)
 	FScopeLock ScopeLock(&RequestLock);
 
 	Callbacks.Add(MoveTemp(Callback));
-	++RequestCount;
+	RequestCount.fetch_add(1, std::memory_order_relaxed);
 }
 
 void FExternalRequests::EnqueueUnique(FFilePlatformRequest&& FileRequest, bool bForceFrontOfQueue)
@@ -195,7 +195,7 @@ void FExternalRequests::EnqueueUnique(FFilePlatformRequest&& FileRequest, bool b
 			Queue.Add(Filename);
 		}
 
-		++RequestCount;
+		RequestCount.fetch_add(1, std::memory_order_relaxed);
 	}
 	else
 	{
@@ -236,12 +236,13 @@ EExternalRequestType FExternalRequests::DequeueNextCluster(TArray<FSchedulerCall
 	else if (Queue.Num())
 	{
 		OutBuildRequests.Reserve(Queue.Num() + OutBuildRequests.Num());
+		int32 NumDequeued = Queue.Num();
 		while (Queue.Num())
 		{
 			FName Filename = Queue.PopFrontValue();
 			OutBuildRequests.Add(RequestMap.FindAndRemoveChecked(Filename));
-			--RequestCount;
 		}
+		RequestCount.fetch_add(-NumDequeued, std::memory_order_relaxed);
 		return EExternalRequestType::Cook;
 	}
 	else
@@ -262,7 +263,7 @@ bool FExternalRequests::ThreadUnsafeDequeueCallbacks(TArray<FSchedulerCallback>&
 	{
 		OutCallbacks = MoveTemp(Callbacks);
 		Callbacks.Empty();
-		RequestCount -= OutCallbacks.Num();
+		RequestCount.fetch_add(-OutCallbacks.Num(), std::memory_order_relaxed);
 		return true;
 	}
 	else
@@ -278,7 +279,7 @@ void FExternalRequests::EmptyRequests()
 	Queue.Empty();
 	RequestMap.Empty();
 	Callbacks.Empty();
-	RequestCount = 0;
+	RequestCount.store(0, std::memory_order_relaxed);
 }
 
 void FExternalRequests::DequeueAll(TArray<FSchedulerCallback>& OutCallbacks, TArray<FFilePlatformRequest>& OutCookRequests)
@@ -295,7 +296,7 @@ void FExternalRequests::DequeueAll(TArray<FSchedulerCallback>& OutCallbacks, TAr
 	RequestMap.Empty();
 	Queue.Empty();
 
-	RequestCount = 0;
+	RequestCount.store(0, std::memory_order_relaxed);
 }
 
 void FExternalRequests::OnRemoveSessionPlatform(const ITargetPlatform* TargetPlatform)
diff --git a/Engine/Source/Editor/UnrealEd/Private/Cooker/CookRequests.h b/Engine/Source/Editor/UnrealEd/Private/Cooker/CookRequests.h
--- a/Engine/Source/Editor/UnrealEd/Private/Cooker/CookRequests.h
+++ b/Engine/Source/Editor/UnrealEd/Private/Cooker/CookRequests.h
@@ -2,6 +2,7 @@
 
 #pragma once
 
+#include <atomic>
 #include "Containers/Array.h"
 #include "Containers/ArrayView.h"
 #include "Containers/Map.h"
@@ -10,7 +11,9 @@
 #include "CookTypes.h"
 #include "Engine/ICookInfo.h"
 #include "HAL/CriticalSection.h"
+#include "HAL/Event.h"
 #include "HAL/Platform.h"
+#include "HAL/PlatformProcess.h"
 #include "UObject/NameTypes.h"
 
 class FEvent;
@@ -19,6 +22,8 @@ class ITargetPlatform;
 namespace UE::Cook
 {
 
+struct FFreeFEvent { void operator()(FEvent* Ptr) const { if (Ptr) FPlatformProcess::ReturnSynchEventToPool(Ptr); } };
+
 /**
  * Structure holding the data for a request for the CookOnTheFlyServer to cook a FileName. Includes platform which file is requested for.
  * These requests are external to the cooker's scheduler, and do not use the FPackageData the scheduler uses internally.
@@ -120,7 +125,7 @@ public:
 
 public:
 	/** An FEvent the scheduler can sleep on when waiting for new cookonthefly requests. */
-	FEvent* CookRequestEvent = nullptr;
+	TUniquePtr<FEvent, FFreeFEvent> CookRequestEvent;
 private:
 	/* Implementation for DequeueCallbacks that assumes the caller has entered the RequestLock. */
 	bool ThreadUnsafeDequeueCallbacks(TArray<FSchedulerCallback>& OutCallbacks);
@@ -131,7 +136,7 @@ private:
 	TMap<FName, FFilePlatformRequest> RequestMap;
 	TArray<FSchedulerCallback> Callbacks;
 	FCriticalSection RequestLock;
-	int32 RequestCount = 0;
+	std::atomic<int32> RequestCount = 0;
 };
 
 }
\ No newline at end of file
diff --git a/Engine/Source/Editor/UnrealEd/Private/Cooker/CookTypes.h b/Engine/Source/Editor/UnrealEd/Private/Cooker/CookTypes.h
--- a/Engine/Source/Editor/UnrealEd/Private/Cooker/CookTypes.h
+++ b/Engine/Source/Editor/UnrealEd/Private/Cooker/CookTypes.h
@@ -235,6 +235,39 @@ namespace UE::Cook
 		};
 		TMap<FName, FPendingBuildData> PendingBuilds;
 	};
+
+	struct FInitializeConfigSettings
+	{
+		void LoadLocal(const FString& InOutputDirectoryOverride);
+
+		FString OutputDirectoryOverride;
+		int32 MaxPrecacheShaderJobs = 0;
+		int32 MaxConcurrentShaderJobs = 0;
+		uint32 PackagesPerGC = 0;
+		float MemoryExpectedFreedToSpreadRatio = 0.f;
+		double IdleTimeToGC = 0.;
+		uint64 MemoryMaxUsedVirtual;
+		uint64 MemoryMaxUsedPhysical;
+		uint64 MemoryMinFreeVirtual;
+		uint64 MemoryMinFreePhysical;
+		int32 MinFreeUObjectIndicesBeforeGC;
+		int32 MaxNumPackagesBeforePartialGC;
+		TArray<FString> ConfigSettingDenyList;
+		TMap<FName, int32> MaxAsyncCacheForType; // max number of objects of a specific type which are allowed to async cache at once
+		bool bHybridIterativeDebug = false;
+	};
+
+	struct FBeginCookConfigSettings
+	{
+		/** Initialize NeverCookPackageList from packaging settings and platform-specific sources. */
+		void LoadLocal(FBeginCookContext& BeginContext);
+		void LoadNeverCookLocal(FBeginCookContext& BeginContext);
+
+		FString CookShowInstigator;
+		bool bHybridIterativeEnabled = true;
+		TArray<FName> NeverCookPackageList;
+		TFastPointerMap<const ITargetPlatform*, TSet<FName>> PlatformSpecificNeverCookPackages;
+	};
 }
 
 inline void RouteBeginCacheForCookedPlatformData(UObject* Obj, const ITargetPlatform* TargetPlatform)
@@ -336,16 +369,24 @@ struct FBeginCookContextPlatform
 	bool bPopulateMemoryResultsFromDiskResults = false;
 	/** If true we are cooking iteratively, from results in a shared build (e.g. from buildfarm) rather than from our previous cook. */
 	bool bIterateSharedBuild = false;
+	/** If true we are a CookWorker, and we are working on a Sandbox directory that has already been populated by a remote Director process. */
+	bool bWorkerOnSharedSandbox = false;
 };
 
 /** Data held on the stack and shared with multiple subfunctions when running StartCookByTheBook or StartCookOnTheFly */
 struct FBeginCookContext
 {
+	FBeginCookContext(UCookOnTheFlyServer& InCOTFS)
+		: COTFS(InCOTFS)
+	{
+	}
+
 	const UCookOnTheFlyServer::FCookByTheBookStartupOptions* StartupOptions = nullptr;
 	/** List of the platforms we are building, with startup context data about each one */
 	TArray<FBeginCookContextPlatform> PlatformContexts;
 	/** The list of platforms by themselves, for passing to functions that need just a list of platforms */
 	TArray<ITargetPlatform*> TargetPlatforms;
+	const UCookOnTheFlyServer& COTFS;
 };
 
 void LogCookerMessage(const FString& MessageText, EMessageSeverity::Type Severity);
diff --git a/Engine/Source/Editor/UnrealEd/Private/Cooker/IWorkerRequests.h b/Engine/Source/Editor/UnrealEd/Private/Cooker/IWorkerRequests.h
new file mode 100644
--- /dev/null
+++ b/Engine/Source/Editor/UnrealEd/Private/Cooker/IWorkerRequests.h
@@ -0,0 +1,86 @@
+// Copyright Epic Games, Inc. All Rights Reserved.
+
+#pragma once
+
+#include "Containers/UnrealString.h"
+#include "CookOnTheSide/CookOnTheFlyServer.h"
+#include "CookRequests.h"
+#include "CookTypes.h"
+
+class FConfigFile;
+namespace UE::Cook { struct FPackageData; }
+
+namespace UE::Cook
+{
+
+/**
+ * Interface for Requests sent to this processes' Cooker, and for information this process needs to send back to its Director.
+ * In SingleProcess cooks the functions are passed through to the local CookOnTheFlyServer.
+ * In MultiProcess cooks the functions are implemented as interprocess messages to/from the CookDirector.
+ **/
+class IWorkerRequests
+{
+public:
+	virtual ~IWorkerRequests() {}
+
+	// Reading Callbacks And Requests
+	/**
+	 * Lockless value for whether the Worker has ExternalRequests it has not yet reported to the Cooker.
+	 * May be out of date after calling; do not assume a true return value means Requests are actually present or a false value means no Requests are present.
+	 * Intended usage is for the Scheduler to be the only consumer of requests, and to use this value for periodic checking of whether there is any work that justifies the expense of taking the lock.
+	 * In a single-consumer case, HasRequests will eventually correctly return true as long as the consumer is not consuming.
+	 */
+	virtual bool HasExternalRequests() const = 0;
+	/**
+	 * Lockless value for the number of ExternalRequests the Worker has which it has not yet reported to the Cooker.
+	 * May be out of date after calling; do not assume the number of actual requests is any one of equal, greater than, or less than the returned value.
+	 * Intended usage is for the Scheduler to be the only consumer of requests, and to use this value for rough reporting of periodic progress.
+	 */
+	virtual int32 GetNumExternalRequests() const = 0;
+	/**
+	 * If this Worker has any callbacks, dequeue them all into OutCallbacks and return EExternalRequestType::Callback; Callbacks take priority over cook requests.
+	 * Otherwise, if there are any cook requests, dequeue them all into OutBuildRequests and return EExternalRequestType::Cook.
+	 * Otherwise, return EExternalRequestType::None.
+	 */
+	virtual EExternalRequestType DequeueNextCluster(TArray<FSchedulerCallback>& OutCallbacks, TArray<FFilePlatformRequest>& OutBuildRequests) = 0;
+	/* Move any existing callbacks onto OutCallbacks, and return whether any were added. */
+	virtual bool DequeueSchedulerCallbacks(TArray<FSchedulerCallback>& OutCallbacks) = 0;
+	/** Move all callbacks into OutCallbacks, and all cook requests into OutCookRequests. This is used when canceling a cook session. */
+	virtual void DequeueAllExternal(TArray<FSchedulerCallback>& OutCallbacks, TArray<FFilePlatformRequest>& OutCookRequests) = 0;
+
+	// Writing Packages during Cook
+	virtual void AddDiscoveredPackage(FPackageData& PackageData, FInstigator& Instigator, bool bLoadReady, bool& bOutShouldAddToQueue)
+	{
+		bOutShouldAddToQueue = true;
+	}
+
+	// Writing Packages from CookByTheBook
+	virtual void AddStartCookByTheBookRequest(FFilePlatformRequest&& Request) = 0;
+
+	// Writing Packages from CookOnTheFly
+	virtual void InitializeCookOnTheFly() = 0;
+	virtual void AddCookOnTheFlyRequest(FFilePlatformRequest&& Request) = 0;
+	virtual void AddCookOnTheFlyCallback(FSchedulerCallback&& Callback) = 0;
+	virtual void WaitForCookOnTheFlyEvents(int TimeoutMs) = 0;
+
+	// Writing Packages from Editor interface
+	virtual void AddEditorActionCallback(FSchedulerCallback&& Callback) = 0;
+	virtual void AddPublicInterfaceRequest(FFilePlatformRequest&& Request, bool bForceFrontOfQueue) = 0;
+
+	// Hooks into cooker events
+	/** Swap all ITargetPlatform* stored on this instance according to the mapping in @param Remap. */
+	virtual void RemapTargetPlatforms(const TMap<ITargetPlatform*, ITargetPlatform*>& Remap) = 0;
+	/** Remove references to the given platform from all cook requests. */
+	virtual void OnRemoveSessionPlatform(const ITargetPlatform* TargetPlatform) = 0;
+
+	// Functions that are written locally for Directors but are sent to Director for Workers
+	virtual void ReportAccessedIniSettings(UCookOnTheFlyServer& COTFS, const FConfigFile& Config) = 0;
+
+	// Functions that read from local for Directors but read from Director for Workers
+	virtual void GetInitializeConfigSettings(UCookOnTheFlyServer& COTFS, const FString& OutputDirectoryOverride, UE::Cook::FInitializeConfigSettings& Settings) = 0;
+	virtual void GetBeginCookConfigSettings(UCookOnTheFlyServer& COTFS, FBeginCookContext& BeginContext, UE::Cook::FBeginCookConfigSettings& Settings) = 0;
+	virtual void GetBeginCookIterativeFlags(UCookOnTheFlyServer& COTFS, FBeginCookContext& BeginContext) = 0;
+	virtual ECookMode::Type GetDirectorCookMode(UCookOnTheFlyServer& COTFS) = 0;
+};
+
+}
\ No newline at end of file
diff --git a/Engine/Source/Editor/UnrealEd/Private/Cooker/LooseCookedPackageWriter.cpp b/Engine/Source/Editor/UnrealEd/Private/Cooker/LooseCookedPackageWriter.cpp
--- a/Engine/Source/Editor/UnrealEd/Private/Cooker/LooseCookedPackageWriter.cpp
+++ b/Engine/Source/Editor/UnrealEd/Private/Cooker/LooseCookedPackageWriter.cpp
@@ -447,10 +447,11 @@ FDateTime FLooseCookedPackageWriter::GetPreviousCookTime() const
 void FLooseCookedPackageWriter::Initialize(const FCookInfo& Info)
 {
 	bIterateSharedBuild = Info.bIterateSharedBuild;
-	if (Info.bFullBuild)
+	if (Info.bFullBuild && !Info.bWorkerOnSharedSandbox)
 	{
 		DeleteSandboxDirectory();
 	}
+	if (!Info.bWorkerOnSharedSandbox)
 	{
 		TRACE_CPUPROFILER_EVENT_SCOPE(SaveScriptObjects);
 		FPackageStoreOptimizer PackageStoreOptimizer;
diff --git a/Engine/Source/Editor/UnrealEd/Private/Cooker/WorkerRequestsLocal.cpp b/Engine/Source/Editor/UnrealEd/Private/Cooker/WorkerRequestsLocal.cpp
new file mode 100644
--- /dev/null
+++ b/Engine/Source/Editor/UnrealEd/Private/Cooker/WorkerRequestsLocal.cpp
@@ -0,0 +1,122 @@
+// Copyright Epic Games, Inc. All Rights Reserved.
+
+#include "WorkerRequestsLocal.h"
+
+#include "CookTypes.h"
+
+namespace UE::Cook
+{
+
+bool FWorkerRequestsLocal::HasExternalRequests() const
+{
+	return ExternalRequests.HasRequests();
+}
+
+int32 FWorkerRequestsLocal::GetNumExternalRequests() const
+{
+	return ExternalRequests.GetNumRequests();
+}
+
+EExternalRequestType FWorkerRequestsLocal::DequeueNextCluster(TArray<FSchedulerCallback>& OutCallbacks, TArray<FFilePlatformRequest>& OutBuildRequests)
+{
+	return ExternalRequests.DequeueNextCluster(OutCallbacks, OutBuildRequests);
+}
+
+bool FWorkerRequestsLocal::DequeueSchedulerCallbacks(TArray<FSchedulerCallback>& OutCallbacks)
+{
+	return ExternalRequests.DequeueCallbacks(OutCallbacks);
+}
+
+void FWorkerRequestsLocal::DequeueAllExternal(TArray<FSchedulerCallback>& OutCallbacks, TArray<FFilePlatformRequest>& OutCookRequests)
+{
+	ExternalRequests.DequeueAll(OutCallbacks, OutCookRequests);
+}
+
+void FWorkerRequestsLocal::AddDiscoveredPackage(FPackageData& PackageData, FInstigator& Instigator, bool bLoadReady, bool& bOutShouldAddToQueue)
+{
+	bOutShouldAddToQueue = true;
+}
+
+void FWorkerRequestsLocal::AddStartCookByTheBookRequest(FFilePlatformRequest&& Request)
+{
+	ExternalRequests.EnqueueUnique(MoveTemp(Request));
+}
+
+void FWorkerRequestsLocal::InitializeCookOnTheFly()
+{
+	ExternalRequests.CookRequestEvent.Reset(FPlatformProcess::GetSynchEventFromPool());
+}
+
+void FWorkerRequestsLocal::AddCookOnTheFlyRequest(FFilePlatformRequest&& Request)
+{
+	ExternalRequests.EnqueueUnique(MoveTemp(Request), true /* bForceFrontOfQueue */);
+	if (ExternalRequests.CookRequestEvent)
+	{
+		ExternalRequests.CookRequestEvent->Trigger();
+	}
+}
+
+void FWorkerRequestsLocal::AddCookOnTheFlyCallback(FSchedulerCallback&& Callback)
+{
+	ExternalRequests.AddCallback(MoveTemp(Callback));
+	if (ExternalRequests.CookRequestEvent)
+	{
+		ExternalRequests.CookRequestEvent->Trigger();
+	}
+}
+
+void FWorkerRequestsLocal::WaitForCookOnTheFlyEvents(int TimeoutMs)
+{
+	if (ExternalRequests.CookRequestEvent)
+	{
+		ExternalRequests.CookRequestEvent->Wait(TimeoutMs, true /* bIgnoreThreadIdleStats */);
+	}
+}
+
+void FWorkerRequestsLocal::AddEditorActionCallback(FSchedulerCallback&& Callback)
+{
+	ExternalRequests.AddCallback(MoveTemp(Callback));
+}
+
+void FWorkerRequestsLocal::AddPublicInterfaceRequest(FFilePlatformRequest&& Request, bool bForceFrontOfQueue)
+{
+	ExternalRequests.EnqueueUnique(MoveTemp(Request), bForceFrontOfQueue);
+}
+
+void FWorkerRequestsLocal::RemapTargetPlatforms(const TMap<ITargetPlatform*, ITargetPlatform*>& Remap)
+{
+	ExternalRequests.RemapTargetPlatforms(Remap);
+}
+
+void FWorkerRequestsLocal::OnRemoveSessionPlatform(const ITargetPlatform* TargetPlatform)
+{
+	ExternalRequests.OnRemoveSessionPlatform(TargetPlatform);
+}
+
+void FWorkerRequestsLocal::ReportAccessedIniSettings(UCookOnTheFlyServer& COTFS, const FConfigFile& Config)
+{
+	COTFS.ProcessAccessedIniSettings(&Config, COTFS.AccessedIniStrings);
+}
+
+void FWorkerRequestsLocal::GetInitializeConfigSettings(UCookOnTheFlyServer& COTFS,
+	const FString& OutputDirectoryOverride, UE::Cook::FInitializeConfigSettings& Settings)
+{
+	Settings.LoadLocal(OutputDirectoryOverride);
+}
+
+void FWorkerRequestsLocal::GetBeginCookConfigSettings(UCookOnTheFlyServer& COTFS, FBeginCookContext& BeginContext, UE::Cook::FBeginCookConfigSettings& Settings)
+{
+	Settings.LoadLocal(BeginContext);
+}
+
+void FWorkerRequestsLocal::GetBeginCookIterativeFlags(UCookOnTheFlyServer& COTFS, FBeginCookContext& BeginContext)
+{
+	BeginContext.COTFS.LoadBeginCookIterativeFlagsLocal(BeginContext);
+}
+
+ECookMode::Type FWorkerRequestsLocal::GetDirectorCookMode(UCookOnTheFlyServer& COTFS)
+{
+	return COTFS.GetCookMode();
+}
+
+}
\ No newline at end of file
diff --git a/Engine/Source/Editor/UnrealEd/Private/Cooker/WorkerRequestsLocal.h b/Engine/Source/Editor/UnrealEd/Private/Cooker/WorkerRequestsLocal.h
new file mode 100644
--- /dev/null
+++ b/Engine/Source/Editor/UnrealEd/Private/Cooker/WorkerRequestsLocal.h
@@ -0,0 +1,41 @@
+// Copyright Epic Games, Inc. All Rights Reserved.
+
+#pragma once
+
+#include "CookRequests.h"
+#include "IWorkerRequests.h"
+
+namespace UE::Cook
+{
+
+/** An IWorkerRequests for SingleProcess cooks: functions are passed through to the CookOnTheFlyServer. */
+class FWorkerRequestsLocal : public IWorkerRequests
+{
+public:
+	virtual bool HasExternalRequests() const override;
+	virtual int32 GetNumExternalRequests() const override;
+	virtual EExternalRequestType DequeueNextCluster(TArray<FSchedulerCallback>& OutCallbacks, TArray<FFilePlatformRequest>& OutBuildRequests) override;
+	virtual bool DequeueSchedulerCallbacks(TArray<FSchedulerCallback>& OutCallbacks) override;
+	virtual void DequeueAllExternal(TArray<FSchedulerCallback>& OutCallbacks, TArray<FFilePlatformRequest>& OutCookRequests) override;
+	virtual void RemapTargetPlatforms(const TMap<ITargetPlatform*, ITargetPlatform*>& Remap) override;
+	virtual void OnRemoveSessionPlatform(const ITargetPlatform* TargetPlatform) override;
+	virtual void AddDiscoveredPackage(FPackageData& PackageData, FInstigator& Instigator, bool bLoadReady, bool& bOutShouldAddToQueue) override;
+	virtual void AddStartCookByTheBookRequest(FFilePlatformRequest&& Request) override;
+	virtual void InitializeCookOnTheFly() override;
+	virtual void AddCookOnTheFlyRequest(FFilePlatformRequest&& Request) override;
+	virtual void AddCookOnTheFlyCallback(FSchedulerCallback&& Callback) override;
+	virtual void WaitForCookOnTheFlyEvents(int TimeoutMs) override;
+	virtual void AddEditorActionCallback(FSchedulerCallback&& Callback) override;
+	virtual void AddPublicInterfaceRequest(FFilePlatformRequest&& Request, bool bForceFrontOfQueue) override;
+
+	virtual void ReportAccessedIniSettings(UCookOnTheFlyServer& COTFS, const FConfigFile& Config) override;
+	virtual void GetInitializeConfigSettings(UCookOnTheFlyServer& COTFS, const FString& OutputOverrideDirectory, UE::Cook::FInitializeConfigSettings& Settings) override;
+	virtual void GetBeginCookConfigSettings(UCookOnTheFlyServer& COTFS, FBeginCookContext& BeginContext, UE::Cook::FBeginCookConfigSettings& Settings) override;
+	virtual void GetBeginCookIterativeFlags(UCookOnTheFlyServer& COTFS, FBeginCookContext& BeginContext) override;
+	virtual ECookMode::Type GetDirectorCookMode(UCookOnTheFlyServer& COTFS) override;
+
+private:
+	FExternalRequests ExternalRequests;
+};
+
+}
\ No newline at end of file
diff --git a/Engine/Source/Editor/UnrealEd/Private/Cooker/WorkerRequestsRemote.cpp b/Engine/Source/Editor/UnrealEd/Private/Cooker/WorkerRequestsRemote.cpp
new file mode 100644
--- /dev/null
+++ b/Engine/Source/Editor/UnrealEd/Private/Cooker/WorkerRequestsRemote.cpp
@@ -0,0 +1,173 @@
+// Copyright Epic Games, Inc. All Rights Reserved.
+
+#include "WorkerRequestsRemote.h"
+
+#include "CookPlatformManager.h"
+#include "CookTypes.h"
+#include "Logging/LogMacros.h"
+#include "Misc/AssertionMacros.h" 
+
+namespace UE::Cook
+{
+
+FWorkerRequestsRemote::FWorkerRequestsRemote(UCookOnTheFlyServer& InCOTFS)
+{
+}
+
+bool FWorkerRequestsRemote::TryConnect(FDirectorConnectionInfo&& ConnectInfo, ECookInitializationFlags& OutCookFlags)
+{
+	return false;
+}
+
+bool FWorkerRequestsRemote::TryGetInitializeSettings(ECookMode::Type& OutCookMode, ECookInitializationFlags& OutCookFlags,
+	FString& OutOutputDirectoryOverride)
+{
+	return false;
+}
+
+bool FWorkerRequestsRemote::HasExternalRequests() const
+{
+	return ExternalRequests.HasRequests();
+}
+
+int32 FWorkerRequestsRemote::GetNumExternalRequests() const
+{
+	return ExternalRequests.GetNumRequests();
+}
+
+EExternalRequestType FWorkerRequestsRemote::DequeueNextCluster(TArray<FSchedulerCallback>& OutCallbacks, TArray<FFilePlatformRequest>& OutBuildRequests)
+{
+	return ExternalRequests.DequeueNextCluster(OutCallbacks, OutBuildRequests);
+}
+
+bool FWorkerRequestsRemote::DequeueSchedulerCallbacks(TArray<FSchedulerCallback>& OutCallbacks)
+{
+	return ExternalRequests.DequeueCallbacks(OutCallbacks);
+}
+
+void FWorkerRequestsRemote::DequeueAllExternal(TArray<FSchedulerCallback>& OutCallbacks, TArray<FFilePlatformRequest>& OutCookRequests)
+{
+	ExternalRequests.DequeueAll(OutCallbacks, OutCookRequests);
+}
+
+void FWorkerRequestsRemote::AddDiscoveredPackage(FPackageData& PackageData, FInstigator& Instigator, bool bLoadReady, bool& bOutShouldAddToQueue)
+{
+	bOutShouldAddToQueue = false;
+	// MPCOOKTODO: Not yet implemented
+	// Send package information to the Director, and make sure to send the instigator.
+	// The director needs to handle GeneratedPackages by forcing them to go to this CookWorker, or if they need to go to another worker, add handling that calls pumpsave on the generator first
+}
+
+void FWorkerRequestsRemote::AddStartCookByTheBookRequest(FFilePlatformRequest&& Request)
+{
+	LogCalledCookByTheBookError(TEXT("AddStartCookByTheBookRequest"));
+}
+
+void FWorkerRequestsRemote::InitializeCookOnTheFly()
+{
+	LogCalledCookOnTheFlyError(TEXT("InitializeCookOnTheFly"));
+}
+
+void FWorkerRequestsRemote::AddCookOnTheFlyRequest(FFilePlatformRequest&& Request)
+{
+	LogCalledCookOnTheFlyError(TEXT("AddCookOnTheFlyRequest"));
+}
+
+void FWorkerRequestsRemote::AddCookOnTheFlyCallback(FSchedulerCallback&& Callback)
+{
+	LogCalledCookOnTheFlyError(TEXT("AddCookOnTheFlyCallback"));
+}
+
+void FWorkerRequestsRemote::WaitForCookOnTheFlyEvents(int TimeoutMs)
+{
+	LogCalledCookOnTheFlyError(TEXT("WaitForCookOnTheFlyEvents"));
+}
+
+void FWorkerRequestsRemote::AddEditorActionCallback(FSchedulerCallback&& Callback)
+{
+	LogCalledEditorActionError(TEXT("AddEditorActionCallback"));
+}
+
+void FWorkerRequestsRemote::AddPublicInterfaceRequest(FFilePlatformRequest&& Request, bool bForceFrontOfQueue)
+{
+	LogCalledPublicInterfaceError(TEXT("AddPublicInterfaceRequest"));
+}
+
+void FWorkerRequestsRemote::RemapTargetPlatforms(const TMap<ITargetPlatform*, ITargetPlatform*>& Remap)
+{
+	ExternalRequests.RemapTargetPlatforms(Remap);
+}
+
+void FWorkerRequestsRemote::OnRemoveSessionPlatform(const ITargetPlatform* TargetPlatform)
+{
+	ExternalRequests.OnRemoveSessionPlatform(TargetPlatform);
+}
+
+void FWorkerRequestsRemote::ReportAccessedIniSettings(UCookOnTheFlyServer& COTFS, const FConfigFile& Config)
+{
+	FIniSettingContainer AccessedIniStrings;
+	COTFS.ProcessAccessedIniSettings(&Config, COTFS.AccessedIniStrings);
+	// MPCOOKTODO: Not yet implemented
+}
+
+void FWorkerRequestsRemote::GetInitializeConfigSettings(UCookOnTheFlyServer& COTFS,
+	const FString& OutputDirectoryOverride, UE::Cook::FInitializeConfigSettings& Settings)
+{
+	// MPCOOKTODO: Not yet implemented
+	Settings.LoadLocal(OutputDirectoryOverride);
+}
+
+void FWorkerRequestsRemote::GetBeginCookConfigSettings(UCookOnTheFlyServer& COTFS, FBeginCookContext& BeginContext, UE::Cook::FBeginCookConfigSettings& Settings)
+{
+	// MPCOOKTODO: Not yet implemented
+	Settings.LoadLocal(BeginContext);
+}
+
+void FWorkerRequestsRemote::GetBeginCookIterativeFlags(UCookOnTheFlyServer& COTFS, FBeginCookContext& BeginContext)
+{
+	for (FBeginCookContextPlatform& PlatformContext : BeginContext.PlatformContexts)
+	{
+		const ITargetPlatform* TargetPlatform = PlatformContext.TargetPlatform;
+		UE::Cook::FPlatformData* PlatformData = PlatformContext.PlatformData;
+		PlatformContext.CurrentCookSettings = COTFS.CalculateCookSettingStrings(); // MPCOOKTODO: Copy from Director
+		PlatformContext.bHasMemoryResults = PlatformData->bIsSandboxInitialized;
+		PlatformContext.bFullBuild = true; // MPCOOKTODO: Copy from Director
+		PlatformContext.bClearMemoryResults = true;
+		PlatformContext.bPopulateMemoryResultsFromDiskResults = false;
+		PlatformContext.bIterateSharedBuild = false; // MPCOOKTODO: Copy from Director
+		PlatformContext.bWorkerOnSharedSandbox = true;
+		PlatformData->bFullBuild = PlatformContext.bFullBuild;
+	}
+}
+
+ECookMode::Type FWorkerRequestsRemote::GetDirectorCookMode(UCookOnTheFlyServer& COTFS)
+{
+	// MPCOOKTODO: Not yet implemented
+	return ECookMode::CookByTheBook;
+}
+
+void FWorkerRequestsRemote::LogCalledCookByTheBookError(const TCHAR* FunctionName) const
+{
+	check(FunctionName);
+	UE_LOG(LogCook, Error, TEXT("Calling %s (a CookByTheBook function) is not allowed in a CookWorker."), FunctionName);
+}
+
+void FWorkerRequestsRemote::LogCalledCookOnTheFlyError(const TCHAR* FunctionName) const
+{
+	check(FunctionName);
+	UE_LOG(LogCook, Error, TEXT("Calling %s (a CookOnTheFly function) is not allowed in a CookWorker."), FunctionName);
+}
+
+void FWorkerRequestsRemote::LogCalledPublicInterfaceError(const TCHAR* FunctionName) const
+{
+	check(FunctionName);
+	UE_LOG(LogCook, Error, TEXT("Calling %s (a CookOnTheFlyServer public interface function) is not allowed in a CookWorker."), FunctionName);
+}
+
+void FWorkerRequestsRemote::LogCalledEditorActionError(const TCHAR* FunctionName) const
+{
+	check(FunctionName);
+	UE_LOG(LogCook, Error, TEXT("Calling %s (an editor-mode-only function) is not allowed in a CookWorker."), FunctionName);
+}
+
+}
\ No newline at end of file
diff --git a/Engine/Source/Editor/UnrealEd/Private/Cooker/WorkerRequestsRemote.h b/Engine/Source/Editor/UnrealEd/Private/Cooker/WorkerRequestsRemote.h
new file mode 100644
--- /dev/null
+++ b/Engine/Source/Editor/UnrealEd/Private/Cooker/WorkerRequestsRemote.h
@@ -0,0 +1,59 @@
+// Copyright Epic Games, Inc. All Rights Reserved.
+
+#pragma once
+
+#include "CookRequests.h"
+#include "Containers/StringView.h"
+#include "IWorkerRequests.h"
+
+namespace UE::Cook
+{
+
+/** An IWorkerRequests for CookWorkers in MultiProcess cooks: functions are implemented as interprocess messages to/from the CookDirector. */
+struct FDirectorConnectionInfo
+{
+	FString HostURI;
+};
+
+class FWorkerRequestsRemote : public IWorkerRequests
+{
+public:
+	FWorkerRequestsRemote(UCookOnTheFlyServer& InCOTFS);
+
+	bool TryConnect(FDirectorConnectionInfo&& ConnectInfo, ECookInitializationFlags& OutCookFlags);
+	bool TryGetInitializeSettings(ECookMode::Type& OutCookMode, ECookInitializationFlags& OutCookFlags,
+		FString& OutOutputDirectoryOverride);
+
+	virtual bool HasExternalRequests() const override;
+	virtual int32 GetNumExternalRequests() const override;
+	virtual EExternalRequestType DequeueNextCluster(TArray<FSchedulerCallback>& OutCallbacks, TArray<FFilePlatformRequest>& OutBuildRequests) override;
+	virtual bool DequeueSchedulerCallbacks(TArray<FSchedulerCallback>& OutCallbacks) override;
+	virtual void DequeueAllExternal(TArray<FSchedulerCallback>& OutCallbacks, TArray<FFilePlatformRequest>& OutCookRequests) override;
+	virtual void RemapTargetPlatforms(const TMap<ITargetPlatform*, ITargetPlatform*>& Remap) override;
+	virtual void OnRemoveSessionPlatform(const ITargetPlatform* TargetPlatform) override;
+	virtual void AddDiscoveredPackage(FPackageData& PackageData, FInstigator& Instigator, bool bLoadReady, bool& bOutShouldAddToQueue) override;
+	virtual void AddStartCookByTheBookRequest(FFilePlatformRequest&& Request) override;
+	virtual void InitializeCookOnTheFly() override;
+	virtual void AddCookOnTheFlyRequest(FFilePlatformRequest&& Request) override;
+	virtual void AddCookOnTheFlyCallback(FSchedulerCallback&& Callback) override;
+	virtual void WaitForCookOnTheFlyEvents(int TimeoutMs) override;
+	virtual void AddEditorActionCallback(FSchedulerCallback&& Callback) override;
+	virtual void AddPublicInterfaceRequest(FFilePlatformRequest&& Request, bool bForceFrontOfQueue) override;
+
+	virtual void ReportAccessedIniSettings(UCookOnTheFlyServer& COTFS, const FConfigFile& Config) override;
+	virtual void GetInitializeConfigSettings(UCookOnTheFlyServer& COTFS, const FString& OutputOverrideDirectory, UE::Cook::FInitializeConfigSettings& Settings) override;
+	virtual void GetBeginCookConfigSettings(UCookOnTheFlyServer& COTFS, FBeginCookContext& BeginContext, UE::Cook::FBeginCookConfigSettings& Settings) override;
+	virtual void GetBeginCookIterativeFlags(UCookOnTheFlyServer& COTFS, FBeginCookContext& BeginContext) override;
+	virtual ECookMode::Type GetDirectorCookMode(UCookOnTheFlyServer& COTFS) override;
+
+private:
+	void LogCalledCookByTheBookError(const TCHAR* FunctionName) const;
+	void LogCalledCookOnTheFlyError(const TCHAR* FunctionName) const;
+	void LogCalledPublicInterfaceError(const TCHAR* FunctionName) const;
+	void LogCalledEditorActionError(const TCHAR* FunctionName) const;
+
+private:
+	FExternalRequests ExternalRequests;
+};
+
+}
\ No newline at end of file
diff --git a/Engine/Source/Runtime/Core/Public/Serialization/PackageWriter.h b/Engine/Source/Runtime/Core/Public/Serialization/PackageWriter.h
--- a/Engine/Source/Runtime/Core/Public/Serialization/PackageWriter.h
+++ b/Engine/Source/Runtime/Core/Public/Serialization/PackageWriter.h
@@ -230,6 +230,7 @@ public:
 		ECookMode CookMode = CookByTheBookMode;
 		bool bFullBuild = true;
 		bool bIterateSharedBuild = false;
+		bool bWorkerOnSharedSandbox = false;
 	};
 
 	/** Delete outdated cooked data, etc.
