Cooker: ConcurrentSave: Fix crash due to unguraded multithreaded access to AllPackageHashes.
#rnx

diff --git a/Engine/Source/Editor/UnrealEd/Private/CookOnTheFlyServer.cpp b/Engine/Source/Editor/UnrealEd/Private/CookOnTheFlyServer.cpp
--- a/Engine/Source/Editor/UnrealEd/Private/CookOnTheFlyServer.cpp
+++ b/Engine/Source/Editor/UnrealEd/Private/CookOnTheFlyServer.cpp
@@ -10226,10 +10226,11 @@ uint32 UCookOnTheFlyServer::CookFullLoadAndSave()
 		{
 			Generators.Add(PlatformManager->GetPlatformData(Target)->RegistryGenerator.Get());
 		}
+		FCriticalSection AssetGeneratorLock;
 
 		int64 ParallelSavedPackages = 0;
 		ParallelFor(PackagesToSave.Num(),
-			[this, &PackagesToSave, &TargetPlatforms, &Generators, &ParallelSavedPackages, SaveFlags, bSaveConcurrent](int32 PackageIdx)
+			[this, &PackagesToSave, &TargetPlatforms, &Generators, &ParallelSavedPackages, SaveFlags, bSaveConcurrent, &AssetGeneratorLock](int32 PackageIdx)
 		{
 			UE::Cook::FPackageData& PackageData = *PackagesToSave[PackageIdx];
 			UPackage* Package = PackageData.GetPackage();
@@ -10355,8 +10356,12 @@ uint32 UCookOnTheFlyServer::CookFullLoadAndSave()
 						}
 
 						// Update asset registry
-						Generator.UpdateAssetRegistryPackageData(*Package, SaveResult, MoveTemp(*CookContext.GetCookTagList()));
-						FAssetPackageData* AssetPackageData = Generator.GetAssetPackageData(Package->GetFName());
+						FAssetPackageData* AssetPackageData = nullptr;
+						{
+							FScopeLock AssetGeneratorScopeLock(&AssetGeneratorLock);
+							Generator.UpdateAssetRegistryPackageData(*Package, SaveResult, MoveTemp(*CookContext.GetCookTagList()));
+							AssetPackageData = Generator.GetAssetPackageData(Package->GetFName());
+						}
 						check(AssetPackageData);
 
 						ICookedPackageWriter::FCommitPackageInfo CommitInfo;
diff --git a/Engine/Source/Editor/UnrealEd/Private/Cooker/LooseCookedPackageWriter.cpp b/Engine/Source/Editor/UnrealEd/Private/Cooker/LooseCookedPackageWriter.cpp
--- a/Engine/Source/Editor/UnrealEd/Private/Cooker/LooseCookedPackageWriter.cpp
+++ b/Engine/Source/Editor/UnrealEd/Private/Cooker/LooseCookedPackageWriter.cpp
@@ -271,14 +271,19 @@ void FLooseCookedPackageWriter::AsyncSaveOutputFiles(FRecord& Record, FCommitCon
 	{
 		ThisPackageHashes = new FPackageHashes();
 
-		// This looks weird but we're finding the _refcount_, not the hashes. So if it gets
-		// constructed, it's not actually assigned a pointer.
-		TRefCountPtr<FPackageHashes>& ExistingPackageHashes = AllPackageHashes.FindOrAdd(Context.Info.PackageName);
-		if (ExistingPackageHashes.IsValid())
+		bool bAlreadyExisted = false;
+		{
+			FScopeLock ConcurrentSaveScopeLock(&ConcurrentSaveLock);
+			TRefCountPtr<FPackageHashes>& ExistingPackageHashes = AllPackageHashes.FindOrAdd(Context.Info.PackageName);
+			// This calculation of bAlreadyExisted looks weird but we're finding the _refcount_, not the hashes. So if it gets
+			// constructed, it's not actually assigned a pointer.
+			bAlreadyExisted = ExistingPackageHashes.IsValid();
+			ExistingPackageHashes = ThisPackageHashes;
+		}
+		if (bAlreadyExisted)
 		{
 			UE_LOG(LogSavePackage, Error, TEXT("FLooseCookedPackageWriter encountered the same package twice in a cook! (%s)"), *Context.Info.PackageName.ToString());
 		}
-		ExistingPackageHashes = ThisPackageHashes;
 	}
 
 	UE::Tasks::Launch(TEXT("HashAndWriteLooseCookedFile"), [OutputFiles = MoveTemp(Context.OutputFiles), WriteOptions = Context.Info.WriteOptions, ThisPackageHashes = MoveTemp(ThisPackageHashes)]()
diff --git a/Engine/Source/Editor/UnrealEd/Private/Cooker/LooseCookedPackageWriter.h b/Engine/Source/Editor/UnrealEd/Private/Cooker/LooseCookedPackageWriter.h
--- a/Engine/Source/Editor/UnrealEd/Private/Cooker/LooseCookedPackageWriter.h
+++ b/Engine/Source/Editor/UnrealEd/Private/Cooker/LooseCookedPackageWriter.h
@@ -127,6 +127,7 @@ private:
 	TMap<FName, TRefCountPtr<FPackageHashes>> AllPackageHashes;
 
 	TMap<FName, FName> UncookedPathToCookedPath;
+	FCriticalSection ConcurrentSaveLock;
 	FString OutputPath;
 	FString MetadataDirectoryPath;
 	const ITargetPlatform& TargetPlatform;
