UE-153947 PIE needs setting to specify which of multiple instances get the VR headset.
-Added PrimaryPIEClientIndex setting to Editor Preferences->Level Editor->Play->Multiplayer Options->Client.
-Defaults to 0 which is the first pie client window (which means it could be the server in listen-server mode where the first window is both the server and a client).  -1 means no primary at all.  If it is larger than the number of indices we primary the last client.  Other negative settings are disallowed by the UI, but would also result in no primary.
-In VRPreview the primary client will track the hmd, render in stereo, etc.  The primary client's window size and position are also handled slightly differently.
-This could be hooked up to more things fairly easily.  VR provides a number of examples of how to do it, because it is a fairly extensive feature.  Much of the time the bool cached in the WorldContext is the best option.

diff --git a/Engine/Source/Editor/UnrealEd/Classes/Editor/EditorEngine.h b/Engine/Source/Editor/UnrealEd/Classes/Editor/EditorEngine.h
--- a/Engine/Source/Editor/UnrealEd/Classes/Editor/EditorEngine.h
+++ b/Engine/Source/Editor/UnrealEd/Classes/Editor/EditorEngine.h
@@ -3102,7 +3102,7 @@ protected:
 	/** Stores the position of the window for that instance index. Doesn't save it to the CDO until the play session ends. */
 	void StoreWindowSizeAndPositionForInstanceIndex(const int32 InInstanceIndex, const FIntPoint& InSize, const FIntPoint& InPosition);
 	/** Returns the stored position at that index. */
-	void GetWindowSizeAndPositionForInstanceIndex(ULevelEditorPlaySettings& InEditorPlaySettings, const int32 InInstanceIndex, FIntPoint& OutSize, FIntPoint& OutPosition);
+	void GetWindowSizeAndPositionForInstanceIndex(ULevelEditorPlaySettings& InEditorPlaySettings, const int32 InInstanceIndex, const FWorldContext& InWorldContext, FIntPoint& OutSize, FIntPoint& OutPosition);
 
 	/** Start the queued Play Session Request. After this is called the queued play session request will be cleared. */
 	virtual void StartQueuedPlaySessionRequestImpl();
diff --git a/Engine/Source/Editor/UnrealEd/Classes/Settings/LevelEditorPlaySettings.h b/Engine/Source/Editor/UnrealEd/Classes/Settings/LevelEditorPlaySettings.h
--- a/Engine/Source/Editor/UnrealEd/Classes/Settings/LevelEditorPlaySettings.h
+++ b/Engine/Source/Editor/UnrealEd/Classes/Settings/LevelEditorPlaySettings.h
@@ -373,6 +373,10 @@
 	UPROPERTY(config, EditAnywhere, Category="Multiplayer Options|Client", meta=(ClampMin = "1", UIMin = "1", UIMax = "64"))
 	int32 PlayNumberOfClients;
 
+	/** In multiplayer PIE which client will be the 'primary'. (default = 0, the first client)*/
+	UPROPERTY(config, EditAnywhere, Category = "Multiplayer Options|Client", meta = (ClampMin = "-1", UIMin = "-1", UIMax = "64", ToolTip = "In multiplayer PIE which client will be the 'primary'. Considered most important and given a larger client window, access to unique hardware like a VirtualReality HMD, etc. Intended to help test issues that affect the second, etc client.  0 is the first client. If the setting is >= than the number of clients the last will be primary. -1 will result in no primary.  Note that this is an index only of PIE instance windows, in netmode 'Play as Client' pie instance zero is a windowless dedicated server, so setting 0 here would make the fist pie window the primary which would be PIEInstance 1, rather than 0 as in other netmodes.", DisplayName = "Primary PIE Client Index"))
+	int PrimaryPIEClientIndex = 0;
+
 	/** What port used by the server for simple networking */
 	UPROPERTY(config, EditAnywhere, Category = "Multiplayer Options|Server", meta=(ClampMin="1", UIMin="1", ClampMax="65535", EditCondition = "PlayNetMode != EPlayNetMode::PIE_Standalone || bLaunchSeparateServer"))
 	uint16 ServerPort;
@@ -498,6 +502,8 @@
 	bool IsPlayNumberOfClientsActive() const { return true; }
 	bool GetPlayNumberOfClients( int32 &OutPlayNumberOfClients ) const { OutPlayNumberOfClients = PlayNumberOfClients; return IsPlayNumberOfClientsActive(); }
 
+	int GetPrimaryPIEClientIndex() const { return PrimaryPIEClientIndex; }
+
 	void SetServerPort(const uint16 InServerPort) { ServerPort = InServerPort; }
 	bool IsServerPortActive() const { return (PlayNetMode != PIE_Standalone) || RunUnderOneProcess; }
 	bool GetServerPort(uint16 &OutServerPort) const { OutServerPort = ServerPort; return IsServerPortActive(); }
diff --git a/Engine/Source/Editor/UnrealEd/Private/PlayLevel.cpp b/Engine/Source/Editor/UnrealEd/Private/PlayLevel.cpp
--- a/Engine/Source/Editor/UnrealEd/Private/PlayLevel.cpp
+++ b/Engine/Source/Editor/UnrealEd/Private/PlayLevel.cpp
@@ -126,7 +126,8 @@ DEFINE_LOG_CATEGORY(LogPlayLevel);
 const static FName NAME_CategoryPIE("PIE");
 
 // Forward declare local utility functions
-FText GeneratePIEViewportWindowTitle(const EPlayNetMode InNetMode, const ERHIFeatureLevel::Type InFeatureLevel, const FRequestPlaySessionParams& InSessionParams, const int32 ClientIndex, const float FixedTick);
+FText GeneratePIEViewportWindowTitle(const EPlayNetMode InNetMode, const ERHIFeatureLevel::Type InFeatureLevel, const FRequestPlaySessionParams& InSessionParams, const int32 ClientIndex, const int32 InViewportIndex, const float FixedTick, const bool bVRPreview);
+bool IsPrimaryPIEClient(const FRequestPlaySessionParams& InPlaySessionParams, const int32 InClientIndex);
 
 // This class listens to output log messages, and forwards warnings and errors to the message log
 class FOutputLogErrorsToMessageLogProxy : public FOutputDevice
@@ -556,11 +557,11 @@ void UEditorEngine::EndPlayMap()
 				PlaySettingsConfig->MultipleInstancePositions.Add(PlayInEditorSessionInfo->CachedWindowInfo[WindowIndex].Position);
 			}
 
-			// Update the position where the first PIE window will be opened (this also updates its displayed value in "Editor Preferences" --> "Level Editor" --> "Play" --> "New Window Position")
-			if (WindowIndex == 0)
+			// Update the position where the primary PIE window will be opened (this also updates its displayed value in "Editor Preferences" --> "Level Editor" --> "Play" --> "New Window Position")
+			if (IsPrimaryPIEClient(PlayInEditorSessionInfo->OriginalRequestParams, WindowIndex))
 			{
 				// Remember last known size
-				PlaySettingsConfig->LastSize = PlayInEditorSessionInfo->CachedWindowInfo[0].Size;
+				PlaySettingsConfig->LastSize = PlayInEditorSessionInfo->CachedWindowInfo[WindowIndex].Size;
 
 				// Only update it if "Always center window to screen" is disabled, and the size was not 0 (which means it is attached to the editor rather than being an standalone window)
 				if (!PlaySettingsConfig->CenterNewWindow && PlaySettingsConfig->LastSize.X > 0 && PlaySettingsConfig->LastSize.Y > 0)
@@ -737,7 +738,7 @@ void UEditorEngine::TeardownPlaySession(FWorldContext& PieWorldContext)
 			{
 				if (FWorldContext* WorldContext = GEngine->GetWorldContextFromWorld(PlayWorld))
 				{
-					if (WorldContext->PIEInstance == 0)
+					if (WorldContext->bIsPrimaryPIEInstance)
 					{
 						if (FAudioDeviceManager* DeviceManager = AudioDevice->GetAudioDeviceManager())
 						{
@@ -1821,6 +1822,7 @@ void UEditorEngine::CreateNewPlayInEditorInstance(FRequestPlaySessionParams &InR
 		GameInstancePIEParameters.WorldFeatureLevel = PreviewPlatform.GetEffectivePreviewFeatureLevel();
 		GameInstancePIEParameters.NetMode = InNetMode;
 		GameInstancePIEParameters.OverrideMapURL = InRequestParams.GlobalMapOverride;
+		GameInstancePIEParameters.bIsPrimaryPIEClient = !bInDedicatedInstance && IsPrimaryPIEClient(InRequestParams, PlayInEditorSessionInfo->NumClientInstancesCreated);
 
 		PIELoginInfo.GameInstancePIEParameters = GameInstancePIEParameters;
 
@@ -3182,7 +3184,7 @@ UGameInstance* UEditorEngine::CreateInnerProcessPIEGameInstance(FRequestPlaySess
 	return GameInstance;
 }
 
-FText GeneratePIEViewportWindowTitle(const EPlayNetMode InNetMode, const ERHIFeatureLevel::Type InFeatureLevel, const FRequestPlaySessionParams& InSessionParams, const int32 ClientIndex, const float FixedTick)
+FText GeneratePIEViewportWindowTitle(const EPlayNetMode InNetMode, const ERHIFeatureLevel::Type InFeatureLevel, const FRequestPlaySessionParams& InSessionParams, const int32 ClientIndex, const float FixedTick, const bool bVRPreview)
 {
 #if PLATFORM_64BITS
 	const FString PlatformBitsString(TEXT("64"));
@@ -3203,15 +3205,14 @@ FText GeneratePIEViewportWindowTitle(const EPlayNetMode InNetMode, const ERHIFea
 	}
 	else if (InNetMode == PIE_ListenServer)
 	{
-		Args.Add(TEXT("NetMode"), FText::FromString(TEXT("Server")));
+		Args.Add(TEXT("NetMode"), FText::FromString(FString::Printf(TEXT("Server %d"), ClientIndex)));
 	}
 	else
 	{
-		Args.Add(TEXT("NetMode"), FText::FromString(TEXT("Standalone")));
+		Args.Add(TEXT("NetMode"), FText::FromString(FString::Printf(TEXT("Standalone %d"), ClientIndex)));
 	}
 
-	if (GEngine->StereoRenderingDevice && GEngine->StereoRenderingDevice.IsValid() && InSessionParams.SessionPreviewTypeOverride == EPlaySessionPreviewType::VRPreview &&
-		GEngine->XRSystem && GEngine->XRSystem.IsValid())
+	if (bVRPreview)
 	{
 		Args.Add(TEXT("XRSystemName"), FText::FromName(GEngine->XRSystem->GetSystemName()));
 		Args.Add(TEXT("XRRuntimeVersion"), FText::FromString(GEngine->XRSystem->GetVersionString()));
@@ -3280,11 +3281,11 @@ void UEditorEngine::TransferEditorSelectionToPlayInstances(const bool bInSelectI
 TSharedRef<SPIEViewport> UEditorEngine::GeneratePIEViewportWindow(const FRequestPlaySessionParams& InSessionParams, int32 InViewportIndex, const FWorldContext& InWorldContext, EPlayNetMode InNetMode, UGameViewportClient* InViewportClient, FSlatePlayInEditorInfo& InSlateInfo)
 {
 	FIntPoint WindowSize, WindowPosition;
-	GetWindowSizeAndPositionForInstanceIndex(*InSessionParams.EditorPlaySettings, InViewportIndex, WindowSize, WindowPosition);
+	GetWindowSizeAndPositionForInstanceIndex(*InSessionParams.EditorPlaySettings, InViewportIndex, InWorldContext, WindowSize, WindowPosition);
 	bool bCenterNewWindowOverride = false;
 	
 	// VR Preview overrides window location.
-	const bool bVRPreview = InSessionParams.SessionPreviewTypeOverride.Get(EPlaySessionPreviewType::NoPreview) == EPlaySessionPreviewType::VRPreview;
+	const bool bVRPreview = InWorldContext.bIsPrimaryPIEInstance && InSessionParams.SessionPreviewTypeOverride.Get(EPlaySessionPreviewType::NoPreview) == EPlaySessionPreviewType::VRPreview;
 	bool bUseOSWndBorder = bVRPreview;
 	if (bVRPreview)
 	{
@@ -3298,7 +3299,7 @@ TSharedRef<SPIEViewport> UEditorEngine::GeneratePIEViewportWindow(const FRequest
 	// If they haven't provided a Slate Window (common), we will create one.
 	if (!bHasCustomWindow)
 	{
-		FText ViewportName = GeneratePIEViewportWindowTitle(InNetMode, PreviewPlatform.GetEffectivePreviewFeatureLevel(), InSessionParams, InWorldContext.PIEInstance, InWorldContext.PIEFixedTickSeconds);
+		FText ViewportName = GeneratePIEViewportWindowTitle(InNetMode, PreviewPlatform.GetEffectivePreviewFeatureLevel(), InSessionParams, InWorldContext.PIEInstance, InWorldContext.PIEFixedTickSeconds, bVRPreview);
 		PieWindow = SNew(SWindow)
 			.Title(ViewportName)
 			.Tag("PIEWindow")
@@ -3345,7 +3346,7 @@ TSharedRef<SPIEViewport> UEditorEngine::GeneratePIEViewportWindow(const FRequest
 
 
 	bool bRenderDirectlyToWindow = bVRPreview;
-	bool bEnableStereoRendering = bVRPreview && (/* only first PIE instance can be VR */ InViewportIndex == 0);
+	bool bEnableStereoRendering = bVRPreview;
 
 	static const auto CVarPropagateAlpha = IConsoleManager::Get().FindTConsoleVariableDataInt(TEXT("r.PostProcessing.PropagateAlpha"));
 	const EAlphaChannelMode::Type PropagateAlpha = EAlphaChannelMode::FromInt(CVarPropagateAlpha->GetValueOnGameThread());
@@ -3430,11 +3431,11 @@ TSharedRef<SPIEViewport> UEditorEngine::GeneratePIEViewportWindow(const FRequest
 			}
 		};
 
-		FDelegateHandle PreviewFeatureLevelChangedHandle = OnPreviewFeatureLevelChanged().AddLambda([PieWindow, bHasCustomWindow, InViewportClient, InNetMode, InSessionParams, InWorldContext](ERHIFeatureLevel::Type NewFeatureLevel)
+		FDelegateHandle PreviewFeatureLevelChangedHandle = OnPreviewFeatureLevelChanged().AddLambda([PieWindow, bHasCustomWindow, InViewportClient, InNetMode, InSessionParams, InWorldContext, bVRPreview](ERHIFeatureLevel::Type NewFeatureLevel)
 			{
 				if (!bHasCustomWindow)
 				{
-					FText ViewportName = GeneratePIEViewportWindowTitle(InNetMode, NewFeatureLevel, InSessionParams, InWorldContext.PIEInstance, InWorldContext.PIEFixedTickSeconds);
+					FText ViewportName = GeneratePIEViewportWindowTitle(InNetMode, NewFeatureLevel, InSessionParams, InWorldContext.PIEInstance, InWorldContext.PIEFixedTickSeconds, bVRPreview);
 					PieWindow->SetTitle(ViewportName);
 				}
 				InViewportClient->GetWorld()->ChangeFeatureLevel(NewFeatureLevel);
@@ -3539,7 +3540,35 @@ void FitWindowPositionToWorkArea(FIntPoint &WinPos, FIntPoint &WinSize, const FM
 	WinPos.Y = FMath::Clamp(WinPos.Y, (int32)PreferredWorkArea.Top, (int32)PreferredWorkArea.Bottom);
 }
 
-void UEditorEngine::GetWindowSizeAndPositionForInstanceIndex(ULevelEditorPlaySettings& InEditorPlaySettings, const int32 InInstanceIndex, FIntPoint& OutSize, FIntPoint& OutPosition)
+bool IsPrimaryPIEClient(const FRequestPlaySessionParams& InPlaySessionParams, const int32 InClientIndex)
+{
+	// Note: the InClientIndex here is the index of the pie clinet instances, meaning instances with pie windows, created not the pie instance index. These may differ, for example it is possible that pie instance 0 is a windowless dedicated server in netmode 'play as client'.
+
+	const ULevelEditorPlaySettings& EditorPlaySettings = InPlaySessionParams.EditorPlaySettings ? *InPlaySessionParams.EditorPlaySettings : *GetDefault<ULevelEditorPlaySettings>();
+
+	if (InClientIndex == EditorPlaySettings.GetPrimaryPIEClientIndex())
+	{
+		return true;
+	}
+
+	int32 ClientCount;
+	if (!EditorPlaySettings.GetPlayNumberOfClients(ClientCount))
+	{
+		return InClientIndex == 0;  // If we aren't doing 'number of clients' just primary the first client.
+	}
+
+	if ((EditorPlaySettings.GetPrimaryPIEClientIndex() >= ClientCount) && (InClientIndex == ClientCount - 1))
+	{
+		// If the number is set too high use the last client.  This could easily happen if the user reduces the number of clients and forgets to update this setting. We won't assume they are intentionally avoiding having a 'primary' in this way.
+		return true;
+	}
+
+	// Note any negative value returned by GetPrimaryPIEClientIndex() would result in no primary client at all.
+	return false;
+
+}
+
+void UEditorEngine::GetWindowSizeAndPositionForInstanceIndex(ULevelEditorPlaySettings& InEditorPlaySettings, const int32 InViewportIndex, const FWorldContext& InWorldContext, FIntPoint& OutSize, FIntPoint& OutPosition)
 {
 	if (!ensureMsgf(PlayInEditorSessionInfo.IsSet(), TEXT("Cannot get saved Window Size/Position if a session has not been started.")))
 	{
@@ -3557,8 +3586,8 @@ void UEditorEngine::GetWindowSizeAndPositionForInstanceIndex(ULevelEditorPlaySet
 	}
 
 	// Alright, they don't have a saved position or don't want to load from the saved position. First, figure out
-	// how big the window should be. If it is the first client, it uses a different resolution source than additional.
-	if (InInstanceIndex == 0)
+	// how big the window should be. If it is the primary client, it uses a different resolution source than additional.
+	if (InWorldContext.bIsPrimaryPIEInstance)
 	{
 		OutSize = FIntPoint(InEditorPlaySettings.NewWindowWidth, InEditorPlaySettings.NewWindowHeight);
 	}
@@ -3586,7 +3615,7 @@ void UEditorEngine::GetWindowSizeAndPositionForInstanceIndex(ULevelEditorPlaySet
 
 
 	// Now we can position the window. If it is the first window, we can respect the center window flag.
-	if (InInstanceIndex == 0)
+	if (InWorldContext.bIsPrimaryPIEInstance)
 	{
 		// Center window if CenterNewWindow checked or if NewWindowPosition is FIntPoint::NoneValue (-1,-1)
 		if (InEditorPlaySettings.CenterNewWindow || InEditorPlaySettings.NewWindowPosition == FIntPoint::NoneValue)
@@ -3603,9 +3632,9 @@ void UEditorEngine::GetWindowSizeAndPositionForInstanceIndex(ULevelEditorPlaySet
 	}
 	else
 	{
-		if (InInstanceIndex < PlayInEditorSessionInfo->CachedWindowInfo.Num())
+		if (InViewportIndex < PlayInEditorSessionInfo->CachedWindowInfo.Num())
 		{
-			OutPosition = PlayInEditorSessionInfo->CachedWindowInfo[InInstanceIndex].Position;
+			OutPosition = PlayInEditorSessionInfo->CachedWindowInfo[InViewportIndex].Position;
 			FitWindowPositionToWorkArea(OutPosition, OutSize, WindowBorderSize);
 		}
 		// Add a new entry.
@@ -3623,21 +3652,22 @@ void UEditorEngine::GetWindowSizeAndPositionForInstanceIndex(ULevelEditorPlaySet
 		}
 	}
 	// Store this Size/Position for this duration
-	StoreWindowSizeAndPositionForInstanceIndex(InInstanceIndex, OutSize, OutPosition);
+	StoreWindowSizeAndPositionForInstanceIndex(InViewportIndex, OutSize, OutPosition);
 }
-void UEditorEngine::StoreWindowSizeAndPositionForInstanceIndex(const int32 InInstanceIndex, const FIntPoint& InSize, const FIntPoint& InPosition)
+void UEditorEngine::StoreWindowSizeAndPositionForInstanceIndex(const int32 InViewportIndex, const FIntPoint& InSize, const FIntPoint& InPosition)
 {
 	// Overwrite an existing one if we have it
-	if (InInstanceIndex < PlayInEditorSessionInfo->CachedWindowInfo.Num())
+	if (InViewportIndex < PlayInEditorSessionInfo->CachedWindowInfo.Num())
 	{
-		PlayInEditorSessionInfo->CachedWindowInfo[InInstanceIndex].Size = InSize;
-		PlayInEditorSessionInfo->CachedWindowInfo[InInstanceIndex].Position = InPosition;
+		PlayInEditorSessionInfo->CachedWindowInfo[InViewportIndex].Size = InSize;
+		PlayInEditorSessionInfo->CachedWindowInfo[InViewportIndex].Position = InPosition;
 	}
 	else
 	{
 		FPlayInEditorSessionInfo::FWindowSizeAndPos& NewInfo = PlayInEditorSessionInfo->CachedWindowInfo.Add_GetRef(FPlayInEditorSessionInfo::FWindowSizeAndPos());
 		NewInfo.Size = InSize;
 		NewInfo.Position = InPosition;
+		check(PlayInEditorSessionInfo->CachedWindowInfo.Num() == InViewportIndex + 1);
 	}
 }
 
diff --git a/Engine/Source/Runtime/Engine/Classes/Engine/Engine.h b/Engine/Source/Runtime/Engine/Classes/Engine/Engine.h
--- a/Engine/Source/Runtime/Engine/Classes/Engine/Engine.h
+++ b/Engine/Source/Runtime/Engine/Classes/Engine/Engine.h
@@ -372,6 +372,9 @@ struct FWorldContext
 	/** Is this world context waiting for an online login to complete (for PIE) */
 	bool	bWaitingOnOnlineSubsystem;
 
+	/** Is this the 'primary' PIE instance.  Primary is preferred when, for example, unique hardware like a VR headset can be used by only one PIE instance. */
+	bool	bIsPrimaryPIEInstance;
+
 	/** Handle to this world context's audio device.*/
 	uint32 AudioDeviceID;
 
@@ -428,6 +431,7 @@ struct FWorldContext
 		, PIEWorldFeatureLevel(ERHIFeatureLevel::Num)
 		, RunAsDedicated(false)
 		, bWaitingOnOnlineSubsystem(false)
+		, bIsPrimaryPIEInstance(false)
 		, AudioDeviceID(INDEX_NONE)
 		, ThisCurrentWorld(nullptr)
 	{ }
diff --git a/Engine/Source/Runtime/Engine/Classes/Engine/GameInstance.h b/Engine/Source/Runtime/Engine/Classes/Engine/GameInstance.h
--- a/Engine/Source/Runtime/Engine/Classes/Engine/GameInstance.h
+++ b/Engine/Source/Runtime/Engine/Classes/Engine/GameInstance.h
@@ -107,6 +107,7 @@ struct FGameInstancePIEParameters
 		, bAnyBlueprintErrors(false)
 		, bStartInSpectatorMode(false)
 		, bRunAsDedicated(false)
+		, bIsPrimaryPIEClient(false)
 		, WorldFeatureLevel(ERHIFeatureLevel::Num)
 		, EditorPlaySettings(nullptr)
 		, NetMode(EPlayNetMode::PIE_Standalone)
@@ -124,6 +125,9 @@ struct FGameInstancePIEParameters
 	// Is this a dedicated server instance for PIE?
 	bool bRunAsDedicated;
 
+	// Is this the primary PIE client?
+	bool bIsPrimaryPIEClient;
+
 	// What time did we start PIE in the editor?
 	double PIEStartTime = 0;
 
diff --git a/Engine/Source/Runtime/Engine/Private/GameInstance.cpp b/Engine/Source/Runtime/Engine/Private/GameInstance.cpp
--- a/Engine/Source/Runtime/Engine/Private/GameInstance.cpp
+++ b/Engine/Source/Runtime/Engine/Private/GameInstance.cpp
@@ -274,6 +274,8 @@ FGameInstancePIEResult UGameInstance::InitializeForPlayInEditor(int32 PIEInstanc
 
 	WorldContext->RunAsDedicated = Params.bRunAsDedicated;
 
+	WorldContext->bIsPrimaryPIEInstance = Params.bIsPrimaryPIEClient;
+
 	WorldContext->OwningGameInstance = this;
 	
 	const FString WorldPackageName = EditorEngine->EditorWorld->GetOutermost()->GetName();
diff --git a/Engine/Source/Runtime/HeadMountedDisplay/Private/IXRTrackingSystem.cpp b/Engine/Source/Runtime/HeadMountedDisplay/Private/IXRTrackingSystem.cpp
--- a/Engine/Source/Runtime/HeadMountedDisplay/Private/IXRTrackingSystem.cpp
+++ b/Engine/Source/Runtime/HeadMountedDisplay/Private/IXRTrackingSystem.cpp
@@ -27,8 +27,7 @@ void IXRTrackingSystem::GetHMDData(UObject* WorldContext, FXRHMDData& HMDData)
 bool IXRTrackingSystem::IsHeadTrackingAllowedForWorld(UWorld& World) const
 {
 #if WITH_EDITOR
-	// For VR PIE only the first instance uses the headset
-	// This implementation is constrained by hotfix rules.  It would be better to cache this somewhere.
+	// For VR PIE only the primary instance uses the headset.
 
 	if (!IsHeadTrackingAllowed())
 	{
@@ -40,16 +39,8 @@ bool IXRTrackingSystem::IsHeadTrackingAllowedForWorld(UWorld& World) const
 		return true;
 	}
 
-	// If we are a pie instance then the first pie world that is not a dedicated server uses head tracking
-	const int32 MyPIEInstanceID = World.GetOutermost()->GetPIEInstanceID();
-	for (const FWorldContext& WorldContext : GEngine->GetWorldContexts())
-	{
-		if (WorldContext.WorldType == EWorldType::PIE && WorldContext.RunAsDedicated == false && WorldContext.World())
-		{
-			return WorldContext.World()->GetOutermost()->GetPIEInstanceID() == MyPIEInstanceID;
-		}
-	}
-	return false;
+	FWorldContext* const WorldContext = GEngine->GetWorldContextFromWorld(&World);
+	return WorldContext && WorldContext->bIsPrimaryPIEInstance;
 #else
 	return IsHeadTrackingAllowed();
 #endif
