Amend naming for the native scene import.

diff --git a/Engine/Source/Editor/UnrealEd/Private/Fbx/FbxMainImport.cpp b/Engine/Source/Editor/UnrealEd/Private/Fbx/FbxMainImport.cpp
--- a/Engine/Source/Editor/UnrealEd/Private/Fbx/FbxMainImport.cpp
+++ b/Engine/Source/Editor/UnrealEd/Private/Fbx/FbxMainImport.cpp
@@ -1799,14 +1799,8 @@ bool FFbxImporter::ImportFromFile(const FString& Filename, const FString& Type,
 	return Result;
 }
 
-void FFbxImporter::ImportFromScene(FbxScene* InScene)
+void FFbxImporter::SetScene(FbxScene* InScene)
 {
-	if (!InScene)
-	{
-		UE_LOG(LogFbx, Warning, TEXT("Attempted to import invalid scene."));
-		return;
-	}
-
 	ClearAllCaches();
 	Scene = InScene;
 	FbxCreator = EFbxCreator::Unknow;
diff --git a/Engine/Source/Editor/UnrealEd/Public/FbxImporter.h b/Engine/Source/Editor/UnrealEd/Public/FbxImporter.h
--- a/Engine/Source/Editor/UnrealEd/Public/FbxImporter.h
+++ b/Engine/Source/Editor/UnrealEd/Public/FbxImporter.h
@@ -693,7 +693,7 @@ public:
 	 * @param Scene - The scene we want to load.
 	 */
 	UE_DEPRECATED(5.1, "Do not use this function.")
-	UNREALED_API void ImportFromScene(FbxScene* InScene);
+	UNREALED_API void SetScene(FbxScene* InScene);
 
 	/**
 	 * Retrieve the FBX loader's error message explaining its failure to read a given FBX file.
