Trace: Create interfaces in TraceServices for external Analysis

### Features
This change sets up new interface APIs and moves concrete APIs into those interfaces to allow for providers to be implemented external to TraceServices. These Provider interfaces are given to new construction APIs in ExternalAnalysis.h that construct the appropriate Analyzer with the specified Providers.

This change creates a new and standalone Analyzer for Bookmarks spinning logic directly out of the MiscAnalyzer. This is necessary to prevent bookmark processing necessitating additional Providers that the client may not want to implement if they just want access to Bookmark data.

SummarizeTraceCommandlet implements these provider interfaces and migrates it's state tracking implementation from raw event decoding to simply forwarding on the calls from the Provider interfaces. The internal SummarizeTraceCommandlet Analyzer API in the Commandlet are mostly left unchanged as a middleman between Trace and the Commandlet existed prior to this change.

### Implementation
Since Providers can now implement multiple interfaces, we need to upcast them to the appropriate interface at registration time. This in turn requires moving the Provider memory model into IAnalysisSession to TSharedPtr/MakeShared to prevent double deletes at destruction time.

### Performance
Opening a large trace in Insights yields this performance change:
Pre: Analyzed in 8.29s at 122.8X speed.
Post: Analyzed in 8.36s at 121.8X speed.
Commandlet runtime is comparable (I didn't measure, less sensitive than Insights), and memory pressure remains more or less steady in the low teens of gigabytes.

### Testing
A/B testing and spot checking the data shows that nearly all data is either more correct or withing many-decimal-places rounding error.
UnrealInsights and Editor builds without error in NoPCH, DisableUnity configurations

### Hashtags

diff --git a/Engine/Plugins/Animation/GameplayInsights/Source/GameplayInsights/Private/GameplayTraceModule.cpp b/Engine/Plugins/Animation/GameplayInsights/Source/GameplayInsights/Private/GameplayTraceModule.cpp
--- a/Engine/Plugins/Animation/GameplayInsights/Source/GameplayInsights/Private/GameplayTraceModule.cpp
+++ b/Engine/Plugins/Animation/GameplayInsights/Source/GameplayInsights/Private/GameplayTraceModule.cpp
@@ -16,9 +16,9 @@ void FGameplayTraceModule::GetModuleInfo(TraceServices::FModuleInfo& OutModuleIn
 
 void FGameplayTraceModule::OnAnalysisBegin(TraceServices::IAnalysisSession& InSession)
 {
-	FGameplayProvider* GameplayProvider = new FGameplayProvider(InSession);
+	TSharedPtr<FGameplayProvider> GameplayProvider = MakeShared<FGameplayProvider>(InSession);
 	InSession.AddProvider(FGameplayProvider::ProviderName, GameplayProvider);
-	FAnimationProvider* AnimationProvider = new FAnimationProvider(InSession, *GameplayProvider);
+	TSharedPtr<FAnimationProvider> AnimationProvider = MakeShared<FAnimationProvider>(InSession, *GameplayProvider);
 	InSession.AddProvider(FAnimationProvider::ProviderName, AnimationProvider);
 
 	InSession.AddAnalyzer(new FAnimationAnalyzer(InSession, *AnimationProvider));
diff --git a/Engine/Plugins/Animation/GameplayInsights/Source/RewindDebuggerVLog/Private/VLogTraceModule.cpp b/Engine/Plugins/Animation/GameplayInsights/Source/RewindDebuggerVLog/Private/VLogTraceModule.cpp
--- a/Engine/Plugins/Animation/GameplayInsights/Source/RewindDebuggerVLog/Private/VLogTraceModule.cpp
+++ b/Engine/Plugins/Animation/GameplayInsights/Source/RewindDebuggerVLog/Private/VLogTraceModule.cpp
@@ -14,7 +14,7 @@ void FVLogTraceModule::GetModuleInfo(TraceServices::FModuleInfo& OutModuleInfo)
 
 void FVLogTraceModule::OnAnalysisBegin(TraceServices::IAnalysisSession& InSession)
 {
-	FVisualLoggerProvider* VisualLoggerProvider = new FVisualLoggerProvider(InSession);
+	TSharedPtr<FVisualLoggerProvider> VisualLoggerProvider = MakeShared<FVisualLoggerProvider>(InSession);
 	InSession.AddProvider(FVisualLoggerProvider::ProviderName, VisualLoggerProvider);
 
 	InSession.AddAnalyzer(new FVisualLoggerAnalyzer(InSession, *VisualLoggerProvider));
diff --git a/Engine/Plugins/Experimental/Animation/PoseSearch/Source/Editor/Private/Trace/PoseSearchTraceModule.cpp b/Engine/Plugins/Experimental/Animation/PoseSearch/Source/Editor/Private/Trace/PoseSearchTraceModule.cpp
--- a/Engine/Plugins/Experimental/Animation/PoseSearch/Source/Editor/Private/Trace/PoseSearchTraceModule.cpp
+++ b/Engine/Plugins/Experimental/Animation/PoseSearch/Source/Editor/Private/Trace/PoseSearchTraceModule.cpp
@@ -17,7 +17,7 @@ void FTraceModule::GetModuleInfo(TraceServices::FModuleInfo& OutModuleInfo)
 void FTraceModule::OnAnalysisBegin(TraceServices::IAnalysisSession& InSession)
 {
 	// Add our provider and analyzer, starting our systems
-	FTraceProvider* PoseSearchProvider = new FTraceProvider(InSession);
+	TSharedPtr<FTraceProvider> PoseSearchProvider = MakeShared<FTraceProvider>(InSession);
 	InSession.AddProvider(FTraceProvider::ProviderName, PoseSearchProvider);
 	InSession.AddAnalyzer(new FTraceAnalyzer(InSession, *PoseSearchProvider));
 }
diff --git a/Engine/Plugins/RenderGraphInsights/Source/RenderGraphInsights/Private/RenderGraphTraceModule.cpp b/Engine/Plugins/RenderGraphInsights/Source/RenderGraphInsights/Private/RenderGraphTraceModule.cpp
--- a/Engine/Plugins/RenderGraphInsights/Source/RenderGraphInsights/Private/RenderGraphTraceModule.cpp
+++ b/Engine/Plugins/RenderGraphInsights/Source/RenderGraphInsights/Private/RenderGraphTraceModule.cpp
@@ -19,7 +19,7 @@ void FRenderGraphTraceModule::GetModuleInfo(TraceServices::FModuleInfo& OutModul
 
 void FRenderGraphTraceModule::OnAnalysisBegin(TraceServices::IAnalysisSession& InSession)
 {
-	FRenderGraphProvider* RenderGraphProvider = new FRenderGraphProvider(InSession);
+	TSharedPtr<FRenderGraphProvider> RenderGraphProvider = MakeShared<FRenderGraphProvider>(InSession);
 	InSession.AddProvider(FRenderGraphProvider::ProviderName, RenderGraphProvider);
 
 	InSession.AddAnalyzer(new FRenderGraphAnalyzer(InSession, *RenderGraphProvider));
diff --git a/Engine/Plugins/Runtime/NetworkPredictionInsights/Source/NetworkPredictionInsights/Private/NetworkPredictionTraceModule.cpp b/Engine/Plugins/Runtime/NetworkPredictionInsights/Source/NetworkPredictionInsights/Private/NetworkPredictionTraceModule.cpp
--- a/Engine/Plugins/Runtime/NetworkPredictionInsights/Source/NetworkPredictionInsights/Private/NetworkPredictionTraceModule.cpp
+++ b/Engine/Plugins/Runtime/NetworkPredictionInsights/Source/NetworkPredictionInsights/Private/NetworkPredictionTraceModule.cpp
@@ -14,7 +14,7 @@ void FNetworkPredictionTraceModule::GetModuleInfo(TraceServices::FModuleInfo& Ou
 
 void FNetworkPredictionTraceModule::OnAnalysisBegin(TraceServices::IAnalysisSession& InSession)
 {
-	FNetworkPredictionProvider* NetworkPredictionProvider = new FNetworkPredictionProvider(InSession);
+	TSharedPtr<FNetworkPredictionProvider> NetworkPredictionProvider = MakeShared<FNetworkPredictionProvider>(InSession);
 	InSession.AddProvider(FNetworkPredictionProvider::ProviderName, NetworkPredictionProvider);
 	
 	InSession.AddAnalyzer(new FNetworkPredictionAnalyzer(InSession, *NetworkPredictionProvider));
diff --git a/Engine/Plugins/Slate/SlateInsights/Source/SlateInsights/Private/SlateTraceModule.cpp b/Engine/Plugins/Slate/SlateInsights/Source/SlateInsights/Private/SlateTraceModule.cpp
--- a/Engine/Plugins/Slate/SlateInsights/Source/SlateInsights/Private/SlateTraceModule.cpp
+++ b/Engine/Plugins/Slate/SlateInsights/Source/SlateInsights/Private/SlateTraceModule.cpp
@@ -19,7 +19,7 @@ void FSlateTraceModule::GetModuleInfo(TraceServices::FModuleInfo& OutModuleInfo)
 
 void FSlateTraceModule::OnAnalysisBegin(TraceServices::IAnalysisSession& InSession)
 {
-	FSlateProvider* SlateProvider = new FSlateProvider(InSession);
+	TSharedPtr<FSlateProvider> SlateProvider = MakeShared<FSlateProvider>(InSession);
 	InSession.AddProvider(FSlateProvider::ProviderName, SlateProvider);
 
 	InSession.AddAnalyzer(new FSlateAnalyzer(InSession, *SlateProvider));
diff --git a/Engine/Source/Developer/TraceAnalysis/Private/DataStream.cpp b/Engine/Source/Developer/TraceAnalysis/Private/DataStream.cpp
new file mode 100644
--- /dev/null
+++ b/Engine/Source/Developer/TraceAnalysis/Private/DataStream.cpp
@@ -0,0 +1,60 @@
+// Copyright Epic Games, Inc. All Rights Reserved.
+
+#include "Trace/DataStream.h"
+#include "GenericPlatform/GenericPlatformFile.h"
+#include "HAL/PlatformFileManager.h"
+
+namespace UE {
+namespace Trace {
+
+FFileDataStream::FFileDataStream()
+	: Handle(nullptr)
+	, Remaining(0)
+{
+}
+
+FFileDataStream::~FFileDataStream()
+{
+}
+
+bool FFileDataStream::Open(const TCHAR* Path)
+{
+	Handle.Reset(FPlatformFileManager::Get().GetPlatformFile().OpenRead(Path));
+	if (Handle == nullptr)
+	{
+		return false;
+	}
+	Remaining = Handle->Size();
+	return true;
+}
+
+int32 FFileDataStream::Read(void* Data, uint32 Size)
+{
+	if (Handle == nullptr)
+	{
+		return -1;
+	}
+
+	if (Remaining <= 0)
+	{
+		return 0;
+	}
+
+	Size = (Size < Remaining) ? Size : Remaining;
+	Remaining -= Size;
+	int32 Result = Handle->Read((uint8*)Data, Size) ? Size : -1;
+	if (Result < 0)
+	{
+		Close();
+	}
+
+	return Result;
+}
+
+void FFileDataStream::Close()
+{
+	Handle.Reset();
+}
+
+} // namespace Trace
+} // namespace UE
diff --git a/Engine/Source/Developer/TraceAnalysis/Public/Trace/DataStream.h b/Engine/Source/Developer/TraceAnalysis/Public/Trace/DataStream.h
--- a/Engine/Source/Developer/TraceAnalysis/Public/Trace/DataStream.h
+++ b/Engine/Source/Developer/TraceAnalysis/Public/Trace/DataStream.h
@@ -2,7 +2,8 @@
 
 #pragma once
 
-#include "CoreTypes.h"
+#include "CoreFwd.h"
+#include "Templates/UniquePtr.h"
 
 namespace UE {
 namespace Trace {
@@ -15,5 +16,43 @@ public:
 	virtual void	Close() {}
 };
 
+/*
+* An implementation of IInDataStream that reads from a file on disk.
+*/
+class TRACEANALYSIS_API FFileDataStream : public IInDataStream
+{
+public:
+	FFileDataStream();
+	~FFileDataStream();
+
+	/*
+	* Open the file.
+	* 
+	* @param Path	The path to the file.
+	* 
+	* @return True if the file was opened successfully.
+	*/
+	bool Open(const TCHAR* Path);
+
+	/*
+	* Read from the file.
+	* 
+	* @param Data	The address in memory to write file data to.
+	* @param Size	The size of available memory at Data.
+	* 
+	* @return The number of bytes read from the file to Data, or zero in the case of a read failure.
+	*/
+	virtual int32 Read(void* Data, uint32 Size) override;
+
+	/*
+	* Close the file.
+	*/
+	virtual void Close() override;
+
+private:
+	TUniquePtr<IFileHandle> Handle;
+	uint64 Remaining;
+};
+
 } // namespace Trace
 } // namespace UE
diff --git a/Engine/Source/Developer/TraceServices/Private/AnalysisService.cpp b/Engine/Source/Developer/TraceServices/Private/AnalysisService.cpp
--- a/Engine/Source/Developer/TraceServices/Private/AnalysisService.cpp
+++ b/Engine/Source/Developer/TraceServices/Private/AnalysisService.cpp
@@ -4,6 +4,7 @@
 #include "AnalysisServicePrivate.h"
 #include "ModuleServicePrivate.h"
 #include "Analyzers/LogTraceAnalysis.h"
+#include "Analyzers/BookmarksTraceAnalysis.h"
 #include "Analyzers/MiscTraceAnalysis.h"
 #include "Analyzers/StringsAnalyzer.h"
 #include "HAL/PlatformFile.h"
@@ -24,6 +25,10 @@
 namespace TraceServices
 {
 
+// if IProvider ever gets member data, it will duplicate state because of "the diamond problem" with multiple inheritance
+// if you hit this you should consider the implications for classes that implement multiple providers
+static_assert(sizeof(IProvider) == sizeof(uintptr_t));
+
 thread_local FAnalysisSessionLock* GThreadCurrentSessionLock;
 thread_local int32 GThreadCurrentReadLockCount;
 thread_local int32 GThreadCurrentWriteLockCount;
@@ -97,10 +102,6 @@ FAnalysisSession::~FAnalysisSession()
 	{
 		delete Analyzers[AnalyzerIndex];
 	}
-	for (int32 ProviderIndex = Providers.Num() - 1; ProviderIndex >= 0; --ProviderIndex)
-	{
-		delete Providers[ProviderIndex];
-	}
 }
 
 void FAnalysisSession::Start()
@@ -133,18 +134,17 @@ void FAnalysisSession::AddAnalyzer(UE::Trace::IAnalyzer* Analyzer)
 	Analyzers.Add(Analyzer);
 }
 
-void FAnalysisSession::AddProvider(const FName& InName, IProvider* Provider)
+void FAnalysisSession::AddProvider(const FName& InName, TSharedPtr<IProvider> Provider, TSharedPtr<IEditableProvider> EditableProvider)
 {
-	Providers.Add(Provider);
-	ProvidersMap.Add(InName, Provider);
+	Providers.Add(InName, MakeTuple(Provider, EditableProvider));
 }
 
 const IProvider* FAnalysisSession::ReadProviderPrivate(const FName& InName) const
 {
-	IProvider* const* FindIt = ProvidersMap.Find(InName);
+	const auto* FindIt = Providers.Find(InName);
 	if (FindIt)
 	{
-		return *FindIt;
+		return FindIt->Key.Get();
 	}
 	else
 	{
@@ -154,10 +154,10 @@ const IProvider* FAnalysisSession::ReadProviderPrivate(const FName& InName) cons
 
 IProvider* FAnalysisSession::EditProviderPrivate(const FName& InName)
 {
-	IProvider** FindIt = ProvidersMap.Find(InName);
+	const auto* FindIt = Providers.Find(InName);
 	if (FindIt)
 	{
-		return *FindIt;
+		return FindIt->Value.Get();
 	}
 	else
 	{
@@ -223,33 +223,38 @@ TSharedPtr<const IAnalysisSession> FAnalysisService::StartAnalysis(uint32 TraceI
 
 	FAnalysisSessionEditScope _(*Session);
 
-	FBookmarkProvider* BookmarkProvider = new FBookmarkProvider(*Session);
-	Session->AddProvider(FBookmarkProvider::ProviderName, BookmarkProvider);
+	// Note that we upcast to the interface we wish to register in cases where a provider implements multiple interfaces
+	// We aren't using virtual inheritance since we don't care about the vtable complexity, and there is no state to worry about
+	// We follow the single data inheritance but multiple interface inheritance idiom, and this necessitates the upcast via TSharedPtr<IInterface>
+
+	TSharedPtr<FBookmarkProvider> BookmarkProvider = MakeShared<FBookmarkProvider>(*Session);
+	Session->AddProvider(FBookmarkProvider::ProviderName, TSharedPtr<IBookmarkProvider>(BookmarkProvider), TSharedPtr<IEditableBookmarkProvider>(BookmarkProvider));
 
-	FLogProvider* LogProvider = new FLogProvider(*Session);
-	Session->AddProvider(FLogProvider::ProviderName, LogProvider);
+	TSharedPtr<FLogProvider> LogProvider = MakeShared<FLogProvider>(*Session);
+	Session->AddProvider(FLogProvider::ProviderName, TSharedPtr<ILogProvider>(LogProvider), TSharedPtr<IEditableLogProvider>(LogProvider));
 
-	FThreadProvider* ThreadProvider = new FThreadProvider(*Session);
-	Session->AddProvider(FThreadProvider::ProviderName, ThreadProvider);
+	TSharedPtr<FThreadProvider> ThreadProvider = MakeShared<FThreadProvider>(*Session);
+	Session->AddProvider(FThreadProvider::ProviderName, TSharedPtr<IThreadProvider>(ThreadProvider), TSharedPtr<IEditableThreadProvider>(ThreadProvider));
 
-	FFrameProvider* FrameProvider = new FFrameProvider(*Session);
+	TSharedPtr<FFrameProvider> FrameProvider = MakeShared<FFrameProvider>(*Session);
 	Session->AddProvider(FFrameProvider::ProviderName, FrameProvider);
 
-	FCounterProvider* CounterProvider = new FCounterProvider(*Session, *FrameProvider);
-	Session->AddProvider(FCounterProvider::ProviderName, CounterProvider);
+	TSharedPtr<FCounterProvider> CounterProvider = MakeShared<FCounterProvider>(*Session, *FrameProvider);
+	Session->AddProvider(FCounterProvider::ProviderName, TSharedPtr<ICounterProvider>(CounterProvider), TSharedPtr<IEditableCounterProvider>(CounterProvider));
 
-	FChannelProvider* ChannelProvider = new FChannelProvider();
+	TSharedPtr<FChannelProvider> ChannelProvider = MakeShared<FChannelProvider>();
 	Session->AddProvider(FChannelProvider::ProviderName, ChannelProvider);
 
-	FScreenshotProvider* ScreenshotProvider = new FScreenshotProvider(*Session);
+	TSharedPtr<FScreenshotProvider> ScreenshotProvider = MakeShared<FScreenshotProvider>(*Session);
 	Session->AddProvider(FScreenshotProvider::ProviderName, ScreenshotProvider);
 
-	Session->AddAnalyzer(new FMiscTraceAnalyzer(*Session, *ThreadProvider, *BookmarkProvider, *LogProvider, *FrameProvider, *ChannelProvider, *ScreenshotProvider));
+	Session->AddAnalyzer(new FMiscTraceAnalyzer(*Session, *ThreadProvider, *LogProvider, *FrameProvider, *ChannelProvider, *ScreenshotProvider));
+	Session->AddAnalyzer(new FBookmarksAnalyzer(*Session, *BookmarkProvider, LogProvider.Get()));
 	Session->AddAnalyzer(new FLogTraceAnalyzer(*Session, *LogProvider));
 
 	Session->AddAnalyzer(new FStringsAnalyzer(*Session));
 
-	FDefinitionProvider* DefProvider = new FDefinitionProvider(&Session.Get());
+	TSharedPtr<FDefinitionProvider> DefProvider = MakeShared<FDefinitionProvider>(&Session.Get());
 	Session->AddProvider(FDefinitionProvider::ProviderName, DefProvider);
 
 	ModuleService.OnAnalysisBegin(*Session);
diff --git a/Engine/Source/Developer/TraceServices/Private/AnalysisServicePrivate.h b/Engine/Source/Developer/TraceServices/Private/AnalysisServicePrivate.h
--- a/Engine/Source/Developer/TraceServices/Private/AnalysisServicePrivate.h
+++ b/Engine/Source/Developer/TraceServices/Private/AnalysisServicePrivate.h
@@ -62,7 +62,7 @@ public:
 	virtual void WriteAccessCheck() override { return Lock.WriteAccessCheck(); }
 
 	virtual void AddAnalyzer(UE::Trace::IAnalyzer* Analyzer) override;
-	virtual void AddProvider(const FName& Name, IProvider* Provider) override;
+	virtual void AddProvider(const FName& Name, TSharedPtr<IProvider> Provider, TSharedPtr<IEditableProvider> EditableProvider = nullptr) override;
 
 	const TArray<UE::Trace::IAnalyzer*> ReadAnalyzers() { return Analyzers; }
 
@@ -79,8 +79,7 @@ private:
 	FStringStore StringStore;
 	FAnalysisCache Cache;
 	TArray<UE::Trace::IAnalyzer*> Analyzers;
-	TArray<IProvider*> Providers;
-	TMap<FName, IProvider*> ProvidersMap;
+	TMap<FName, TTuple<TSharedPtr<IProvider>, TSharedPtr<IEditableProvider>>> Providers;
 	mutable TUniquePtr<UE::Trace::IInDataStream> DataStream;
 	mutable UE::Trace::FAnalysisProcessor Processor;
 };
diff --git a/Engine/Source/Developer/TraceServices/Private/AnalyzerFactories.cpp b/Engine/Source/Developer/TraceServices/Private/AnalyzerFactories.cpp
new file mode 100644
--- /dev/null
+++ b/Engine/Source/Developer/TraceServices/Private/AnalyzerFactories.cpp
@@ -0,0 +1,29 @@
+// Copyright Epic Games, Inc. All Rights Reserved.
+
+#include "TraceServices/AnalyzerFactories.h"
+#include "Trace/DataStream.h"
+
+#include "AnalysisServicePrivate.h"
+#include "Analyzers/BookmarksTraceAnalysis.h"
+#include "Analyzers/CountersTraceAnalysis.h"
+#include "Analyzers/CpuProfilerTraceAnalysis.h"
+
+TSharedPtr<TraceServices::IAnalysisSession> TraceServices::CreateAnalysisSession(uint32 InTraceId, const TCHAR* InSessionName, TUniquePtr<UE::Trace::IInDataStream>&& InDataStream)
+{
+	return MakeShared<FAnalysisSession>(InTraceId, InSessionName, MoveTemp(InDataStream));
+}
+
+TSharedPtr<UE::Trace::IAnalyzer> TraceServices::CreateBookmarksAnalyzer(IAnalysisSession& InSession, IEditableBookmarkProvider& InEditableBookmarkProvider)
+{
+	return MakeShared<FBookmarksAnalyzer>(InSession, InEditableBookmarkProvider, nullptr);
+}
+
+TSharedPtr<UE::Trace::IAnalyzer> TraceServices::CreateCountersAnalyzer(IAnalysisSession& InSession, IEditableCounterProvider& InEditableCounterProvider)
+{
+	return MakeShared<FCountersAnalyzer>(InSession, InEditableCounterProvider);
+}
+
+TSharedPtr<UE::Trace::IAnalyzer> TraceServices::CreateCpuProfilerAnalyzer(IAnalysisSession& InSession, IEditableTimingProfilerProvider& InEditableTimingProfilerProvider, IEditableThreadProvider& InEditableThreadProvider)
+{
+	return MakeShared<FCpuProfilerAnalyzer>(InSession, InEditableTimingProfilerProvider, InEditableThreadProvider);
+}
diff --git a/Engine/Source/Developer/TraceServices/Private/Analyzers/BookmarksTraceAnalysis.cpp b/Engine/Source/Developer/TraceServices/Private/Analyzers/BookmarksTraceAnalysis.cpp
new file mode 100644
--- /dev/null
+++ b/Engine/Source/Developer/TraceServices/Private/Analyzers/BookmarksTraceAnalysis.cpp
@@ -0,0 +1,115 @@
+// Copyright Epic Games, Inc. All Rights Reserved.
+
+#include "BookmarksTraceAnalysis.h"
+
+#include "Common/Utils.h"
+#include "Common/FormatArgs.h"
+#include "Common/ProviderLock.h"
+#include "HAL/LowLevelMemTracker.h"
+#include "TraceServices/Model/Bookmarks.h"
+#include "TraceServices/Model/Log.h"
+#include "TraceServices/Model/AnalysisSession.h"
+
+namespace TraceServices
+{
+
+FBookmarksAnalyzer::FBookmarksAnalyzer(IAnalysisSession& InSession, IEditableBookmarkProvider& InEditableBookmarkProvider, IEditableLogProvider* InEditableLogProvider)
+	: Session(InSession)
+	, EditableBookmarkProvider(InEditableBookmarkProvider)
+	, EditableLogProvider(InEditableLogProvider)
+{
+	if (EditableLogProvider)
+	{
+		// Todo: update this to use provider locking instead of session locking
+		// FProviderEditScopeLock EditableLogProviderLock (*EditableLogProvider);
+		FAnalysisSessionEditScope _(Session);
+		BookmarkLogCategoryId = EditableLogProvider->RegisterCategory();
+		FLogCategoryInfo& BookmarkLogCategory = EditableLogProvider->GetCategory(BookmarkLogCategoryId);
+		BookmarkLogCategory.Name = TEXT("LogBookmark");
+		BookmarkLogCategory.DefaultVerbosity = ELogVerbosity::All;
+	}
+}
+
+void FBookmarksAnalyzer::OnAnalysisBegin(const FOnAnalysisContext& Context)
+{
+	auto& Builder = Context.InterfaceBuilder;
+
+	Builder.RouteEvent(RouteId_BookmarkSpec, "Misc", "BookmarkSpec");
+	Builder.RouteEvent(RouteId_Bookmark, "Misc", "Bookmark");
+}
+
+bool FBookmarksAnalyzer::OnEvent(uint16 RouteId, EStyle Style, const FOnEventContext& Context)
+{
+	LLM_SCOPE_BYNAME(TEXT("Insights/FBookmarksAnalyzer"));
+
+	// Todo: update this to use provider locking instead of session locking
+	FAnalysisSessionEditScope _(Session);
+
+	const auto& EventData = Context.EventData;
+	switch (RouteId)
+	{
+	case RouteId_BookmarkSpec:
+	{
+		uint64 BookmarkPoint = EventData.GetValue<uint64>("BookmarkPoint");
+		const TCHAR* BookmarkFormatString = nullptr;
+		const TCHAR* BookmarkFile = nullptr;
+		int32 BookmarkLine = EventData.GetValue<int32>("Line");
+
+		FString FileName;
+		if (EventData.GetString("FileName", FileName))
+		{
+			BookmarkFile = Session.StoreString(*FileName);
+
+			FString FormatString;
+			EventData.GetString("FormatString", FormatString);
+			BookmarkFormatString = Session.StoreString(*FormatString);
+		}
+		else
+		{
+			const ANSICHAR* File = reinterpret_cast<const ANSICHAR*>(EventData.GetAttachment());
+			BookmarkFile = Session.StoreString(ANSI_TO_TCHAR(File));
+			BookmarkFormatString = Session.StoreString(reinterpret_cast<const TCHAR*>(EventData.GetAttachment() + strlen(File) + 1));
+		}
+
+		{
+			// Todo: update this to use provider locking instead of session locking
+			// FProviderEditScopeLock EditableBookmarkProviderLock (EditableBookmarkProvider);
+			EditableBookmarkProvider.UpdateBookmarkSpec(BookmarkPoint, BookmarkFormatString, BookmarkFile, BookmarkLine);
+		}
+
+		if (EditableLogProvider)
+		{
+			// Todo: update this to use provider locking instead of session locking
+			// FProviderEditScopeLock EditableLogProviderLock (*EditableLogProvider);
+			EditableLogProvider->UpdateMessageSpec(BookmarkPoint, BookmarkLogCategoryId, BookmarkFormatString, BookmarkFile, BookmarkLine, ELogVerbosity::Log);
+		}
+		break;
+	}
+
+	case RouteId_Bookmark:
+	{
+		uint64 BookmarkPoint = EventData.GetValue<uint64>("BookmarkPoint");
+		uint64 Cycle = EventData.GetValue<uint64>("Cycle");
+		double Timestamp = Context.EventTime.AsSeconds(Cycle);
+		TArrayView<const uint8> FormatArgsView = FTraceAnalyzerUtils::LegacyAttachmentArray("FormatArgs", Context);
+
+		{
+			// Todo: update this to use provider locking instead of session locking
+			// FProviderEditScopeLock EditableBookmarkProviderLock (EditableBookmarkProvider);
+			EditableBookmarkProvider.AppendBookmark(BookmarkPoint, Timestamp, FormatArgsView.GetData());
+		}
+
+		if (EditableLogProvider)
+		{
+			// Todo: update this to use provider locking instead of session locking
+			// FProviderEditScopeLock EditableLogProviderLock (*EditableLogProvider);
+			EditableLogProvider->AppendMessage(BookmarkPoint, Timestamp, FormatArgsView.GetData());
+		}
+		break;
+	}
+	}
+
+	return true;
+}
+
+} // namespace TraceServices
diff --git a/Engine/Source/Developer/TraceServices/Private/Analyzers/BookmarksTraceAnalysis.h b/Engine/Source/Developer/TraceServices/Private/Analyzers/BookmarksTraceAnalysis.h
new file mode 100644
--- /dev/null
+++ b/Engine/Source/Developer/TraceServices/Private/Analyzers/BookmarksTraceAnalysis.h
@@ -0,0 +1,46 @@
+// Copyright Epic Games, Inc. All Rights Reserved.
+
+#pragma once
+
+#include "Trace/Analyzer.h"
+#include "Templates/SharedPointer.h"
+#include "ProfilingDebugging/MiscTrace.h"
+#include "Common/PagedArray.h"
+
+namespace TraceServices
+{
+
+class IAnalysisSession;
+class IEditableBookmarkProvider;
+class ILogProvider;
+class IEditableLogProvider;
+
+class FBookmarksAnalyzer
+	: public UE::Trace::IAnalyzer
+{
+public:
+	FBookmarksAnalyzer(IAnalysisSession& Session, IEditableBookmarkProvider& EditableBookmarkProvider, IEditableLogProvider* InEditableLogProvider);
+	virtual void OnAnalysisBegin(const FOnAnalysisContext& Context) override;
+	virtual bool OnEvent(uint16 RouteId, EStyle Style, const FOnEventContext& Context) override;
+
+private:
+	enum : uint16
+	{
+		RouteId_BookmarkSpec,
+		RouteId_Bookmark,
+	};
+
+	IAnalysisSession& Session;
+	IEditableBookmarkProvider& EditableBookmarkProvider;
+	IEditableLogProvider* EditableLogProvider;
+	uint64 BookmarkLogCategoryId = uint64(-1);
+
+	enum
+	{
+		FormatBufferSize = 65536
+	};
+	TCHAR FormatBuffer[FormatBufferSize];
+	TCHAR TempBuffer[FormatBufferSize];
+};
+
+} // namespace TraceServices
diff --git a/Engine/Source/Developer/TraceServices/Private/Analyzers/CountersTraceAnalysis.cpp b/Engine/Source/Developer/TraceServices/Private/Analyzers/CountersTraceAnalysis.cpp
--- a/Engine/Source/Developer/TraceServices/Private/Analyzers/CountersTraceAnalysis.cpp
+++ b/Engine/Source/Developer/TraceServices/Private/Analyzers/CountersTraceAnalysis.cpp
@@ -10,9 +10,9 @@
 namespace TraceServices
 {
 
-FCountersAnalyzer::FCountersAnalyzer(IAnalysisSession& InSession, ICounterProvider& InCounterProvider)
+FCountersAnalyzer::FCountersAnalyzer(IAnalysisSession& InSession, IEditableCounterProvider& InEditableCounterProvider)
 	: Session(InSession)
-	, CounterProvider(InCounterProvider)
+	, EditableCounterProvider(InEditableCounterProvider)
 {
 }
 
@@ -39,14 +39,14 @@ bool FCountersAnalyzer::OnEvent(uint16 RouteId, EStyle Style, const FOnEventCont
 		uint16 CounterId = EventData.GetValue<uint16>("Id");
 		ETraceCounterType CounterType = static_cast<ETraceCounterType>(EventData.GetValue<uint8>("Type"));
 		ETraceCounterDisplayHint CounterDisplayHint = static_cast<ETraceCounterDisplayHint>(EventData.GetValue<uint8>("DisplayHint"));
-		IEditableCounter* Counter = CounterProvider.CreateCounter();
+		IEditableCounter* EditableCounter = EditableCounterProvider.CreateEditableCounter();
 		if (CounterType == TraceCounterType_Float)
 		{
-			Counter->SetIsFloatingPoint(true);
+			EditableCounter->SetIsFloatingPoint(true);
 		}
 		if (CounterDisplayHint == TraceCounterDisplayHint_Memory)
 		{
-			Counter->SetDisplayHint(CounterDisplayHint_Memory);
+			EditableCounter->SetDisplayHint(CounterDisplayHint_Memory);
 		}
 		FString Name = FTraceAnalyzerUtils::LegacyAttachmentString<TCHAR>("Name", Context);
 		if (Name.IsEmpty())
@@ -54,8 +54,8 @@ bool FCountersAnalyzer::OnEvent(uint16 RouteId, EStyle Style, const FOnEventCont
 			UE_LOG(LogTraceServices, Warning, TEXT("Invalid counter name for counter %u."), uint32(CounterId));
 			Name = FString::Printf(TEXT("<noname counter %u>"), uint32(CounterId));
 		}
-		Counter->SetName(Session.StoreString(*Name));
-		CountersMap.Add(CounterId, Counter);
+		EditableCounter->SetName(Session.StoreString(*Name));
+		EditableCountersMap.Add(CounterId, EditableCounter);
 		break;
 	}
 	case RouteId_SetValueInt:
@@ -63,10 +63,10 @@ bool FCountersAnalyzer::OnEvent(uint16 RouteId, EStyle Style, const FOnEventCont
 		uint16 CounterId = EventData.GetValue<uint16>("CounterId");
 		int64 Value = EventData.GetValue<int64>("Value");
 		double Timestamp = Context.EventTime.AsSeconds(EventData.GetValue<uint64>("Cycle"));
-		IEditableCounter* FindCounter = CountersMap.FindRef(CounterId);
-		if (ensure(FindCounter))
+		IEditableCounter* FindEditableCounter = EditableCountersMap.FindRef(CounterId);
+		if (ensure(FindEditableCounter))
 		{
-			FindCounter->SetValue(Timestamp, Value);
+			FindEditableCounter->SetValue(Timestamp, Value);
 		}
 		break;
 	}
@@ -75,10 +75,10 @@ bool FCountersAnalyzer::OnEvent(uint16 RouteId, EStyle Style, const FOnEventCont
 		uint16 CounterId = EventData.GetValue<uint16>("CounterId");
 		float Value = EventData.GetValue<float>("Value");
 		double Timestamp = Context.EventTime.AsSeconds(EventData.GetValue<uint64>("Cycle"));
-		IEditableCounter* FindCounter = CountersMap.FindRef(CounterId);
-		if (ensure(FindCounter))
+		IEditableCounter* FindEditableCounter = EditableCountersMap.FindRef(CounterId);
+		if (ensure(FindEditableCounter))
 		{
-			FindCounter->SetValue(Timestamp, Value);
+			FindEditableCounter->SetValue(Timestamp, Value);
 		}
 		break;
 	}
diff --git a/Engine/Source/Developer/TraceServices/Private/Analyzers/CountersTraceAnalysis.h b/Engine/Source/Developer/TraceServices/Private/Analyzers/CountersTraceAnalysis.h
--- a/Engine/Source/Developer/TraceServices/Private/Analyzers/CountersTraceAnalysis.h
+++ b/Engine/Source/Developer/TraceServices/Private/Analyzers/CountersTraceAnalysis.h
@@ -10,13 +10,13 @@ namespace TraceServices
 
 class IAnalysisSession;
 class IEditableCounter;
-class ICounterProvider;
+class IEditableCounterProvider;
 
 class FCountersAnalyzer
 	: public UE::Trace::IAnalyzer
 {
 public:
-	FCountersAnalyzer(IAnalysisSession& Session, ICounterProvider& CounterProvider);
+	FCountersAnalyzer(IAnalysisSession& Session, IEditableCounterProvider& EditableCounterProvider);
 	virtual void OnAnalysisBegin(const FOnAnalysisContext& Context) override;
 	virtual bool OnEvent(uint16 RouteId, EStyle Style, const FOnEventContext& Context) override;
 
@@ -29,8 +29,8 @@ private:
 	};
 
 	IAnalysisSession& Session;
-	ICounterProvider& CounterProvider;
-	TMap<uint16, IEditableCounter*> CountersMap;
+	IEditableCounterProvider& EditableCounterProvider;
+	TMap<uint16, IEditableCounter*> EditableCountersMap;
 };
 
 } // namespace TraceServices
diff --git a/Engine/Source/Developer/TraceServices/Private/Analyzers/CpuProfilerTraceAnalysis.cpp b/Engine/Source/Developer/TraceServices/Private/Analyzers/CpuProfilerTraceAnalysis.cpp
--- a/Engine/Source/Developer/TraceServices/Private/Analyzers/CpuProfilerTraceAnalysis.cpp
+++ b/Engine/Source/Developer/TraceServices/Private/Analyzers/CpuProfilerTraceAnalysis.cpp
@@ -20,10 +20,10 @@ namespace TraceServices
 
 ////////////////////////////////////////////////////////////////////////////////////////////////////
 
-FCpuProfilerAnalyzer::FCpuProfilerAnalyzer(IAnalysisSession& InSession, FTimingProfilerProvider& InTimingProfilerProvider, FThreadProvider& InThreadProvider)
+FCpuProfilerAnalyzer::FCpuProfilerAnalyzer(IAnalysisSession& InSession, IEditableTimingProfilerProvider& InEditableTimingProfilerProvider, IEditableThreadProvider& InEditableThreadProvider)
 	: Session(InSession)
-	, TimingProfilerProvider(InTimingProfilerProvider)
-	, ThreadProvider(InThreadProvider)
+	, EditableTimingProfilerProvider(InEditableTimingProfilerProvider)
+	, EditableThreadProvider(InEditableThreadProvider)
 {
 }
 
@@ -422,7 +422,7 @@ uint64 FCpuProfilerAnalyzer::ProcessBufferV2(const FEventTime& EventTime, FThrea
 						CborWriter.WriteValue("C", 1); // continuation?
 						CborWriter.WriteValue(false);
 					}
-					uint32 MetadataTimerId = TimingProfilerProvider.AddMetadata(CoroutineTimerId, MoveTemp(CborData));
+					uint32 MetadataTimerId = EditableTimingProfilerProvider.AddMetadata(CoroutineTimerId, MoveTemp(CborData));
 
 					FEventScopeState& ScopeState = ThreadState.ScopeStack.AddDefaulted_GetRef();
 					ScopeState.StartCycle = ActualCycle;
@@ -483,7 +483,7 @@ uint64 FCpuProfilerAnalyzer::ProcessBufferV2(const FEventTime& EventTime, FThrea
 					if (ThreadState.ScopeStack.Num() > 0)
 					{
 						uint32 MetadataTimerId = ThreadState.ScopeStack.Top().EventTypeId;
-						TArrayView<const uint8> Metadata = TimingProfilerProvider.GetMetadata(MetadataTimerId);
+						TArrayView<const uint8> Metadata = EditableTimingProfilerProvider.GetMetadata(MetadataTimerId);
 						// Change the last byte in metadata to "true".
 						const_cast<uint8*>(Metadata.GetData())[Metadata.Num() - 1] = (uint8)(ECborCode::Prim | ECborCode::True);
 					}
@@ -658,7 +658,7 @@ void FCpuProfilerAnalyzer::OnCpuScopeEnter(const FOnEventContext& Context)
 
 	TArray<uint8> CborData;
 	Context.EventData.SerializeToCbor(CborData);
-	TimerId = TimingProfilerProvider.AddMetadata(TimerId, MoveTemp(CborData));
+	TimerId = EditableTimingProfilerProvider.AddMetadata(TimerId, MoveTemp(CborData));
 
 	uint64 Cycle = Context.EventTime.AsCycle64();
 	ThreadState.PendingEvents.Add({Cycle, TimerId});
@@ -699,7 +699,7 @@ uint32 FCpuProfilerAnalyzer::DefineTimer(uint32 SpecId, const TCHAR* Name, const
 		{
 			// Yes, SpecId was already mapped to a timer (ex. as an <unknown> timer).
 			// Update name for mapped timer.
-			TimingProfilerProvider.SetTimerNameAndLocation(*FindTimerId, Name, File, Line);
+			EditableTimingProfilerProvider.SetTimerNameAndLocation(*FindTimerId, Name, File, Line);
 			// In this case, we do not remap the SpecId to the previously defined timer with same name.
 			// This is becasue the two timers are already used in timelines.
 			// So we will continue to use separate timers, even if those have same name.
@@ -721,7 +721,7 @@ uint32 FCpuProfilerAnalyzer::DefineTimer(uint32 SpecId, const TCHAR* Name, const
 		{
 			// Yes, SpecId was already mapped to a timer (ex. as an <unknown> timer).
 			// Update name for mapped timer.
-			TimingProfilerProvider.SetTimerNameAndLocation(*FindTimerId, Name, File, Line);
+			EditableTimingProfilerProvider.SetTimerNameAndLocation(*FindTimerId, Name, File, Line);
 			if (bMergeByName)
 			{
 				// Map the name to the timer.
@@ -732,7 +732,7 @@ uint32 FCpuProfilerAnalyzer::DefineTimer(uint32 SpecId, const TCHAR* Name, const
 		else
 		{
 			// Define a new Cpu timer.
-			uint32 NewTimerId = TimingProfilerProvider.AddCpuTimer(Name, File, Line);
+			uint32 NewTimerId = EditableTimingProfilerProvider.AddCpuTimer(Name, File, Line);
 			// Map the SpecId to the timer.
 			SpecIdToTimerIdMap.Add(SpecId, NewTimerId);
 			if (bMergeByName)
@@ -750,7 +750,7 @@ uint32 FCpuProfilerAnalyzer::DefineTimer(uint32 SpecId, const TCHAR* Name, const
 uint32 FCpuProfilerAnalyzer::DefineNewTimerChecked(uint32 SpecId, const TCHAR* TimerName, const TCHAR* File, uint32 Line)
 {
 	TimerName = Session.StoreString(TimerName);
-	uint32 NewTimerId = TimingProfilerProvider.AddCpuTimer(TimerName, File, Line);
+	uint32 NewTimerId = EditableTimingProfilerProvider.AddCpuTimer(TimerName, File, Line);
 	SpecIdToTimerIdMap.Add(SpecId, NewTimerId);
 	return NewTimerId;
 }
@@ -781,12 +781,12 @@ FCpuProfilerAnalyzer::FThreadState& FCpuProfilerAnalyzer::GetThreadState(uint32
 	{
 		ThreadState = new FThreadState();
 		ThreadState->ThreadId = ThreadId;
-		ThreadState->Timeline = &TimingProfilerProvider.EditCpuThreadTimeline(ThreadId);
+		ThreadState->Timeline = &EditableTimingProfilerProvider.GetCpuThreadEditableTimeline(ThreadId);
 		ThreadStatesMap.Add(ThreadId, ThreadState);
 
 		// Just in case the rest of Insight's reporting/analysis doesn't know about
 		// this thread, we'll explicitly add it. For fault tolerance.
-		ThreadProvider.AddThread(ThreadId, nullptr, TPri_Normal);
+		EditableThreadProvider.AddThread(ThreadId, nullptr, TPri_Normal);
 	}
 	return *ThreadState;
 }
diff --git a/Engine/Source/Developer/TraceServices/Private/Analyzers/CpuProfilerTraceAnalysis.h b/Engine/Source/Developer/TraceServices/Private/Analyzers/CpuProfilerTraceAnalysis.h
--- a/Engine/Source/Developer/TraceServices/Private/Analyzers/CpuProfilerTraceAnalysis.h
+++ b/Engine/Source/Developer/TraceServices/Private/Analyzers/CpuProfilerTraceAnalysis.h
@@ -5,19 +5,21 @@
 #include "Trace/Analyzer.h"
 #include "Containers/Map.h"
 #include "Containers/UnrealString.h"
-#include "Model/TimingProfilerPrivate.h"
+#include "TraceServices/Model/TimingProfiler.h"
+#include "Model/MonotonicTimeline.h"
 
 namespace TraceServices
 {
 
 class IAnalysisSession;
-class FThreadProvider;
+class IEditableTimingProfilerProvider;
+class IEditableThreadProvider;
 
 class FCpuProfilerAnalyzer
 	: public UE::Trace::IAnalyzer
 {
 public:
-	FCpuProfilerAnalyzer(IAnalysisSession& Session, FTimingProfilerProvider& TimingProfilerProvider, FThreadProvider& InThreadProvider);
+	FCpuProfilerAnalyzer(IAnalysisSession& Session, IEditableTimingProfilerProvider& InEditableTimingProfilerProvider, IEditableThreadProvider& InEditableThreadProvider);
 	~FCpuProfilerAnalyzer();
 	virtual void OnAnalysisBegin(const FOnAnalysisContext& Context) override;
 	virtual void OnAnalysisEnd(/*const FOnAnalysisEndContext& Context*/) override;
@@ -41,7 +43,7 @@ private:
 		uint32 ThreadId = 0;
 		TArray<FEventScopeState> ScopeStack;
 		TArray<FPendingEvent> PendingEvents;
-		FTimingProfilerProvider::TimelineInternal* Timeline = nullptr;
+		IEditableTimeline<FTimingProfilerEvent>* Timeline = nullptr;
 		uint64 LastCycle = 0;
 	};
 
@@ -69,8 +71,8 @@ private:
 	};
 
 	IAnalysisSession& Session;
-	FTimingProfilerProvider& TimingProfilerProvider;
-	FThreadProvider& ThreadProvider;
+	IEditableTimingProfilerProvider& EditableTimingProfilerProvider;
+	IEditableThreadProvider& EditableThreadProvider;
 	TMap<uint32, FThreadState*> ThreadStatesMap;
 	TMap<uint32, uint32> SpecIdToTimerIdMap;
 	TMap<const TCHAR*, uint32> ScopeNameToTimerIdMap;
diff --git a/Engine/Source/Developer/TraceServices/Private/Analyzers/CsvProfilerTraceAnalysis.cpp b/Engine/Source/Developer/TraceServices/Private/Analyzers/CsvProfilerTraceAnalysis.cpp
--- a/Engine/Source/Developer/TraceServices/Private/Analyzers/CsvProfilerTraceAnalysis.cpp
+++ b/Engine/Source/Developer/TraceServices/Private/Analyzers/CsvProfilerTraceAnalysis.cpp
@@ -12,10 +12,10 @@
 namespace TraceServices
 {
 
-FCsvProfilerAnalyzer::FCsvProfilerAnalyzer(IAnalysisSession& InSession, FCsvProfilerProvider& InCsvProfilerProvider, ICounterProvider& InCounterProvider, const IFrameProvider& InFrameProvider, const IThreadProvider& InThreadProvider)
+FCsvProfilerAnalyzer::FCsvProfilerAnalyzer(IAnalysisSession& InSession, FCsvProfilerProvider& InCsvProfilerProvider, IEditableCounterProvider& InEditableCounterProvider, const IFrameProvider& InFrameProvider, const IThreadProvider& InThreadProvider)
 	: Session(InSession)
 	, CsvProfilerProvider(InCsvProfilerProvider)
-	, CounterProvider(InCounterProvider)
+	, EditableCounterProvider(InEditableCounterProvider)
 	, FrameProvider(InFrameProvider)
 	, ThreadProvider(InThreadProvider)
 {
@@ -286,7 +286,7 @@ FCsvProfilerAnalyzer::FStatSeriesInstance& FCsvProfilerAnalyzer::GetStatSeries(u
 	const TCHAR* StatSeriesName = GetStatSeriesName(Definition, Type, ThreadState, false);
 	Instance->ProviderHandle = CsvProfilerProvider.AddSeries(StatSeriesName, Type);
 	Instance->ProviderCountHandle = CsvProfilerProvider.AddSeries(GetStatSeriesName(Definition, Type, ThreadState, true), CsvStatSeriesType_CustomStatInt);
-	Instance->Counter = CounterProvider.CreateCounter();
+	Instance->Counter = EditableCounterProvider.CreateEditableCounter();
 	Instance->Counter->SetName(StatSeriesName);
 	Instance->Counter->SetIsFloatingPoint(Type != CsvStatSeriesType_CustomStatInt);
 	Instance->Type = Type;
diff --git a/Engine/Source/Developer/TraceServices/Private/Analyzers/CsvProfilerTraceAnalysis.h b/Engine/Source/Developer/TraceServices/Private/Analyzers/CsvProfilerTraceAnalysis.h
--- a/Engine/Source/Developer/TraceServices/Private/Analyzers/CsvProfilerTraceAnalysis.h
+++ b/Engine/Source/Developer/TraceServices/Private/Analyzers/CsvProfilerTraceAnalysis.h
@@ -13,14 +13,14 @@ namespace TraceServices
 class IAnalysisSession;
 class IFrameProvider;
 class IThreadProvider;
-class ICounterProvider;
+class IEditableCounterProvider;
 class IEditableCounter;
 
 class FCsvProfilerAnalyzer
 	: public UE::Trace::IAnalyzer
 {
 public:
-	FCsvProfilerAnalyzer(IAnalysisSession& Session, FCsvProfilerProvider& CsvProfilerProvider, ICounterProvider& CounterProvider, const IFrameProvider& FrameProvider, const IThreadProvider& ThreadProvider);
+	FCsvProfilerAnalyzer(IAnalysisSession& Session, FCsvProfilerProvider& CsvProfilerProvider, IEditableCounterProvider& EditableCounterProvider, const IFrameProvider& FrameProvider, const IThreadProvider& ThreadProvider);
 	~FCsvProfilerAnalyzer();
 	virtual void OnAnalysisBegin(const FOnAnalysisContext& Context) override;
 	virtual bool OnEvent(uint16 RouteId, EStyle Style, const FOnEventContext& Context) override;
@@ -130,7 +130,7 @@ private:
 	
 	IAnalysisSession& Session;
 	FCsvProfilerProvider& CsvProfilerProvider;
-	ICounterProvider& CounterProvider;
+	IEditableCounterProvider& EditableCounterProvider;
 	const IFrameProvider& FrameProvider;
 	const IThreadProvider& ThreadProvider;
 
diff --git a/Engine/Source/Developer/TraceServices/Private/Analyzers/MiscTraceAnalysis.cpp b/Engine/Source/Developer/TraceServices/Private/Analyzers/MiscTraceAnalysis.cpp
--- a/Engine/Source/Developer/TraceServices/Private/Analyzers/MiscTraceAnalysis.cpp
+++ b/Engine/Source/Developer/TraceServices/Private/Analyzers/MiscTraceAnalysis.cpp
@@ -17,22 +17,23 @@
 
 FMiscTraceAnalyzer::FMiscTraceAnalyzer(IAnalysisSession& InSession,
 									   FThreadProvider& InThreadProvider,
-									   FBookmarkProvider& InBookmarkProvider,
 									   FLogProvider& InLogProvider,
 									   FFrameProvider& InFrameProvider,
 									   FChannelProvider& InChannelProvider,
 									   FScreenshotProvider& InScreenshotProvider)
 	: Session(InSession)
 	, ThreadProvider(InThreadProvider)
-	, BookmarkProvider(InBookmarkProvider)
 	, LogProvider(InLogProvider)
 	, FrameProvider(InFrameProvider)
 	, ChannelProvider(InChannelProvider)
 	, ScreenshotProvider(InScreenshotProvider)
 {
-	FLogCategoryInfo& BookmarkLogCategory = LogProvider.GetCategory(FLogProvider::ReservedLogCategory_Bookmark);
-	BookmarkLogCategory.Name = TEXT("LogBookmark");
-	BookmarkLogCategory.DefaultVerbosity = ELogVerbosity::All;
+	// Todo: update this to use provider locking instead of session locking
+	// FProviderEditScopeLock LogProviderLock (LogProvider);
+	FAnalysisSessionEditScope _(Session);
+	ScreenshotLogCategoryId = LogProvider.RegisterCategory();
+	FLogCategoryInfo& ScreenshotLogCategory = LogProvider.GetCategory(ScreenshotLogCategoryId);
+	ScreenshotLogCategory.Name = TEXT("Screenshot");
 }
 
 void FMiscTraceAnalyzer::OnAnalysisBegin(const FOnAnalysisContext& Context)
@@ -86,48 +87,6 @@
 	const auto& EventData = Context.EventData;
 	switch (RouteId)
 	{
-	case RouteId_BookmarkSpec:
-	{
-		uint64 BookmarkPoint = EventData.GetValue<uint64>("BookmarkPoint");
-		FBookmarkSpec& Spec = BookmarkProvider.GetSpec(BookmarkPoint);
-		Spec.Line = EventData.GetValue<int32>("Line");
-
-		FString FileName;
-		if (EventData.GetString("FileName", FileName))
-		{
-			Spec.File = Session.StoreString(*FileName);
-
-			FString FormatString;
-			EventData.GetString("FormatString", FormatString);
-			Spec.FormatString = Session.StoreString(*FormatString);
-		}
-		else
-		{
-			const ANSICHAR* File = reinterpret_cast<const ANSICHAR*>(EventData.GetAttachment());
-			Spec.File = Session.StoreString(ANSI_TO_TCHAR(File));
-			Spec.FormatString = Session.StoreString(reinterpret_cast<const TCHAR*>(EventData.GetAttachment() + strlen(File) + 1));
-		}
-
-		FLogMessageSpec& LogMessageSpec = LogProvider.GetMessageSpec(BookmarkPoint);
-		LogMessageSpec.Category = &LogProvider.GetCategory(FLogProvider::ReservedLogCategory_Bookmark);
-		LogMessageSpec.Line = Spec.Line;
-		LogMessageSpec.File = Spec.File;
-		LogMessageSpec.FormatString = Spec.FormatString;
-		LogMessageSpec.Verbosity = ELogVerbosity::Log;
-		break;
-	}
-
-	case RouteId_Bookmark:
-	{
-		uint64 BookmarkPoint = EventData.GetValue<uint64>("BookmarkPoint");
-		uint64 Cycle = EventData.GetValue<uint64>("Cycle");
-		double Timestamp = Context.EventTime.AsSeconds(Cycle);
-		TArrayView<const uint8> FormatArgsView = FTraceAnalyzerUtils::LegacyAttachmentArray("FormatArgs", Context);
-		BookmarkProvider.AppendBookmark(Timestamp, BookmarkPoint, FormatArgsView.GetData());
-		LogProvider.AppendMessage(BookmarkPoint, Timestamp, FormatArgsView.GetData());
-		break;
-	}
-
 	case RouteId_BeginFrame:
 	{
 		uint64 Cycle = EventData.GetValue<uint64>("Cycle");
@@ -201,8 +160,7 @@
 		Screenshot->Data.Reserve(Screenshot->Size);
 
 		FLogMessageSpec& LogMessageSpec = LogProvider.GetMessageSpec(Cycle);
-		LogMessageSpec.Category = &LogProvider.GetCategory(FLogProvider::ReservedLogCategory_Screenshot);
-		LogMessageSpec.Category->Name = TEXT("Screenshot");
+		LogMessageSpec.Category = &LogProvider.GetCategory(ScreenshotLogCategoryId);
 		LogMessageSpec.Line = Id;
 		LogMessageSpec.File = nullptr;
 		LogMessageSpec.FormatString = nullptr;
diff --git a/Engine/Source/Developer/TraceServices/Private/Analyzers/MiscTraceAnalysis.h b/Engine/Source/Developer/TraceServices/Private/Analyzers/MiscTraceAnalysis.h
--- a/Engine/Source/Developer/TraceServices/Private/Analyzers/MiscTraceAnalysis.h
+++ b/Engine/Source/Developer/TraceServices/Private/Analyzers/MiscTraceAnalysis.h
@@ -12,7 +12,6 @@ namespace TraceServices
 
 class IAnalysisSession;
 class FThreadProvider;
-class FBookmarkProvider;
 class FLogProvider;
 class FFrameProvider;
 class FChannelProvider;
@@ -24,7 +23,6 @@ class FMiscTraceAnalyzer
 public:
 	FMiscTraceAnalyzer(IAnalysisSession& Session,
 					   FThreadProvider& ThreadProvider,
-					   FBookmarkProvider& BookmarkProvider,
 					   FLogProvider& LogProvider,
 					   FFrameProvider& FrameProvider, 
 					   FChannelProvider& ChannelProvider,
@@ -66,13 +64,13 @@ private:
 
 	IAnalysisSession& Session;
 	FThreadProvider& ThreadProvider;
-	FBookmarkProvider& BookmarkProvider;
 	FLogProvider& LogProvider;
 	FFrameProvider& FrameProvider;
 	FChannelProvider& ChannelProvider;
 	FScreenshotProvider& ScreenshotProvider;
 	TMap<uint32, TSharedRef<FThreadState>> ThreadStateMap;
 	uint64 LastFrameCycle[TraceFrameType_Count] = { 0, 0 };
+	uint64 ScreenshotLogCategoryId = uint64(-1);
 };
 
 
diff --git a/Engine/Source/Developer/TraceServices/Private/Analyzers/ModuleAnalysis.cpp b/Engine/Source/Developer/TraceServices/Private/Analyzers/ModuleAnalysis.cpp
--- a/Engine/Source/Developer/TraceServices/Private/Analyzers/ModuleAnalysis.cpp
+++ b/Engine/Source/Developer/TraceServices/Private/Analyzers/ModuleAnalysis.cpp
@@ -64,10 +64,11 @@ bool FModuleAnalyzer::OnEvent(uint16 RouteId, EStyle Style, const FOnEventContex
 				if (Context.EventData.GetString("SymbolFormat", SymbolFormat))
 				{
 					check(Provider == nullptr); // Should only get one init message
-					Provider = CreateModuleProvider(Session, SymbolFormat);
-					if (Provider)
+					TSharedPtr<IModuleAnalysisProvider> ModuleProvider = CreateModuleProvider(Session, SymbolFormat);
+					if (ModuleProvider)
 					{
-						Session.AddProvider(GetModuleProviderName(), Provider);
+						Session.AddProvider(GetModuleProviderName(), ModuleProvider);
+						Provider = ModuleProvider.Get();
 					}
 				}
 			}
diff --git a/Engine/Source/Developer/TraceServices/Private/Analyzers/StatsTraceAnalysis.cpp b/Engine/Source/Developer/TraceServices/Private/Analyzers/StatsTraceAnalysis.cpp
--- a/Engine/Source/Developer/TraceServices/Private/Analyzers/StatsTraceAnalysis.cpp
+++ b/Engine/Source/Developer/TraceServices/Private/Analyzers/StatsTraceAnalysis.cpp
@@ -15,9 +15,9 @@
 namespace TraceServices
 {
 
-FStatsAnalyzer::FStatsAnalyzer(IAnalysisSession& InSession, ICounterProvider& InCounterProvider)
+FStatsAnalyzer::FStatsAnalyzer(IAnalysisSession& InSession, IEditableCounterProvider& InEditableCounterProvider)
 	: Session(InSession)
-	, CounterProvider(InCounterProvider)
+	, EditableCounterProvider(InEditableCounterProvider)
 {
 }
 
@@ -48,11 +48,11 @@ bool FStatsAnalyzer::OnEvent(uint16 RouteId, EStyle Style, const FOnEventContext
 		case RouteId_Spec:
 		{
 			uint32 StatId = EventData.GetValue<uint32>("Id");
-			IEditableCounter* Counter = CountersMap.FindRef(StatId);
-			if (!Counter)
+			IEditableCounter* EditableCounter = EditableCountersMap.FindRef(StatId);
+			if (!EditableCounter)
 			{
-				Counter = CounterProvider.CreateCounter();
-				CountersMap.Add(StatId, Counter);
+				EditableCounter = EditableCounterProvider.CreateEditableCounter();
+				EditableCountersMap.Add(StatId, EditableCounter);
 			}
 
 			FString Name;
@@ -74,23 +74,45 @@ bool FStatsAnalyzer::OnEvent(uint16 RouteId, EStyle Style, const FOnEventContext
 				UE_LOG(LogTraceServices, Warning, TEXT("Invalid counter name for Stats counter %u."), StatId);
 				Name = FString::Printf(TEXT("<noname stats counter %u>"), StatId);
 			}
-			Counter->SetName(Session.StoreString(*Name));
+			EditableCounter->SetName(Session.StoreString(*Name));
 
 			if (!Group.IsEmpty())
 			{
-				Counter->SetGroup(Session.StoreString(*Group));
+				EditableCounter->SetGroup(Session.StoreString(*Group));
+			}
+			EditableCounter->SetDescription(Session.StoreString(*Description));
+
+			bool bIsFloatingPoint = EventData.GetValue<bool>("IsFloatingPoint");
+			EditableCounter->SetIsFloatingPoint(bIsFloatingPoint);
+			
+			bool bIsResetEveryFrame = EventData.GetValue<bool>("ShouldClearEveryFrame");
+			EditableCounter->SetIsResetEveryFrame(bIsResetEveryFrame);
+			if (bIsResetEveryFrame)
+			{
+				if (bIsFloatingPoint)
+				{
+					IEditableCounter* ResetEveryFrameEditableCounter = FloatResetEveryFrameCountersMap.FindRef(StatId);
+					if (!ResetEveryFrameEditableCounter)
+					{
+						FloatResetEveryFrameCountersMap.Add(StatId, EditableCounter);
+					}
+				}
+				else
+				{
+					IEditableCounter* ResetEveryFrameEditableCounter = Int64ResetEveryFrameCountersMap.FindRef(StatId);
+					if (!ResetEveryFrameEditableCounter)
+					{
+						Int64ResetEveryFrameCountersMap.Add(StatId, EditableCounter);
+					}
+				}
 			}
-			Counter->SetDescription(Session.StoreString(*Description));
-
-			Counter->SetIsFloatingPoint(EventData.GetValue<bool>("IsFloatingPoint"));
-			Counter->SetIsResetEveryFrame(EventData.GetValue<bool>("ShouldClearEveryFrame"));
 
 			ECounterDisplayHint DisplayHint = CounterDisplayHint_None;
 			if (EventData.GetValue<bool>("IsMemory"))
 			{
 				DisplayHint = CounterDisplayHint_Memory;
 			}
-			Counter->SetDisplayHint(DisplayHint);
+			EditableCounter->SetDisplayHint(DisplayHint);
 			break;
 		}
 
@@ -116,13 +138,13 @@ bool FStatsAnalyzer::OnEvent(uint16 RouteId, EStyle Style, const FOnEventContext
 
 				uint64 DecodedIdAndOp = FTraceAnalyzerUtils::Decode7bit(BufferPtr);
 				uint32 StatId = DecodedIdAndOp >> 3;
-				IEditableCounter* Counter = CountersMap.FindRef(StatId);
-				if (!Counter)
+				IEditableCounter* EditableCounter = EditableCountersMap.FindRef(StatId);
+				if (!EditableCounter)
 				{
-					Counter = CounterProvider.CreateCounter();
+					EditableCounter = EditableCounterProvider.CreateEditableCounter();
 					FString Name = FString::Printf(TEXT("<unknown stats counter %u>"), StatId);
-					Counter->SetName(Session.StoreString(*Name));
-					CountersMap.Add(StatId, Counter);
+					EditableCounter->SetName(Session.StoreString(*Name));
+					EditableCountersMap.Add(StatId, EditableCounter);
 				}
 
 				uint8 Op = DecodedIdAndOp & 0x7;
@@ -135,27 +157,27 @@ bool FStatsAnalyzer::OnEvent(uint16 RouteId, EStyle Style, const FOnEventContext
 				{
 				case Increment:
 				{
-					Counter->AddValue(Time, int64(1));
+					EditableCounter->AddValue(Time, int64(1));
 					STATS_ANALYZER_DEBUG_LOG(StatId, TEXT("%f INC() %u"), Time, ThreadId);
 					break;
 				}
 				case Decrement:
 				{
-					Counter->AddValue(Time, int64(-1));
+					EditableCounter->AddValue(Time, int64(-1));
 					STATS_ANALYZER_DEBUG_LOG(StatId, TEXT("%f DEC() %u"), Time, ThreadId);
 					break;
 				}
 				case AddInteger:
 				{
 					int64 Amount = FTraceAnalyzerUtils::DecodeZigZag(BufferPtr);
-					Counter->AddValue(Time, Amount);
+					EditableCounter->AddValue(Time, Amount);
 					STATS_ANALYZER_DEBUG_LOG(StatId, TEXT("%f ADD(%lli) %u"), Time, Amount, ThreadId);
 					break;
 				}
 				case SetInteger:
 				{
 					int64 Value = FTraceAnalyzerUtils::DecodeZigZag(BufferPtr);
-					Counter->SetValue(Time, Value);
+					EditableCounter->SetValue(Time, Value);
 					STATS_ANALYZER_DEBUG_LOG(StatId, TEXT("%f SET(%lli) %u"), Time, Value, ThreadId);
 					break;
 				}
@@ -164,7 +186,7 @@ bool FStatsAnalyzer::OnEvent(uint16 RouteId, EStyle Style, const FOnEventContext
 					double Amount;
 					memcpy(&Amount, BufferPtr, sizeof(double));
 					BufferPtr += sizeof(double);
-					Counter->AddValue(Time, Amount);
+					EditableCounter->AddValue(Time, Amount);
 					STATS_ANALYZER_DEBUG_LOG(StatId, TEXT("%f ADD(%f) %u"), Time, Amount, ThreadId);
 					break;
 				}
@@ -173,7 +195,7 @@ bool FStatsAnalyzer::OnEvent(uint16 RouteId, EStyle Style, const FOnEventContext
 					double Value;
 					memcpy(&Value, BufferPtr, sizeof(double));
 					BufferPtr += sizeof(double);
-					Counter->SetValue(Time, Value);
+					EditableCounter->SetValue(Time, Value);
 					STATS_ANALYZER_DEBUG_LOG(StatId, TEXT("%f SET(%f) %u"), Time, Value, ThreadId);
 					break;
 				}
@@ -192,20 +214,15 @@ bool FStatsAnalyzer::OnEvent(uint16 RouteId, EStyle Style, const FOnEventContext
 			{
 				uint64 Cycle = EventData.GetValue<uint64>("Cycle");
 				double Time = Context.EventTime.AsSeconds(Cycle);
-				for (auto& KV : CountersMap)
+				for (auto& KV : FloatResetEveryFrameCountersMap)
 				{
-					if (KV.Value->IsResetEveryFrame())
-					{
-						STATS_ANALYZER_DEBUG_LOG(KV.Key, TEXT("%f RESET"), Time);
-						if (KV.Value->IsFloatingPoint())
-						{
-							KV.Value->SetValue(Time, 0.0);
-						}
-						else
-						{
-							KV.Value->SetValue(Time, (int64)0);
-						}
-					}
+					STATS_ANALYZER_DEBUG_LOG(KV.Key, TEXT("%f RESET"), Time);
+					KV.Value->SetValue(Time, 0.0);
+				}
+				for (auto& KV : Int64ResetEveryFrameCountersMap)
+				{
+					STATS_ANALYZER_DEBUG_LOG(KV.Key, TEXT("%f RESET"), Time);
+					KV.Value->SetValue(Time, 0ll);
 				}
 			}
 			break;
@@ -236,23 +253,15 @@ void FStatsAnalyzer::CreateFrameCounters()
 
 	FAnalysisSessionEditScope _(Session);
 
-	TMap<uint32, IEditableCounter*> ResetEveryFrameCountersMap;
-
-	for (auto& KV : CountersMap)
+	auto CreateFrameCounter = [this](IEditableCounter* EditableCounter, const ICounter*& Counter, IEditableCounter*& FrameCounter)
 	{
-		uint32 StatId = KV.Key;
-		IEditableCounter* Counter = KV.Value;
-
-		if (Counter->IsResetEveryFrame())
+		Counter = EditableCounterProvider.GetCounter(EditableCounter);
+		if (!Counter)
 		{
-			ResetEveryFrameCountersMap.Add(StatId, Counter);
+			return;
 		}
-	}
 
-	for (auto& KV : ResetEveryFrameCountersMap)
-	{
-		IEditableCounter* Counter = KV.Value;
-		IEditableCounter* FrameCounter = CounterProvider.CreateCounter();
+		FrameCounter = EditableCounterProvider.CreateEditableCounter();
 
 		FString FrameCounterName = FString(Counter->GetName()) + TEXT(" (1/frame)");
 		FrameCounter->SetName(Session.StoreString(*FrameCounterName));
@@ -268,49 +277,66 @@ void FStatsAnalyzer::CreateFrameCounters()
 		FrameCounter->SetIsFloatingPoint(Counter->IsFloatingPoint());
 		FrameCounter->SetIsResetEveryFrame(false);
 		FrameCounter->SetDisplayHint(Counter->GetDisplayHint());
+	};
 
-		constexpr double InfiniteTime = std::numeric_limits<double>::infinity();
+	constexpr double InfiniteTime = std::numeric_limits<double>::infinity();
 
-		if (Counter->IsFloatingPoint())
+	for (auto& KV : FloatResetEveryFrameCountersMap)
+	{
+		const ICounter* Counter = nullptr;
+		IEditableCounter* FrameCounter = nullptr;
+		CreateFrameCounter(KV.Value, Counter, FrameCounter);
+		if (!FrameCounter)
 		{
-			bool bFirst = true;
-			double FrameTime = 0.0;
-			double FrameValue = 0.0;
-			Counter->EnumerateFloatValues(-InfiniteTime, InfiniteTime, false, [FrameCounter, &bFirst, &FrameTime, &FrameValue](double Time, double Value)
-				{
-					if (bFirst && Value != 0.0)
-					{
-						bFirst = false;
-						FrameCounter->SetValue(Time, 0.0);
-					}
-					if (Value == 0.0 && FrameValue != 0.0)
-					{
-						FrameCounter->SetValue(FrameTime, FrameValue);
-					}
-					FrameTime = Time;
-					FrameValue = Value;
-				});
+			continue;
 		}
-		else
-		{
-			bool bFirst = true;
-			double FrameTime = 0.0;
-			int64 FrameValue = 0;
-			Counter->EnumerateValues(-InfiniteTime, InfiniteTime, false, [FrameCounter, &bFirst, &FrameTime, &FrameValue](double Time, int64 Value)
+
+		bool bFirst = true;
+		double FrameTime = 0.0;
+		double FrameValue = 0.0;
+		Counter->EnumerateFloatValues(-InfiniteTime, InfiniteTime, false, [FrameCounter, &bFirst, &FrameTime, &FrameValue](double Time, double Value)
+			{
+				if (bFirst && Value != 0.0)
 				{
-					if (bFirst && Value != 0)
-					{
-						bFirst = false;
-						FrameCounter->SetValue(Time, (int64)0);
-					}
-					if (Value == 0 && FrameValue != 0)
-					{
-						FrameCounter->SetValue(FrameTime, FrameValue);
-					}
-					FrameTime = Time;
-					FrameValue = Value;
-				});
+					bFirst = false;
+					FrameCounter->SetValue(Time, 0.0);
+				}
+				if (Value == 0.0 && FrameValue != 0.0)
+				{
+					FrameCounter->SetValue(FrameTime, FrameValue);
+				}
+				FrameTime = Time;
+				FrameValue = Value;
+			});
+	}
+
+	for (auto& KV : Int64ResetEveryFrameCountersMap)
+	{
+		const ICounter* Counter = nullptr;
+		IEditableCounter* FrameCounter = nullptr;
+		CreateFrameCounter(KV.Value, Counter, FrameCounter);
+		if (!FrameCounter)
+		{
+			continue;
 		}
+
+		bool bFirst = true;
+		double FrameTime = 0.0;
+		int64 FrameValue = 0;
+		Counter->EnumerateValues(-InfiniteTime, InfiniteTime, false, [FrameCounter, &bFirst, &FrameTime, &FrameValue](double Time, int64 Value)
+			{
+				if (bFirst && Value != 0)
+				{
+					bFirst = false;
+					FrameCounter->SetValue(Time, (int64)0);
+				}
+				if (Value == 0 && FrameValue != 0)
+				{
+					FrameCounter->SetValue(FrameTime, FrameValue);
+				}
+				FrameTime = Time;
+				FrameValue = Value;
+			});
 	}
 }
 
diff --git a/Engine/Source/Developer/TraceServices/Private/Analyzers/StatsTraceAnalysis.h b/Engine/Source/Developer/TraceServices/Private/Analyzers/StatsTraceAnalysis.h
--- a/Engine/Source/Developer/TraceServices/Private/Analyzers/StatsTraceAnalysis.h
+++ b/Engine/Source/Developer/TraceServices/Private/Analyzers/StatsTraceAnalysis.h
@@ -9,14 +9,16 @@ namespace TraceServices
 {
 
 class IAnalysisSession;
-class IEditableCounter;
+class ICounter;
 class ICounterProvider;
+class IEditableCounter;
+class IEditableCounterProvider;
 
 class FStatsAnalyzer
 	: public UE::Trace::IAnalyzer
 {
 public:
-	FStatsAnalyzer(IAnalysisSession& Session, ICounterProvider& CounterProvider);
+	FStatsAnalyzer(IAnalysisSession& Session, IEditableCounterProvider& InEditableCounterProvider);
 	virtual void OnAnalysisBegin(const FOnAnalysisContext& Context) override;
 	virtual void OnAnalysisEnd() override;
 	virtual bool OnEvent(uint16 RouteId, EStyle Style, const FOnEventContext& Context) override;
@@ -41,8 +43,10 @@ private:
 	TSharedRef<FThreadState> GetThreadState(uint32 ThreadId);
 
 	IAnalysisSession& Session;
-	ICounterProvider& CounterProvider;
-	TMap<uint32, IEditableCounter*> CountersMap;
+	IEditableCounterProvider& EditableCounterProvider;
+	TMap<uint32, IEditableCounter*> EditableCountersMap;
+	TMap<uint32, IEditableCounter*> Int64ResetEveryFrameCountersMap;
+	TMap<uint32, IEditableCounter*> FloatResetEveryFrameCountersMap;
 	TMap<uint32, TSharedRef<FThreadState>> ThreadStatesMap;
 };
 
diff --git a/Engine/Source/Developer/TraceServices/Private/Common/Utils.cpp b/Engine/Source/Developer/TraceServices/Private/Common/Utils.cpp
--- a/Engine/Source/Developer/TraceServices/Private/Common/Utils.cpp
+++ b/Engine/Source/Developer/TraceServices/Private/Common/Utils.cpp
@@ -1,6 +1,8 @@
 // Copyright Epic Games, Inc. All Rights Reserved.
 
 #include "Utils.h"
+#include "TraceServices/Utils.h"
+#include "Common/FormatArgs.h"
 #include "TraceServices/Containers/Tables.h"
 #include "Templates/SharedPointer.h"
 #include "HAL/FileManager.h"
@@ -72,4 +74,9 @@ void Table2Csv(const IUntypedTable& Table, const TCHAR* Filename)
 	OutputFile->Close();
 }
 
+void StringFormat(TCHAR* Out, uint64 MaxOut, TCHAR* Temp, uint64 MaxTemp, const TCHAR* FormatString, const uint8* FormatArgs)
+{
+	FFormatArgsHelper::Format(Out, MaxOut, Temp, MaxTemp, FormatString, FormatArgs);
+}
+
 } // namespace TraceServices
diff --git a/Engine/Source/Developer/TraceServices/Private/Model/Bookmarks.cpp b/Engine/Source/Developer/TraceServices/Private/Model/Bookmarks.cpp
--- a/Engine/Source/Developer/TraceServices/Private/Model/Bookmarks.cpp
+++ b/Engine/Source/Developer/TraceServices/Private/Model/Bookmarks.cpp
@@ -1,9 +1,9 @@
 // Copyright Epic Games, Inc. All Rights Reserved.
 
 #include "TraceServices/Model/Bookmarks.h"
+#include "Common/FormatArgs.h"
 #include "Model/BookmarksPrivate.h"
 #include "AnalysisServicePrivate.h"
-#include "Common/FormatArgs.h"
 
 namespace TraceServices
 {
@@ -19,32 +19,50 @@ FBookmarkProvider::FBookmarkProvider(IAnalysisSession& InSession)
 FBookmarkSpec& FBookmarkProvider::GetSpec(uint64 BookmarkPoint)
 {
 	Session.WriteAccessCheck();
-	if (SpecMap.Contains(BookmarkPoint))
-	{
-		return *SpecMap[BookmarkPoint].Get();
-	}
-	else
+	TSharedPtr<FBookmarkSpec>* Found = SpecMap.Find(BookmarkPoint);
+	if (Found)
 	{
-		TSharedPtr<FBookmarkSpec> Spec = MakeShared<FBookmarkSpec>();
-		Spec->File = TEXT("<unknown>");
-		Spec->FormatString = TEXT("<unknown>");
-		SpecMap.Add(BookmarkPoint, Spec);
-		return *Spec.Get();
+		return *Found->Get();
 	}
+
+	TSharedPtr<FBookmarkSpec> Spec = MakeShared<FBookmarkSpec>();
+	Spec->File = TEXT("<unknown>");
+	Spec->FormatString = TEXT("<unknown>");
+	SpecMap.Add(BookmarkPoint, Spec);
+	return *Spec.Get();
 }
 
-void FBookmarkProvider::AppendBookmark(double Time, uint64 BookmarkPoint, const uint8* FormatArgs)
+void FBookmarkProvider::UpdateBookmarkSpec(uint64 BookmarkPoint, const TCHAR* FormatString, const TCHAR* File, int32 Line)
 {
 	Session.WriteAccessCheck();
-	FBookmarkSpec& Spec = GetSpec(BookmarkPoint);
+	FBookmarkSpec& BookmarkSpec = GetSpec(BookmarkPoint);
+	BookmarkSpec.FormatString = FormatString;
+	BookmarkSpec.File = File;
+	BookmarkSpec.Line = Line;
+}
+
+void FBookmarkProvider::AppendBookmark(uint64 BookmarkPoint, double Time, const uint8* FormatArgs)
+{
+	Session.WriteAccessCheck();
+	FBookmarkSpec Spec = GetSpec(BookmarkPoint);
+	FFormatArgsHelper::Format(FormatBuffer, FormatBufferSize - 1, TempBuffer, FormatBufferSize - 1, Spec.FormatString, FormatArgs);
 	TSharedRef<FBookmarkInternal> Bookmark = MakeShared<FBookmarkInternal>();
 	Bookmark->Time = Time;
-	FFormatArgsHelper::Format(FormatBuffer, FormatBufferSize - 1, TempBuffer, FormatBufferSize - 1, Spec.FormatString, FormatArgs);
 	Bookmark->Text = Session.StoreString(FormatBuffer);
 	Bookmarks.Add(Bookmark);
 	Session.UpdateDurationSeconds(Time);
 }
 
+void FBookmarkProvider::AppendBookmark(uint64 BookmarkPoint, double Time, const TCHAR* Text)
+{
+	Session.WriteAccessCheck();
+	TSharedRef<FBookmarkInternal> Bookmark = MakeShared<FBookmarkInternal>();
+	Bookmark->Time = Time;
+	Bookmark->Text = Text;
+	Bookmarks.Add(Bookmark);
+	Session.UpdateDurationSeconds(Time);
+}
+
 void FBookmarkProvider::EnumerateBookmarks(double IntervalStart, double IntervalEnd, TFunctionRef<void(const FBookmark &)> Callback) const
 {
 	Session.ReadAccessCheck();
@@ -85,4 +103,9 @@ const IBookmarkProvider& ReadBookmarkProvider(const IAnalysisSession& Session)
 	return *Session.ReadProvider<IBookmarkProvider>(FBookmarkProvider::ProviderName);
 }
 
+IEditableBookmarkProvider& EditBookmarkProvider(IAnalysisSession& Session)
+{
+	return *Session.EditProvider<IEditableBookmarkProvider>(FBookmarkProvider::ProviderName);
+}
+
 } // namespace TraceServices
diff --git a/Engine/Source/Developer/TraceServices/Private/Model/BookmarksPrivate.h b/Engine/Source/Developer/TraceServices/Private/Model/BookmarksPrivate.h
--- a/Engine/Source/Developer/TraceServices/Private/Model/BookmarksPrivate.h
+++ b/Engine/Source/Developer/TraceServices/Private/Model/BookmarksPrivate.h
@@ -26,6 +26,7 @@ struct FBookmarkInternal
 
 class FBookmarkProvider
 	: public IBookmarkProvider
+	, public IEditableBookmarkProvider
 {
 public:
 	static const FName ProviderName;
@@ -35,18 +36,22 @@ public:
 
 	FBookmarkSpec& GetSpec(uint64 BookmarkPoint);
 	virtual uint64 GetBookmarkCount() const override { return Bookmarks.Num(); }
-	void AppendBookmark(double Time, uint64 BookmarkPoint, const uint8* FormatArgs);
 	virtual void EnumerateBookmarks(double IntervalStart, double IntervalEnd, TFunctionRef<void(const FBookmark&)> Callback) const override;
 
+	// implement IEditableBookmarkProvider
+	virtual void UpdateBookmarkSpec(uint64 BookmarkPoint, const TCHAR* FormatString, const TCHAR* File, int32 Line) override;
+	virtual void AppendBookmark(uint64 BookmarkPoint, double Time, const uint8* FormatArgs) override;
+	virtual void AppendBookmark(uint64 BookmarkPoint, double Time, const TCHAR* Text) override;
+
 private:
+	IAnalysisSession& Session;
+	TMap<uint64, TSharedPtr<FBookmarkSpec>> SpecMap;
+	TArray<TSharedRef<FBookmarkInternal>> Bookmarks;
+
 	enum
 	{
 		FormatBufferSize = 65536
 	};
-
-	IAnalysisSession& Session;
-	TMap<uint64, TSharedPtr<FBookmarkSpec>> SpecMap;
-	TArray<TSharedRef<FBookmarkInternal>> Bookmarks;
 	TCHAR FormatBuffer[FormatBufferSize];
 	TCHAR TempBuffer[FormatBufferSize];
 };
diff --git a/Engine/Source/Developer/TraceServices/Private/Model/Counters.cpp b/Engine/Source/Developer/TraceServices/Private/Model/Counters.cpp
--- a/Engine/Source/Developer/TraceServices/Private/Model/Counters.cpp
+++ b/Engine/Source/Developer/TraceServices/Private/Model/Counters.cpp
@@ -177,7 +177,12 @@ bool FCounterProvider::ReadCounter(uint32 CounterId, TFunctionRef<void(const ICo
 	return true;
 }
 
-IEditableCounter* FCounterProvider::CreateCounter()
+const ICounter* FCounterProvider::GetCounter(IEditableCounter* EditableCounter)
+{
+	return static_cast<FCounter*>(EditableCounter);
+}
+
+IEditableCounter* FCounterProvider::CreateEditableCounter()
 {
 	FCounter* Counter = new FCounter(Session.GetLinearAllocator(), FrameProvider.GetFrameStartTimes(TraceFrameType_Game));
 	Counters.Add(Counter);
@@ -194,9 +199,9 @@ const ICounterProvider& ReadCounterProvider(const IAnalysisSession& Session)
 	return *Session.ReadProvider<ICounterProvider>(FCounterProvider::ProviderName);
 }
 
-ICounterProvider& EditCounterProvider(IAnalysisSession& Session)
+IEditableCounterProvider& EditCounterProvider(IAnalysisSession& Session)
 {
-	return *Session.EditProvider<ICounterProvider>(FCounterProvider::ProviderName);
+	return *Session.EditProvider<IEditableCounterProvider>(FCounterProvider::ProviderName);
 }
 
 } // namespace TraceServices
diff --git a/Engine/Source/Developer/TraceServices/Private/Model/CountersPrivate.h b/Engine/Source/Developer/TraceServices/Private/Model/CountersPrivate.h
--- a/Engine/Source/Developer/TraceServices/Private/Model/CountersPrivate.h
+++ b/Engine/Source/Developer/TraceServices/Private/Model/CountersPrivate.h
@@ -190,7 +190,8 @@ private:
 };
 
 class FCounter
-	: public IEditableCounter
+	: public ICounter
+	, public IEditableCounter
 {
 public:
 	FCounter(ILinearAllocator& Allocator, const TArray64<double>& FrameStartTimes);
@@ -228,6 +229,7 @@ private:
 
 class FCounterProvider
 	: public ICounterProvider
+	, public IEditableCounterProvider
 {
 public:
 	static const FName ProviderName;
@@ -237,7 +239,10 @@ public:
 	virtual uint64 GetCounterCount() const override { return Counters.Num(); }
 	virtual void EnumerateCounters(TFunctionRef<void(uint32, const ICounter&)> Callback) const override;
 	virtual bool ReadCounter(uint32 CounterId, TFunctionRef<void(const ICounter&)> Callback) const override;
-	virtual IEditableCounter* CreateCounter() override;
+
+	// implement IEditableCounterProvider
+	virtual const ICounter* GetCounter(IEditableCounter* EditableCounter) override;
+	virtual IEditableCounter* CreateEditableCounter() override;
 	virtual void AddCounter(const ICounter* Counter) override;
 
 private:
diff --git a/Engine/Source/Developer/TraceServices/Private/Model/LoadTimeProfiler.cpp b/Engine/Source/Developer/TraceServices/Private/Model/LoadTimeProfiler.cpp
--- a/Engine/Source/Developer/TraceServices/Private/Model/LoadTimeProfiler.cpp
+++ b/Engine/Source/Developer/TraceServices/Private/Model/LoadTimeProfiler.cpp
@@ -92,9 +92,9 @@ void FLoadTimeProfilerProvider::FLoaderFrameCounter::EnumerateFloatValues(double
 	}
 }
 
-FLoadTimeProfilerProvider::FLoadTimeProfilerProvider(IAnalysisSession& InSession, ICounterProvider& InCounterProvider)
+FLoadTimeProfilerProvider::FLoadTimeProfilerProvider(IAnalysisSession& InSession, IEditableCounterProvider& InEditableCounterProvider)
 	: Session(InSession)
-	, CounterProvider(InCounterProvider)
+	, EditableCounterProvider(InEditableCounterProvider)
 	, ClassInfos(Session.GetLinearAllocator(), 16384)
 	, Requests(Session.GetLinearAllocator(), 16384)
 	, Packages(Session.GetLinearAllocator(), 16384)
@@ -568,20 +568,20 @@ void FLoadTimeProfilerProvider::CreateCounters()
 	{
 		for (int32 CounterType = 0; CounterType < FLoaderFrameCounter::LoaderFrameCounterType_Count; ++CounterType)
 		{
-			CounterProvider.AddCounter(new FLoaderFrameCounter(FLoaderFrameCounter::ELoaderFrameCounterType(CounterType), Frames));
+			EditableCounterProvider.AddCounter(new FLoaderFrameCounter(FLoaderFrameCounter::ELoaderFrameCounterType(CounterType), Frames));
 		}
 
-		ActiveIoDispatcherBatchesCounter = CounterProvider.CreateCounter();
+		ActiveIoDispatcherBatchesCounter = EditableCounterProvider.CreateEditableCounter();
 		ActiveIoDispatcherBatchesCounter->SetName(TEXT("AssetLoading/IoDispatcher/ActiveBatches"));
 
-		TotalIoDispatcherBytesReadCounter = CounterProvider.CreateCounter();
+		TotalIoDispatcherBytesReadCounter = EditableCounterProvider.CreateEditableCounter();
 		TotalIoDispatcherBytesReadCounter->SetName(TEXT("AssetLoading/IoDispatcher/TotalBytesRead"));
 		TotalIoDispatcherBytesReadCounter->SetDisplayHint(CounterDisplayHint_Memory);
 
-		LoadingPackagesCounter = CounterProvider.CreateCounter();
+		LoadingPackagesCounter = EditableCounterProvider.CreateEditableCounter();
 		LoadingPackagesCounter->SetName(TEXT("AssetLoading/AsyncLoading/ActiveLoadingPackages"));
 
-		TotalLoaderBytesLoadedCounter = CounterProvider.CreateCounter();
+		TotalLoaderBytesLoadedCounter = EditableCounterProvider.CreateEditableCounter();
 		TotalLoaderBytesLoadedCounter->SetName(TEXT("AssetLoading/AsyncLoading/TotalBytesLoaded"));
 		TotalLoaderBytesLoadedCounter->SetDisplayHint(CounterDisplayHint_Memory);
 
diff --git a/Engine/Source/Developer/TraceServices/Private/Model/LoadTimeProfilerPrivate.h b/Engine/Source/Developer/TraceServices/Private/Model/LoadTimeProfilerPrivate.h
--- a/Engine/Source/Developer/TraceServices/Private/Model/LoadTimeProfilerPrivate.h
+++ b/Engine/Source/Developer/TraceServices/Private/Model/LoadTimeProfilerPrivate.h
@@ -21,7 +21,7 @@ class FLoadTimeProfilerProvider
 public:
 	typedef TMonotonicTimeline<FLoadTimeProfilerCpuEvent> CpuTimelineInternal;
 
-	explicit FLoadTimeProfilerProvider(IAnalysisSession& Session, ICounterProvider& CounterProvider);
+	explicit FLoadTimeProfilerProvider(IAnalysisSession& Session, IEditableCounterProvider& EditableCounterProvider);
 	virtual ~FLoadTimeProfilerProvider() {}
 
 	virtual uint64 GetTimelineCount() const override { return CpuTimelines.Num(); }
@@ -100,7 +100,7 @@ private:
 	};
 
 	IAnalysisSession& Session;
-	ICounterProvider& CounterProvider;
+	IEditableCounterProvider& EditableCounterProvider;
 	TPagedArray<FClassInfo> ClassInfos;
 	TPagedArray<FLoadRequest> Requests;
 	TPagedArray<FPackageInfo> Packages;
diff --git a/Engine/Source/Developer/TraceServices/Private/Model/Log.cpp b/Engine/Source/Developer/TraceServices/Private/Model/Log.cpp
--- a/Engine/Source/Developer/TraceServices/Private/Model/Log.cpp
+++ b/Engine/Source/Developer/TraceServices/Private/Model/Log.cpp
@@ -42,6 +42,12 @@ FLogProvider::FLogProvider(IAnalysisSession& InSession)
 		AddColumn(&FLogMessageInternal::Message, TEXT("Message"));
 }
 
+uint64 FLogProvider::RegisterCategory()
+{
+	static uint64 IdGenerator = 0;
+	return IdGenerator++;
+}
+
 FLogCategoryInfo& FLogProvider::GetCategory(uint64 CategoryPointer)
 {
 	Session.WriteAccessCheck();
@@ -74,6 +80,46 @@ FLogMessageSpec& FLogProvider::GetMessageSpec(uint64 LogPoint)
 	}
 }
 
+void FLogProvider::UpdateMessageCategory(uint64 LogPoint, uint64 InCategoryPointer)
+{
+	Session.WriteAccessCheck();
+	FLogMessageSpec& LogMessageSpec = GetMessageSpec(LogPoint);
+	LogMessageSpec.Category = &GetCategory(InCategoryPointer);
+}
+
+void FLogProvider::UpdateMessageFormatString(uint64 LogPoint, const TCHAR* InFormatString)
+{
+	Session.WriteAccessCheck();
+	FLogMessageSpec& LogMessageSpec = GetMessageSpec(LogPoint);
+	LogMessageSpec.FormatString = InFormatString;
+}
+
+void FLogProvider::UpdateMessageFile(uint64 LogPoint, const TCHAR* InFile, int32 InLine)
+{
+	Session.WriteAccessCheck();
+	FLogMessageSpec& LogMessageSpec = GetMessageSpec(LogPoint);
+	LogMessageSpec.File = InFile;
+	LogMessageSpec.Line = InLine;
+}
+
+void FLogProvider::UpdateMessageVebosity(uint64 LogPoint, ELogVerbosity::Type InVerbosity)
+{
+	Session.WriteAccessCheck();
+	FLogMessageSpec& LogMessageSpec = GetMessageSpec(LogPoint);
+	LogMessageSpec.Verbosity = InVerbosity;
+}
+
+void FLogProvider::UpdateMessageSpec(uint64 LogPoint, uint64 InCategoryPointer, const TCHAR* InFormatString, const TCHAR* InFile, int32 InLine, ELogVerbosity::Type InVerbosity)
+{
+	Session.WriteAccessCheck();
+	FLogMessageSpec& LogMessageSpec = GetMessageSpec(LogPoint);
+	LogMessageSpec.Category = &GetCategory(InCategoryPointer);
+	LogMessageSpec.FormatString = InFormatString;
+	LogMessageSpec.File = InFile;
+	LogMessageSpec.Line = InLine;
+	LogMessageSpec.Verbosity = InVerbosity;
+}
+
 void FLogProvider::AppendMessage(uint64 LogPoint, double Time, const uint8* FormatArgs)
 {
 	Session.WriteAccessCheck();
@@ -89,6 +135,20 @@ void FLogProvider::AppendMessage(uint64 LogPoint, double Time, const uint8* Form
 	}
 }
 
+void FLogProvider::AppendMessage(uint64 LogPoint, double Time, const TCHAR* Text)
+{
+	Session.WriteAccessCheck();
+	FLogMessageSpec** FindSpec = SpecMap.Find(LogPoint);
+	if (FindSpec && (*FindSpec)->Verbosity != ELogVerbosity::SetColor)
+	{
+		FLogMessageInternal& InternalMessage = Messages.PushBack();
+		InternalMessage.Time = Time;
+		InternalMessage.Spec = *FindSpec;
+		InternalMessage.Message = Text;
+		Session.UpdateDurationSeconds(Time);
+	}
+}
+
 void FLogProvider::AppendMessage(uint64 LogPoint, double Time, const FString& Message)
 {
 	Session.WriteAccessCheck();
@@ -188,6 +248,11 @@ const ILogProvider& ReadLogProvider(const IAnalysisSession& Session)
 	return *Session.ReadProvider<ILogProvider>(FLogProvider::ProviderName);
 }
 
+IEditableLogProvider& EditLogProvider(IAnalysisSession& Session)
+{
+	return *Session.EditProvider<IEditableLogProvider>(FLogProvider::ProviderName);
+}
+
 void FormatString(TCHAR* OutputString, uint32 OutputStringCount, const TCHAR* FormatString, const uint8* FormatArgs)
 {
 	TCHAR* TempBuffer = (TCHAR*)FMemory_Alloca(OutputStringCount * sizeof(TCHAR));
diff --git a/Engine/Source/Developer/TraceServices/Private/Model/LogPrivate.h b/Engine/Source/Developer/TraceServices/Private/Model/LogPrivate.h
--- a/Engine/Source/Developer/TraceServices/Private/Model/LogPrivate.h
+++ b/Engine/Source/Developer/TraceServices/Private/Model/LogPrivate.h
@@ -31,23 +31,27 @@ struct FLogMessageInternal
 };
 
 class FLogProvider :
-	public ILogProvider
+	public ILogProvider,
+	public IEditableLogProvider
 {
 public:
-	enum
-	{
-		ReservedLogCategory_Bookmark = 0,
-		ReservedLogCategory_Screenshot = 1
-	};
-
 	static const FName ProviderName;
 
 	explicit FLogProvider(IAnalysisSession& Session);
 	virtual ~FLogProvider() {}
 
-	FLogCategoryInfo& GetCategory(uint64 CategoryPointer);
+	// implement IEditableLogProvider
+	virtual uint64 RegisterCategory() override;
+	virtual FLogCategoryInfo& GetCategory(uint64 CategoryPointer) override;
+	virtual void UpdateMessageCategory(uint64 LogPoint, uint64 InCategoryPointer) override;
+	virtual void UpdateMessageFormatString(uint64 LogPoint, const TCHAR* InFormatString) override;
+	virtual void UpdateMessageFile(uint64 LogPoint, const TCHAR* InFile, int32 InLine) override;
+	virtual void UpdateMessageVebosity(uint64 LogPoint, ELogVerbosity::Type InVerbosity) override;
+	virtual void UpdateMessageSpec(uint64 LogPoint, uint64 InCategoryPointer, const TCHAR* InFormatString, const TCHAR* InFile, int32 InLine, ELogVerbosity::Type InVerbosity) override;
+	virtual void AppendMessage(uint64 LogPoint, double Time, const uint8* FormatArgs) override;
+	virtual void AppendMessage(uint64 LogPoint, double Time, const TCHAR* Text) override;
+
 	FLogMessageSpec& GetMessageSpec(uint64 LogPoint);
-	void AppendMessage(uint64 LogPoint, double Time, const uint8* FormatArgs);
 	void AppendMessage(uint64 LogPoint, double Time, const FString& Message);
 
 	virtual uint64 GetMessageCount() const override;
diff --git a/Engine/Source/Developer/TraceServices/Private/Model/ModuleProvider.cpp b/Engine/Source/Developer/TraceServices/Private/Model/ModuleProvider.cpp
--- a/Engine/Source/Developer/TraceServices/Private/Model/ModuleProvider.cpp
+++ b/Engine/Source/Developer/TraceServices/Private/Model/ModuleProvider.cpp
@@ -467,19 +467,20 @@ void FResolvedSymbolFilter::Update(FResolvedSymbol& InSymbol) const
 }
 
 /////////////////////////////////////////////////////////////////////
-IModuleAnalysisProvider* CreateModuleProvider(IAnalysisSession& InSession, const FAnsiStringView& InSymbolFormat)
+TSharedPtr<IModuleAnalysisProvider> CreateModuleProvider(IAnalysisSession& InSession, const FAnsiStringView& InSymbolFormat)
 {
-	IModuleAnalysisProvider* Provider(nullptr);
+	TSharedPtr<IModuleAnalysisProvider> Provider;
+
 #if PLATFORM_WINDOWS && USE_SYMSLIB
 	if (!Provider && (InSymbolFormat.Equals("pdb") || InSymbolFormat.Equals("dwarf")))
 	{
-		Provider = new TModuleProvider<FSymslibResolver>(InSession);
+		Provider = MakeShared<TModuleProvider<FSymslibResolver>>(InSession);
 	}
 #endif // PLATFORM_WINDOWS && USE_SYMSLIB
 #if PLATFORM_WINDOWS && USE_DBGHELP
 	if (!Provider && InSymbolFormat.Equals("pdb"))
 	{
-		Provider = new TModuleProvider<FDbgHelpResolver>(InSession);
+		Provider = MakeShared<TModuleProvider<FDbgHelpResolver>>(InSession);
 	}
 #endif // PLATFORM_WINDOWS && USE_DBGHELP
 	return Provider;
diff --git a/Engine/Source/Developer/TraceServices/Private/Model/ModuleProvider.h b/Engine/Source/Developer/TraceServices/Private/Model/ModuleProvider.h
--- a/Engine/Source/Developer/TraceServices/Private/Model/ModuleProvider.h
+++ b/Engine/Source/Developer/TraceServices/Private/Model/ModuleProvider.h
@@ -19,6 +19,6 @@ public:
 };
 
 /** Create a module provider with the given symbol format. */
-IModuleAnalysisProvider* CreateModuleProvider(IAnalysisSession& Session, const FAnsiStringView& SymbolFormat);
+TSharedPtr<IModuleAnalysisProvider> CreateModuleProvider(IAnalysisSession& Session, const FAnsiStringView& SymbolFormat);
 
 } // namespace TraceServices
diff --git a/Engine/Source/Developer/TraceServices/Private/Model/TasksProfiler.cpp b/Engine/Source/Developer/TraceServices/Private/Model/TasksProfiler.cpp
--- a/Engine/Source/Developer/TraceServices/Private/Model/TasksProfiler.cpp
+++ b/Engine/Source/Developer/TraceServices/Private/Model/TasksProfiler.cpp
@@ -13,7 +13,7 @@ namespace TraceServices
 {
 	FTasksProvider::FTasksProvider(IAnalysisSession& InSession)
 		: Session(InSession)
-		, CounterProvider(EditCounterProvider(Session))
+		, EditableCounterProvider(EditCounterProvider(Session))
 	{
 #if 0
 		////////////////////////////////
@@ -61,32 +61,32 @@ namespace TraceServices
 		check(bCountersCreated == false);
 		FAnalysisSessionEditScope _(Session);
 
-		WaitingForPrerequisitesTasksCounter = CounterProvider.CreateCounter();
+		WaitingForPrerequisitesTasksCounter = EditableCounterProvider.CreateEditableCounter();
 		WaitingForPrerequisitesTasksCounter->SetName(TEXT("Tasks::WaitingForPrerequisitesTasks"));
 		WaitingForPrerequisitesTasksCounter->SetDescription(TEXT("Tasks: the number of tasks waiting for prerequisites (blocked by dependency)"));
 		WaitingForPrerequisitesTasksCounter->SetIsFloatingPoint(false);
 
-		TaskLatencyCounter = CounterProvider.CreateCounter();
+		TaskLatencyCounter = EditableCounterProvider.CreateEditableCounter();
 		TaskLatencyCounter->SetName(TEXT("Tasks::TaskLatency"));
 		TaskLatencyCounter->SetDescription(TEXT("Tasks: tasks latency - the time from scheduling to execution start"));
 		TaskLatencyCounter->SetIsFloatingPoint(true);
 
-		ScheduledTasksCounter = CounterProvider.CreateCounter();
+		ScheduledTasksCounter = EditableCounterProvider.CreateEditableCounter();
 		ScheduledTasksCounter->SetName(TEXT("Tasks::ScheduledTasks"));
 		ScheduledTasksCounter->SetDescription(TEXT("Tasks: number of scheduled tasks excluding named threads (the size of the queue)"));
 		ScheduledTasksCounter->SetIsFloatingPoint(false);
 
-		NamedThreadsScheduledTasksCounter = CounterProvider.CreateCounter();
+		NamedThreadsScheduledTasksCounter = EditableCounterProvider.CreateEditableCounter();
 		NamedThreadsScheduledTasksCounter->SetName(TEXT("Tasks::NamedThreadsScheduledTasks"));
 		NamedThreadsScheduledTasksCounter->SetDescription(TEXT("Tasks: number of scheduled tasks for named threads"));
 		NamedThreadsScheduledTasksCounter->SetIsFloatingPoint(false);
 
-		RunningTasksCounter = CounterProvider.CreateCounter();
+		RunningTasksCounter = EditableCounterProvider.CreateEditableCounter();
 		RunningTasksCounter->SetName(TEXT("Tasks::RunningTasks"));
 		RunningTasksCounter->SetDescription(TEXT("Tasks: level of parallelism - the number of tasks being executed"));
 		RunningTasksCounter->SetIsFloatingPoint(false);
 
-		ExecutionTimeCounter = CounterProvider.CreateCounter();
+		ExecutionTimeCounter = EditableCounterProvider.CreateEditableCounter();
 		ExecutionTimeCounter->SetName(TEXT("Tasks::ExecutionTime"));
 		ExecutionTimeCounter->SetDescription(TEXT("Tasks: execution time"));
 		ExecutionTimeCounter->SetIsFloatingPoint(true);
diff --git a/Engine/Source/Developer/TraceServices/Private/Model/TasksProfilerPrivate.h b/Engine/Source/Developer/TraceServices/Private/Model/TasksProfilerPrivate.h
--- a/Engine/Source/Developer/TraceServices/Private/Model/TasksProfilerPrivate.h
+++ b/Engine/Source/Developer/TraceServices/Private/Model/TasksProfilerPrivate.h
@@ -91,7 +91,7 @@ namespace TraceServices
 		// a bool that is set when the counters are created
 		bool bCountersCreated = false;
 
-		ICounterProvider& CounterProvider;
+		IEditableCounterProvider& EditableCounterProvider;
 		IEditableCounter* TaskLatencyCounter;
 		IEditableCounter* WaitingForPrerequisitesTasksCounter;
 		IEditableCounter* ScheduledTasksCounter;
diff --git a/Engine/Source/Developer/TraceServices/Private/Model/Threads.cpp b/Engine/Source/Developer/TraceServices/Private/Model/Threads.cpp
--- a/Engine/Source/Developer/TraceServices/Private/Model/Threads.cpp
+++ b/Engine/Source/Developer/TraceServices/Private/Model/Threads.cpp
@@ -219,4 +219,9 @@ const IThreadProvider& ReadThreadProvider(const IAnalysisSession& Session)
 	return *Session.ReadProvider<IThreadProvider>(FThreadProvider::ProviderName);
 }
 
+IEditableThreadProvider& EditThreadProvider(IAnalysisSession& Session)
+{
+	return *Session.EditProvider<IEditableThreadProvider>(FThreadProvider::ProviderName);
+}
+
 } // namespace TraceServices
diff --git a/Engine/Source/Developer/TraceServices/Private/Model/ThreadsPrivate.h b/Engine/Source/Developer/TraceServices/Private/Model/ThreadsPrivate.h
--- a/Engine/Source/Developer/TraceServices/Private/Model/ThreadsPrivate.h
+++ b/Engine/Source/Developer/TraceServices/Private/Model/ThreadsPrivate.h
@@ -13,6 +13,7 @@ namespace TraceServices
 
 class FThreadProvider
 	: public IThreadProvider
+	, public IEditableThreadProvider
 {
 public:
 	static const FName ProviderName;
@@ -21,7 +22,7 @@ public:
 	virtual ~FThreadProvider();
 
 	void AddGameThread(uint32 Id);
-	void AddThread(uint32 Id, const TCHAR* Name, EThreadPriority Priority);
+	virtual void AddThread(uint32 Id, const TCHAR* Name, EThreadPriority Priority) override;
 	void SetThreadPriority(uint32 Id, EThreadPriority Priority);
 	void SetThreadGroup(uint32 Id, const TCHAR* GroupName);
 	virtual uint64 GetModCount() const override { return ModCount; }
diff --git a/Engine/Source/Developer/TraceServices/Private/Model/TimingProfiler.cpp b/Engine/Source/Developer/TraceServices/Private/Model/TimingProfiler.cpp
--- a/Engine/Source/Developer/TraceServices/Private/Model/TimingProfiler.cpp
+++ b/Engine/Source/Developer/TraceServices/Private/Model/TimingProfiler.cpp
@@ -133,6 +133,11 @@ uint32 FTimingProfilerProvider::AddMetadata(uint32 MasterTimerId, TArray<uint8>&
 	return ~MetadataId;
 }
 
+IEditableTimeline<FTimingProfilerEvent>& FTimingProfilerProvider::GetCpuThreadEditableTimeline(uint32 ThreadId)
+{
+	return EditCpuThreadTimeline(ThreadId);
+}
+
 FTimingProfilerProvider::TimelineInternal& FTimingProfilerProvider::EditCpuThreadTimeline(uint32 ThreadId)
 {
 	Session.WriteAccessCheck();
diff --git a/Engine/Source/Developer/TraceServices/Private/Model/TimingProfilerPrivate.h b/Engine/Source/Developer/TraceServices/Private/Model/TimingProfilerPrivate.h
--- a/Engine/Source/Developer/TraceServices/Private/Model/TimingProfilerPrivate.h
+++ b/Engine/Source/Developer/TraceServices/Private/Model/TimingProfilerPrivate.h
@@ -16,17 +16,15 @@ class FStringStore;
 class FTimingProfilerProvider
 	: public ITimingProfilerProvider
 	, public ITimingProfilerTimerReader
+	, public IEditableTimingProfilerProvider
 {
 public:
-	typedef TMonotonicTimeline<FTimingProfilerEvent> TimelineInternal;
+	typedef TMonotonicTimeline<FTimingProfilerEvent> 	TimelineInternal;
 
 	explicit FTimingProfilerProvider(IAnalysisSession& InSession);
 	virtual ~FTimingProfilerProvider();
-	uint32 AddCpuTimer(FStringView Name, const TCHAR* File = nullptr, uint32 Line = 0);
 	uint32 AddGpuTimer(FStringView Name, const TCHAR* File = nullptr, uint32 Line = 0);
 	void SetTimerName(uint32 TimerId, FStringView Name);
-	void SetTimerNameAndLocation(uint32 TimerId, FStringView Name, const TCHAR* File, uint32 Line);
-	uint32 AddMetadata(uint32 MasterTimerId, TArray<uint8>&& Metadata);
 	TimelineInternal& EditCpuThreadTimeline(uint32 ThreadId);
 	TimelineInternal& EditGpuTimeline();
 	TimelineInternal& EditGpu2Timeline();
@@ -41,7 +39,13 @@ public:
 	virtual ITimingProfilerButterfly* CreateButterfly(double IntervalStart, double IntervalEnd, TFunctionRef<bool(uint32)> CpuThreadFilter, bool IncludeGpu) const override;
 	virtual const FTimingProfilerTimer* GetTimer(uint32 TimerId) const override;
 	virtual uint32 GetTimerCount() const override;
+
+	// implementing IEditableTimingProfilerProvider
+	virtual uint32 AddCpuTimer(FStringView Name, const TCHAR* File = nullptr, uint32 Line = 0) override;
+	virtual void SetTimerNameAndLocation(uint32 TimerId, FStringView Name, const TCHAR* File, uint32 Line) override;
+	virtual uint32 AddMetadata(uint32 MasterTimerId, TArray<uint8>&& Metadata) override;
 	virtual TArrayView<const uint8> GetMetadata(uint32 TimerId) const override;
+	virtual IEditableTimeline<FTimingProfilerEvent>& GetCpuThreadEditableTimeline(uint32 ThreadId) override;
 
 private:
 	FTimingProfilerTimer& AddTimerInternal(FStringView Name, const TCHAR* File, uint32 Line, bool IsGpuEvent);
diff --git a/Engine/Source/Developer/TraceServices/Private/Modules/CookProfilerModule.cpp b/Engine/Source/Developer/TraceServices/Private/Modules/CookProfilerModule.cpp
--- a/Engine/Source/Developer/TraceServices/Private/Modules/CookProfilerModule.cpp
+++ b/Engine/Source/Developer/TraceServices/Private/Modules/CookProfilerModule.cpp
@@ -19,7 +19,7 @@ void FCookProfilerModule::GetModuleInfo(FModuleInfo& OutModuleInfo)
 
 void FCookProfilerModule::OnAnalysisBegin(IAnalysisSession& Session)
 {
-	FCookProfilerProvider* CookProfilerProvider = new FCookProfilerProvider(Session);
+	TSharedPtr<FCookProfilerProvider> CookProfilerProvider = MakeShared<FCookProfilerProvider>(Session);
 	Session.AddProvider(CookProfilerProviderName, CookProfilerProvider);
 	Session.AddAnalyzer(new FCookAnalyzer(Session, *CookProfilerProvider));
 }
diff --git a/Engine/Source/Developer/TraceServices/Private/Modules/CountersModule.cpp b/Engine/Source/Developer/TraceServices/Private/Modules/CountersModule.cpp
--- a/Engine/Source/Developer/TraceServices/Private/Modules/CountersModule.cpp
+++ b/Engine/Source/Developer/TraceServices/Private/Modules/CountersModule.cpp
@@ -19,8 +19,8 @@ void FCountersModule::GetModuleInfo(FModuleInfo& OutModuleInfo)
 	
 void FCountersModule::OnAnalysisBegin(IAnalysisSession& Session)
 {
-	ICounterProvider& CounterProvider = EditCounterProvider(Session);
-	Session.AddAnalyzer(new FCountersAnalyzer(Session, CounterProvider));
+	IEditableCounterProvider& EditableCounterProvider = EditCounterProvider(Session);
+	Session.AddAnalyzer(new FCountersAnalyzer(Session, EditableCounterProvider));
 }
 
 void FCountersModule::GetLoggers(TArray<const TCHAR *>& OutLoggers)
diff --git a/Engine/Source/Developer/TraceServices/Private/Modules/CsvProfilerModule.cpp b/Engine/Source/Developer/TraceServices/Private/Modules/CsvProfilerModule.cpp
--- a/Engine/Source/Developer/TraceServices/Private/Modules/CsvProfilerModule.cpp
+++ b/Engine/Source/Developer/TraceServices/Private/Modules/CsvProfilerModule.cpp
@@ -23,10 +23,10 @@ void FCsvProfilerModule::OnAnalysisBegin(IAnalysisSession& Session)
 {
 	const IFrameProvider& FrameProvider = ReadFrameProvider(Session);
 	const IThreadProvider& ThreadProvider = ReadThreadProvider(Session);
-	ICounterProvider& CounterProvider = EditCounterProvider(Session);
-	FCsvProfilerProvider* CsvProfilerProvider = new FCsvProfilerProvider(Session);
+	IEditableCounterProvider& EditableCounterProvider = EditCounterProvider(Session);
+	TSharedPtr<FCsvProfilerProvider> CsvProfilerProvider = MakeShared<FCsvProfilerProvider>(Session);
 	Session.AddProvider(CsvProfilerProviderName, CsvProfilerProvider);
-	Session.AddAnalyzer(new FCsvProfilerAnalyzer(Session, *CsvProfilerProvider, CounterProvider, FrameProvider, ThreadProvider));
+	Session.AddAnalyzer(new FCsvProfilerAnalyzer(Session, *CsvProfilerProvider, EditableCounterProvider, FrameProvider, ThreadProvider));
 }
 
 void FCsvProfilerModule::GenerateReports(const IAnalysisSession& Session, const TCHAR* CmdLine, const TCHAR* OutputDirectory)
diff --git a/Engine/Source/Developer/TraceServices/Private/Modules/DiagnosticsModule.cpp b/Engine/Source/Developer/TraceServices/Private/Modules/DiagnosticsModule.cpp
--- a/Engine/Source/Developer/TraceServices/Private/Modules/DiagnosticsModule.cpp
+++ b/Engine/Source/Developer/TraceServices/Private/Modules/DiagnosticsModule.cpp
@@ -17,10 +17,10 @@ void FDiagnosticsModule::GetModuleInfo(FModuleInfo& OutModuleInfo)
 
 void FDiagnosticsModule::OnAnalysisBegin(IAnalysisSession& Session)
 {
-	FDiagnosticsProvider* DiagnosticsProvider = new FDiagnosticsProvider(Session);
+	TSharedPtr<FDiagnosticsProvider> DiagnosticsProvider = MakeShared<FDiagnosticsProvider>(Session);
 	Session.AddProvider(GetDiagnosticsProviderName(), DiagnosticsProvider);
 
-	Session.AddAnalyzer(new FDiagnosticsAnalyzer(Session, DiagnosticsProvider));
+	Session.AddAnalyzer(new FDiagnosticsAnalyzer(Session, DiagnosticsProvider.Get()));
 }
 
 void FDiagnosticsModule::GetLoggers(TArray<const TCHAR *>& OutLoggers)
diff --git a/Engine/Source/Developer/TraceServices/Private/Modules/LoadTimeProfilerModule.cpp b/Engine/Source/Developer/TraceServices/Private/Modules/LoadTimeProfilerModule.cpp
--- a/Engine/Source/Developer/TraceServices/Private/Modules/LoadTimeProfilerModule.cpp
+++ b/Engine/Source/Developer/TraceServices/Private/Modules/LoadTimeProfilerModule.cpp
@@ -36,10 +36,10 @@ void FLoadTimeProfilerModule::GetModuleInfo(FModuleInfo& OutModuleInfo)
 
 void FLoadTimeProfilerModule::OnAnalysisBegin(IAnalysisSession& Session)
 {
-	FLoadTimeProfilerProvider* LoadTimeProfilerProvider = new FLoadTimeProfilerProvider(Session, EditCounterProvider(Session));
+	TSharedPtr<FLoadTimeProfilerProvider> LoadTimeProfilerProvider = MakeShared<FLoadTimeProfilerProvider>(Session, EditCounterProvider(Session));
 	Session.AddProvider(LoadTimeProfilerProviderName, LoadTimeProfilerProvider);
 	Session.AddAnalyzer(new FAsyncLoadingTraceAnalyzer(Session, *LoadTimeProfilerProvider));
-	FFileActivityProvider* FileActivityProvider = new FFileActivityProvider(Session);
+	TSharedPtr<FFileActivityProvider> FileActivityProvider = MakeShared<FFileActivityProvider>(Session);
 	Session.AddProvider(FileActivityProviderName, FileActivityProvider);
 	Session.AddAnalyzer(new FPlatformFileTraceAnalyzer(Session, *FileActivityProvider));
 }
diff --git a/Engine/Source/Developer/TraceServices/Private/Modules/MemoryModule.cpp b/Engine/Source/Developer/TraceServices/Private/Modules/MemoryModule.cpp
--- a/Engine/Source/Developer/TraceServices/Private/Modules/MemoryModule.cpp
+++ b/Engine/Source/Developer/TraceServices/Private/Modules/MemoryModule.cpp
@@ -25,25 +25,25 @@ void FMemoryModule::GetModuleInfo(FModuleInfo& OutModuleInfo)
 void FMemoryModule::OnAnalysisBegin(IAnalysisSession& Session)
 {
 	// LLM Tag Stats
-	FMemoryProvider* MemoryProvider = new FMemoryProvider(Session);
+	TSharedPtr<FMemoryProvider> MemoryProvider = MakeShared<FMemoryProvider>(Session);
 	Session.AddProvider(GetMemoryProviderName(), MemoryProvider);
-	Session.AddAnalyzer(new FMemoryAnalyzer(Session, MemoryProvider));
+	Session.AddAnalyzer(new FMemoryAnalyzer(Session, MemoryProvider.Get()));
 
 	// Module
 	Session.AddAnalyzer(new FModuleAnalyzer(Session));
 
 	// Callstack
-	FCallstacksProvider* CallstacksProvider = new FCallstacksProvider(Session);
+	TSharedPtr<FCallstacksProvider> CallstacksProvider = MakeShared<FCallstacksProvider>(Session);
 	Session.AddProvider(GetCallstacksProviderName(), CallstacksProvider);
-	Session.AddAnalyzer(new FCallstacksAnalyzer(Session, CallstacksProvider));
+	Session.AddAnalyzer(new FCallstacksAnalyzer(Session, CallstacksProvider.Get()));
 
 	// Metadata
-	FMetadataProvider* MetadataProvider = new FMetadataProvider(Session);
+	TSharedPtr<FMetadataProvider> MetadataProvider = MakeShared<FMetadataProvider>(Session);
 	Session.AddProvider(GetMetadataProviderName(), MetadataProvider);
-	Session.AddAnalyzer(new FMetadataAnalysis(Session, MetadataProvider));
+	Session.AddAnalyzer(new FMetadataAnalysis(Session, MetadataProvider.Get()));
 	
 	// Allocations
-	FAllocationsProvider* AllocationsProvider = new FAllocationsProvider(Session, *MetadataProvider);
+	TSharedPtr<FAllocationsProvider> AllocationsProvider = MakeShared<FAllocationsProvider>(Session, *MetadataProvider);
 	Session.AddProvider(GetAllocationsProviderName(), AllocationsProvider);
 	Session.AddAnalyzer(new FAllocationsAnalyzer(Session, *AllocationsProvider, *MetadataProvider));
 
diff --git a/Engine/Source/Developer/TraceServices/Private/Modules/NetProfilerModule.cpp b/Engine/Source/Developer/TraceServices/Private/Modules/NetProfilerModule.cpp
--- a/Engine/Source/Developer/TraceServices/Private/Modules/NetProfilerModule.cpp
+++ b/Engine/Source/Developer/TraceServices/Private/Modules/NetProfilerModule.cpp
@@ -18,7 +18,7 @@ void FNetProfilerModule::GetModuleInfo(FModuleInfo& OutModuleInfo)
 
 void FNetProfilerModule::OnAnalysisBegin(IAnalysisSession& Session)
 {
-	FNetProfilerProvider* NetProfilerProvider = new FNetProfilerProvider(Session);
+	TSharedPtr<FNetProfilerProvider> NetProfilerProvider = MakeShared<FNetProfilerProvider>(Session);
 	Session.AddProvider(GetNetProfilerProviderName(), NetProfilerProvider);
 
 	Session.AddAnalyzer(new FNetTraceAnalyzer(Session, *NetProfilerProvider));
diff --git a/Engine/Source/Developer/TraceServices/Private/Modules/PlatformEventsModule.cpp b/Engine/Source/Developer/TraceServices/Private/Modules/PlatformEventsModule.cpp
--- a/Engine/Source/Developer/TraceServices/Private/Modules/PlatformEventsModule.cpp
+++ b/Engine/Source/Developer/TraceServices/Private/Modules/PlatformEventsModule.cpp
@@ -21,8 +21,8 @@ void FPlatformEventsModule::OnAnalysisBegin(IAnalysisSession& InSession)
 {
 	FAnalysisSession& Session = static_cast<FAnalysisSession&>(InSession);
 
-	FContextSwitchesProvider* ContextSwitchesProvider = new FContextSwitchesProvider(Session);
-	FStackSamplesProvider* StackSamplesProvider = new FStackSamplesProvider(Session);
+	TSharedPtr<FContextSwitchesProvider> ContextSwitchesProvider = MakeShared<FContextSwitchesProvider>(Session);
+	TSharedPtr<FStackSamplesProvider> StackSamplesProvider = MakeShared<FStackSamplesProvider>(Session);
 	
 	Session.AddProvider(FContextSwitchesProvider::ProviderName, ContextSwitchesProvider);
 	Session.AddProvider(FStackSamplesProvider::ProviderName, StackSamplesProvider);
diff --git a/Engine/Source/Developer/TraceServices/Private/Modules/StatsModule.cpp b/Engine/Source/Developer/TraceServices/Private/Modules/StatsModule.cpp
--- a/Engine/Source/Developer/TraceServices/Private/Modules/StatsModule.cpp
+++ b/Engine/Source/Developer/TraceServices/Private/Modules/StatsModule.cpp
@@ -17,8 +17,8 @@ void FStatsModule::GetModuleInfo(FModuleInfo& OutModuleInfo)
 	
 void FStatsModule::OnAnalysisBegin(IAnalysisSession& Session)
 {
-	ICounterProvider& CounterProvider = EditCounterProvider(Session);
-	Session.AddAnalyzer(new FStatsAnalyzer(Session, CounterProvider));
+	IEditableCounterProvider& EditableCounterProvider = EditCounterProvider(Session);
+	Session.AddAnalyzer(new FStatsAnalyzer(Session, EditableCounterProvider));
 }
 
 void FStatsModule::GetLoggers(TArray<const TCHAR *>& OutLoggers)
diff --git a/Engine/Source/Developer/TraceServices/Private/Modules/TasksModule.cpp b/Engine/Source/Developer/TraceServices/Private/Modules/TasksModule.cpp
--- a/Engine/Source/Developer/TraceServices/Private/Modules/TasksModule.cpp
+++ b/Engine/Source/Developer/TraceServices/Private/Modules/TasksModule.cpp
@@ -19,7 +19,7 @@ namespace TraceServices
 
 	void FTasksModule::OnAnalysisBegin(IAnalysisSession& Session)
 	{
-		FTasksProvider* TasksProvider = new FTasksProvider(Session);
+		TSharedPtr<FTasksProvider> TasksProvider = MakeShared<FTasksProvider>(Session);
 		Session.AddProvider(TasksProviderName, TasksProvider);
 		Session.AddAnalyzer(new FTasksAnalyzer(Session, *TasksProvider));
 	}
diff --git a/Engine/Source/Developer/TraceServices/Private/Modules/TimingProfilerModule.cpp b/Engine/Source/Developer/TraceServices/Private/Modules/TimingProfilerModule.cpp
--- a/Engine/Source/Developer/TraceServices/Private/Modules/TimingProfilerModule.cpp
+++ b/Engine/Source/Developer/TraceServices/Private/Modules/TimingProfilerModule.cpp
@@ -23,12 +23,13 @@ void FTimingProfilerModule::OnAnalysisBegin(IAnalysisSession& InSession)
 {
 	FAnalysisSession& Session = static_cast<FAnalysisSession&>(InSession);
 	
-	auto* ThreadProvider = Session.EditProvider<FThreadProvider>(FThreadProvider::ProviderName);
-	check(ThreadProvider != nullptr);
+	IEditableThreadProvider* EditableThreadProvider = Session.EditProvider<IEditableThreadProvider>(FThreadProvider::ProviderName);
+	check(EditableThreadProvider != nullptr);
 
-	FTimingProfilerProvider* TimingProfilerProvider = new FTimingProfilerProvider(Session);
-	Session.AddProvider(TimingProfilerProviderName, TimingProfilerProvider);
-	Session.AddAnalyzer(new FCpuProfilerAnalyzer(Session, *TimingProfilerProvider, *ThreadProvider));
+	// see comment in FAnalysisService::StartAnalysis for details
+	TSharedPtr<FTimingProfilerProvider> TimingProfilerProvider = MakeShared<FTimingProfilerProvider>(Session);
+	Session.AddProvider(TimingProfilerProviderName, TSharedPtr<ITimingProfilerProvider>(TimingProfilerProvider), TSharedPtr<IEditableTimingProfilerProvider>(TimingProfilerProvider));
+	Session.AddAnalyzer(new FCpuProfilerAnalyzer(Session, *TimingProfilerProvider, *EditableThreadProvider));
 	Session.AddAnalyzer(new FGpuProfilerAnalyzer(Session, *TimingProfilerProvider));
 }
 
@@ -43,4 +44,9 @@ const ITimingProfilerProvider* ReadTimingProfilerProvider(const IAnalysisSession
 	return Session.ReadProvider<ITimingProfilerProvider>(TimingProfilerProviderName);
 }
 
+IEditableTimingProfilerProvider* EditTimingProfilerProvider(IAnalysisSession& Session)
+{
+	return Session.EditProvider<IEditableTimingProfilerProvider>(TimingProfilerProviderName);
+}
+
 } // namespace TraceServices
diff --git a/Engine/Source/Developer/TraceServices/Public/Model/MonotonicTimeline.h b/Engine/Source/Developer/TraceServices/Public/Model/MonotonicTimeline.h
--- a/Engine/Source/Developer/TraceServices/Public/Model/MonotonicTimeline.h
+++ b/Engine/Source/Developer/TraceServices/Public/Model/MonotonicTimeline.h
@@ -29,9 +29,35 @@ struct FMonotonicTimelineDefaultSettings
 	}
 };
 
+/*
+* An interface that can consume timed serial events (a timeline).
+*/
+template<typename InEventType>
+class IEditableTimeline
+{
+public:
+	virtual ~IEditableTimeline() = default;
+
+	/*
+	* Begin a new timed event.
+	* 
+	* @param StartTime	The starting timestamp of the event in seconds.
+	* @param Event		The event information.
+	*/
+	virtual void AppendBeginEvent(double StartTime, const InEventType& Event) = 0;
+
+	/*
+	* End a new timed event. This ends the event started by the prior call to AppendBeginEvent.
+	* 
+	* @param EndTime	The ending timestamp of the event in seconds.
+	*/
+	virtual void AppendEndEvent(double EndTime) = 0;
+};
+
 template<typename InEventType, typename SettingsType = FMonotonicTimelineDefaultSettings>
 class TMonotonicTimeline
 	: public ITimeline<InEventType>
+	, public IEditableTimeline<InEventType>
 {
 	friend class FEnumerateAsyncTask<InEventType, SettingsType>;
 
@@ -55,7 +81,7 @@ public:
 
 	virtual ~TMonotonicTimeline() = default;
 
-	void AppendBeginEvent(double StartTime, const EventType& Event)
+	virtual void AppendBeginEvent(double StartTime, const EventType& Event) override
 	{
 		int32 CurrentDepth = DetailLevels[0].InsertionState.CurrentDepth;
 		if (CurrentDepth >= SettingsType::MaxDepth)
@@ -121,7 +147,7 @@ public:
 		++ModCount;
 	}
 
-	void AppendEndEvent(double EndTime)
+	virtual void AppendEndEvent(double EndTime) override
 	{
 		if (ExtraDepthEvents > 0)
 		{
diff --git a/Engine/Source/Developer/TraceServices/Public/TraceServices/AnalyzerFactories.h b/Engine/Source/Developer/TraceServices/Public/TraceServices/AnalyzerFactories.h
new file mode 100644
--- /dev/null
+++ b/Engine/Source/Developer/TraceServices/Public/TraceServices/AnalyzerFactories.h
@@ -0,0 +1,63 @@
+// Copyright Epic Games, Inc. All Rights Reserved.
+
+#pragma once
+
+#include "CoreTypes.h"
+#include "Templates/SharedPointer.h"
+#include "Templates/UniquePtr.h"
+
+namespace UE
+{
+namespace Trace
+{
+
+class IAnalyzer;
+class IInDataStream;
+
+} // namepsace Trace
+} // namespace UE
+
+namespace TraceServices
+{
+
+class IAnalysisSession;
+class IEditableBookmarkProvider;
+class IEditableCounterProvider;
+class IEditableTimingProfilerProvider;
+class IEditableThreadProvider;
+
+/*
+* Create an IAnalysisSession object for external analysis.
+* 
+* @param InTraceId		A user defined identifier for this session.
+* @param InSessionName	A user defined string to present to the user for this session.
+* @param InDataStream	The data stream interface for this analysis session.
+*/
+TRACESERVICES_API TSharedPtr<IAnalysisSession> CreateAnalysisSession(uint32 InTraceId, const TCHAR* InSessionName, TUniquePtr<UE::Trace::IInDataStream>&& InDataStream);
+
+/*
+* Create an IAnalyzer object for analyzing bookmark events.
+*
+* @param InSession					The session object to interact with.
+* @param InEditableBookmarkProvider	The bookmark provider interface to receive the event callbacks.
+*/
+TRACESERVICES_API TSharedPtr<UE::Trace::IAnalyzer> CreateBookmarksAnalyzer(IAnalysisSession& InSession, IEditableBookmarkProvider& InEditableBookmarkProvider);
+
+/*
+* Create an IAnalyzer object for analyzing counter events.
+*
+* @param InSession					The session object to interact with.
+* @param InEditableCounterProvider	The counter provider interface to receive the event callbacks.
+*/
+TRACESERVICES_API TSharedPtr<UE::Trace::IAnalyzer> CreateCountersAnalyzer(IAnalysisSession& InSession, IEditableCounterProvider& InEditableCounterProvider);
+
+/*
+* Create an IAnalyzer object for external analysis.
+*
+* @param InSession							The session object to interact with.
+* @param InEditableTimingProfilerProvider	The timing profiler provider interface to receive the cpu timing event callbacks.
+* @param InEditableThreadProvider			The thread provider interface to receive the thread event callbacks.
+*/
+TRACESERVICES_API TSharedPtr<UE::Trace::IAnalyzer> CreateCpuProfilerAnalyzer(IAnalysisSession& InSession, IEditableTimingProfilerProvider& InEditableTimingProfilerProvider, IEditableThreadProvider& InEditableThreadProvider);
+
+} // namespace TraceServices
diff --git a/Engine/Source/Developer/TraceServices/Public/TraceServices/Model/AnalysisSession.h b/Engine/Source/Developer/TraceServices/Public/TraceServices/Model/AnalysisSession.h
--- a/Engine/Source/Developer/TraceServices/Public/TraceServices/Model/AnalysisSession.h
+++ b/Engine/Source/Developer/TraceServices/Public/TraceServices/Model/AnalysisSession.h
@@ -5,6 +5,7 @@
 #include "CoreTypes.h"
 #include "Containers/StringFwd.h"
 #include "Misc/AssertionMacros.h"
+#include "Templates/SharedPointer.h"
 
 class FName;
 
@@ -36,6 +37,11 @@ public:
 	virtual void ReadAccessCheck() const { unimplemented(); }
 };
 
+class IEditableProvider
+	: public IProvider
+{
+};
+
 class IAnalysisSession
 {
 public:
@@ -65,7 +71,9 @@ public:
 	
 	virtual void AddAnalyzer(UE::Trace::IAnalyzer* Analyzer) = 0;
 
-	virtual void AddProvider(const FName& Name, IProvider* Provider) = 0;
+	UE_DEPRECATED(5.1, "Please use the TSharedPtr overload for registering Provider objects")
+	virtual void AddProvider(const FName& Name, IProvider* Provider) { AddProvider(Name, TSharedPtr<IProvider>(Provider), nullptr); }
+	virtual void AddProvider(const FName& Name, TSharedPtr<IProvider> Provider, TSharedPtr<IEditableProvider> EditableProvider = nullptr) = 0;
 	template<typename ProviderType>
 	const ProviderType* ReadProvider(const FName& Name) const
 	{
diff --git a/Engine/Source/Developer/TraceServices/Public/TraceServices/Model/Bookmarks.h b/Engine/Source/Developer/TraceServices/Public/TraceServices/Model/Bookmarks.h
--- a/Engine/Source/Developer/TraceServices/Public/TraceServices/Model/Bookmarks.h
+++ b/Engine/Source/Developer/TraceServices/Public/TraceServices/Model/Bookmarks.h
@@ -23,6 +23,46 @@ public:
 	virtual void EnumerateBookmarks(double IntervalStart, double IntervalEnd, TFunctionRef<void(const FBookmark&)> Callback) const = 0;
 };
 
+/*
+* The interface to a provider that can consume mutations of bookmark events from a session.
+*/
+class IEditableBookmarkProvider
+	: public IEditableProvider
+{
+public:
+	virtual ~IEditableBookmarkProvider() = default;
+
+	/*
+	* Update the information in the Bookmark.
+	* 
+	* @param BookmarkPoint	The unique identity (memory address) of the bookmark.
+	* @param FormatString	The format string of the bookmark's message.
+	* @param File			The source file the bookmark is in.
+	* @param Line			The line in the source file the bookmark is on.
+	*/
+	virtual void UpdateBookmarkSpec(uint64 BookmarkPoint, const TCHAR* FormatString, const TCHAR* File, int32 Line) = 0;
+
+	/*
+	* Append a new instance of a bookmark from the trace session.
+	*
+	* @param BookmarkPoint	The unique identity (memory address) of the bookmark.
+	* @param Time			The time in seconds of the event.
+	* @param FormatArgs		The arguments to use in conjunction with the spec's FormatString.
+	*/
+	virtual void AppendBookmark(uint64 BookmarkPoint, double Time, const uint8* FormatArgs) = 0;
+
+	/*
+	* Append a new instance of a bookmark from the trace session.
+	* 
+	* @param BookmarkPoint	The unique identity (memory address) of the bookmark.
+	* @param Time			The time in seconds of the event.
+	* @param Text			The fully formatted bookmark string for this instance. This may vary between instances.
+	*						This pointer is valid until the IAnalysisSession is deleted.
+	*/
+	virtual void AppendBookmark(uint64 BookmarkPoint, double Time, const TCHAR* Text) = 0;
+};
+
 TRACESERVICES_API const IBookmarkProvider& ReadBookmarkProvider(const IAnalysisSession& Session);
+TRACESERVICES_API IEditableBookmarkProvider& EditBookmarkProvider(IAnalysisSession& Session);
 
 } // namespace TraceServices
diff --git a/Engine/Source/Developer/TraceServices/Public/TraceServices/Model/Counters.h b/Engine/Source/Developer/TraceServices/Public/TraceServices/Model/Counters.h
--- a/Engine/Source/Developer/TraceServices/Public/TraceServices/Model/Counters.h
+++ b/Engine/Source/Developer/TraceServices/Public/TraceServices/Model/Counters.h
@@ -29,21 +29,88 @@ public:
 	virtual void EnumerateFloatValues(double IntervalStart, double IntervalEnd, bool bIncludeExternalBounds, TFunctionRef<void(double, double)> Callback) const = 0;
 };
 
+class IEditableCounterProvider;
+
+/*
+* An interface that can consume mutations of the state of a counter.
+*/
 class IEditableCounter
-	: public ICounter
 {
 public:
 	virtual ~IEditableCounter() = default;
 
+	/*
+	* Sets the name of the counter.
+	* 
+	* @param Name	The name of the counter.
+	*/
 	virtual void SetName(const TCHAR* Name) = 0;
+
+	/*
+	* Sets the group of the counter.
+	* 
+	* @param Group	The group name of the counter.
+	*/
 	virtual void SetGroup(const TCHAR* Group) = 0;
+
+	/*
+	* Sets the description of the counter.
+	* 
+	* @param Description	The description of the counter.
+	*/
 	virtual void SetDescription(const TCHAR* Description) = 0;
+
+	/*
+	* Sets whether the counter is an integral or floating point number.
+	* 
+	* @param bIsFloatingPoint	True if the counter is floating point, false if it is integral.
+	*/
 	virtual void SetIsFloatingPoint(bool bIsFloatingPoint) = 0;
+
+	/*
+	* Sets whether the counter value is reset every frame. This can be used for counters polled from stats.
+	* 
+	* @param bInIsResetEveryFrame	True if the counter is reset every frame.
+	*/
 	virtual void SetIsResetEveryFrame(bool bInIsResetEveryFrame) = 0;
+
+	/*
+	* Sets the display hint for this counter.
+	* 
+	* @param DisplayHint	The display hint for this counter.
+	*/
 	virtual void SetDisplayHint(ECounterDisplayHint DisplayHint) = 0;
+
+	/*
+	* Add a value to the value counter.
+	* 
+	* @param Time	The time at which the value was added in seconds.
+	* @param Value	The value to add to the current value of the counter.
+	*/
 	virtual void AddValue(double Time, int64 Value) = 0;
+
+	/*
+	* Add a value to the value counter.
+	*
+	* @param Time	The time at which the value was added in seconds.
+	* @param Value	The value to add to the current value of the counter.
+	*/
 	virtual void AddValue(double Time, double Value) = 0;
+
+	/*
+	* Set the value of the counter.
+	* 
+	* @param Time	The time at which the value was set in seconds.
+	* @param Value	The new value for the counter.
+	*/
 	virtual void SetValue(double Time, int64 Value) = 0;
+
+	/*
+	* Set the value of the counter.
+	*
+	* @param Time	The time at which the value was set in seconds.
+	* @param Value	The new value for the counter.
+	*/
 	virtual void SetValue(double Time, double Value) = 0;
 };
 
@@ -55,11 +122,39 @@ public:
 	virtual uint64 GetCounterCount() const = 0;
 	virtual bool ReadCounter(uint32 CounterId, TFunctionRef<void(const ICounter&)> Callback) const = 0;
 	virtual void EnumerateCounters(TFunctionRef<void(uint32, const ICounter&)> Callback) const = 0;
-	virtual IEditableCounter* CreateCounter() = 0;
+};
+
+/*
+* The interface to a provider that can consume mutations of counter events from a session.
+*/
+class IEditableCounterProvider
+	: public IEditableProvider
+{
+public:
+	virtual ~IEditableCounterProvider() = default;
+
+	/*
+	* Retreive the counter interface to an editable counter.
+	*
+	* @return The interface the counter.
+	*         This can be null if the provider doesn't care about deriving new counters from existing editable counters.
+	*/
+	virtual const ICounter* GetCounter(IEditableCounter* EditableCounter) = 0;
+
+	/*
+	* Create a new counter mutation interface.
+	* 
+	* @return The interface to mutate the counter.
+	*/
+	virtual IEditableCounter* CreateEditableCounter() = 0;
+
+	/*
+	* Add a custom counter to the provider.
+	*/
 	virtual void AddCounter(const ICounter* Counter) = 0;
 };
 
 TRACESERVICES_API const ICounterProvider& ReadCounterProvider(const IAnalysisSession& Session);
-TRACESERVICES_API ICounterProvider& EditCounterProvider(IAnalysisSession& Session);
+TRACESERVICES_API IEditableCounterProvider& EditCounterProvider(IAnalysisSession& Session);
 
 } // namespace TraceServices
diff --git a/Engine/Source/Developer/TraceServices/Public/TraceServices/Model/Log.h b/Engine/Source/Developer/TraceServices/Public/TraceServices/Model/Log.h
--- a/Engine/Source/Developer/TraceServices/Public/TraceServices/Model/Log.h
+++ b/Engine/Source/Developer/TraceServices/Public/TraceServices/Model/Log.h
@@ -41,7 +41,94 @@ public:
 	virtual const IUntypedTable& GetMessagesTable() const = 0;
 };
 
+class IEditableLogProvider
+	: public IEditableProvider
+{
+public:
+	virtual ~IEditableLogProvider() = default;
+
+	/*
+	* Register a new log message category.
+	*
+	* @return The category identity.
+	*/
+	virtual uint64 RegisterCategory() = 0;
+
+	/*
+	* Fetch the data structure for a log category.
+	*
+	* @param CategoryPointer	The unique identity (memory address) of the category instrumentation.
+	*
+	* @return A reference to the category stucture.
+	*/
+	virtual FLogCategoryInfo& GetCategory(uint64 CategoryPointer) = 0;
+	
+	/*
+	* Update a log message's Category information.
+	* 
+	* @param LogPoint			The log message to update.
+	* @param InCategoryPointer	The category.
+	*/
+	virtual void UpdateMessageCategory(uint64 LogPoint, uint64 InCategoryPointer) = 0;
+
+	/*
+	* Update a log message's format string.
+	*
+	* @param LogPoint			The log message to update.
+	* @param InFormatString		The format string whose memory is stored in the Session.
+	*/
+	virtual void UpdateMessageFormatString(uint64 LogPoint, const TCHAR* InFormatString) = 0;
+
+	/*
+	* Update a log message's file location.
+	*
+	* @param LogPoint	The log message to update.
+	* @param InFile		The file path of the message's location.
+	* @param InLine		The line number of the message's location.
+	*/
+	virtual void UpdateMessageFile(uint64 LogPoint, const TCHAR* InFile, int32 InLine) = 0;
+
+	/*
+	* Update a log message's verbosity.
+	*
+	* @param LogPoint		The log message to update.
+	* @param InVerbosity	The verbosity of the message.
+	*/
+	virtual void UpdateMessageVebosity(uint64 LogPoint, ELogVerbosity::Type InVerbosity) = 0;
+
+	/*
+	* Update a log message's information.
+	*
+	* @param LogPoint			The log message to update.
+	* @param InCategoryPointer	The category.
+	* @param InFormatString		The format string whose memory is stored in the Session.
+	* @param InFile				The file path of the message's location.
+	* @param InLine				The line number of the message's location.
+	* @param InVerbosity		The verbosity of the message.
+	*/
+	virtual void UpdateMessageSpec(uint64 LogPoint, uint64 InCategoryPointer, const TCHAR* InFormatString, const TCHAR* InFile, int32 InLine, ELogVerbosity::Type InVerbosity) = 0;
+
+	/*
+	* Append a new instance of a message from the trace session.
+	*
+	* @param LogPoint	The unique identity (memory address) of the message instrumentation.
+	* @param Time		The time in seconds of the event.
+	* @param FormatArgs	The arguments to use in conjunction with the spec's FormatString.
+	*/
+	virtual void AppendMessage(uint64 LogPoint, double Time, const uint8* FormatArgs) = 0;
+
+	/*
+	* Append a new instance of a message from the trace session.
+	*
+	* @param LogPoint	The unique identity (memory address) of the message instrumentation.
+	* @param Time		The time in seconds of the event.
+	* @param Text		The the message text.
+	*/
+	virtual void AppendMessage(uint64 LogPoint, double Time, const TCHAR* Text) = 0;
+};
+
 TRACESERVICES_API const ILogProvider& ReadLogProvider(const IAnalysisSession& Session);
+TRACESERVICES_API IEditableLogProvider& EditLogProvider(IAnalysisSession& Session);
 TRACESERVICES_API void FormatString(TCHAR* OutputString, uint32 OutputStringCount, const TCHAR* FormatString, const uint8* FormatArgs);
 
 } // namespace TraceServices
diff --git a/Engine/Source/Developer/TraceServices/Public/TraceServices/Model/Threads.h b/Engine/Source/Developer/TraceServices/Public/TraceServices/Model/Threads.h
--- a/Engine/Source/Developer/TraceServices/Public/TraceServices/Model/Threads.h
+++ b/Engine/Source/Developer/TraceServices/Public/TraceServices/Model/Threads.h
@@ -26,6 +26,26 @@ public:
 	virtual const TCHAR* GetThreadName(uint32 ThreadId) const = 0;
 };
 
+/*
+* An interface that can consume thread related events from a session.
+*/
+class IEditableThreadProvider
+	: public IEditableProvider
+{
+public:
+	virtual ~IEditableThreadProvider() = default;
+
+	/*
+	* Note the existence of a new thread.
+	* 
+	* @param Id			The thread identity.
+	* @param Name		The name the user may know the thread by, if available.
+	* @param Priority	The system priority level of the thread, if available.
+	*/
+	virtual void AddThread(uint32 Id, const TCHAR* Name, EThreadPriority Priority) = 0;
+};
+
 TRACESERVICES_API const IThreadProvider& ReadThreadProvider(const IAnalysisSession& Session);
+TRACESERVICES_API IEditableThreadProvider& EditThreadProvider(IAnalysisSession& Session);
 
 } // namespace TraceServices
diff --git a/Engine/Source/Developer/TraceServices/Public/TraceServices/Model/TimingProfiler.h b/Engine/Source/Developer/TraceServices/Public/TraceServices/Model/TimingProfiler.h
--- a/Engine/Source/Developer/TraceServices/Public/TraceServices/Model/TimingProfiler.h
+++ b/Engine/Source/Developer/TraceServices/Public/TraceServices/Model/TimingProfiler.h
@@ -93,6 +93,66 @@ public:
 	virtual ITimingProfilerButterfly* CreateButterfly(double IntervalStart, double IntervalEnd, TFunctionRef<bool(uint32)> CpuThreadFilter, bool IncludeGpu) const = 0;
 };
 
+/*
+* An interface that can consume timeline CpuProfiler events from a session.
+*/
+class IEditableTimingProfilerProvider
+	: public IEditableProvider
+{
+public:
+	virtual ~IEditableTimingProfilerProvider() = default;
+
+	/*
+	* A new CPU timer object has been found.
+	* 
+	* @param Name	The name attached to the timer.
+	* @param File	The source file in which the timer is defined.
+	* @param Line	The line number of the source file in which the timer is defined.
+	* 
+	* @return The identity of the CPU timer object.
+	*/
+	virtual uint32 AddCpuTimer(FStringView Name, const TCHAR* File, uint32 Line) = 0;
+
+	/*
+	* Update an existing timer with information. Some information is unavailable when it's created.
+	* 
+	* @param TimerId	The identity of the timer to update.
+	* @param Name		The name attached to the timer.
+	* @param File		The source file in which the timer is defined.
+	* @param Line		The line number of the source file in which the timer is defined.
+	*/
+	virtual void SetTimerNameAndLocation(uint32 TimerId, FStringView Name, const TCHAR* File, uint32 Line) = 0;
+
+	/*
+	* Add metadata to a timer.
+	* 
+	* @param MasterTimerId	The identity of the timer to add metadata to.
+	* @param Metadata		The metadata.
+	* 
+	* @return The identity of the metadata.
+	*/
+	virtual uint32 AddMetadata(uint32 MasterTimerId, TArray<uint8>&& Metadata) = 0;
+
+	/*
+	* Get metadata for a timer.
+	* 
+	* @param TimerId	The identity of the metadata to retrieve.
+	* 
+	* @return The metadata.
+	*/
+	virtual TArrayView<const uint8> GetMetadata(uint32 TimerId) const = 0;
+
+	/*
+	* Get an object to receive ordered timeline events for a thread.
+	* 
+	* @param ThreadId	The thread for which the events are for.
+	* 
+	* @return The object to receive the serial events for the specified thread.
+	*/
+	virtual IEditableTimeline<FTimingProfilerEvent>& GetCpuThreadEditableTimeline(uint32 ThreadId) = 0;
+};
+
 TRACESERVICES_API const ITimingProfilerProvider* ReadTimingProfilerProvider(const IAnalysisSession& Session);
+TRACESERVICES_API IEditableTimingProfilerProvider* EditTimingProfilerProvider(IAnalysisSession& Session);
 
 } // namespace TraceServices
diff --git a/Engine/Source/Developer/TraceServices/Public/TraceServices/Utils.h b/Engine/Source/Developer/TraceServices/Public/TraceServices/Utils.h
--- a/Engine/Source/Developer/TraceServices/Public/TraceServices/Utils.h
+++ b/Engine/Source/Developer/TraceServices/Public/TraceServices/Utils.h
@@ -39,4 +39,7 @@ inline TArrayView<const uint8> LegacyAttachmentArray(
 }
 
 } // namespace FTraceAnalyzerUtils
+
+TRACESERVICES_API void StringFormat(TCHAR* Out, uint64 MaxOut, TCHAR* Temp, uint64 MaxTemp, const TCHAR* FormatString, const uint8* FormatArgs);
+
 } // namespace TraceServices
diff --git a/Engine/Source/Editor/UnrealEd/Private/Commandlets/SummarizeTraceCommandlet.cpp b/Engine/Source/Editor/UnrealEd/Private/Commandlets/SummarizeTraceCommandlet.cpp
--- a/Engine/Source/Editor/UnrealEd/Private/Commandlets/SummarizeTraceCommandlet.cpp
+++ b/Engine/Source/Editor/UnrealEd/Private/Commandlets/SummarizeTraceCommandlet.cpp
@@ -13,41 +13,23 @@
 #include "Misc/Crc.h"
 #include "Misc/FileHelper.h"
 #include "Misc/Paths.h"
+#include "ProfilingDebugging/CountersTrace.h"
 #include "String/ParseTokens.h"
 #include "Trace/Analysis.h"
-#include "Trace/Analyzer.h"
 #include "Trace/DataStream.h"
-#include "TraceServices/Model/Log.h"
+#include "TraceServices/AnalyzerFactories.h"
 #include "TraceServices/Utils.h"
-#include "ProfilingDebugging/CountersTrace.h"
-
-/*
- * The following could be in TraceServices. This way if the format of CPU scope
- * events change this interface acts as a compatibility contract for external
- * tools. In the future it may be in a separate library so that Trace and
- * Insights' instrumentation are more broadly available.
- */
-
-static uint64 Decode7bit(const uint8*& Cursor)
-{
-	uint64 Value = 0;
-	uint64 ByteIndex = 0;
-	bool bHasMoreBytes;
-	do
-	{
-		uint8 ByteValue = *Cursor++;
-		bHasMoreBytes = ByteValue & 0x80;
-		Value |= uint64(ByteValue & 0x7f) << (ByteIndex * 7);
-		++ByteIndex;
-	} 	while (bHasMoreBytes);
-	return Value;
-}
+#include "TraceServices/Model/AnalysisSession.h"
+#include "TraceServices/Model/Bookmarks.h"
+#include "TraceServices/Model/Counters.h"
+#include "TraceServices/Model/Threads.h"
+#include "TraceServices/Model/TimingProfiler.h"
 
 /**
  * Base class to extend for CPU scope analysis. Derived class instances are meant to be registered as
- * delegates to the FCpuScopeStreamProcessor to handle scope events and perform meaningful analysis.
+ * delegates to the FSummarizeCpuProfilerProvider to handle scope events and perform meaningful analysis.
  */
-class FCpuScopeAnalyzer
+class FSummarizeCpuScopeAnalyzer
 {
 public:
 	enum class EScopeEventType : uint32
@@ -73,13 +55,13 @@ public:
 	};
 
 public:
-	virtual ~FCpuScopeAnalyzer() = default;
+	virtual ~FSummarizeCpuScopeAnalyzer() = default;
 
 	/** Invoked when a CPU scope is discovered. This function is always invoked first when a CPU scope is encountered for the first time.*/
 	virtual void OnCpuScopeDiscovered(uint32 ScopeId) {}
 
 	/** Invoked when CPU scope specification is encountered in the trace stream. */
-	virtual void OnCpuScopeName(uint32 ScopeId, const FString& ScopeName) {};
+	virtual void OnCpuScopeName(uint32 ScopeId, const FStringView& ScopeName) {};
 
 	/** Invoked when a scope is entered. The scope name might not be known yet. */
 	virtual void OnCpuScopeEnter(const FScopeEvent& ScopeEnter, const FString* ScopeName) {};
@@ -87,8 +69,8 @@ public:
 	/** Invoked when a scope is exited. The scope name might not be known yet. */
 	virtual void OnCpuScopeExit(const FScope& Scope, const FString* ScopeName) {};
 
-	/** Invoked when a root event on the specified thread along with all child events down to the leaves are know. */
-	virtual void OnCpuScopeTree(uint32 ThreadId, const TArray<FCpuScopeAnalyzer::FScopeEvent>& ScopeEvents, const TFunction<const FString*(uint32)>& ScopeLookupNameFn) {};
+	/** Invoked when a root event on the specified thread along with all child events down to the leaves are known. */
+	virtual void OnCpuScopeTree(uint32 ThreadId, const TArray<FSummarizeCpuScopeAnalyzer::FScopeEvent>& ScopeEvents, const TFunction<const FString*(uint32)>& ScopeLookupNameFn) {};
 
 	/** Invoked when the trace stream has been fully consumed/processed. */
 	virtual void OnCpuScopeAnalysisEnd() {};
@@ -103,14 +85,17 @@ public:
  * is to decode the stream only once and let several registered analyzers reuse the small state built up by this processor. This
  * matters when processing very large traces with possibly billions of scope events.
  */
-class FCpuScopeStreamProcessor
-	: public UE::Trace::IAnalyzer
+class FSummarizeCpuProfilerProvider
+	: public TraceServices::IEditableThreadProvider
+	, public TraceServices::IEditableTimingProfilerProvider
 {
 public:
-	FCpuScopeStreamProcessor();
+	FSummarizeCpuProfilerProvider();
 
 	/** Register an analyzer with this processor. The processor decodes the trace stream and invokes the registered analyzers when a CPU scope event occurs.*/
-	void AddCpuScopeAnalyzer(TSharedPtr<FCpuScopeAnalyzer> Analyzer);
+	void AddCpuScopeAnalyzer(TSharedPtr<FSummarizeCpuScopeAnalyzer> Analyzer);
+
+	void AnalysisComplete();
 
 private:
 	struct FScopeEnter
@@ -123,7 +108,7 @@ private:
 	struct FScopeTreeInfo
 	{
 		// Records the current root scope and its children to run analysis that needs to know the parent/child relationship.
-		TArray<FCpuScopeAnalyzer::FScopeEvent> ScopeEvents;
+		TArray<FSummarizeCpuScopeAnalyzer::FScopeEvent> ScopeEvents;
 
 		// Indicates if one of the scope in the current hierarchy is nameless. (Its names specs hasn't been received yet).
 		bool bHasNamelessScopes = false;
@@ -137,7 +122,25 @@ private:
 
 	// For each thread we track what the stack of scopes are, for matching end-to-start
 	struct FThread
+		: public TraceServices::IEditableTimeline<TraceServices::FTimingProfilerEvent>
 	{
+		FThread(uint32 InThreadId, FSummarizeCpuProfilerProvider* InProvider)
+			: ThreadId(InThreadId)
+			, Provider(InProvider)
+		{
+
+		}
+
+		virtual void AppendBeginEvent(double StartTime, const TraceServices::FTimingProfilerEvent& Event) override;
+		virtual void AppendEndEvent(double EndTime) override;
+
+		// The ThreadId of this thread
+		uint32 ThreadId;
+
+		// The provider to forward calls to
+		FSummarizeCpuProfilerProvider* Provider;
+
+		// The current stack state
 		TArray<FScopeEnter> ScopeStack;
 
 		// The events recorded for the current root scope and its children to run analysis that needs to know the parent/child relationship, for example to compute time including childreen and time
@@ -150,532 +153,245 @@ private:
 	};
 
 private:
-	// UE::Trace::IAnalyzer interface.
-	virtual void OnAnalysisBegin(const FOnAnalysisContext& Context) override;
-	virtual void OnAnalysisEnd() override;
-	virtual bool OnEvent(uint16 RouteId, EStyle Style, const FOnEventContext& Context) override;
-
-	// Internal implementation.
-	void OnEventSpec(const FOnEventContext& Context);
-	void OnBatch(const FOnEventContext& Context);
-	void OnBatchV2(const FOnEventContext& Context);
-	void OnCpuScopeSpec(uint32 ScopeId, const FString* ScopeName);
-	void OnCpuScopeEnter(uint32 ScopeId, uint32 ThreadId, double Timestamp);
-	void OnCpuScopeExit(uint32 ThreadId, double Timestamp);
+	// TraceServices::IEditableThreadProvider
+	virtual void					AddThread(uint32 Id, const TCHAR* Name, EThreadPriority Priority) override;
+
+	// TraceServices::IEditableTimingProfilerProvider
+	virtual uint32					AddCpuTimer(FStringView Name, const TCHAR* File, uint32 Line) override;
+	virtual void					SetTimerNameAndLocation(uint32 TimerId, FStringView Name, const TCHAR* File, uint32 Line) override;
+	virtual uint32					AddMetadata(uint32 MasterTimerId, TArray<uint8>&& Metadata) override;
+	virtual TArrayView<const uint8> GetMetadata(uint32 TimerId) const override;
+	virtual TraceServices::IEditableTimeline<TraceServices::FTimingProfilerEvent>& GetCpuThreadEditableTimeline(uint32 ThreadId) override;
+
+	// Callbacks from FThread to forward calls to Analyzers
+	virtual void AppendBeginEvent(const FSummarizeCpuScopeAnalyzer::FScopeEvent& ScopeEvent);
+	virtual void AppendEndEvent(const FSummarizeCpuScopeAnalyzer::FScope& ScopeEvent, const FString* ScopeName);
+
+	// Processing the ScopeTree once it's complete
 	void OnCpuScopeTree(uint32 ThreadId, const FScopeTreeInfo& ScopeTreeInfo);
-	const FString* LookupScopeName(uint32 ScopeId) { return ScopeId < static_cast<uint32>(ScopeNames.Num()) && ScopeNames[ScopeId] ? &ScopeNames[ScopeId].GetValue() : nullptr; }
+
+	// Resolve a ScopeId to string
+	const FString* LookupScopeName(uint32 ScopeId);
 
 private:
-	// The state at any moment of the threads, indexes are doled out on the process-side
-	TArray<FThread> Threads;
+	// The state at any moment of the threads
+	TMap<uint32, TUniquePtr<FThread>> Threads;
 
 	// The scope names, the array index correspond to the scope Id. If the optional is not set, the scope hasn't been encountered yet.
 	TArray<TOptional<FString>> ScopeNames;
 
 	// List of analyzers to invoke when a scope event is decoded.
-	TArray<TSharedPtr<FCpuScopeAnalyzer>> ScopeAnalyzers;
+	TArray<TSharedPtr<FSummarizeCpuScopeAnalyzer>> ScopeAnalyzers;
 
 	// Scope name lookup function, cached for efficiency.
 	TFunction<const FString*(uint32 ScopeId)> LookupScopeNameFn;
-};
 
-enum
-{
-	// CpuProfilerTrace.cpp
-	RouteId_CpuProfiler_EventSpec,
-	RouteId_CpuProfiler_EventBatch,
-	RouteId_CpuProfiler_EndCapture,
-	RouteId_CpuProfiler_EventBatchV2,
-	RouteId_CpuProfiler_EndCaptureV2
+	struct FMetadata
+	{
+		TArray<uint8> Payload;
+		uint32 TimerId;
+	};
+
+	TArray<FMetadata> Metadatas;
 };
 
-FCpuScopeStreamProcessor::FCpuScopeStreamProcessor()
+FSummarizeCpuProfilerProvider::FSummarizeCpuProfilerProvider()
  : LookupScopeNameFn([this](uint32 ScopeId) { return LookupScopeName(ScopeId); })
 {
 }
 
-void FCpuScopeStreamProcessor::AddCpuScopeAnalyzer(TSharedPtr<FCpuScopeAnalyzer> Analyzer)
+void FSummarizeCpuProfilerProvider::AddCpuScopeAnalyzer(TSharedPtr<FSummarizeCpuScopeAnalyzer> Analyzer)
 {
 	ScopeAnalyzers.Add(MoveTemp(Analyzer));
 }
 
-void FCpuScopeStreamProcessor::OnAnalysisBegin(const FOnAnalysisContext& Context)
-{
-	Context.InterfaceBuilder.RouteEvent(RouteId_CpuProfiler_EventSpec, "CpuProfiler", "EventSpec");
-	Context.InterfaceBuilder.RouteEvent(RouteId_CpuProfiler_EventBatch, "CpuProfiler", "EventBatch");
-	Context.InterfaceBuilder.RouteEvent(RouteId_CpuProfiler_EndCapture, "CpuProfiler", "EndCapture");
-	Context.InterfaceBuilder.RouteEvent(RouteId_CpuProfiler_EventBatchV2, "CpuProfiler", "EventBatchV2");
-	Context.InterfaceBuilder.RouteEvent(RouteId_CpuProfiler_EndCaptureV2, "CpuProfiler", "EndCaptureV2");
-}
-
-void FCpuScopeStreamProcessor::OnAnalysisEnd()
+void FSummarizeCpuProfilerProvider::AnalysisComplete()
 {
 	// Analyze scope trees that contained 'nameless' context when they were captured. Unless the trace was truncated,
 	// all scope names should be known now.
-	uint32 ThreadId = 0;
-	for (FThread& Thread : Threads)
+	for (const auto& Thread : Threads)
 	{
-		for (FScopeTreeInfo& DelayedScopeTree : Threads[ThreadId].DelayedScopeTreeInfo)
+		for (FScopeTreeInfo& DelayedScopeTree : Thread.Value->DelayedScopeTreeInfo)
 		{
 			// Run summary analysis for this delayed hierarchy.
-			OnCpuScopeTree(ThreadId, DelayedScopeTree);
+			OnCpuScopeTree(Thread.Key, DelayedScopeTree);
 		}
-
-		++ThreadId;
 	}
 
-	for (TSharedPtr<FCpuScopeAnalyzer>& Analyzer : ScopeAnalyzers)
+	for (TSharedPtr<FSummarizeCpuScopeAnalyzer>& Analyzer : ScopeAnalyzers)
 	{
 		Analyzer->OnCpuScopeAnalysisEnd();
 	}
 }
 
-bool FCpuScopeStreamProcessor::OnEvent(uint16 RouteId, EStyle Style, const FOnEventContext& Context)
+void FSummarizeCpuProfilerProvider::AddThread(uint32 Id, const TCHAR* Name, EThreadPriority Priority)
 {
-	switch (RouteId)
-	{
-	case RouteId_CpuProfiler_EventSpec:
-		OnEventSpec(Context);
-		break;
-
-	case RouteId_CpuProfiler_EventBatch:
-	case RouteId_CpuProfiler_EndCapture:
-		OnBatch(Context);
-		break;
-	case RouteId_CpuProfiler_EventBatchV2:
-	case RouteId_CpuProfiler_EndCaptureV2:
-		OnBatchV2(Context);
-		break;
-	};
-
-	return true;
-}
-
-void FCpuScopeStreamProcessor::OnEventSpec(const FOnEventContext& Context)
-{
-	const FEventData& EventData = Context.EventData;
-	FString Name;
-	uint32 Id = EventData.GetValue<uint32>("Id");
-	EventData.GetString("Name", Name);
-	OnCpuScopeSpec(Id - 1, &Name);
+	TUniquePtr<FThread>* Found = Threads.Find(Id);
+	if (!Found)
+	{
+		Threads.Add(Id, MakeUnique<FThread>(Id, this));
+	}
 }
 
-void FCpuScopeStreamProcessor::OnBatch(const FOnEventContext& Context)
+uint32 FSummarizeCpuProfilerProvider::AddCpuTimer(FStringView Name, const TCHAR* File, uint32 Line)
 {
-	const FEventData& EventData = Context.EventData;
-	const FEventTime& EventTime = Context.EventTime;
+	TOptional<FString> ScopeName;
 
-	uint32 ThreadId = Context.ThreadInfo.GetId();
+	if (!Name.IsEmpty())
+	{
+		ScopeName.Emplace(FString(Name));
+	}
+	
+	uint32 TimerId = ScopeNames.Add(ScopeName);
 
-	TArrayView<const uint8> DataView = TraceServices::FTraceAnalyzerUtils::LegacyAttachmentArray("Data", Context);
-	const uint8* Cursor = DataView.GetData();
-	const uint8* End = Cursor + DataView.Num();
-	uint64 LastCycle = 0;
-	while (Cursor < End)
+	// Notify the analyzers.
+	for (TSharedPtr<FSummarizeCpuScopeAnalyzer>& Analyzer : ScopeAnalyzers)
 	{
-		uint64 Value = Decode7bit(Cursor);
-		uint64 Cycle = LastCycle + (Value >> 1);
-		LastCycle = Cycle;
+		Analyzer->OnCpuScopeDiscovered(TimerId);
 
-		double TimeStamp = EventTime.AsSeconds(Cycle);
-		if (Value & 1)
+		if (!Name.IsEmpty())
 		{
-			uint64 ScopeId = Decode7bit(Cursor);
-			OnCpuScopeEnter(static_cast<uint32>(ScopeId - 1), ThreadId, TimeStamp);
-		}
-		else
-		{
-			OnCpuScopeExit(ThreadId, TimeStamp);
+			Analyzer->OnCpuScopeName(TimerId, Name);
 		}
 	}
+
+	return TimerId;
 }
 
-void FCpuScopeStreamProcessor::OnBatchV2(const FOnEventContext& Context)
+void FSummarizeCpuProfilerProvider::SetTimerNameAndLocation(uint32 TimerId, FStringView Name, const TCHAR* File, uint32 Line)
 {
-	const FEventData& EventData = Context.EventData;
-	const FEventTime& EventTime = Context.EventTime;
+	check(TimerId < uint32(ScopeNames.Num()));
+	check(!Name.IsEmpty());
 
-	uint32 ThreadId = Context.ThreadInfo.GetId();
+	ScopeNames[TimerId].Emplace(FString(Name));
 
-	TArrayView<const uint8> DataView = TraceServices::FTraceAnalyzerUtils::LegacyAttachmentArray("Data", Context);
-	const uint8* Cursor = DataView.GetData();
-	const uint8* End = Cursor + DataView.Num();
-	uint64 LastCycle = 0;
-	while (Cursor < End)
+	// Notify the registered scope analyzers.
+	for (TSharedPtr<FSummarizeCpuScopeAnalyzer>& Analyzer : ScopeAnalyzers)
 	{
-		uint64 Value = Decode7bit(Cursor);
-		uint64 Cycle = LastCycle + (Value >> 2);
-		LastCycle = Cycle;
-		double TimeStamp = EventTime.AsSeconds(Cycle);
-
-		if (Value & 2ull)
-		{		
-			if (Value & 1ull)			
-			{
-				uint64 CoroutineId = Decode7bit(Cursor);
-				uint32 TimerScopeDepth = Decode7bit(Cursor);
-				OnCpuScopeEnter(FCpuScopeAnalyzer::CoroutineSpecId, ThreadId, TimeStamp);
-				for (uint32 i = 0; i < TimerScopeDepth; ++i)
-				{
-					OnCpuScopeEnter(FCpuScopeAnalyzer::CoroutineUnknownSpecId, ThreadId, TimeStamp);
-				}
-			}
-			else
-			{			
-				uint32 TimerScopeDepth = Decode7bit(Cursor);
-				for (uint32 i = 0; i < TimerScopeDepth; ++i)
-				{
-					OnCpuScopeExit(ThreadId, TimeStamp);
-				}
-				OnCpuScopeExit(ThreadId, TimeStamp);
-			}
-		}
-		else
-		{
-			if (Value & 1)
-			{
-				uint64 ScopeId = Decode7bit(Cursor);
-				OnCpuScopeEnter(static_cast<uint32>(ScopeId - 1), ThreadId, TimeStamp);
-			}
-			else
-			{
-				OnCpuScopeExit(ThreadId, TimeStamp);
-			}
-		}
+		Analyzer->OnCpuScopeName(TimerId, Name);
 	}
 }
 
-void FCpuScopeStreamProcessor::OnCpuScopeSpec(uint32 ScopeId, const FString* ScopeName)
+uint32 FSummarizeCpuProfilerProvider::AddMetadata(uint32 MasterTimerId, TArray<uint8>&& Metadata)
 {
-	if (ScopeId >= uint32(ScopeNames.Num()))
+	uint32 MetadataId = Metadatas.Num();
+	Metadatas.Add({ MoveTemp(Metadata), MasterTimerId });
+	return ~MetadataId;
+}
+
+TArrayView<const uint8> FSummarizeCpuProfilerProvider::GetMetadata(uint32 TimerId) const
+{
+	if (int32(TimerId) >= 0)
 	{
-		uint32 Num = (ScopeId + 128) & ~127;
-		ScopeNames.SetNum(Num);
+		return TArrayView<const uint8>();
 	}
 
-	// If the optional is not set, the scope hasn't been encountered yet (Set to a blank names means encountered, but names hasn't be encountered yet)
-	bool bDiscovered = !ScopeNames[ScopeId].IsSet();
-
-	// Store the discovery state (setting the optional) and the actual name if known.
-	ScopeNames[ScopeId].Emplace(ScopeName ? *ScopeName : FString());
-
-	// Notify the analyzers.
-	for (TSharedPtr<FCpuScopeAnalyzer>& Analyzer : ScopeAnalyzers)
+	TimerId = ~TimerId;
+	if (TimerId >= uint32(Metadatas.Num()))
 	{
-		if (bDiscovered)
-		{
-			Analyzer->OnCpuScopeDiscovered(ScopeId);
-		}
-		if (ScopeName)
-		{
-			Analyzer->OnCpuScopeName(ScopeId, *ScopeName);
-		}
+		return TArrayView<const uint8>();
 	}
+
+	const FMetadata& Metadata = Metadatas[TimerId];
+	return Metadata.Payload;
 }
 
-void FCpuScopeStreamProcessor::OnCpuScopeEnter(uint32 ScopeId, uint32 ThreadId, double Timestamp)
+TraceServices::IEditableTimeline<TraceServices::FTimingProfilerEvent>& FSummarizeCpuProfilerProvider::GetCpuThreadEditableTimeline(uint32 ThreadId)
 {
-	if (ThreadId >= uint32(Threads.Num()))
+	TUniquePtr<FThread>* Found = Threads.Find(ThreadId);
+	if (Found)
 	{
-		Threads.SetNum(ThreadId + 1);
+		return *(Found->Get());
 	}
 
-	FCpuScopeAnalyzer::FScopeEvent ScopeEvent{FCpuScopeAnalyzer::EScopeEventType::Enter, ScopeId, ThreadId, Timestamp};
-	Threads[ThreadId].ScopeStack.Add(FScopeEnter{ScopeId, Timestamp});
-	Threads[ThreadId].ScopeTreeInfo.ScopeEvents.Add(ScopeEvent);
+	return *Threads.Add(ThreadId, MakeUnique<FThread>(ThreadId, this));
+}
 
-	// Scope specs/events can be received out of order depending on engine thread scheduling.
-	if (ScopeId >= static_cast<uint32>(ScopeNames.Num()) || !ScopeNames[ScopeId].IsSet())
-	{
-		// This is a newly discovered scope. Create an entry for it and notify the discovery.
-		OnCpuScopeSpec(ScopeId, nullptr);
-	}
+void FSummarizeCpuProfilerProvider::FThread::AppendBeginEvent(double StartTime, const TraceServices::FTimingProfilerEvent& Event)
+{
+	FScopeEnter ScopeEnter{ Event.TimerIndex, StartTime };
+	ScopeStack.Add(ScopeEnter);
 
-	// Notify the registered scope analyzers.
-	for (TSharedPtr<FCpuScopeAnalyzer>& Analyzer : ScopeAnalyzers)
-	{
-		Analyzer->OnCpuScopeEnter(ScopeEvent, LookupScopeName(ScopeId));
-	}
+	FSummarizeCpuScopeAnalyzer::FScopeEvent ScopeEvent{ FSummarizeCpuScopeAnalyzer::EScopeEventType::Enter, Event.TimerIndex, ThreadId, StartTime };
+	ScopeTreeInfo.ScopeEvents.Add(ScopeEvent);
+
+	Provider->AppendBeginEvent(ScopeEvent);
 }
 
-void FCpuScopeStreamProcessor::OnCpuScopeExit(uint32 ThreadId, double Timestamp)
+void FSummarizeCpuProfilerProvider::FThread::AppendEndEvent(double EndTime)
 {
-	if (ThreadId >= uint32(Threads.Num()) || Threads[ThreadId].ScopeStack.Num() <= 0)
+	if (ScopeStack.IsEmpty())
 	{
 		return;
 	}
 
-	FScopeEnter ScopeEnter = Threads[ThreadId].ScopeStack.Pop();
-	
-	Threads[ThreadId].ScopeTreeInfo.ScopeEvents.Add(FCpuScopeAnalyzer::FScopeEvent{FCpuScopeAnalyzer::EScopeEventType::Exit, ScopeEnter.ScopeId, ThreadId, Timestamp});
-	Threads[ThreadId].ScopeTreeInfo.bHasNamelessScopes |= ScopeNames[ScopeEnter.ScopeId].GetValue().IsEmpty();
+	FScopeEnter ScopeEnter = ScopeStack.Pop();
 
-	// Notify the registered scope analyzers.
-	FCpuScopeAnalyzer::FScope Scope{ScopeEnter.ScopeId, ThreadId, ScopeEnter.Timestamp, Timestamp};
-	for (TSharedPtr<FCpuScopeAnalyzer>& Analyzer : ScopeAnalyzers)
-	{
-		Analyzer->OnCpuScopeExit(Scope, LookupScopeName(ScopeEnter.ScopeId));
-	}
+	FSummarizeCpuScopeAnalyzer::FScopeEvent ScopeEvent{ FSummarizeCpuScopeAnalyzer::EScopeEventType::Exit, ScopeEnter.ScopeId, ThreadId, EndTime };
+	ScopeTreeInfo.ScopeEvents.Add(ScopeEvent);
 
-	// The root scope on this thread just poped out.
-	if (Threads[ThreadId].ScopeStack.IsEmpty())
+	// Check if at this point if the scope has a name
+	const FString* ScopeName = Provider->LookupScopeName(ScopeEnter.ScopeId);
+	ScopeTreeInfo.bHasNamelessScopes |= ScopeName == nullptr || ScopeName->IsEmpty();
+
+	FSummarizeCpuScopeAnalyzer::FScope Scope{ ScopeEnter.ScopeId, ThreadId, ScopeEnter.Timestamp, EndTime };
+	Provider->AppendEndEvent(Scope, ScopeName);
+
+	// The root scope on this thread just popped out.
+	if (ScopeStack.IsEmpty())
 	{
-		if (Threads[ThreadId].ScopeTreeInfo.bHasNamelessScopes)
+		if (ScopeTreeInfo.bHasNamelessScopes)
 		{
 			// Delay the analysis until all the scope names are known.
-			Threads[ThreadId].DelayedScopeTreeInfo.Add(MoveTemp(Threads[ThreadId].ScopeTreeInfo));
+			DelayedScopeTreeInfo.Add(MoveTemp(ScopeTreeInfo));
 		}
 		else
 		{
 			// Run analysis for this scope tree.
-			OnCpuScopeTree(ThreadId, Threads[ThreadId].ScopeTreeInfo);
+			Provider->OnCpuScopeTree(ThreadId, ScopeTreeInfo);
 		}
 
-		Threads[ThreadId].ScopeTreeInfo.Reset();
+		ScopeTreeInfo.Reset();
 	}
 }
 
-void FCpuScopeStreamProcessor::OnCpuScopeTree(uint32 ThreadId, const FScopeTreeInfo& ScopeTreeInfo)
+void FSummarizeCpuProfilerProvider::AppendBeginEvent(const FSummarizeCpuScopeAnalyzer::FScopeEvent& ScopeEvent)
 {
-	for (TSharedPtr<FCpuScopeAnalyzer>& Analyzer : ScopeAnalyzers)
+	// Notify the registered scope analyzers.
+	for (TSharedPtr<FSummarizeCpuScopeAnalyzer>& Analyzer : ScopeAnalyzers)
 	{
-		Analyzer->OnCpuScopeTree(ThreadId, ScopeTreeInfo.ScopeEvents, LookupScopeNameFn);
+		Analyzer->OnCpuScopeEnter(ScopeEvent, LookupScopeName(ScopeEvent.ScopeId));
 	}
 }
 
-/**
- * Base class to extend for counters analysis.
- */
-class FCountersAnalyzer
-	: public UE::Trace::IAnalyzer
-{
-public:
-	struct FCounterName
-	{
-		const TCHAR*		Name;
-		ETraceCounterType	Type;
-		uint16				Id;
-	};
-
-	struct FCounterIntValue
-	{
-		uint16				Id;
-		double				Timestamp;
-		int64				Value;
-	};
-
-	struct FCounterFloatValue
-	{
-		uint16				Id;
-		double				Timestamp;
-		double				Value;
-	};
-
-	virtual void		OnCounterName(const FCounterName& CounterName) = 0;
-	virtual void		OnCounterIntValue(const FCounterIntValue& NewValue) = 0;
-	virtual void		OnCounterFloatValue(const FCounterFloatValue& NewValue) = 0;
-
-private:
-	virtual void		OnAnalysisBegin(const FOnAnalysisContext& Context) override;
-	virtual bool		OnEvent(uint16 RouteId, EStyle Style, const FOnEventContext& Context) override;
-	void				OnCountersSpec(const FOnEventContext& Context);
-	void				OnCountersSetValueInt(const FOnEventContext& Context);
-	void				OnCountersSetValueFloat(const FOnEventContext& Context);
-};
-
-enum
+void FSummarizeCpuProfilerProvider::AppendEndEvent(const FSummarizeCpuScopeAnalyzer::FScope& Scope, const FString* ScopeName)
 {
-	// CountersTrace.cpp
-	RouteId_Counters_Spec,
-	RouteId_Counters_SetValueInt,
-	RouteId_Counters_SetValueFloat,
-};
-
-void FCountersAnalyzer::OnAnalysisBegin(const FOnAnalysisContext& Context)
-{
-	Context.InterfaceBuilder.RouteEvent(RouteId_Counters_Spec, "Counters", "Spec");
-	Context.InterfaceBuilder.RouteEvent(RouteId_Counters_SetValueInt, "Counters", "SetValueInt");
-	Context.InterfaceBuilder.RouteEvent(RouteId_Counters_SetValueFloat, "Counters", "SetValueFloat");
-}
-
-bool FCountersAnalyzer::OnEvent(uint16 RouteId, EStyle Style, const FOnEventContext& Context)
-{
-	switch (RouteId)
+	// Notify the registered scope analyzers.
+	for (TSharedPtr<FSummarizeCpuScopeAnalyzer>& Analyzer : ScopeAnalyzers)
 	{
-	case RouteId_Counters_Spec:
-		OnCountersSpec(Context);
-		break;
-
-	case RouteId_Counters_SetValueInt:
-		OnCountersSetValueInt(Context);
-		break;
-
-	case RouteId_Counters_SetValueFloat:
-		OnCountersSetValueFloat(Context);
-		break;
+		Analyzer->OnCpuScopeExit(Scope, ScopeName);
 	}
-
-	return true;
-}
-
-void FCountersAnalyzer::OnCountersSpec(const FOnEventContext& Context)
-{
-	const FEventData& EventData = Context.EventData;
-	FString Name;
-	uint16 Id = EventData.GetValue<uint16>("Id");
-	ETraceCounterType Type = static_cast<ETraceCounterType>(EventData.GetValue<uint8>("Type"));
-	EventData.GetString("Name", Name);
-	OnCounterName({ *Name, Type, uint16(Id - 1) });
-}
-
-void FCountersAnalyzer::OnCountersSetValueInt(const FOnEventContext& Context)
-{
-	const FEventData& EventData = Context.EventData;
-	uint16 CounterId = EventData.GetValue<uint16>("CounterId");
-	int64 Value = EventData.GetValue<int64>("Value");
-	uint64 Cycle = EventData.GetValue<uint64>("Cycle");
-	double Timestamp = Context.EventTime.AsSeconds(Cycle);
-	OnCounterIntValue({ uint16(CounterId - 1), Timestamp, Value });
-}
-
-void FCountersAnalyzer::OnCountersSetValueFloat(const FOnEventContext& Context)
-{
-	const FEventData& EventData = Context.EventData;
-	uint16 CounterId = EventData.GetValue<uint16>("CounterId");
-	double Value = EventData.GetValue<double>("Value");
-	uint64 Cycle = EventData.GetValue<uint64>("Cycle");
-	double Timestamp = Context.EventTime.AsSeconds(Cycle);
-	OnCounterFloatValue({ uint16(CounterId - 1), Timestamp, Value });
 }
 
-/**
- * Base class to extend for bookmarks analysis.
- */
-class FBookmarksAnalyzer
-	: public UE::Trace::IAnalyzer
-{
-public:
-	struct FBookmarkSpecEvent
-	{
-		uint64			Id;
-		const TCHAR*	FileName;
-		int32			Line;
-		const TCHAR*	FormatString;
-	};
-
-	struct FBookmarkEvent
-	{
-		uint64					Id;
-		double					Timestamp;
-		TArrayView<const uint8>	FormatArgs;
-	};
-
-	virtual void		OnBookmarkSpecEvent(const FBookmarkSpecEvent& BookmarkSpecEvent) = 0;
-	virtual void		OnBookmarkEvent(const FBookmarkEvent& BookmarkEvent) = 0;
-
-private:
-	virtual void		OnAnalysisBegin(const FOnAnalysisContext& Context) override;
-	virtual bool		OnEvent(uint16 RouteId, EStyle Style, const FOnEventContext& Context) override;
-	void				OnBookmarksSpec(const FOnEventContext& Context);
-	void				OnBookmarksBookmark(const FOnEventContext& Context);
-};
-
-enum
-{
-	// MiscTrace.cpp
-	RouteId_Bookmarks_BookmarkSpec,
-	RouteId_Bookmarks_Bookmark,
-};
-
-void FBookmarksAnalyzer::OnAnalysisBegin(const FOnAnalysisContext& Context)
+void FSummarizeCpuProfilerProvider::OnCpuScopeTree(uint32 ThreadId, const FScopeTreeInfo& ScopeTreeInfo)
 {
-	Context.InterfaceBuilder.RouteEvent(RouteId_Bookmarks_BookmarkSpec, "Misc", "BookmarkSpec");
-	Context.InterfaceBuilder.RouteEvent(RouteId_Bookmarks_Bookmark, "Misc", "Bookmark");
-}
-
-bool FBookmarksAnalyzer::OnEvent(uint16 RouteId, EStyle Style, const FOnEventContext& Context)
-{
-	switch (RouteId)
+	// Notify the registered scope analyzers.
+	for (TSharedPtr<FSummarizeCpuScopeAnalyzer>& Analyzer : ScopeAnalyzers)
 	{
-	case RouteId_Bookmarks_BookmarkSpec:
-		OnBookmarksSpec(Context);
-		break;
-
-	case RouteId_Bookmarks_Bookmark:
-		OnBookmarksBookmark(Context);
-		break;
+		Analyzer->OnCpuScopeTree(ThreadId, ScopeTreeInfo.ScopeEvents, LookupScopeNameFn);
 	}
-
-	return true;
-}
-
-void FBookmarksAnalyzer::OnBookmarksSpec(const FOnEventContext& Context)
-{
-	const FEventData& EventData = Context.EventData;
-	uint64 Id = EventData.GetValue<uint64>("BookmarkPoint");
-	int32 Line = EventData.GetValue<int32>("Line");
-	FString FileName;
-	EventData.GetString("FileName", FileName);
-	FString FormatString;
-	EventData.GetString("FormatString", FormatString);
-	OnBookmarkSpecEvent({ Id, *FileName, Line, *FormatString });
 }
 
-void FBookmarksAnalyzer::OnBookmarksBookmark(const FOnEventContext& Context)
+const FString* FSummarizeCpuProfilerProvider::LookupScopeName(uint32 ScopeId)
 {
-	const FEventData& EventData = Context.EventData;
-	uint64 Id = EventData.GetValue<uint64>("BookmarkPoint");
-	uint64 Cycle = EventData.GetValue<uint64>("Cycle");
-	double Timestamp = Context.EventTime.AsSeconds(Cycle);
-	TArrayView<const uint8> FormatArgsView = TraceServices::FTraceAnalyzerUtils::LegacyAttachmentArray("FormatArgs", Context);
-	OnBookmarkEvent({ Id, Timestamp, FormatArgsView });
-}
-
-/*
- * This too could be housed elsewhere, along with an API to make it easier to
- * run analysis on trace files. The current model is influenced a little too much
- * on the store model that Insights' browser mode hosts.
- */
-
-class FFileDataStream : public UE::Trace::IInDataStream
-{
-public:
-	FFileDataStream()
-		: Handle(nullptr)
+	if (ScopeId < static_cast<uint32>(ScopeNames.Num()) && ScopeNames[ScopeId])
 	{
+		return &ScopeNames[ScopeId].GetValue();
 	}
+	
+	return nullptr;
+}
 
-	~FFileDataStream()
-	{
-		delete Handle;
-	}
-
-	bool Open(const TCHAR* Path)
-	{
-		Handle = FPlatformFileManager::Get().GetPlatformFile().OpenRead(Path);
-		if (Handle == nullptr)
-		{
-			return false;
-		}
-		Remaining = Handle->Size();
-		return true;
-	}
-
-	virtual int32 Read(void* Data, uint32 Size) override
-	{
-		if (Remaining <= 0 || Handle == nullptr)
-		{
-			return 0;
-		}
-
-		Size = (Size < Remaining) ? Size : Remaining;
-		Remaining -= Size;
-		return Handle->Read((uint8*)Data, Size) ? Size : 0;
-	}
-
-	IFileHandle* Handle;
-	uint64 Remaining;
-};
 
 /*
 * 
@@ -795,6 +511,16 @@ public:
 		VarianceAccumulator = NewVarianceAccumulator;
 	}
 
+	/**
+	* Reset state back to initialized.
+	*/
+	void Reset()
+	{
+		Count = 0;
+		Mean = 0.0;
+		VarianceAccumulator = 0.0;
+	}
+
 private:
 	uint64 Count;
 	double Mean;
@@ -802,7 +528,7 @@ private:
 };
 
 //
-// FSummarizeCpuAnalyzer - Helpers & Analyzers for Cpu channel events
+// Helpers for Cpu channel events
 //
 
 struct FSummarizeScope
@@ -980,15 +706,15 @@ static uint32 GetTypeHash(const FSummarizeScope& Scope)
  * This analyzer aggregates CPU scopes having the same name together, counting the number of occurrences, total duration,
  * average occurrence duration, for each scope name.
  */
-class FSummarizeCpuAnalyzer
-	: public FCpuScopeAnalyzer
+class FSummarizeCpuScopeDurationAnalyzer
+	: public FSummarizeCpuScopeAnalyzer
 {
 public:
-	FSummarizeCpuAnalyzer(TFunction<void(const TArray<FSummarizeScope>&)> InPublishFn);
+	FSummarizeCpuScopeDurationAnalyzer(TFunction<void(const TArray<FSummarizeScope>&)> InPublishFn);
 
 protected:
 	virtual void OnCpuScopeDiscovered(uint32 ScopeId) override;
-	virtual void OnCpuScopeName(uint32 ScopeId, const FString& ScopeName) override;
+	virtual void OnCpuScopeName(uint32 ScopeId, const FStringView& ScopeName) override;
 	virtual void OnCpuScopeExit(const FScope& Scope, const FString* ScopeName) override;
 	virtual void OnCpuScopeAnalysisEnd() override;
 
@@ -999,16 +725,16 @@ private:
 	TFunction<void(const TArray<FSummarizeScope>&)> PublishFn;
 };
 
-FSummarizeCpuAnalyzer::FSummarizeCpuAnalyzer(TFunction<void(const TArray<FSummarizeScope>&)> InPublishFn)
+FSummarizeCpuScopeDurationAnalyzer::FSummarizeCpuScopeDurationAnalyzer(TFunction<void(const TArray<FSummarizeScope>&)> InPublishFn)
 	: PublishFn(MoveTemp(InPublishFn))
 {
-	OnCpuScopeDiscovered(FCpuScopeAnalyzer::CoroutineSpecId);
-	OnCpuScopeDiscovered(FCpuScopeAnalyzer::CoroutineUnknownSpecId);
-	OnCpuScopeName(FCpuScopeAnalyzer::CoroutineSpecId, TEXT("Coroutine"));
-	OnCpuScopeName(FCpuScopeAnalyzer::CoroutineUnknownSpecId, TEXT("<unknown>"));
+	OnCpuScopeDiscovered(FSummarizeCpuScopeAnalyzer::CoroutineSpecId);
+	OnCpuScopeDiscovered(FSummarizeCpuScopeAnalyzer::CoroutineUnknownSpecId);
+	OnCpuScopeName(FSummarizeCpuScopeAnalyzer::CoroutineSpecId, TEXT("Coroutine"));
+	OnCpuScopeName(FSummarizeCpuScopeAnalyzer::CoroutineUnknownSpecId, TEXT("<unknown>"));
 }
 
-void FSummarizeCpuAnalyzer::OnCpuScopeDiscovered(uint32 ScopeId)
+void FSummarizeCpuScopeDurationAnalyzer::OnCpuScopeDiscovered(uint32 ScopeId)
 {
 	if (!Scopes.Find(ScopeId))
 	{
@@ -1016,17 +742,22 @@ void FSummarizeCpuAnalyzer::OnCpuScopeDiscovered(uint32 ScopeId)
 	}
 }
 
-void FSummarizeCpuAnalyzer::OnCpuScopeName(uint32 ScopeId, const FString& ScopeName)
+void FSummarizeCpuScopeDurationAnalyzer::OnCpuScopeName(uint32 ScopeId, const FStringView& ScopeName)
 {
 	Scopes[ScopeId].Name = ScopeName;
 }
 
-void FSummarizeCpuAnalyzer::OnCpuScopeExit(const FScope& Scope, const FString* ScopeName)
+void FSummarizeCpuScopeDurationAnalyzer::OnCpuScopeExit(const FScope& Scope, const FString* ScopeName)
 {
-	Scopes[Scope.ScopeId].AddDuration(Scope.EnterTimestamp, Scope.ExitTimestamp);
+	// This can miss if we are given a metadata timer's index, which is negative signed
+	FSummarizeScope* Found = Scopes.Find(Scope.ScopeId);
+	if (Found)
+	{
+		Found->AddDuration(Scope.EnterTimestamp, Scope.ExitTimestamp);
+	}
 }
 
-void FSummarizeCpuAnalyzer::OnCpuScopeAnalysisEnd()
+void FSummarizeCpuScopeDurationAnalyzer::OnCpuScopeAnalysisEnd()
 {
 	TArray<FSummarizeScope> LocalScopes;
 	// Eliminates scopes that don't have a name. (On scope discovery, the array is expended to creates blank scopes that may never be filled).
@@ -1072,7 +803,7 @@ void FSummarizeCpuAnalyzer::OnCpuScopeAnalysisEnd()
  */
 
 class FSummarizeCpuScopeHierarchyAnalyzer
-	: public FCpuScopeAnalyzer
+	: public FSummarizeCpuScopeAnalyzer
 {
 public:
 	/**
@@ -1094,7 +825,7 @@ public:
 	 * @param ScopeEvents The scopes events containing one root event along with its hierarchy.
 	 * @param InScopeNameLookup Callback function to lookup scope names from scope ID.
 	 */
-	virtual void OnCpuScopeTree(uint32 ThreadId, const TArray<FCpuScopeAnalyzer::FScopeEvent>& ScopeEvents, const TFunction<const FString*(uint32 /*ScopeId*/)>& InScopeNameLookup) override;
+	virtual void OnCpuScopeTree(uint32 ThreadId, const TArray<FSummarizeCpuScopeAnalyzer::FScopeEvent>& ScopeEvents, const TFunction<const FString*(uint32 /*ScopeId*/)>& InScopeNameLookup) override;
 
 	/**
 	 * Invoked to notify that the trace session ended and that the analyzer can publish the statistics gathered.
@@ -1114,7 +845,7 @@ private:
 	// from the 'exclusive' time (itself - duration of immediate children).
 	TMap<FString, FSummarizeScope> ExactNameMatchScopesSummaries;
 
-	// Invoked at the end of the analysis to publich scope summaries.
+	// Invoked at the end of the analysis to publish scope summaries.
 	TFunction<void(const FSummarizeScope&, TArray<FSummarizeScope>&&)> PublishFn;
 };
 
@@ -1125,7 +856,7 @@ FSummarizeCpuScopeHierarchyAnalyzer::FSummarizeCpuScopeHierarchyAnalyzer(const F
 	MatchedScopesSummary.Name = InAnalyzerName;
 }
 
-void FSummarizeCpuScopeHierarchyAnalyzer::OnCpuScopeTree(uint32 ThreadId, const TArray<FCpuScopeAnalyzer::FScopeEvent>& ScopeEvents, const TFunction<const FString*(uint32 /*ScopeId*/)>& InScopeNameLookup)
+void FSummarizeCpuScopeHierarchyAnalyzer::OnCpuScopeTree(uint32 ThreadId, const TArray<FSummarizeCpuScopeAnalyzer::FScopeEvent>& ScopeEvents, const TFunction<const FString*(uint32 /*ScopeId*/)>& InScopeNameLookup)
 {
 	// Scope matching the pattern.
 	struct FMatchScopeEnter
@@ -1139,7 +870,7 @@ void FSummarizeCpuScopeHierarchyAnalyzer::OnCpuScopeTree(uint32 ThreadId, const
 	TArray<FMatchScopeEnter> ScopeStack;
 
 	// Replay and filter this scope hierarchy to only keep the ones matching the condition/regex. (See class documentation for a visual example)
-	for (const FCpuScopeAnalyzer::FScopeEvent& ScopeEvent : ScopeEvents)
+	for (const FSummarizeCpuScopeAnalyzer::FScopeEvent& ScopeEvent : ScopeEvents)
 	{
 		const FString* ScopeName = InScopeNameLookup(ScopeEvent.ScopeId);
 		if (!ScopeName || !MatchesFn(*ScopeName))
@@ -1147,7 +878,7 @@ void FSummarizeCpuScopeHierarchyAnalyzer::OnCpuScopeTree(uint32 ThreadId, const
 			continue;
 		}
 
-		if (ScopeEvent.ScopeEventType == FCpuScopeAnalyzer::EScopeEventType::Enter)
+		if (ScopeEvent.ScopeEventType == FSummarizeCpuScopeAnalyzer::EScopeEventType::Enter)
 		{
 			ScopeStack.Emplace(ScopeEvent.ScopeId, ScopeEvent.Timestamp);
 		}
@@ -1200,9 +931,10 @@ void FSummarizeCpuScopeHierarchyAnalyzer::OnCpuScopeAnalysisEnd()
 //
 
 struct FSummarizeCounter
+	: public TraceServices::IEditableCounter
 {
 	FString Name;
-	ETraceCounterType Type;
+	ETraceCounterType Type = TraceCounterType_Int;
 
 	union FValue
 	{
@@ -1221,29 +953,98 @@ struct FSummarizeCounter
 	double FirstSeconds = 0.0;
 	double LastSeconds = 0.0;
 
-	FSummarizeCounter(FString InName, ETraceCounterType InType)
+	FSummarizeCounter()
+	{
+		SetIsFloatingPoint(false);
+	}
+
+	virtual void SetName(const TCHAR* InName) override
 	{
 		Name = InName;
-		Type = InType;
-		switch (Type)
+	}
+
+	virtual void SetGroup(const TCHAR* Group) override
+	{}
+
+	virtual void SetDescription(const TCHAR* Description) override
+	{}
+
+	virtual void SetIsFloatingPoint(bool bIsFloatingPoint) override
+	{
+		ETraceCounterType Previous = Type;
+
+		if (bIsFloatingPoint)
 		{
-			case TraceCounterType_Int:
-				First.Int = 0;
-				Last.Int = 0;
-				Minimum.Int = TNumericLimits<int64>::Max();
-				Maximum.Int = TNumericLimits<int64>::Min();
-				break;
+			Type = TraceCounterType_Float;
+			First.Float = 0.0;
+			Last.Float = 0.0;
+			Minimum.Float = TNumericLimits<double>::Max();
+			Maximum.Float = TNumericLimits<double>::Min();
+		}
+		else
+		{
+			Type = TraceCounterType_Int;
+			First.Int = 0;
+			Last.Int = 0;
+			Minimum.Int = TNumericLimits<int64>::Max();
+			Maximum.Int = TNumericLimits<int64>::Min();
+		}
 
-			case TraceCounterType_Float:
-				First.Float = 0.0;
-				Last.Float = 0.0;
-				Minimum.Float = TNumericLimits<double>::Max();
-				Maximum.Float = TNumericLimits<double>::Min();
-				break;
+		if (Previous != Type)
+		{
+			Variance.Reset();
+			FirstSeconds = 0.0;
+			LastSeconds = 0.0;
+		}
+	}
+
+	virtual void SetIsResetEveryFrame(bool bInIsResetEveryFrame) override
+	{}
+
+	virtual void SetDisplayHint(TraceServices::ECounterDisplayHint DisplayHint) override
+	{}
+
+	virtual void AddValue(double Time, int64 Value) override
+	{
+		int64 Current;
+		if (Get(Current))
+		{
+			Set(Current + Value, Time);
 		}
 	}
 
-	void SetValue(int64 InValue, double InTimestamp)
+	virtual void AddValue(double Time, double Value) override
+	{
+		double Current;
+		if (Get(Current))
+		{
+			Set(Current + Value, Time);
+		}
+	}
+
+	virtual void SetValue(double Time, int64 Value) override
+	{
+		Set(Value, Time);
+	}
+
+	virtual void SetValue(double Time, double Value) override
+	{
+		Set(Value, Time);
+	}
+
+	bool Get(int64& Value)
+	{
+		ensure(Type == TraceCounterType_Int);
+		if (Type == TraceCounterType_Int)
+		{
+			Value = Last.Int;
+			return true;
+		}
+
+		return false;
+	}
+
+	void Set(int64 InValue, double InTimestamp)
 	{
 		ensure(Type == TraceCounterType_Int);
 		if (Type == TraceCounterType_Int)
@@ -1264,7 +1065,19 @@ struct FSummarizeCounter
 		}
 	}
 
-	void SetValue(double InValue, double InTimestamp)
+	bool Get(double& Value)
+	{
+		ensure(Type == TraceCounterType_Float);
+		if (Type == TraceCounterType_Float)
+		{
+			Value = Last.Float;
+			return true;
+		}
+
+		return false;
+	}
+
+	void Set(double InValue, double InTimestamp)
 	{
 		ensure(Type == TraceCounterType_Float);
 		if (Type == TraceCounterType_Float)
@@ -1385,44 +1198,32 @@ struct FSummarizeCounter
 	}
 };
 
-class FSummarizeCountersAnalyzer
-	: public FCountersAnalyzer
+class FSummarizeCountersProvider
+	: public TraceServices::IEditableCounterProvider
 {
 public:
-	virtual void OnCounterName(const FCounterName& CounterName) override;
-	virtual void OnCounterIntValue(const FCounterIntValue& NewValue) override;
-	virtual void OnCounterFloatValue(const FCounterFloatValue& NewValue) override;
-
-	TMap<uint16, FSummarizeCounter> Counters;
-};
-
-void FSummarizeCountersAnalyzer::OnCounterName(const FCounterName& CounterName)
-{
-	Counters.Add(CounterName.Id, FSummarizeCounter(CounterName.Name, CounterName.Type));
-}
+	const TraceServices::ICounter* GetCounter(TraceServices::IEditableCounter* EditableCounter)
+	{
+		// we don't want derived counters to be created by the analyzer
+		return nullptr;
+	}
 
-void FSummarizeCountersAnalyzer::OnCounterIntValue(const FCounterIntValue& NewValue)
-{
-	FSummarizeCounter* FoundCounter = Counters.Find(NewValue.Id);
-	ensure(FoundCounter);
-	if (FoundCounter)
+	virtual TraceServices::IEditableCounter* CreateEditableCounter()
 	{
-		FoundCounter->SetValue(NewValue.Value, NewValue.Timestamp);
+		Counters.Add(MakeUnique<FSummarizeCounter>());
+		return Counters.Last().Get();
 	}
-}
 
-void FSummarizeCountersAnalyzer::OnCounterFloatValue(const FCounterFloatValue& NewValue)
-{
-	FSummarizeCounter* FoundCounter = Counters.Find(NewValue.Id);
-	ensure(FoundCounter);
-	if (FoundCounter)
+	virtual void AddCounter(const TraceServices::ICounter* Counter)
 	{
-		FoundCounter->SetValue(NewValue.Value, NewValue.Timestamp);
+		// we don't use custom counter objects
 	}
-}
+
+	TArray<TUniquePtr<FSummarizeCounter>> Counters;
+};
 
 //
-// FSummarizeBookmarksAnalyzer - Tally Bookmarks from bookmark events
+// FSummarizeBookmarksProvider - Tally Bookmarks from bookmark events
 //
 
 struct FSummarizeBookmark
@@ -1479,23 +1280,24 @@ static uint32 GetTypeHash(const FSummarizeBookmark& Bookmark)
 	return FCrc::StrCrc32(*Bookmark.Name);
 }
 
-class FSummarizeBookmarksAnalyzer
-	: public FBookmarksAnalyzer
+class FSummarizeBookmarksProvider
+	: public TraceServices::IEditableBookmarkProvider
 {
-	virtual void OnBookmarkSpecEvent(const FBookmarkSpecEvent& BookmarkSpecEvent) override;
-	virtual void OnBookmarkEvent(const FBookmarkEvent& BookmarkEvent) override;
+	virtual void UpdateBookmarkSpec(uint64 BookmarkPoint, const TCHAR* FormatString, const TCHAR* File, int32 Line) override;
+	virtual void AppendBookmark(uint64 BookmarkPoint, double Time, const uint8* FormatString) override;
+	virtual void AppendBookmark(uint64 BookmarkPoint, double Time, const TCHAR* Text) override;
 
-	FSummarizeBookmark* FindStartBookmarkForEndBookmark(const FString& Name);
-
-public:
 	struct FBookmarkSpec
 	{
-		uint64	Id;
-		FString	FileName;
-		int32	Line;
-		FString	FormatString;
+		const TCHAR* File = nullptr;
+		const TCHAR* FormatString = nullptr;
+		int32 Line = 0;
 	};
 
+	FBookmarkSpec& GetSpec(uint64 BookmarkPoint);
+	FSummarizeBookmark* FindStartBookmarkForEndBookmark(const FString& Name);
+
+public:
 	// Keyed by a unique memory address
 	TMap<uint64, FBookmarkSpec> BookmarkSpecs;
 
@@ -1504,54 +1306,74 @@ public:
 
 	// Bookmarks named formed to scopes, see FindStartBookmarkForEndBookmark
 	TMap<FString, FSummarizeScope> Scopes;
+
+private:
+	enum
+	{
+		FormatBufferSize = 65536
+	};
+	TCHAR FormatBuffer[FormatBufferSize];
+	TCHAR TempBuffer[FormatBufferSize];
 };
 
-void FSummarizeBookmarksAnalyzer::OnBookmarkSpecEvent(const FBookmarkSpecEvent& BookmarkSpecEvent)
+FSummarizeBookmarksProvider::FBookmarkSpec& FSummarizeBookmarksProvider::GetSpec(uint64 BookmarkPoint)
 {
-	BookmarkSpecs.Add(BookmarkSpecEvent.Id, {
-		BookmarkSpecEvent.Id,
-		BookmarkSpecEvent.FileName,
-		BookmarkSpecEvent.Line,
-		BookmarkSpecEvent.FormatString
-		});
+	FBookmarkSpec* Found = BookmarkSpecs.Find(BookmarkPoint);
+	if (Found)
+	{
+		return *Found;
+	}
+
+	FBookmarkSpec& Spec = BookmarkSpecs.Add(BookmarkPoint, FBookmarkSpec());
+	Spec.File = TEXT("<unknown>");
+	Spec.FormatString = TEXT("<unknown>");
+	return Spec;
 }
 
-void FSummarizeBookmarksAnalyzer::OnBookmarkEvent(const FBookmarkEvent& BookmarkEvent)
+void FSummarizeBookmarksProvider::UpdateBookmarkSpec(uint64 BookmarkPoint, const TCHAR* FormatString, const TCHAR* File, int32 Line)
 {
-	FBookmarkSpec* Spec = BookmarkSpecs.Find(BookmarkEvent.Id);
-	if (Spec)
-	{
-		TCHAR FormattedString[65535];
-		TraceServices::FormatString(FormattedString, sizeof(FormattedString) / sizeof(FormattedString[0]), *Spec->FormatString, BookmarkEvent.FormatArgs.GetData());
+	FBookmarkSpec& BookmarkSpec = GetSpec(BookmarkPoint);
+	BookmarkSpec.FormatString = FormatString;
+	BookmarkSpec.File = File;
+	BookmarkSpec.Line = Line;
+}
 
-		FString Name (FormattedString);
+void FSummarizeBookmarksProvider::AppendBookmark(uint64 BookmarkPoint, double Time, const uint8* FormatArgs)
+{
+	FBookmarkSpec& BookmarkSpec = GetSpec(BookmarkPoint);
+	TraceServices::StringFormat(FormatBuffer, FormatBufferSize - 1, TempBuffer, FormatBufferSize - 1, BookmarkSpec.FormatString, FormatArgs);
+	FSummarizeBookmarksProvider::AppendBookmark(BookmarkPoint, Time, FormatBuffer);
+}
 
-		FSummarizeBookmark* FoundBookmark = Bookmarks.Find(Name);
-		if (!FoundBookmark)
-		{
-			FoundBookmark = &Bookmarks.Add(Name, FSummarizeBookmark());
-			FoundBookmark->Name = Name;
-		}
+void FSummarizeBookmarksProvider::AppendBookmark(uint64 BookmarkPoint, double Time, const TCHAR* Text)
+{
+	FString Name(Text);
+
+	FSummarizeBookmark* FoundBookmark = Bookmarks.Find(Name);
+	if (!FoundBookmark)
+	{
+		FoundBookmark = &Bookmarks.Add(Name, FSummarizeBookmark());
+		FoundBookmark->Name = Name;
+	}
 
-		FoundBookmark->AddTimestamp(BookmarkEvent.Timestamp);
+	FoundBookmark->AddTimestamp(Time);
 
-		FSummarizeBookmark* StartBookmark = FindStartBookmarkForEndBookmark(Name);
-		if (StartBookmark)
+	FSummarizeBookmark* StartBookmark = FindStartBookmarkForEndBookmark(Name);
+	if (StartBookmark)
+	{
+		FString ScopeName = FString(TEXT("Generated Scope for ")) + StartBookmark->Name;
+		FSummarizeScope* FoundScope = Scopes.Find(ScopeName);
+		if (!FoundScope)
 		{
-			FString ScopeName = FString(TEXT("Generated Scope for ")) + StartBookmark->Name;
-			FSummarizeScope* FoundScope = Scopes.Find(ScopeName);
-			if (!FoundScope)
-			{
-				FoundScope = &Scopes.Add(ScopeName, FSummarizeScope());
-				FoundScope->Name = ScopeName;
-			}
-
-			FoundScope->AddDuration(StartBookmark->LastSeconds, BookmarkEvent.Timestamp);
+			FoundScope = &Scopes.Add(ScopeName, FSummarizeScope());
+			FoundScope->Name = ScopeName;
 		}
+
+		FoundScope->AddDuration(StartBookmark->LastSeconds, Time);
 	}
 }
 
-FSummarizeBookmark* FSummarizeBookmarksAnalyzer::FindStartBookmarkForEndBookmark(const FString& Name)
+FSummarizeBookmark* FSummarizeBookmarksProvider::FindStartBookmarkForEndBookmark(const FString& Name)
 {
 	int32 Index = Name.Find(TEXT("Complete"));
 	if (Index != -1)
@@ -2051,8 +1873,8 @@ int32 USummarizeTraceCommandlet::Main(const FString& CmdLineParams)
 		return 1;
 	}
 
-	FFileDataStream DataStream;
-	if (!DataStream.Open(*TraceFileName))
+	UE::Trace::FFileDataStream* DataStream = new UE::Trace::FFileDataStream();
+	if (!DataStream->Open(*TraceFileName))
 	{
 		UE_LOG(LogSummarizeTrace, Error, TEXT("Unable to open trace file '%s' for read"), *TraceFileName);
 		return 1;
@@ -2065,7 +1887,7 @@ int32 USummarizeTraceCommandlet::Main(const FString& CmdLineParams)
 	TArray<FSummarizeScope> CollectedScopeSummaries;
 
 	// Analyze CPU scope timer individually.
-	TSharedPtr<FSummarizeCpuAnalyzer> IndividualScopeAnalyzer = MakeShared<FSummarizeCpuAnalyzer>(
+	TSharedPtr<FSummarizeCpuScopeDurationAnalyzer> IndividualScopeAnalyzer = MakeShared<FSummarizeCpuScopeDurationAnalyzer>(
 		[&CollectedScopeSummaries](const TArray<FSummarizeScope>& ScopeSummaries)
 		{
 			// Collect all individual scopes summary from this analyzer.
@@ -2096,22 +1918,30 @@ int32 USummarizeTraceCommandlet::Main(const FString& CmdLineParams)
 			CollectedScopeSummaries.Append(ModuleStats.GetData(), 10);
 		});
 
-	FCpuScopeStreamProcessor CpuScopeStreamProcessor;
-	CpuScopeStreamProcessor.AddCpuScopeAnalyzer(IndividualScopeAnalyzer);
-	CpuScopeStreamProcessor.AddCpuScopeAnalyzer(HierarchicalScopeAnalyzer);
-	AnalysisContext.AddAnalyzer(CpuScopeStreamProcessor);
+	TSharedPtr<TraceServices::IAnalysisSession> Session = TraceServices::CreateAnalysisSession(0, nullptr, TUniquePtr<UE::Trace::IInDataStream>(DataStream));
+
+	FSummarizeBookmarksProvider BookmarksProvider;
+	TSharedPtr<UE::Trace::IAnalyzer> BookmarksAnalyzer = TraceServices::CreateBookmarksAnalyzer(*Session, BookmarksProvider);
+	AnalysisContext.AddAnalyzer(*BookmarksAnalyzer);
 
-	FSummarizeCountersAnalyzer CountersAnalyzer;
-	AnalysisContext.AddAnalyzer(CountersAnalyzer);
-	FSummarizeBookmarksAnalyzer BookmarksAnalyzer;
-	AnalysisContext.AddAnalyzer(BookmarksAnalyzer);
+	FSummarizeCountersProvider CountersProvider;
+	TSharedPtr<UE::Trace::IAnalyzer> CountersAnalyzer = TraceServices::CreateCountersAnalyzer(*Session, CountersProvider);
+	AnalysisContext.AddAnalyzer(*CountersAnalyzer);
+
+	FSummarizeCpuProfilerProvider CpuProfilerProvider;
+	CpuProfilerProvider.AddCpuScopeAnalyzer(IndividualScopeAnalyzer);
+	CpuProfilerProvider.AddCpuScopeAnalyzer(HierarchicalScopeAnalyzer);
+	TSharedPtr<UE::Trace::IAnalyzer> CpuProfilerAnalyzer = TraceServices::CreateCpuProfilerAnalyzer(*Session, CpuProfilerProvider, CpuProfilerProvider);
+	AnalysisContext.AddAnalyzer(*CpuProfilerAnalyzer);
 
 	// kick processing on a thread
-	UE::Trace::FAnalysisProcessor AnalysisProcessor = AnalysisContext.Process(DataStream);
+	UE::Trace::FAnalysisProcessor AnalysisProcessor = AnalysisContext.Process(*DataStream);
 
 	// sync on completion
 	AnalysisProcessor.Wait();
 
+	CpuProfilerProvider.AnalysisComplete();
+
 	TSet<FSummarizeScope> DeduplicatedScopes;
 	auto IngestScope = [](TSet<FSummarizeScope>& DeduplicatedScopes, const FSummarizeScope& Scope)
 	{
@@ -2139,7 +1969,7 @@ int32 USummarizeTraceCommandlet::Main(const FString& CmdLineParams)
 	{
 		IngestScope(DeduplicatedScopes, Scope);
 	}
-	for (const TMap<FString, FSummarizeScope>::ElementType& ScopeItem : BookmarksAnalyzer.Scopes)
+	for (const TMap<FString, FSummarizeScope>::ElementType& ScopeItem : BookmarksProvider.Scopes)
 	{
 		IngestScope(DeduplicatedScopes, ScopeItem.Value);
 	}
@@ -2211,20 +2041,20 @@ int32 USummarizeTraceCommandlet::Main(const FString& CmdLineParams)
 	{
 		// no newline, see row printfs
 		WriteUTF8(CsvHandle, FString::Printf(TEXT("Name,Count,CountPerSecond,First,FirstSeconds,Last,LastSeconds,Minimum,Maximum,Mean,Deviation,")));
-		for (const TMap<uint16, FSummarizeCounter>::ElementType& Counter : CountersAnalyzer.Counters)
+		for (const TUniquePtr<FSummarizeCounter>& Counter : CountersProvider.Counters)
 		{
-			if (!IsCsvSafeString(Counter.Value.Name))
+			if (!IsCsvSafeString(Counter->Name))
 			{
 				continue;
 			}
 
 			// note newline is at the front of every data line to prevent final extraneous newline, per customary for csv
 			WriteUTF8(CsvHandle, FString::Printf(TEXT("\n%s,%llu,%f,%s,%f,%s,%f,%s,%s,%f,%f,"),
-				*Counter.Value.Name, Counter.Value.GetCount(), Counter.Value.GetCountPerSecond(),
-				*Counter.Value.GetValue(TEXT("First")), Counter.Value.FirstSeconds,
-				*Counter.Value.GetValue(TEXT("Last")), Counter.Value.LastSeconds,
-				*Counter.Value.GetValue(TEXT("Minimum")), *Counter.Value.GetValue(TEXT("Maximum")),
-				Counter.Value.GetMean(), Counter.Value.GetDeviation()));
+				*Counter->Name, Counter->GetCount(), Counter->GetCountPerSecond(),
+				*Counter->GetValue(TEXT("First")), Counter->FirstSeconds,
+				*Counter->GetValue(TEXT("Last")), Counter->LastSeconds,
+				*Counter->GetValue(TEXT("Minimum")), *Counter->GetValue(TEXT("Maximum")),
+				Counter->GetMean(), Counter->GetDeviation()));
 		}
 		CsvHandle->Flush();
 		delete CsvHandle;
@@ -2244,7 +2074,7 @@ int32 USummarizeTraceCommandlet::Main(const FString& CmdLineParams)
 	{
 		// no newline, see row printfs
 		WriteUTF8(CsvHandle, FString::Printf(TEXT("Name,Count,FirstSeconds,LastSeconds,")));
-		for (const TMap<FString, FSummarizeBookmark>::ElementType& Bookmark : BookmarksAnalyzer.Bookmarks)
+		for (const TMap<FString, FSummarizeBookmark>::ElementType& Bookmark : BookmarksProvider.Bookmarks)
 		{
 			if (!IsCsvSafeString(Bookmark.Value.Name))
 			{
@@ -2324,31 +2154,31 @@ int32 USummarizeTraceCommandlet::Main(const FString& CmdLineParams)
 		}
 
 		// resolve counters to telemetry
-		for (const TMap<uint16, FSummarizeCounter>::ElementType& Counter : CountersAnalyzer.Counters)
+		for (const TUniquePtr<FSummarizeCounter>& Counter : CountersProvider.Counters)
 		{
-			if (!IsCsvSafeString(Counter.Value.Name))
+			if (!IsCsvSafeString(Counter->Name))
 			{
 				continue;
 			}
 
 			if (bAllTelemetry)
 			{
-				TelemetryData.Add(TelemetryDefinition(TestName, Counter.Value.Name, TEXT("Count"), Counter.Value.GetValue(TEXT("Count")), TEXT("Count")));
-				TelemetryData.Add(TelemetryDefinition(TestName, Counter.Value.Name, TEXT("First"), Counter.Value.GetValue(TEXT("First")), TEXT("Count")));
-				TelemetryData.Add(TelemetryDefinition(TestName, Counter.Value.Name, TEXT("FirstSeconds"), Counter.Value.GetValue(TEXT("FirstSeconds")), TEXT("Seconds")));
-				TelemetryData.Add(TelemetryDefinition(TestName, Counter.Value.Name, TEXT("Last"), Counter.Value.GetValue(TEXT("Last")), TEXT("Count")));
-				TelemetryData.Add(TelemetryDefinition(TestName, Counter.Value.Name, TEXT("LastSeconds"), Counter.Value.GetValue(TEXT("LastSeconds")), TEXT("Seconds")));
-				TelemetryData.Add(TelemetryDefinition(TestName, Counter.Value.Name, TEXT("Minimum"), Counter.Value.GetValue(TEXT("Minimum")), TEXT("Count")));
-				TelemetryData.Add(TelemetryDefinition(TestName, Counter.Value.Name, TEXT("Maximum"), Counter.Value.GetValue(TEXT("Maximum")), TEXT("Count")));
-				TelemetryData.Add(TelemetryDefinition(TestName, Counter.Value.Name, TEXT("Mean"), Counter.Value.GetValue(TEXT("Mean")), TEXT("Count")));
-				TelemetryData.Add(TelemetryDefinition(TestName, Counter.Value.Name, TEXT("Deviation"), Counter.Value.GetValue(TEXT("Deviation")), TEXT("Count")));
+				TelemetryData.Add(TelemetryDefinition(TestName, Counter->Name, TEXT("Count"), Counter->GetValue(TEXT("Count")), TEXT("Count")));
+				TelemetryData.Add(TelemetryDefinition(TestName, Counter->Name, TEXT("First"), Counter->GetValue(TEXT("First")), TEXT("Count")));
+				TelemetryData.Add(TelemetryDefinition(TestName, Counter->Name, TEXT("FirstSeconds"), Counter->GetValue(TEXT("FirstSeconds")), TEXT("Seconds")));
+				TelemetryData.Add(TelemetryDefinition(TestName, Counter->Name, TEXT("Last"), Counter->GetValue(TEXT("Last")), TEXT("Count")));
+				TelemetryData.Add(TelemetryDefinition(TestName, Counter->Name, TEXT("LastSeconds"), Counter->GetValue(TEXT("LastSeconds")), TEXT("Seconds")));
+				TelemetryData.Add(TelemetryDefinition(TestName, Counter->Name, TEXT("Minimum"), Counter->GetValue(TEXT("Minimum")), TEXT("Count")));
+				TelemetryData.Add(TelemetryDefinition(TestName, Counter->Name, TEXT("Maximum"), Counter->GetValue(TEXT("Maximum")), TEXT("Count")));
+				TelemetryData.Add(TelemetryDefinition(TestName, Counter->Name, TEXT("Mean"), Counter->GetValue(TEXT("Mean")), TEXT("Count")));
+				TelemetryData.Add(TelemetryDefinition(TestName, Counter->Name, TEXT("Deviation"), Counter->GetValue(TEXT("Deviation")), TEXT("Count")));
 			}
 			else
 			{
-				NameToDefinitionMap.MultiFind(Counter.Value.Name, Statistics, true);
+				NameToDefinitionMap.MultiFind(Counter->Name, Statistics, true);
 				for (const StatisticDefinition& Statistic : Statistics)
 				{
-					TelemetryData.Add(TelemetryDefinition(TestName, Statistic.TelemetryContext, Statistic.TelemetryDataPoint, Statistic.TelemetryUnit, Counter.Value.GetValue(Statistic.Statistic)));
+					TelemetryData.Add(TelemetryDefinition(TestName, Statistic.TelemetryContext, Statistic.TelemetryDataPoint, Statistic.TelemetryUnit, Counter->GetValue(Statistic.Statistic)));
 					ResolvedStatistics.Add(Statistic.Name);
 				}
 				Statistics.Reset();			
@@ -2356,7 +2186,7 @@ int32 USummarizeTraceCommandlet::Main(const FString& CmdLineParams)
 		}
 
 		// resolve bookmarks to telemetry
-		for (const TMap<FString, FSummarizeBookmark>::ElementType& Bookmark : BookmarksAnalyzer.Bookmarks)
+		for (const TMap<FString, FSummarizeBookmark>::ElementType& Bookmark : BookmarksProvider.Bookmarks)
 		{
 			if (!IsCsvSafeString(Bookmark.Value.Name))
 			{
