Fix an issue where the OverrideOutputDirectory was not locally set after the settings were parsed.
Not setting the OverrideOutputDirectory can result in cooked packages ending up in unexpected folder locations.

diff --git a/Engine/Source/Editor/UnrealEd/Private/CookOnTheFlyServer.cpp b/Engine/Source/Editor/UnrealEd/Private/CookOnTheFlyServer.cpp
--- a/Engine/Source/Editor/UnrealEd/Private/CookOnTheFlyServer.cpp
+++ b/Engine/Source/Editor/UnrealEd/Private/CookOnTheFlyServer.cpp
@@ -5720,6 +5720,7 @@ void FInitializeConfigSettings::LoadLocal(const FString& InOutputDirectoryOverri
 
 void UCookOnTheFlyServer::SetInitializeConfigSettings(UE::Cook::FInitializeConfigSettings&& Settings)
 {
+	OutputDirectoryOverride = MoveTemp(Settings.OutputDirectoryOverride);
 	MaxPrecacheShaderJobs = Settings.MaxPrecacheShaderJobs;
 	MaxConcurrentShaderJobs = Settings.MaxConcurrentShaderJobs;
 	PackagesPerGC = Settings.PackagesPerGC;
