UE-157689 Assertion failed: PlayInEditorSessionInfo->CachedWindowInfo.Num() == InViewportIndex + 1 when settting Net Mode to Play as Client or listen Server on Lyra
-There are cases, like PIE with "Selected Viewport" where we don't have a meaningful window size/position to cache for the pie instance.  But in multiplayer pie we were simply appending to an array to cache a previously-unused index.  This meant we could end up caching for index 1 into index 0 on the first attempt.  A new assert I added triggered on that case.  Because we would typically re-cache into a no-longer empty array, adding another element on PIE shutdown this bug was pretty invisible.
-Now we are padding the cache up to our current index when caching so we never end up off by one.

diff --git a/Engine/Source/Editor/UnrealEd/Private/PlayLevel.cpp b/Engine/Source/Editor/UnrealEd/Private/PlayLevel.cpp
--- a/Engine/Source/Editor/UnrealEd/Private/PlayLevel.cpp
+++ b/Engine/Source/Editor/UnrealEd/Private/PlayLevel.cpp
@@ -3664,10 +3664,15 @@ void UEditorEngine::StoreWindowSizeAndPositionForInstanceIndex(const int32 InVie
 	}
 	else
 	{
-		FPlayInEditorSessionInfo::FWindowSizeAndPos& NewInfo = PlayInEditorSessionInfo->CachedWindowInfo.Add_GetRef(FPlayInEditorSessionInfo::FWindowSizeAndPos());
-		NewInfo.Size = InSize;
-		NewInfo.Position = InPosition;
-		check(PlayInEditorSessionInfo->CachedWindowInfo.Num() == InViewportIndex + 1);
+		// It is possible for PIE to play in viewports which are not independent and for which we do not want to cache window size/position information.  For example "Selected Viewport" will use a docked editor viewport for PIE.
+		// Appending as many as we need to get to our proper index.  These cache entries may be overwritten later, or they may be left with what must be at least a plausible size/pos.
+		while (PlayInEditorSessionInfo->CachedWindowInfo.Num() < InViewportIndex + 1)
+		{
+			FPlayInEditorSessionInfo::FWindowSizeAndPos& NewInfo = PlayInEditorSessionInfo->CachedWindowInfo.Add_GetRef(FPlayInEditorSessionInfo::FWindowSizeAndPos());
+			NewInfo.Size = InSize;
+			NewInfo.Position = InPosition;
+		}
+		check(PlayInEditorSessionInfo->CachedWindowInfo.Num() == InViewportIndex + 1); 
 	}
 }
 
