MultiprocessCook: CookSockets and CompactBinaryTCP helper classes for communication between Director and CookWorkers.
#rnx

diff --git a/Engine/Source/Editor/UnrealEd/Private/Cooker/CompactBinaryTCP.cpp b/Engine/Source/Editor/UnrealEd/Private/Cooker/CompactBinaryTCP.cpp
new file mode 100644
--- /dev/null
+++ b/Engine/Source/Editor/UnrealEd/Private/Cooker/CompactBinaryTCP.cpp
@@ -0,0 +1,345 @@
+// Copyright Epic Games, Inc. All Rights Reserved.
+
+#include "CompactBinaryTCP.h"
+
+#include "Serialization/CompactBinaryValidation.h"
+#include "Serialization/CompactBinaryWriter.h"
+#include "Sockets.h"
+#include "SocketSubsystem.h"
+
+namespace UE::CompactBinaryTCP
+{
+
+const TCHAR* DescribeStatus(EConnectionStatus Status)
+{
+	switch (Status)
+	{
+	case EConnectionStatus::Okay: return TEXT("Connection is okay");
+	case EConnectionStatus::Terminated: return TEXT("Connection terminated");
+	case EConnectionStatus::FormatError: return TEXT("Connection received invalidly formatted bytes");
+	case EConnectionStatus::Failed: return TEXT("Operation failed");
+	case EConnectionStatus::Incomplete: return TEXT("Operation is incomplete");
+	default: checkNoEntry(); return TEXT("InvalidConnectionStatus");
+	}
+}
+
+struct FPacketHeader
+{
+	uint64 Magic;
+	uint64 Size;
+};
+constexpr uint64 MessageHeaderExpectedMagic = 0x854EBA92854EBA92;
+
+void FReceiveBuffer::Reset()
+{
+	Payload.Reset();
+	bParsedHeader = false;
+	BytesRead = 0;
+}
+
+void FSendBuffer::Reset()
+{
+	PendingMessages.Empty();
+	Payload.Reset();
+	BytesWritten = 0;
+	bCreatedPayload = false;
+	bSentHeader = false;
+}
+
+class FCompactBinaryTCPImpl
+{
+public:
+	static EConnectionStatus TryReadPacket(FSocket* Socket, FReceiveBuffer& Buffer,
+		TArray<FMarshalledMessage>& Messages, uint64 MaxPacketSize);
+	static bool TryCreatePayload(TArray<FMarshalledMessage>& PendingMessages,
+		uint64 MaxPacketSize, FUniqueBuffer& OutBuffer);
+	static EConnectionStatus PollSendBytes(FSocket* Socket, const void* Data,
+		uint64 DataSize, uint64& InOutBytesWritten);
+	/** Note that the elements in AppendMessages are mutated */
+	static EConnectionStatus TryWritePacket(FSocket* Socket, FSendBuffer& Buffer,
+		TArrayView<FMarshalledMessage>&& AppendMessages, uint64 MaxPacketSize);
+};
+
+EConnectionStatus FCompactBinaryTCPImpl::TryReadPacket(FSocket* Socket, FReceiveBuffer& Buffer,
+	TArray<FMarshalledMessage>& Messages, uint64 MaxPacketSize)
+{
+	int32 BytesRead;
+	if (!Buffer.bParsedHeader)
+	{
+		// When reading the relatively small PacketHeader, we wait for its total size to be ready
+		// using Peek to reduce complexity in this function.
+		FPacketHeader MessageHeader;
+		bool bStillAlive = Socket->Recv(reinterpret_cast<uint8*>(&MessageHeader), sizeof(MessageHeader),
+			BytesRead, ESocketReceiveFlags::Peek);
+		check(BytesRead <= sizeof(FPacketHeader));
+		if (BytesRead < sizeof(FPacketHeader))
+		{
+			if (!bStillAlive)
+			{
+				return EConnectionStatus::Terminated;
+			}
+			else
+			{
+				return EConnectionStatus::Okay;
+			}
+		}
+
+		Socket->Recv(reinterpret_cast<uint8*>(&MessageHeader), sizeof(MessageHeader), BytesRead,
+			ESocketReceiveFlags::None);
+		check(BytesRead == sizeof(FPacketHeader));
+		if (MessageHeader.Magic != MessageHeaderExpectedMagic ||
+			(MaxPacketSize > 0 && MessageHeader.Size > MaxPacketSize))
+		{
+			return EConnectionStatus::FormatError;
+		}
+
+		Buffer.Payload = FUniqueBuffer::Alloc(MessageHeader.Size);
+		Buffer.bParsedHeader = true;
+		Buffer.BytesRead = 0;
+	}
+
+
+	while (Buffer.Payload.GetSize() > Buffer.BytesRead)
+	{
+		// When reading the possibly large payload size that will block the remote socket
+		// from continuing until we have read some of it, we read as much of the payload
+		// as is available and store it in our dynamic and larger buffer in Buffer.Payload.
+		uint64 RemainingSize = Buffer.Payload.GetSize() - Buffer.BytesRead;
+		// Avoid possible OS restrictions on maximum read size by imposing our own moderate
+		// size per call to Recv.
+		constexpr uint64 MaxReadSize = 1000 * 1000 * 64;
+		int32 ReadSize = static_cast<int32>(FMath::Min(RemainingSize, MaxReadSize));
+		uint8* ReadData = static_cast<uint8*>(Buffer.Payload.GetData()) + Buffer.BytesRead;
+		bool bConnectionAlive = Socket->Recv(ReadData, ReadSize, BytesRead);
+		if (BytesRead <= 0)
+		{
+			if (!bConnectionAlive)
+			{
+				return EConnectionStatus::Terminated;
+			}
+			else
+			{
+				return EConnectionStatus::Okay;
+			}
+		}
+
+		check(BytesRead <= ReadSize);
+		Buffer.BytesRead += BytesRead;
+	}
+
+	// The FCbObjects we return have a pointer to this buffer and will keep it allocated
+	// until they are destructed.
+	FSharedBuffer SharedBuffer = Buffer.Payload.MoveToShared();
+	Buffer.Reset();
+	if (ValidateCompactBinary(SharedBuffer, ECbValidateMode::Default) != ECbValidateError::None)
+	{
+		return EConnectionStatus::FormatError;
+	}
+
+	FCbFieldIterator It = FCbFieldIterator::MakeRange(SharedBuffer);
+	while (It)
+	{
+		FMarshalledMessage& Message = Messages.Emplace_GetRef();
+		FCbFieldView MessageTypeView = *It;
+		Message.MessageType = MessageTypeView.AsUuid();
+		if (MessageTypeView.HasError())
+		{
+			return EConnectionStatus::FormatError;
+		}
+		++It;
+		FCbField ObjectField = *It;
+		Message.Object = ObjectField.AsObject();
+		if (ObjectField.HasError())
+		{
+			return EConnectionStatus::FormatError;
+		}
+		++It;
+	}
+	return EConnectionStatus::Okay;
+}
+
+bool FCompactBinaryTCPImpl::TryCreatePayload(TArray<FMarshalledMessage>& PendingMessages,
+	uint64 MaxPacketSize, FUniqueBuffer& OutBuffer)
+{
+	int32 NumMessages = PendingMessages.Num();
+	checkSlow(NumMessages > 0); // Caller guarantees
+
+	int32 NumMessagesTaken = 0;
+	FCbWriter Writer;
+	int32 MaxNumMessagesTaken = NumMessages;
+	int32 RetryCount = 0;
+	for (;;)
+	{
+		for (NumMessagesTaken = 0; NumMessagesTaken < MaxNumMessagesTaken; ++NumMessagesTaken)
+		{
+			FMarshalledMessage& Message = PendingMessages[NumMessagesTaken];
+			if (MaxPacketSize > 0 && RetryCount == 0)
+			{
+				// We need to know the size the writer will have after serializing the object
+				// so we can avoid exceeding our max packet size. If we exceed we have to retry
+				// and write everything over again. But there's no cheap way to compute
+				// that, and the max message size is usually large and won't be exceeded
+				// go over, so we use a cheap hardcoded value rather than an accurate estimate.
+				constexpr uint64 EstimatedSize = 1000;
+				if (NumMessagesTaken > 0 && Writer.GetSaveSize() + EstimatedSize > MaxPacketSize)
+				{
+					break;
+				}
+			}
+			Writer << Message.MessageType;
+			Writer << Message.Object;
+		}
+		uint64 SaveSize = Writer.GetSaveSize();
+		if (MaxPacketSize > 0 && Writer.GetSaveSize() > MaxPacketSize)
+		{
+			if (NumMessagesTaken == 1)
+			{
+				UE_LOG(LogSockets, Error, TEXT("Could not WritePacket to Socket. A single message with Guid %s is larger than MaxPacketSize %" UINT64_FMT "."),
+					*PendingMessages[0].MessageType.ToString(), MaxPacketSize);
+				return false;
+			}
+			check(RetryCount < 1); // If this is the Retry, then MaxNumMessagesTaken should have blocked going over size
+			MaxNumMessagesTaken = NumMessagesTaken - 1;
+			Writer.Reset();
+			++RetryCount;
+			continue;
+		}
+
+		// Successful
+		OutBuffer = FUniqueBuffer::Alloc(Writer.GetSaveSize());
+		Writer.Save(OutBuffer.GetView());
+		PendingMessages.RemoveAt(0, NumMessagesTaken);
+		return true;
+	}
+}
+
+EConnectionStatus FCompactBinaryTCPImpl::PollSendBytes(FSocket* Socket, const void* Data, 
+	uint64 DataSize, uint64& InOutBytesWritten)
+{
+	check(DataSize <= MAX_int32);
+	check(InOutBytesWritten < DataSize);
+
+	int32 WriteSize = static_cast<int32>(DataSize - InOutBytesWritten);
+	const uint8* WriteData = static_cast<const uint8*>(Data) + InOutBytesWritten;
+	int32 BytesSent;
+	bool bResult = Socket->Send(WriteData, WriteSize, BytesSent);
+	if (BytesSent < WriteSize)
+	{
+		if (BytesSent >= 0)
+		{
+			InOutBytesWritten += BytesSent;
+		}
+		if (!bResult)
+		{
+			UE_LOG(LogSockets, Error, TEXT("Could not WritePacket to Socket. Socket->Send failed."));
+			return EConnectionStatus::Failed;
+		}
+		else
+		{
+			return EConnectionStatus::Incomplete;
+		}
+	}
+	InOutBytesWritten = 0;
+	return EConnectionStatus::Okay;
+};
+
+EConnectionStatus FCompactBinaryTCPImpl::TryWritePacket(FSocket* Socket, FSendBuffer& Buffer,
+	TArrayView<FMarshalledMessage>&& AppendMessages, uint64 MaxPacketSize)
+{
+	if (!AppendMessages.IsEmpty())
+	{
+		Buffer.PendingMessages.Reserve(Buffer.PendingMessages.Num() + AppendMessages.Num());
+		for (FMarshalledMessage& NewMessage : AppendMessages)
+		{
+			Buffer.PendingMessages.Add(MoveTemp(NewMessage));
+		}
+	}
+	// Our Socket API restricts us to int32 Count field (Max == 1 << 31 - 1), but we further
+	// restrict that by a factor of two to prevent impact from any OS limitations less than that
+	// and to allow room for the 16 byte PacketHeader.
+	constexpr uint64 MaxOSPacketSize = (1 << 30) - 1;
+	MaxPacketSize = MaxPacketSize == 0 ? MaxOSPacketSize : FMath::Min(MaxPacketSize, MaxOSPacketSize);
+
+	for (;;)
+	{
+		if (!Buffer.bCreatedPayload)
+		{
+			if (Buffer.PendingMessages.IsEmpty())
+			{
+				return EConnectionStatus::Okay;
+			}
+			if (!TryCreatePayload(Buffer.PendingMessages, MaxPacketSize, Buffer.Payload))
+			{
+				return EConnectionStatus::Failed;
+			}
+			Buffer.bCreatedPayload = true;
+			Buffer.BytesWritten = 0;
+		}
+
+		if (!Buffer.bSentHeader)
+		{
+			FPacketHeader MessageHeader;
+			MessageHeader.Magic = MessageHeaderExpectedMagic;
+			MessageHeader.Size = Buffer.Payload.GetSize();
+			EConnectionStatus Status = PollSendBytes(Socket, &MessageHeader, 
+				sizeof(FPacketHeader), Buffer.BytesWritten);
+			if (Status != EConnectionStatus::Okay)
+			{
+				return Status;
+			}
+			Buffer.bSentHeader = true;
+		}
+
+		{
+			EConnectionStatus Status = PollSendBytes(Socket, Buffer.Payload.GetData(),
+				Buffer.Payload.GetSize(), Buffer.BytesWritten);
+			if (Status != EConnectionStatus::Okay)
+			{
+				return Status;
+			}
+			Buffer.bCreatedPayload = false;
+			Buffer.bSentHeader = false;
+			Buffer.Payload.Reset();
+		}
+	}
+}
+
+EConnectionStatus TryReadPacket(FSocket* Socket, FReceiveBuffer& Buffer,
+	TArray<FMarshalledMessage>& Messages, uint64 MaxPacketSize)
+{
+	return FCompactBinaryTCPImpl::TryReadPacket(Socket, Buffer, Messages, MaxPacketSize);
+}
+
+EConnectionStatus TryWritePacket(FSocket* Socket, FSendBuffer& Buffer, 
+	TArray<FMarshalledMessage>&& AppendMessages, uint64 MaxPacketSize)
+{
+	return FCompactBinaryTCPImpl::TryWritePacket(Socket, Buffer, MoveTemp(AppendMessages), MaxPacketSize);
+}
+
+EConnectionStatus TryWritePacket(FSocket* Socket, FSendBuffer& Buffer,
+	FMarshalledMessage&& AppendMessage, uint64 MaxPacketSize)
+{
+	return FCompactBinaryTCPImpl::TryWritePacket(Socket, Buffer,
+		TArrayView<FMarshalledMessage>(&AppendMessage, 1), MaxPacketSize);
+}
+
+EConnectionStatus TryWritePacket(FSocket* Socket, FSendBuffer& Buffer,
+	const IMessage& AppendMessage, uint64 MaxPacketSize)
+{
+	FMarshalledMessage Marshalled;
+	Marshalled.MessageType = AppendMessage.GetMessageType();
+	FCbWriter Writer;
+	Writer.BeginObject();
+	AppendMessage.Write(Writer);
+	Writer.EndObject();
+	Marshalled.Object = Writer.Save().AsObject();
+	return FCompactBinaryTCPImpl::TryWritePacket(Socket, Buffer,
+		TArrayView<FMarshalledMessage>(&Marshalled, 1), MaxPacketSize);
+}
+
+EConnectionStatus TryFlushBuffer(FSocket* Socket, FSendBuffer& Buffer, uint64 MaxPacketSize)
+{
+	return FCompactBinaryTCPImpl::TryWritePacket(Socket, Buffer, TArrayView<FMarshalledMessage>(), MaxPacketSize);
+}
+
+}
\ No newline at end of file
diff --git a/Engine/Source/Editor/UnrealEd/Private/Cooker/CompactBinaryTCP.h b/Engine/Source/Editor/UnrealEd/Private/Cooker/CompactBinaryTCP.h
new file mode 100644
--- /dev/null
+++ b/Engine/Source/Editor/UnrealEd/Private/Cooker/CompactBinaryTCP.h
@@ -0,0 +1,128 @@
+// Copyright Epic Games, Inc. All Rights Reserved.
+
+#pragma once
+
+#include "Containers/Array.h"
+#include "Containers/ArrayView.h"
+#include "Memory/SharedBuffer.h"
+#include "Misc/Guid.h"
+#include "Serialization/CompactBinary.h"
+
+class FSocket;
+
+namespace UE::CompactBinaryTCP
+{
+
+class FCompactBinaryTCPImpl;
+struct FReceiveBuffer;
+struct FSendBuffer;
+
+enum EConnectionStatus
+{
+	/** Connection is still okay or operation succeeded. */
+	Okay,
+	/** Connection failed and no further use of the Socket is possible. */
+	Terminated,
+	/** The Socket is still active but the data received was invalid and recovery is not possible. */
+	FormatError,
+	/** The operation failed and the Socket is no longer usable. */
+	Failed,
+	/** The Socket is busy and the operation needs to be retried later */
+	Incomplete,
+};
+const TCHAR* DescribeStatus(EConnectionStatus Status);
+
+/**
+ * The base class of messages that can be sent through CompactBinaryTCP. Messages are identified to
+ * the remote connection by the Guid identifier from GetMessageType.
+ */
+class IMessage
+{
+public:
+	virtual ~IMessage() {}
+	/** Marshall the message to a CompactBinaryObject. */
+	virtual void Write(FCbWriter& Writer) const = 0;
+	/** Unmarshall the message from a CompactBinaryObject. */
+	virtual bool TryRead(FCbObjectView Object) = 0;
+	/** Return the Guid that identifies the message to the remote connection. */
+	virtual FGuid GetMessageType() const = 0;
+};
+
+/** IMessages are marshalled into a Guid and a CompactBinaryObject for serialization through the socket. */
+struct FMarshalledMessage
+{
+	FGuid MessageType;
+	FCbObject Object;
+};
+
+/**
+ * Attempt to read messages from the socket. Returns all messages currently available and returns. Partially
+ * transferred messages are stored in Buffer and can be completed by further calls to TryReadPacket.
+ * 
+ * @param MaxPacketSize If non-zero, all messages are assumed to be <= and a FormatError is given if they are larger.
+ * 
+ * @return EConnectionStatus::Okay on success and EConnectionStatus::Terminated or EConnection::FormatError on error.
+ * If EConnectionStatus::Okay is returned, Messages may be empty (no message yet received) or populated.
+ */
+EConnectionStatus TryReadPacket(FSocket* Socket, FReceiveBuffer& Buffer,
+	TArray<FMarshalledMessage>& Messages, uint64 MaxPacketSize=0);
+
+/**
+ * Attempt to write messages to the socket. If the socket is busy the new messages are stored in Buffer and will be
+ * sent on a future call to TryWritePacket.
+ * 
+ * @param MaxPacketSize When combining multiple messages into a single packet, messages will stopped being added and
+ * will be pushed into the next packet if the packet size is > this amount. If 0, will use the maximum size allowed
+ * by the socket, which is >= 2^30. If a single message has larger size than MaxPacketSize, EConnectionStatus::Failed
+ * is return and the socket will no longer be usable.
+ * 
+ * @return EConnectionStatus::Okay on success and all messages have been sent. EConnectionStatus::Incomplete if some
+ * messages could not yet be sent. EConnectionStatus::Failed if connection failed and socket is no longer usable.
+ */
+EConnectionStatus TryWritePacket(FSocket* Socket, FSendBuffer& Buffer,
+	TArray<FMarshalledMessage>&& AppendMessages, uint64 MaxPacketSize = 0);
+EConnectionStatus TryWritePacket(FSocket* Socket, FSendBuffer& Buffer,
+	FMarshalledMessage&& AppendMessage, uint64 MaxPacketSize = 0);
+EConnectionStatus TryWritePacket(FSocket* Socket, FSendBuffer& Buffer,
+	const IMessage& AppendMessage, uint64 MaxPacketSize = 0);
+
+/** Attempt to finish sending packets that were previously queued by TryWritePacket */
+EConnectionStatus TryFlushBuffer(FSocket* Socket, FSendBuffer& Buffer, uint64 MaxPacketSize = 0);
+
+/**
+ * Holds the state of the CompactBinaryTCP communication read from a Socket. Expects the socket to be a
+ * series of packets of messages and keeps track of the next packet when a packet is as yet only partially received.
+ */
+struct FReceiveBuffer
+{
+public:
+	void Reset();
+
+private:
+	FUniqueBuffer Payload;
+	uint64 BytesRead = 0;
+	bool bParsedHeader = false;
+
+	friend class UE::CompactBinaryTCP::FCompactBinaryTCPImpl;
+};
+
+/**
+ * Holds the state of the CompactBinaryTCP communication written to a Socket. Allows the socket to be a
+ * series of packets of messages and keeps track of a partially sent packet and of pending messages.
+ */
+struct FSendBuffer
+{
+public:
+	void Reset();
+
+private:
+	TArray<FMarshalledMessage> PendingMessages;
+	FUniqueBuffer Payload;
+	uint64 BytesWritten = 0;
+	bool bCreatedPayload = false;
+	bool bSentHeader = false;
+
+	friend class UE::CompactBinaryTCP::FCompactBinaryTCPImpl;
+};
+
+}
diff --git a/Engine/Source/Editor/UnrealEd/Private/Cooker/CookSockets.cpp b/Engine/Source/Editor/UnrealEd/Private/Cooker/CookSockets.cpp
new file mode 100644
--- /dev/null
+++ b/Engine/Source/Editor/UnrealEd/Private/Cooker/CookSockets.cpp
@@ -0,0 +1,126 @@
+// Copyright Epic Games, Inc. All Rights Reserved.
+
+#include "CookSockets.h"
+
+#include "Containers/StringView.h"
+#include "Containers/UnrealString.h"
+#include "CoreGlobals.h"
+#include "Misc/OutputDeviceRedirector.h"
+#include "Misc/ScopeExit.h"
+#include "Sockets.h"
+#include "SocketSubsystem.h"
+
+namespace UE::Cook::Sockets
+{
+
+void CloseSocket(FSocket*& Socket)
+{
+	if (Socket)
+	{
+		Socket->Close();
+		ISocketSubsystem::Get()->DestroySocket(Socket);
+		Socket = nullptr;
+	}
+}
+
+FSocket* CreateListenSocket(int32& InOutPort, TSharedPtr<FInternetAddr>& OutAddr, FString& OutConnectAuthority,
+	const TCHAR* SocketDebugName, FString& OutErrorReason)
+{
+	ISocketSubsystem* SocketSubsystem = ISocketSubsystem::Get();
+	OutAddr = SocketSubsystem->GetLocalBindAddr(*GLog);
+
+	// REMOTEMPCOOKTODO: Calculate ConnectAuthorityAddr from the proper element of GetLocalAdapterAddresses
+	// rather than hardcoding to 127.0.0.1
+	TSharedPtr<FInternetAddr> ConnectAuthorityAddr = SocketSubsystem->GetAddressFromString(TEXT("127.0.0.1"));
+
+	OutAddr->SetPort(InOutPort);
+	FSocket* LocalSocket = SocketSubsystem->CreateSocket(NAME_Stream, SocketDebugName, OutAddr->GetProtocolType());
+	if (!LocalSocket)
+	{
+		OutErrorReason = TEXT("Could not create listen socket");
+		return nullptr;
+	}
+
+	ON_SCOPE_EXIT
+	{
+		CloseSocket(LocalSocket);
+	};
+
+	LocalSocket->SetReuseAddr(false);
+	LocalSocket->SetNoDelay(true);
+	LocalSocket->SetNonBlocking(true);
+
+	if (!LocalSocket->Bind(*OutAddr))
+	{
+		OutErrorReason = FString::Printf(TEXT("Failed to bind listen socket %s"), *OutAddr->ToString(true /* bAppendPort */));
+		return nullptr;
+	}
+	int32 MaxBacklog = 16;
+	if (!LocalSocket->Listen(MaxBacklog))
+	{
+		OutErrorReason = FString::Printf(TEXT("Failed to listen to socket %s"), *OutAddr->ToString(true /* bAppendPort */));
+		return nullptr;
+	}
+
+	InOutPort = LocalSocket->GetPortNo();
+	OutAddr->SetPort(InOutPort);
+	ConnectAuthorityAddr->SetPort(InOutPort);
+	OutConnectAuthority = ConnectAuthorityAddr->ToString(true /* bAppendPort */);
+
+	FSocket* OutSocket = LocalSocket;
+	LocalSocket = nullptr;
+	return OutSocket;
+}
+
+FSocket* ConnectToHost(FInternetAddr& HostAddr, const TCHAR* SocketDebugName)
+{
+	ISocketSubsystem* SocketSubsystem = ISocketSubsystem::Get();
+	FSocket* Socket = SocketSubsystem->CreateSocket(NAME_Stream, SocketDebugName, HostAddr.GetProtocolType());
+	if (Socket == nullptr)
+	{
+		return nullptr;
+	}
+
+	Socket->SetNonBlocking(true);
+	if (!Socket->Connect(HostAddr))
+	{
+		CloseSocket(Socket);
+		return nullptr;
+	}
+	return Socket;
+}
+
+TSharedPtr<FInternetAddr> GetAddressFromStringWithPort(FStringView Text)
+{
+	int32 PortValue = -1;
+	int32 PortStartIndex;
+	if (Text.FindLastChar(':', PortStartIndex))
+	{
+		if (PortStartIndex == 0)
+		{
+			UE_LOG(LogSockets, Warning, TEXT("Could not serialize %.*s, it is missing the <Addr>: prefix."),
+				Text.Len(), Text.GetData());
+			return TSharedPtr<FInternetAddr>();
+		}
+		if (!LexTryParseString(PortValue, *WriteToString<16>(Text.RightChop(PortStartIndex + 1))))
+		{
+			PortValue = -1;
+		}
+	}
+	if (PortValue < 0)
+	{
+		UE_LOG(LogSockets, Warning, TEXT("Could not serialize %.*s, it is missing the :<Port> suffix."),
+			Text.Len(), Text.GetData());
+		return TSharedPtr<FInternetAddr>();
+	}
+	
+	ISocketSubsystem* SocketSubsystem = ISocketSubsystem::Get();
+	TSharedPtr<FInternetAddr> Result = SocketSubsystem->GetAddressFromString(FString(Text.Mid(0, PortStartIndex)));
+	if (Result)
+	{
+		Result->SetPort(PortValue);
+	}
+	return Result;
+}
+
+}
\ No newline at end of file
diff --git a/Engine/Source/Editor/UnrealEd/Private/Cooker/CookSockets.h b/Engine/Source/Editor/UnrealEd/Private/Cooker/CookSockets.h
new file mode 100644
--- /dev/null
+++ b/Engine/Source/Editor/UnrealEd/Private/Cooker/CookSockets.h
@@ -0,0 +1,41 @@
+// Copyright Epic Games, Inc. All Rights Reserved.
+
+#pragma once
+
+#include "Containers/StringFwd.h"
+#include "HAL/Platform.h"
+#include "IPAddress.h"
+#include "Templates/SharedPointer.h"
+
+class FSocket;
+
+namespace UE::Cook::Sockets
+{
+
+constexpr int32 COOKDIRECTOR_DEFAULT_REQUEST_CONNECTION_PORT = 41897;
+
+/** Destruct Socket as gracefully as possible. Sets Socket to nullptr. Early exits if Socket is already null. */
+void CloseSocket(FSocket*& Socket);
+
+/**
+ * Create a socket to listen for incoming connections on the given port.
+ * 
+ * @param InOutPort: On input, the desired listening port.
+ *     On output, the selected port if desired listening port was 0, otherwise unchanged.
+ * @param OutAddr: The address being listened to at the OS socket level, e.g. 0.0.0.0:<port>. 
+ * @param OutConnectAuthority: The URI authority other processes should connect to, e.g. 192.168.1.10:<port>.
+ * @param SocketDebugName: String used in logs sent for the socket.
+ * @param OutErrorReason: Reason why creation failed if return value is nullptr.
+ * 
+ * @return Valid socket or nullptr on failure.
+ */
+FSocket* CreateListenSocket(int32& InOutPort, TSharedPtr<FInternetAddr>& OutAddr, FString& OutConnectAuthority,
+	const TCHAR* SocketDebugName, FString& OutErrorReason);
+
+/** Create a socket to create to a Host using an InternetAddr created from e.g. OutConnectURI from CreateListenSocket. */
+FSocket* ConnectToHost(FInternetAddr& Host, const TCHAR* SocketDebugName);
+
+/** Convert "<addr>:<port>"" into InternetAddr, e.g. "192.168.1.10:48000" or "2001:0db8:85a3:0000:0000:8a2e:0370:7334:48000" */
+TSharedPtr<FInternetAddr> GetAddressFromStringWithPort(FStringView Text);
+
+}
\ No newline at end of file
