Cooker Multiprocess: Add AssignedToWorker state for PackageDatas on the Director. Implement Server/Client state machine for communication from CookDirector process's CookWorkerServers to the CookWorker processes' CookWorkerClients.
#rnx

diff --git a/Engine/Config/BaseEditor.ini b/Engine/Config/BaseEditor.ini
--- a/Engine/Config/BaseEditor.ini
+++ b/Engine/Config/BaseEditor.ini
@@ -330,6 +330,9 @@ IgnoreHeaderDiffs=True
 bUseFullGCAssetClassNames=true
 +FullGCAssetClassNames=/Script/Engine.World
 
+CookMultiProcessEnabled=false
+CookWorkerCount=1
+
 +ExcludedIniParams=Game:/Script/BlueprintContext.McpContext
 +ExcludedIniParams=Editor:CookSettings
 
diff --git a/Engine/Source/Editor/UnrealEd/Classes/CookOnTheSide/CookOnTheFlyServer.h b/Engine/Source/Editor/UnrealEd/Classes/CookOnTheSide/CookOnTheFlyServer.h
--- a/Engine/Source/Editor/UnrealEd/Classes/CookOnTheSide/CookOnTheFlyServer.h
+++ b/Engine/Source/Editor/UnrealEd/Classes/CookOnTheSide/CookOnTheFlyServer.h
@@ -127,6 +127,7 @@ namespace UE::Cook
 	class FBuildDefinitions;
 	class FCookDirector;
 	class FRequestCluster;
+	class FRequestQueue;
 	class FSaveCookedPackageContext;
 	class FWorkerRequestsLocal;
 	class FWorkerRequestsRemote;
@@ -135,6 +136,8 @@ namespace UE::Cook
 	class IWorkerRequests;
 	enum class EPollStatus : uint8;
 	enum class EReleaseSaveReason : uint8;
+	enum class ESuppressCookReason : uint8;
+	enum class ESendFlags : uint8;
 	struct FBeginCookConfigSettings;
 	struct FConstructPackageData;
 	struct FCookByTheBookOptions;
@@ -419,12 +422,21 @@ private:
 	};
 	/** Inspect all tasks the scheduler could do and return which one it should do. */
 	ECookAction DecideNextCookAction(UE::Cook::FTickStackData& StackData);
+	/** Is the local CookOnTheFlyServer a CookDirector and it has finished all locally processed Requests? */
+	bool IsMultiprocessLocalWorkerIdle() const;
 	/** Execute any existing external callbacks and push any existing external cook requests into new RequestClusters. */
 	void PumpExternalRequests(const UE::Cook::FCookerTimer& CookerTimer);
 	/** Send the PackageData back to request state to create a request cluster, if it has not yet been explored. */
 	bool TryCreateRequestCluster(UE::Cook::FPackageData& PackageData);
 	/** Inspect the next package in the RequestQueue and push it on to its next state. Report the number of PackageDatas that were pushed to load. */
 	void PumpRequests(UE::Cook::FTickStackData& StackData, int32& OutNumPushed);
+	/** Assign the requests found in PumpRequests; either pushing them to ReadyRequests if SingleProcess Cook or to CookWorkers if MultiProcess. */
+	void AssignRequests(TArrayView<UE::Cook::FPackageData*> Requests, UE::Cook::FRequestQueue& RequestQueue);
+	/**
+	 * Multiprocess cook: Notify the CookDirector that a package it assigned to a worker was demoted
+	 * due to e.g. cancelled cook and should be removed from the CookWorker.
+	 */
+	void NotifyRemovedFromWorker(UE::Cook::FPackageData& PackageData);
 
 	/** Load all packages in the LoadQueue until it's time to break. Report the number of loads that were pushed to save. */
 	void PumpLoads(UE::Cook::FTickStackData& StackData, uint32 DesiredQueueLength, int32& OutNumPushed, bool& bOutBusy);
@@ -435,7 +447,7 @@ private:
 	/** Load the given PackageData that was in the load queue and send it on to its next state. Report the number of PackageDatas that were pushed to save (0 or 1) */
 	void LoadPackageInQueue(UE::Cook::FPackageData& PackageData, uint32& ResultFlags, int32& OutNumPushed);
 	/** Mark that the given PackageData failed to load and return it to idle. */
-	void RejectPackageToLoad(UE::Cook::FPackageData& PackageData, const TCHAR* Reason);
+	void RejectPackageToLoad(UE::Cook::FPackageData& PackageData, const TCHAR* ReasonText, UE::Cook::ESuppressCookReason Reason);
 
 	/** Try to save all packages in the SaveQueue until it's time to break. Report the number of requests that were completed (either skipped or successfully saved or failed to save) */
 	void PumpSaves(UE::Cook::FTickStackData& StackData, uint32 DesiredQueueLength, int32& OutNumPushed, bool& bOutBusy);
@@ -457,6 +469,11 @@ private:
 	void QueueDiscoveredPackageData(UE::Cook::FPackageData& PackageData, UE::Cook::FInstigator&& Instigator,
 		bool bLoadReady = false);
 
+	/** Function called whenever an inprogress package is cancelled and returned to idle. Notifies the CookDirector when on a CookWorker. */
+	void DemoteToIdle(UE::Cook::FPackageData& PackageData, UE::Cook::ESendFlags SendFlags, UE::Cook::ESuppressCookReason Reason);
+	/** Function called whenever an inprogress package completes its save and returns to idle. Notifies the CookDirector when on a CookWorker. */
+	void PromoteToSaveComplete(UE::Cook::FPackageData& PackageData, UE::Cook::ESendFlags SendFlags);
+
 	/** Check whether the package filter has change, and if so call QueueDiscoveredPackage again on each existing package. */
 	void UpdatePackageFilter();
 
@@ -1318,6 +1335,7 @@ private:
 	FCriticalSection PollablesLock;
 	TRefCountPtr<FPollable> RecompileRequestsPollable;
 	TRefCountPtr<FPollable> QueuedCancelPollable;
+	TRefCountPtr<FPollable> DirectorPollable;
 	double PollNextTimeSeconds = 0.;
 	double PollNextTimeIdleSeconds = 0.;
 	double IdleStatusStartTime = 0.0;
@@ -1327,6 +1345,7 @@ private:
 	EIdleStatus IdleStatus = EIdleStatus::Done;
 
 	friend UE::Cook::FBeginCookConfigSettings;
+	friend UE::Cook::FCookDirector;
 	friend UE::Cook::FGeneratorPackage;
 	friend UE::Cook::FPackageData;
 	friend UE::Cook::FPendingCookedPlatformData;
diff --git a/Engine/Source/Editor/UnrealEd/Private/Commandlets/CookCommandlet.cpp b/Engine/Source/Editor/UnrealEd/Private/Commandlets/CookCommandlet.cpp
--- a/Engine/Source/Editor/UnrealEd/Private/Commandlets/CookCommandlet.cpp
+++ b/Engine/Source/Editor/UnrealEd/Private/Commandlets/CookCommandlet.cpp
@@ -901,6 +901,7 @@ bool UCookCommandlet::CookAsCookWorker()
 
 	if (!CookOnTheFlyServer->TryInitializeCookWorker())
 	{
+		UE_LOG(LogCook, Display, TEXT("CookWorker initialization failed, aborting CookCommandlet."));
 		return false;
 	}
 
diff --git a/Engine/Source/Editor/UnrealEd/Private/CookOnTheFlyServer.cpp b/Engine/Source/Editor/UnrealEd/Private/CookOnTheFlyServer.cpp
--- a/Engine/Source/Editor/UnrealEd/Private/CookOnTheFlyServer.cpp
+++ b/Engine/Source/Editor/UnrealEd/Private/CookOnTheFlyServer.cpp
@@ -2072,6 +2072,11 @@ void UCookOnTheFlyServer::InitializePollables()
 		Pollables.Emplace(new FPollable(TEXT("RequestManager"), 0.5f, 0.5f, [this](FTickStackData&) { TickRequestManager(); }));
 
 	}
+	if (CookDirector)
+	{
+		DirectorPollable = new FPollable(TEXT("CookDirector"), 1.0f, 1.0f, [this](FTickStackData& StackData) { CookDirector->TickFromSchedulerThread(); });
+		Pollables.Add(FPollableQueueKey(DirectorPollable));
+	}
 	Pollables.Heapify();
 
 	PollNextTimeSeconds = 0.;
@@ -2142,6 +2147,11 @@ UCookOnTheFlyServer::ECookAction UCookOnTheFlyServer::DecideNextCookAction(UE::C
 			SetIdleStatus(StackData, EIdleStatus::Active);
 			return ECookAction::Request;
 		}
+
+		if (Monitor.GetNumUrgent(UE::Cook::EPackageState::AssignedToWorker) > 0)
+		{
+			// Fall through and do non-urgent while we wait for the Worker to finish
+		}
 		else
 		{
 			checkf(false, TEXT("Urgent request is in state not yet handled by DecideNextCookAction"));
@@ -2225,6 +2235,20 @@ UCookOnTheFlyServer::ECookAction UCookOnTheFlyServer::DecideNextCookAction(UE::C
 	return ECookAction::Done;
 }
 
+bool UCookOnTheFlyServer::IsMultiprocessLocalWorkerIdle() const
+{
+	if (!CookDirector.IsValid())
+	{
+		return false;
+	}
+	UE::Cook::FPackageDataMonitor& Monitor = PackageDatas->GetMonitor();
+	return !WorkerRequests->HasExternalRequests() &&
+		PackageDatas->GetRequestQueue().IsEmpty() &&
+		PackageDatas->GetLoadPrepareQueue().IsEmpty() &&
+		PackageDatas->GetLoadReadyQueue().IsEmpty() &&
+		PackageDatas->GetSaveQueue().IsEmpty();
+}
+
 void UCookOnTheFlyServer::PumpExternalRequests(const UE::Cook::FCookerTimer& CookerTimer)
 {
 	if (!WorkerRequests->HasExternalRequests())
@@ -2310,15 +2334,12 @@ void UCookOnTheFlyServer::PumpRequests(UE::Cook::FTickStackData& StackData, int3
 		if (bComplete)
 		{
 			TArray<FPackageData*> RequestsToLoad;
-			TArray<FPackageData*> RequestsToDemote;
+			TArray<TPair<FPackageData*, ESuppressCookReason>> RequestsToDemote;
 			RequestCluster.ClearAndDetachOwnedPackageDatas(RequestsToLoad, RequestsToDemote);
-			for (FPackageData* PackageData : RequestsToLoad)
-			{
-				RequestQueue.AddReadyRequest(PackageData);
-			}
-			for (FPackageData* PackageData : RequestsToDemote)
+			AssignRequests(RequestsToLoad, RequestQueue);
+			for (TPair<FPackageData*, ESuppressCookReason>& Pair : RequestsToDemote)
 			{
-				PackageData->SendToState(EPackageState::Idle, ESendFlags::QueueAdd);
+				DemoteToIdle(*Pair.Key, ESendFlags::QueueAdd, Pair.Value);
 			}
 			RequestClusters.PopFront();
 			OnRequestClusterCompleted(RequestCluster);
@@ -2344,7 +2365,7 @@ void UCookOnTheFlyServer::PumpRequests(UE::Cook::FTickStackData& StackData, int3
 #if DEBUG_COOKONTHEFLY
 			UE_LOG(LogCook, Display, TEXT("Package for platform already cooked %s, discarding request"), *PackageData.GetFileName().ToString());
 #endif
-			PackageData->SendToState(EPackageState::Idle, ESendFlags::QueueAdd);
+			DemoteToIdle(*PackageData, ESendFlags::QueueAdd, ESuppressCookReason::AlreadyCooked);
 			continue;
 		}
 		PackageData->SendToState(EPackageState::LoadPrepare, ESendFlags::QueueAdd);
@@ -2353,6 +2374,68 @@ void UCookOnTheFlyServer::PumpRequests(UE::Cook::FTickStackData& StackData, int3
 	OutNumPushed += NumInBatch;
 }
 
+void UCookOnTheFlyServer::AssignRequests(TArrayView<UE::Cook::FPackageData*> Requests, UE::Cook::FRequestQueue& RequestQueue)
+{
+	using namespace UE::Cook;
+
+	if (CookDirector)
+	{
+		int32 NumRequests = Requests.Num();
+		if (NumRequests == 0)
+		{
+			return;
+		}
+		TArray<FWorkerId> Assignments;
+		CookDirector->AssignRequests(Requests, Assignments);
+		check(Assignments.Num() == NumRequests);
+		for (int32 Index = 0; Index < NumRequests; ++Index)
+		{
+			FPackageData* PackageData = Requests[Index];
+			FWorkerId Assignment = Assignments[Index];
+			if (Assignment.IsLocal())
+			{
+				RequestQueue.AddReadyRequest(PackageData);
+			}
+			else
+			{
+				PackageData->SendToState(EPackageState::AssignedToWorker, ESendFlags::QueueAdd);
+				PackageData->SetWorkerAssignment(Assignment);
+			}
+		}
+	}
+	else
+	{
+		for (FPackageData* PackageData : Requests)
+		{
+			RequestQueue.AddReadyRequest(PackageData);
+		}
+	}
+}
+
+void UCookOnTheFlyServer::NotifyRemovedFromWorker(UE::Cook::FPackageData& PackageData)
+{
+	check(CookDirector);
+	CookDirector->RemoveFromWorker(PackageData);
+}
+
+void UCookOnTheFlyServer::DemoteToIdle(UE::Cook::FPackageData& PackageData, UE::Cook::ESendFlags SendFlags, UE::Cook::ESuppressCookReason Reason)
+{
+	if (PackageData.IsInProgress())
+	{
+		WorkerRequests->ReportDemoteToIdle(PackageData, Reason);
+	}
+	PackageData.SendToState(UE::Cook::EPackageState::Idle, SendFlags);
+}
+
+void UCookOnTheFlyServer::PromoteToSaveComplete(UE::Cook::FPackageData& PackageData, UE::Cook::ESendFlags SendFlags)
+{
+	if (PackageData.IsInProgress())
+	{
+		WorkerRequests->ReportPromoteToSaveComplete(PackageData);
+	}
+	PackageData.SendToState(UE::Cook::EPackageState::Idle, SendFlags);
+}
+
 void UCookOnTheFlyServer::PumpLoads(UE::Cook::FTickStackData& StackData, uint32 DesiredQueueLength, int32& OutNumPushed, bool& bOutBusy)
 {
 	using namespace UE::Cook;
@@ -2473,7 +2556,7 @@ void UCookOnTheFlyServer::LoadPackageInQueue(UE::Cook::FPackageData& PackageData
 		{
 			ResultFlags |= COSR_ErrorLoadingPackage;
 			UE_LOG(LogCook, Verbose, TEXT("Not cooking package %s"), *PackageFileName.ToString());
-			RejectPackageToLoad(PackageData, TEXT("failed to load"));
+			RejectPackageToLoad(PackageData, TEXT("failed to load"), ESuppressCookReason::LoadError);
 			return;
 		}
 		check(LoadedPackage != nullptr && LoadedPackage->IsFullyLoaded());
@@ -2492,7 +2575,7 @@ void UCookOnTheFlyServer::LoadPackageInQueue(UE::Cook::FPackageData& PackageData
 			QueueDiscoveredPackageData(OtherPackageData, FInstigator(PackageData.GetInstigator()));
 
 			PackageData.SetPlatformsCooked(PlatformManager->GetSessionPlatforms(), true);
-			RejectPackageToLoad(PackageData, TEXT("is redirected to another filename"));
+			RejectPackageToLoad(PackageData, TEXT("is redirected to another filename"), ESuppressCookReason::Redirected);
 			return;
 		}
 	}
@@ -2502,14 +2585,14 @@ void UCookOnTheFlyServer::LoadPackageInQueue(UE::Cook::FPackageData& PackageData
 		if (!Generator)
 		{
 			UE_LOG(LogCook, Error, TEXT("Package %s is an out-of-date generated package with a no-longer-available generator. It can not be loaded."), *PackageFileName.ToString());
-			RejectPackageToLoad(PackageData, TEXT("is an orphaned generated package"));
+			RejectPackageToLoad(PackageData, TEXT("is an orphaned generated package"), ESuppressCookReason::OrphanedGenerated);
 			return;
 		}
 		FCookGenerationInfo* Info = Generator->FindInfo(PackageData);
 		if (!Info)
 		{
 			UE_LOG(LogCook, Error, TEXT("Package %s is a generated package but its generator no longer has a record of it. It can not be loaded."), *PackageFileName.ToString());
-			RejectPackageToLoad(PackageData, TEXT("is an orphaned generated package"));
+			RejectPackageToLoad(PackageData, TEXT("is an orphaned generated package"), ESuppressCookReason::OrphanedGenerated);
 			return;
 		}
 
@@ -2532,7 +2615,7 @@ void UCookOnTheFlyServer::LoadPackageInQueue(UE::Cook::FPackageData& PackageData
 				ResultFlags |= COSR_ErrorLoadingPackage;
 				UE_LOG(LogCook, Error, TEXT("Package %s is a generated package and we could not load its generator package %s. It can not be loaded."),
 					*PackageFileName.ToString(), *OwnerPackageData.GetFileName().ToString());
-				RejectPackageToLoad(PackageData, TEXT("is a generated package which could not load its generator"));
+				RejectPackageToLoad(PackageData, TEXT("is a generated package which could not load its generator"), ESuppressCookReason::LoadError);
 				return;
 			}
 			Generator->GetCookPackageSplitterInstance()->OnOwnerReloaded(OwnerPackage, SplitterDataObject);
@@ -2542,7 +2625,7 @@ void UCookOnTheFlyServer::LoadPackageInQueue(UE::Cook::FPackageData& PackageData
 		LoadedPackage = TryCreateGeneratedPackage(*Generator, *Info);
 		if (!LoadedPackage)
 		{
-			RejectPackageToLoad(PackageData, TEXT("is a generated package which could not be populated"));
+			RejectPackageToLoad(PackageData, TEXT("is a generated package which could not be populated"), ESuppressCookReason::LoadError);
 			return;
 		}
 	}
@@ -2551,7 +2634,7 @@ void UCookOnTheFlyServer::LoadPackageInQueue(UE::Cook::FPackageData& PackageData
 	{
 		// Already cooked. This can happen if we needed to load a package that was previously cooked and garbage collected because it is a loaddependency of a new request.
 		// Send the package back to idle, nothing further to do with it.
-		PackageData.SendToState(EPackageState::Idle, ESendFlags::QueueAdd);
+		DemoteToIdle(PackageData, ESendFlags::QueueAdd, ESuppressCookReason::AlreadyCooked);
 		return;
 	}
 
@@ -2561,7 +2644,7 @@ void UCookOnTheFlyServer::LoadPackageInQueue(UE::Cook::FPackageData& PackageData
 	++OutNumPushed;
 }
 
-void UCookOnTheFlyServer::RejectPackageToLoad(UE::Cook::FPackageData& PackageData, const TCHAR* Reason)
+void UCookOnTheFlyServer::RejectPackageToLoad(UE::Cook::FPackageData& PackageData, const TCHAR* ReasonText, UE::Cook::ESuppressCookReason Reason)
 {
 	// make sure this package doesn't exist
 	for (const TPair<const ITargetPlatform*, UE::Cook::FPackageData::FPlatformData>& Pair : PackageData.GetPlatformDatas())
@@ -2578,11 +2661,11 @@ void UCookOnTheFlyServer::RejectPackageToLoad(UE::Cook::FPackageData& PackageDat
 			// if we find the file this means it was cooked on a previous cook, however source package can't be found now. 
 			// this could be because the source package was deleted or renamed, and we are using iterative cooking
 			// perhaps in this case we should delete it?
-			UE_LOG(LogCook, Warning, TEXT("Found cooked file '%s' which shouldn't exist as it %s."), *SandboxFilename, Reason);
+			UE_LOG(LogCook, Warning, TEXT("Found cooked file '%s' which shouldn't exist as it %s."), *SandboxFilename, ReasonText);
 			IFileManager::Get().Delete(*SandboxFilename);
 		}
 	}
-	PackageData.SendToState(UE::Cook::EPackageState::Idle, UE::Cook::ESendFlags::QueueAdd);
+	DemoteToIdle(PackageData, UE::Cook::ESendFlags::QueueAdd, Reason);
 }
 
 //////////////////////////////////////////////////////////////////////////
@@ -3897,7 +3980,7 @@ void UCookOnTheFlyServer::PumpSaves(UE::Cook::FTickStackData& StackData, uint32
 		if (Package->IsLoadedByEditorPropertiesOnly() && PackageTracker->UncookedEditorOnlyPackages.Contains(Package->GetFName()))
 		{
 			// We already attempted to cook this package and it's still not referenced by any non editor-only properties.
-			PackageData.SendToState(EPackageState::Idle, ESendFlags::QueueAdd);
+			DemoteToIdle(PackageData, ESendFlags::QueueAdd, ESuppressCookReason::OnlyEditorOnly);
 			++OutNumPushed;
 			continue;
 		}
@@ -3909,7 +3992,7 @@ void UCookOnTheFlyServer::PumpSaves(UE::Cook::FTickStackData& StackData, uint32
 		if (PackageTracker->NeverCookPackageList.Contains(PackageData.GetFileName()))
 		{
 			// refuse to save this package, it's clearly one of the undesirables
-			PackageData.SendToState(EPackageState::Idle, ESendFlags::QueueAdd);
+			DemoteToIdle(PackageData, ESendFlags::QueueAdd, ESuppressCookReason::NeverCook);
 			++OutNumPushed;
 			continue;
 		}
@@ -3921,7 +4004,7 @@ void UCookOnTheFlyServer::PumpSaves(UE::Cook::FTickStackData& StackData, uint32
 			// We've already saved all possible platforms for this package; this should not be possible.
 			// All places that add a package to the save queue check for existence of incomplete platforms before adding
 			UE_LOG(LogCook, Warning, TEXT("Package '%s' in SaveQueue has no more platforms left to cook; this should not be possible!"), *PackageData.GetFileName().ToString());
-			PackageData.SendToState(EPackageState::Idle, ESendFlags::QueueAdd);
+			DemoteToIdle(PackageData, ESendFlags::QueueAdd, ESuppressCookReason::AlreadyCooked);
 			++OutNumPushed;
 			continue;
 		}
@@ -3985,7 +4068,7 @@ void UCookOnTheFlyServer::PumpSaves(UE::Cook::FTickStackData& StackData, uint32
 				check(PackageData.HasPrepareSaveFailed()); // Should have been set by PrepareSave; we rely on this for cleanup
 				ReleaseCookedPlatformData(PackageData, EReleaseSaveReason::AbortSave);
 				PackageData.SetPlatformsCooked(PlatformsForPackage, false /* bSucceeded */);
-				PackageData.SendToState(EPackageState::Idle, ESendFlags::QueueAdd);
+				DemoteToIdle(PackageData, ESendFlags::QueueAdd, ESuppressCookReason::SaveError);
 				++OutNumPushed;
 				continue;
 			}
@@ -4086,7 +4169,7 @@ void UCookOnTheFlyServer::PumpSaves(UE::Cook::FTickStackData& StackData, uint32
 		}
 
 		ReleaseCookedPlatformData(PackageData, !Context.bHasRetryErrorCode ? EReleaseSaveReason::Completed : EReleaseSaveReason::DoneForNow);
-		PackageData.SendToState(EPackageState::Idle, ESendFlags::QueueAdd);
+		PromoteToSaveComplete(PackageData, ESendFlags::QueueAdd);
 		++OutNumPushed;
 	}
 }
@@ -5547,9 +5630,9 @@ void UCookOnTheFlyServer::Initialize( ECookMode::Type DesiredCookMode, ECookInit
 	if (IsCookByTheBookMode() && !IsCookingInEditor())
 	{
 		bool bCookMultiProcess = false;;
-		GConfig->GetBool(TEXT("CookSettings"), TEXT("CookMultiProcess"), bCookMultiProcess, GEditorIni);
-		bCookMultiProcess |= FParse::Param(FCommandLine::Get(), TEXT("-CookMultiProcess"));
-		bCookMultiProcess &= FParse::Param(FCommandLine::Get(), TEXT("-CookSingleProcess"));
+		GConfig->GetBool(TEXT("CookSettings"), TEXT("CookMultiProcessEnabled"), bCookMultiProcess, GEditorIni);
+		bCookMultiProcess |= FParse::Param(FCommandLine::Get(), TEXT("CookMultiProcess"));
+		bCookMultiProcess &= !FParse::Param(FCommandLine::Get(), TEXT("CookSingleProcess"));
 		if (bCookMultiProcess)
 		{
 			CookDirector = MakeUnique<UE::Cook::FCookDirector>(*this);
@@ -5775,15 +5858,14 @@ bool UCookOnTheFlyServer::TryInitializeCookWorker()
 {
 	using namespace UE::Cook;
 
-	FString CookDirectorHost;
-	if (!FParse::Value(FCommandLine::Get(), TEXT("-CookDirectorHost="), CookDirectorHost))
+	FDirectorConnectionInfo ConnectInfo;
+	if (!ConnectInfo.TryParseCommandLine())
 	{
-		UE_LOG(LogCook, Error, TEXT("CookWorker startup failed: no CookDirector specified on commandline."))
 		return false;
 	}
 	TUniquePtr<FWorkerRequestsRemote> RemoteTasks = MakeUnique<FWorkerRequestsRemote>(*this);
 	ECookInitializationFlags LocalCookFlags;
-	if (!RemoteTasks->TryConnect(FDirectorConnectionInfo{ CookDirectorHost }, LocalCookFlags))
+	if (!RemoteTasks->TryConnect(MoveTemp(ConnectInfo), LocalCookFlags))
 	{
 		return false;
 	}
@@ -8022,6 +8104,7 @@ void UCookOnTheFlyServer::CookByTheBookFinished()
 	check(IsCookByTheBookMode());
 	check(IsInSession());
 	check(PackageDatas->GetRequestQueue().IsEmpty());
+	check(PackageDatas->GetAssignedToWorkerSet().IsEmpty());
 	check(PackageDatas->GetLoadPrepareQueue().IsEmpty());
 	check(PackageDatas->GetLoadReadyQueue().IsEmpty());
 	check(PackageDatas->GetSaveQueue().IsEmpty());
@@ -8424,45 +8507,50 @@ void UCookOnTheFlyServer::CancelAllQueues()
 	FPackageDataQueue& SaveQueue = PackageDatas->GetSaveQueue();
 	while (!SaveQueue.IsEmpty())
 	{
-		SaveQueue.PopFrontValue()->SendToState(EPackageState::Idle, ESendFlags::QueueAdd);
+		DemoteToIdle(*SaveQueue.PopFrontValue(), ESendFlags::QueueAdd, ESuppressCookReason::CookCanceled);
 	}
 	FPackageDataQueue& LoadReadyQueue = PackageDatas->GetLoadReadyQueue();
 	while (!LoadReadyQueue.IsEmpty())
 	{
-		LoadReadyQueue.PopFrontValue()->SendToState(EPackageState::Idle, ESendFlags::QueueAdd);
+		DemoteToIdle(*LoadReadyQueue.PopFrontValue(), ESendFlags::QueueAdd, ESuppressCookReason::CookCanceled);
 	}
 	FLoadPrepareQueue& LoadPrepareQueue = PackageDatas->GetLoadPrepareQueue();
 	while (!LoadPrepareQueue.IsEmpty())
 	{
-		LoadPrepareQueue.PopFront()->SendToState(EPackageState::Idle, ESendFlags::QueueAdd);
+		DemoteToIdle(*LoadPrepareQueue.PopFront(), ESendFlags::QueueAdd, ESuppressCookReason::CookCanceled);
+	}
+	for (FPackageData* PackageData : PackageDatas->GetAssignedToWorkerSet())
+	{
+		DemoteToIdle(*PackageData, ESendFlags::QueueAdd, ESuppressCookReason::CookCanceled);
 	}
+	PackageDatas->GetAssignedToWorkerSet().Empty();
 	FRequestQueue& RequestQueue = PackageDatas->GetRequestQueue();
 	FPackageDataSet& UnclusteredRequests = RequestQueue.GetUnclusteredRequests();
 	for (FPackageData* PackageData : UnclusteredRequests)
 	{
-		PackageData->SendToState(EPackageState::Idle, ESendFlags::QueueAdd);
+		DemoteToIdle(*PackageData, ESendFlags::QueueAdd, ESuppressCookReason::CookCanceled);
 	}
 	UnclusteredRequests.Empty();
 	TRingBuffer<FRequestCluster>& RequestClusters = RequestQueue.GetRequestClusters();
 	for (FRequestCluster& RequestCluster : RequestClusters)
 	{
 		TArray<FPackageData*> RequestsToLoad;
-		TArray<FPackageData*> RequestsToDemote;
+		TArray<TPair<FPackageData*, ESuppressCookReason>> RequestsToDemote;
 		RequestCluster.ClearAndDetachOwnedPackageDatas(RequestsToLoad, RequestsToDemote);
 		for (FPackageData* PackageData : RequestsToLoad)
 		{
-			PackageData->SendToState(EPackageState::Idle, ESendFlags::QueueAdd);
+			DemoteToIdle(*PackageData, ESendFlags::QueueAdd, ESuppressCookReason::CookCanceled);
 		}
-		for (FPackageData* PackageData : RequestsToDemote)
+		for (TPair<FPackageData*, ESuppressCookReason>& Pair : RequestsToDemote)
 		{
-			PackageData->SendToState(EPackageState::Idle, ESendFlags::QueueAdd);
+			DemoteToIdle(*Pair.Key, ESendFlags::QueueAdd, ESuppressCookReason::CookCanceled);
 		}
 	}
 	RequestClusters.Empty();
 
 	while (!RequestQueue.IsReadyRequestsEmpty())
 	{
-		RequestQueue.PopReadyRequest()->SendToState(EPackageState::Idle, ESendFlags::QueueAdd);
+		DemoteToIdle(*RequestQueue.PopReadyRequest(), ESendFlags::QueueAdd, ESuppressCookReason::CookCanceled);
 	}
 
 	SetLoadBusy(false);
diff --git a/Engine/Source/Editor/UnrealEd/Private/Cooker/CookDirector.cpp b/Engine/Source/Editor/UnrealEd/Private/Cooker/CookDirector.cpp
new file mode 100644
--- /dev/null
+++ b/Engine/Source/Editor/UnrealEd/Private/Cooker/CookDirector.cpp
@@ -0,0 +1,485 @@
+// Copyright Epic Games, Inc. All Rights Reserved.
+
+#pragma once
+
+#include "CookDirector.h"
+
+#include "CompactBinaryTCP.h"
+#include "CookPackageData.h"
+#include "CookWorkerServer.h"
+#include "CookOnTheSide/CookOnTheFlyServer.h"
+#include "CoreGlobals.h"
+#include "Math/NumericLimits.h"
+#include "Misc/CommandLine.h"
+#include "Misc/ConfigCacheIni.h"
+#include "Serialization/CompactBinary.h"
+#include "Serialization/CompactBinaryWriter.h"
+#include "Sockets.h"
+#include "SocketSubsystem.h"
+#include "String/ParseTokens.h"
+
+namespace UE::Cook
+{
+
+FCookDirector::FCookDirector(UCookOnTheFlyServer& InCOTFS)
+	: COTFS(InCOTFS)
+{
+	WorkersStalledStartTimeSeconds = MAX_flt;
+	WorkersStalledWarnTimeSeconds = MAX_flt;
+
+	const TCHAR* CommandLine = FCommandLine::Get();
+	GConfig->GetInt(TEXT("CookSettings"), TEXT("CookWorkerCount"), DesiredNumRemoteWorkers, GEditorIni);
+	FParse::Value(CommandLine, TEXT("-CookWorkerCount="), DesiredNumRemoteWorkers);
+	WorkerConnectPort = Sockets::COOKDIRECTOR_DEFAULT_REQUEST_CONNECTION_PORT;
+	FParse::Value(CommandLine, TEXT("-CookDirectorListenPort="), WorkerConnectPort);
+
+	if (DesiredNumRemoteWorkers > 0)
+	{
+		ISocketSubsystem* SocketSubsystem = ISocketSubsystem::Get();
+		if (!SocketSubsystem)
+		{
+			UE_LOG(LogCook, Error, TEXT("CookDirector initialization failure: platform does not support network sockets. CookWorkers will be disabled."));
+			DesiredNumRemoteWorkers = 0;
+		}
+	}
+
+	if (DesiredNumRemoteWorkers > 0)
+	{
+		TryCreateWorkerConnectSocket();
+	}
+}
+
+FCookDirector::~FCookDirector()
+{
+	TSet<FPackageData*> AbortedAssignments;
+	for (TPair<int32, TUniquePtr<FCookWorkerServer>>& Pair : RemoteWorkers)
+	{
+		Pair.Value->AbortWorker(AbortedAssignments);
+	}
+	for (FPackageData* PackageData : AbortedAssignments)
+	{
+		check(PackageData->IsInProgress()); // Packages that were assigned to workers should be in the AssignedToWorker state
+		PackageData->SetWorkerAssignment(FWorkerId::Invalid());
+		PackageData->SendToState(UE::Cook::EPackageState::Request, ESendFlags::QueueAddAndRemove);
+	}
+	RemoteWorkers.Empty();
+	PendingConnections.Empty();
+	Sockets::CloseSocket(WorkerConnectSocket);
+}
+
+void FCookDirector::AssignRequests(TArrayView<UE::Cook::FPackageData*> Requests, TArray<FWorkerId>& OutAssignments)
+{
+	InitializeWorkers();
+
+	int32 NumRemoteWorkers = RemoteWorkers.Num();
+	if (NumRemoteWorkers == 0)
+	{
+		OutAssignments.SetNum(Requests.Num());
+		for (FWorkerId& Assignment : OutAssignments)
+		{
+			Assignment = FWorkerId::Local();
+		}
+		return;
+	}
+
+	// Convert the Map of RemoteWorkers into an array sorted by WorkerIndex
+	TArray<FCookWorkerServer*> SortedWorkers;
+	SortedWorkers.Reserve(NumRemoteWorkers);
+	int32 MaxRemoteIndex = -1;
+	for (TPair<int32, TUniquePtr<FCookWorkerServer>>& Pair : RemoteWorkers)
+	{
+		SortedWorkers.Add(Pair.Value.Get());
+		MaxRemoteIndex = FMath::Max(Pair.Key, MaxRemoteIndex);
+	}
+	check(MaxRemoteIndex >= 0);
+	SortedWorkers.Sort([](const FCookWorkerServer& A, const FCookWorkerServer& B) { return A.GetWorkerId() < B.GetWorkerId(); });
+
+	// Call a LoadBalancing algorithm to split the requests among the LocalWorker and RemoteWorkers
+	// MPCOOKTODO: Implement LoadBalanceGreedy
+	LoadBalanceStriped(SortedWorkers, Requests, OutAssignments);
+
+	// Split the output array of WorkerId assignments into a batch for each of the RemoteWorkers 
+	TArray<TArray<FPackageData*>> RemoteBatches;
+	RemoteBatches.SetNum(MaxRemoteIndex+1);
+	for (int32 RequestIndex = 0; RequestIndex < Requests.Num(); ++RequestIndex)
+	{
+		FWorkerId WorkerId = OutAssignments[RequestIndex];
+		if (!WorkerId.IsLocal())
+		{
+			uint8 RemoteIndex = WorkerId.GetRemoteIndex();
+			check(RemoteIndex < RemoteBatches.Num());
+			TArray<FPackageData*>& RemoteBatch = RemoteBatches[RemoteIndex];
+			if (RemoteBatch.Num() == 0)
+			{
+				RemoteBatch.Reserve(2 * Requests.Num() / (NumRemoteWorkers + 1));
+			}
+			RemoteBatch.Add(Requests[RequestIndex]);
+		}
+	}
+
+	// MPCOOKTODO: Sort each batch from leaf to root
+
+	// Assign each batch to the FCookWorkerServer in RemoteWorkers;
+	// the CookWorkerServer's tick will handle sending the message to the remote process
+	for (TPair<int32, TUniquePtr<FCookWorkerServer>>& Pair : RemoteWorkers)
+	{
+		Pair.Value->AppendAssignments(RemoteBatches[Pair.Key]);
+	}
+	TickWorkerConnects();
+}
+
+void FCookDirector::RemoveFromWorker(FPackageData& PackageData)
+{
+	for (TPair<int32, TUniquePtr<FCookWorkerServer>>& Pair : RemoteWorkers)
+	{
+		Pair.Value->AbortAssignment(PackageData);
+	}
+}
+
+
+void FCookDirector::TickFromSchedulerThread()
+{
+	TickWorkerConnects();
+	for (TPair<int32, TUniquePtr<FCookWorkerServer>>& Pair : RemoteWorkers)
+	{
+		FCookWorkerServer& RemoteWorker = *Pair.Value;
+		RemoteWorker.TickFromSchedulerThread();
+		if (RemoteWorker.IsShuttingDown())
+		{
+			TUniquePtr<FCookWorkerServer>& Existing = ShuttingDownWorkers.FindOrAdd(&RemoteWorker);
+			check(!Existing); // We should not be able to send the same pointer into ShuttingDown twice
+		}
+	}
+
+	TickWorkerShutdowns();
+
+	bool bIsStalled = COTFS.IsMultiprocessLocalWorkerIdle() && !COTFS.PackageDatas->GetAssignedToWorkerSet().IsEmpty();
+	SetWorkersStalled(bIsStalled);
+}
+
+void FCookDirector::SetWorkersStalled(bool bInWorkersStalled)
+{
+	if (bInWorkersStalled != bWorkersStalled)
+	{
+		bWorkersStalled = bInWorkersStalled;
+		if (bWorkersStalled)
+		{
+			const double CurrentTime = FPlatformTime::Seconds();
+			WorkersStalledStartTimeSeconds = CurrentTime;
+			WorkersStalledWarnTimeSeconds = CurrentTime + GCookProgressWarnBusyTime;
+		}
+		else
+		{
+			WorkersStalledStartTimeSeconds = MAX_flt;
+			WorkersStalledWarnTimeSeconds = MAX_flt;
+		}
+	}
+	else if (bWorkersStalled)
+	{
+		const double CurrentTime = FPlatformTime::Seconds();
+		if (CurrentTime >= WorkersStalledWarnTimeSeconds)
+		{
+			UE_LOG(LogCook, Warning, TEXT("Cooker has been blocked with no results from remote CookWorkers for %.0f seconds."),
+				(float)(CurrentTime - WorkersStalledStartTimeSeconds));
+			WorkersStalledWarnTimeSeconds = CurrentTime + GCookProgressWarnBusyTime;
+		}
+	}
+}
+
+FCookDirector::FPendingConnection::FPendingConnection(FPendingConnection&& Other)
+{
+	Swap(Socket, Other.Socket);
+	Buffer = MoveTemp(Other.Buffer);
+}
+
+FCookDirector::FPendingConnection::~FPendingConnection()
+{
+	Sockets::CloseSocket(Socket);
+}
+
+FSocket* FCookDirector::FPendingConnection::DetachSocket()
+{
+	FSocket* Result = Socket;
+	Socket = nullptr;
+	return Result;
+}
+
+void FWorkerConnectMessage::Write(FCbWriter& Writer) const
+{
+	Writer << "RemoteIndex" << RemoteIndex;
+}
+
+bool FWorkerConnectMessage::TryRead(FCbObjectView Object)
+{
+	RemoteIndex = Object["RemoteIndex"].AsInt32(-1);
+	return RemoteIndex >= 0;
+}
+
+FGuid FWorkerConnectMessage::MessageType(TEXT("302096E887DA48F7B079FAFAD0EE5695"));
+
+bool FCookDirector::TryCreateWorkerConnectSocket()
+{
+	FString ErrorReason;
+	TSharedPtr<FInternetAddr> ListenAddr;
+	WorkerConnectSocket = Sockets::CreateListenSocket(WorkerConnectPort, ListenAddr, WorkerConnectAuthority,
+		TEXT("FCookDirector-WorkerConnect"), ErrorReason);
+	if (!WorkerConnectSocket)
+	{
+		UE_LOG(LogCook, Error, TEXT("CookDirector could not create listen socket, CookWorkers will be disabled. Reason: %s."),
+			*ErrorReason);
+		DesiredNumRemoteWorkers = 0;
+		return false;
+	}
+	return true;
+}
+
+
+void FCookDirector::InitializeWorkers()
+{
+	if (RemoteWorkers.Num() >= DesiredNumRemoteWorkers)
+	{
+		return;
+	}
+
+	// Find any unused RemoteIndex less than the maximum used RemoteIndex
+	TArray<int32> UnusedRemoteIndexes;
+	RemoteWorkers.KeySort(TLess<>());
+	int32 NextPossiblyOpenIndex = 0;
+	for (TPair<int32, TUniquePtr<FCookWorkerServer>>& Pair : RemoteWorkers)
+	{
+		check(NextPossiblyOpenIndex <= Pair.Key);
+		while (NextPossiblyOpenIndex != Pair.Key)
+		{
+			UnusedRemoteIndexes.Add(NextPossiblyOpenIndex++);
+		}
+	}
+	// Add RemoteWorkers, pulling the RemoteIndex id from the UnusedRemoteIndexes if any exist
+	// otherwise use the next integer because all indexes up to RemoteWorkers.Num() are in use.
+	while (RemoteWorkers.Num() < DesiredNumRemoteWorkers)
+	{
+		int32 RemoteIndex;
+		if (UnusedRemoteIndexes.Num())
+		{
+			RemoteIndex = UnusedRemoteIndexes[0];
+			UnusedRemoteIndexes.RemoveAtSwap(0);
+		}
+		else
+		{
+			RemoteIndex = RemoteWorkers.Num();
+		}
+		RemoteWorkers.Add(RemoteIndex, MakeUnique<FCookWorkerServer>(*this, FWorkerId::FromRemoteIndex(RemoteIndex)));
+	}
+}
+
+void FCookDirector::TickWorkerConnects()
+{
+	using namespace UE::CompactBinaryTCP;
+
+	if (!WorkerConnectSocket)
+	{
+		return;
+	}
+
+	bool bReadReady;
+	while (WorkerConnectSocket->HasPendingConnection(bReadReady) && bReadReady)
+	{
+		FSocket* WorkerSocket = WorkerConnectSocket->Accept(TEXT("Client Connection"));
+		if (!WorkerSocket)
+		{
+			UE_LOG(LogCook, Warning, TEXT("Pending connection failed to create a ClientSocket."));
+		}
+		else
+		{
+			WorkerSocket->SetNonBlocking(true);
+			PendingConnections.Add(FPendingConnection(WorkerSocket));
+		}
+	}
+
+	for (TArray<FPendingConnection>::TIterator Iter(PendingConnections); Iter; ++Iter)
+	{
+		FPendingConnection& Conn = *Iter;
+		TArray<FMarshalledMessage> Messages;
+		EConnectionStatus Status;
+		Status = TryReadPacket(Conn.Socket, Conn.Buffer, Messages);
+		if (Status != EConnectionStatus::Okay)
+		{
+			UE_LOG(LogCook, Warning, TEXT("Pending connection failed before sending a WorkerPacket: %s"), DescribeStatus(Status));
+			Iter.RemoveCurrent();
+		}
+		if (Messages.Num() == 0)
+		{
+			continue;
+		}
+		FPendingConnection LocalConn(MoveTemp(Conn));
+		Iter.RemoveCurrent();
+
+		if (Messages[0].MessageType != FWorkerConnectMessage::MessageType)
+		{
+			UE_LOG(LogCook, Warning, TEXT("Pending connection sent a different message before sending a connection message. MessageType: %s. Connection will be ignored."),
+				*Messages[0].MessageType.ToString());
+			continue;
+		}
+		FWorkerConnectMessage Message;
+		if (!Message.TryRead(Messages[0].Object))
+		{
+			UE_LOG(LogCook, Warning, TEXT("Pending connection sent an invalid Connection Message. Connection will be ignored."));
+			continue;
+		}
+		TUniquePtr<FCookWorkerServer>* RemoteWorkerPtr = RemoteWorkers.Find(Message.RemoteIndex);
+		if (!RemoteWorkerPtr)
+		{
+			TStringBuilder<256> ValidIndexes;
+			if (RemoteWorkers.Num())
+			{
+				RemoteWorkers.KeySort(TLess<>());
+				for (TPair<int32, TUniquePtr<FCookWorkerServer>>& Pair : RemoteWorkers)
+				{
+					ValidIndexes.Appendf(TEXT("%d,"), Pair.Key);
+				}
+				ValidIndexes.RemoveSuffix(1); // Remove the terminating comma
+			}
+			UE_LOG(LogCook, Warning, TEXT("Pending connection sent a Connection Message with invalid RemoteIndex %d. ValidIndexes = {%s}. Connection will be ignored."),
+				Message.RemoteIndex, *ValidIndexes);
+			continue;
+		}
+		FCookWorkerServer& RemoteWorker = **RemoteWorkerPtr;
+		FSocket* LocalSocket = LocalConn.DetachSocket();
+		Messages.RemoveAt(0);
+		if (!RemoteWorker.TryHandleConnectMessage(Message, LocalSocket, MoveTemp(Messages)))
+		{
+			UE_LOG(LogCook, Warning, TEXT("Pending connection sent a Connection Message with an already in-use RemoteIndex. Connection will be ignored."));
+			Sockets::CloseSocket(LocalSocket);
+			continue;
+		}
+	}
+}
+
+void FCookDirector::TickWorkerShutdowns()
+{
+	if (ShuttingDownWorkers.IsEmpty())
+	{
+		return;
+	}
+
+	// Move any newly shutting down workers from RemoteWorkers
+	TArray<FCookWorkerServer*> NewShutdowns;
+	for (TPair<FCookWorkerServer*, TUniquePtr<FCookWorkerServer>>& Pair : ShuttingDownWorkers)
+	{
+		if (!Pair.Value)
+		{
+			NewShutdowns.Add(Pair.Key);
+		}
+	}
+	for (FCookWorkerServer* NewShutdown : NewShutdowns)
+	{
+		AbortWorker(NewShutdown->GetWorkerId());
+	}
+	for (TMap<FCookWorkerServer*, TUniquePtr<FCookWorkerServer>>::TIterator Iter(ShuttingDownWorkers); Iter; ++Iter)
+	{
+		FCookWorkerServer& Worker = *Iter.Key();
+		check(Iter.Value()); // All non-yet-moved workers should have been moved by the AbortWorker calls above
+		if (Worker.IsShutdownComplete())
+		{
+			// Worker is now deleted, do not access
+			Iter.RemoveCurrent();
+		}
+	}
+}
+
+FString FCookDirector::GetWorkerCommandLine(FWorkerId WorkerId)
+{
+	FString CommandLine = FCommandLine::Get();
+
+	const TCHAR* ProjectName = FApp::GetProjectName();
+	checkf(ProjectName && ProjectName[0], TEXT("Expected UnrealEditor to be running with a non-empty project name"));
+	TArray<FString> Tokens;
+	UE::String::ParseTokensMultiple(CommandLine, { ' ', '\t', '\r', '\n' }, [&Tokens](FStringView Token)
+		{
+			if (Token.StartsWith(TEXT("-run=")) ||
+				Token == TEXT("-CookOnTheFly") ||
+				Token == TEXT("-CookWorker") ||
+				Token == TEXT("-CookMultiProcess") ||
+				Token == TEXT("-CookSingleProcess") ||
+				Token.StartsWith(TEXT("-TargetPlatform")) ||
+				Token.StartsWith(TEXT("-CookCultures")) ||
+				Token.StartsWith(TEXT("-CookDirectorCount=")) ||
+				Token.StartsWith(TEXT("-CookDirectorHost=")) ||
+				Token.StartsWith(TEXT("-CookWorkerId="))
+				)
+			{
+				return;
+			}
+			Tokens.Add(FString(Token));
+		}, UE::String::EParseTokensOptions::SkipEmpty);
+	if (Tokens[0] != ProjectName)
+	{
+		Tokens.Insert(ProjectName, 0);
+	}
+	Tokens.Insert(TEXT("-run=cook"), 1);
+	Tokens.Insert(TEXT("-cookworker"), 2);
+	check(!WorkerConnectAuthority.IsEmpty()); // This should have been constructed in TryCreateWorkerConnectSocket before any CookWorkerServers could exist to call GetWorkerCommandLine
+	Tokens.Add(FString::Printf(TEXT("-CookDirectorHost=%s"), *WorkerConnectAuthority));
+	Tokens.Add(FString::Printf(TEXT("-CookWorkerId=%d"), WorkerId.GetRemoteIndex()));
+
+	return FString::Join(Tokens, TEXT(" "));
+}
+
+bool FDirectorConnectionInfo::TryParseCommandLine()
+{
+	if (!FParse::Value(FCommandLine::Get(), TEXT("-CookDirectorHost="), HostURI))
+	{
+		UE_LOG(LogCook, Error, TEXT("CookWorker startup failed: no CookDirector specified on commandline."));
+		return false;
+	}
+	if (!FParse::Value(FCommandLine::Get(), TEXT("-CookWorkerId="), RemoteIndex))
+	{
+		UE_LOG(LogCook, Error, TEXT("CookWorker startup failed: no CookWorkerId specified on commandline."));
+		return false;
+	}
+	return true;
+}
+
+void FCookDirector::LoadBalanceStriped(TArrayView<FCookWorkerServer*> SortedWorkers, TArrayView<FPackageData*> Requests, TArray<FWorkerId>& OutAssignments)
+{
+	TArray<FWorkerId> AllWorkers;
+	int32 NumAllWorkers = SortedWorkers.Num() + 1;
+	AllWorkers.Reserve(NumAllWorkers);
+	AllWorkers.Add(FWorkerId::Local());
+	for (FCookWorkerServer* Worker : SortedWorkers)
+	{
+		AllWorkers.Add(Worker->GetWorkerId());
+	}
+	OutAssignments.Reset(Requests.Num());
+	int32 AllWorkersIndex = 0;
+	for (FPackageData* Request : Requests)
+	{
+		OutAssignments.Add(AllWorkers[AllWorkersIndex]);
+		AllWorkersIndex = (AllWorkersIndex + 1) % NumAllWorkers;
+	}
+}
+
+void FCookDirector::AbortWorker(FWorkerId WorkerId)
+{
+	check(!WorkerId.IsLocal());
+	int32 Index = WorkerId.GetRemoteIndex();
+	TUniquePtr<FCookWorkerServer> RemoteWorker;
+	RemoteWorkers.RemoveAndCopyValue(Index, RemoteWorker);
+	if (!RemoteWorker)
+	{
+		return;
+	}
+	--DesiredNumRemoteWorkers;
+	TSet<FPackageData*> PackagesToReassign;
+	RemoteWorker->AbortWorker(PackagesToReassign);
+	for (FPackageData* PackageData : PackagesToReassign)
+	{
+		check(PackageData->IsInProgress()); // Packages that were assigned to a worker should be in the AssignedToWorker state
+		PackageData->SetWorkerAssignment(FWorkerId::Invalid());
+		PackageData->SendToState(UE::Cook::EPackageState::Request, ESendFlags::QueueAddAndRemove);
+	}
+	TUniquePtr<FCookWorkerServer>& Existing = ShuttingDownWorkers.FindOrAdd(RemoteWorker.Get());
+	check(!Existing); // We should not be able to abort a worker twice because we removed it from RemoteWorkers above
+	Existing = MoveTemp(RemoteWorker);
+}
+
+}
diff --git a/Engine/Source/Editor/UnrealEd/Private/Cooker/CookDirector.h b/Engine/Source/Editor/UnrealEd/Private/Cooker/CookDirector.h
--- a/Engine/Source/Editor/UnrealEd/Private/Cooker/CookDirector.h
+++ b/Engine/Source/Editor/UnrealEd/Private/Cooker/CookDirector.h
@@ -2,16 +2,118 @@
 
 #pragma once
 
+#include "CompactBinaryTCP.h"
+#include "Containers/Array.h"
+#include "Containers/ArrayView.h"
+#include "Containers/Map.h"
+#include "Containers/UnrealString.h"
+#include "CookSockets.h"
+#include "CookTypes.h"
+#include "Memory/SharedBuffer.h"
+#include "Misc/Guid.h"
+#include "Serialization/CompactBinary.h"
+#include "Templates/UniquePtr.h"
+
+class FCbWriter;
 class UCookOnTheFlyServer;
+namespace UE::Cook { class FCookWorkerServer; }
+namespace UE::Cook { struct FPackageData; }
+namespace UE::Cook { struct FWorkerId; }
 
 namespace UE::Cook
 {
 
+/**
+ * Helper for CookOnTheFlyServer that sends requests to CookWorker processes for load/save and merges
+ * their replies into the local process's cook results.
+ */
 class FCookDirector
 {
 public:
-	FCookDirector(UCookOnTheFlyServer& COTFS)
-	{}
+	FCookDirector(UCookOnTheFlyServer& InCOTFS);
+	~FCookDirector();
+
+	/** Assign the given requests out to CookWorkers (or keep on local COTFS), return the list of assignments. */
+	void AssignRequests(TArrayView<UE::Cook::FPackageData*> Requests, TArray<FWorkerId>& OutAssignments);
+	/** Notify the CookWorker that owns the cook of the package that the Director wants to take it back. */
+	void RemoveFromWorker(FPackageData& PackageData);
+	/** Periodic tick function. Sends/Receives messages to CookWorkers. */
+	void TickFromSchedulerThread();
+
+private:
+	/** CookWorker connections that have not yet identified which CookWorker they are. */
+	struct FPendingConnection
+	{
+		explicit FPendingConnection(FSocket* InSocket = nullptr)
+		:Socket(InSocket)
+		{
+		}
+		FPendingConnection(FPendingConnection&& Other);
+		FPendingConnection(const FPendingConnection& Other) = delete;
+		~FPendingConnection();
+
+		FSocket* DetachSocket();
+
+		FSocket* Socket = nullptr;
+		UE::CompactBinaryTCP::FReceiveBuffer Buffer;
+	};
+	/** Initialization helper: create the listen socket. */
+	bool TryCreateWorkerConnectSocket();
+	/** Initialization helper: add the local Server for a remote worker, worker process not yet created. */
+	void InitializeWorkers();
+	/** Tick helper: tick any workers that have not yet finished initialization. */
+	void TickWorkerConnects();
+	/** Tick helper: tick any workers that are shutting down. */
+	void TickWorkerShutdowns();
+	/** Get the commandline to launch a worker process with. */
+	FString GetWorkerCommandLine(FWorkerId WorkerId);
+	/** Simple assignment that divides requests evenly without considering dependencies or load burden. */
+	void LoadBalanceStriped(TArrayView<FCookWorkerServer*> SortedWorkers, TArrayView<FPackageData*> Requests,
+		TArray<FWorkerId>& OutAssignments);
+	/** Move the given worker from active workers to the list of workers shutting down. */
+	void AbortWorker(FWorkerId WorkerId);
+	/**
+	 * Periodically update whether (1) local server is done and (2) no results from cookworkers have come in.
+	 * Send warning when it goes on too long.
+	 */
+	void SetWorkersStalled(bool bInWorkersStalled);
+
+private:
+	TMap<int32, TUniquePtr<FCookWorkerServer>> RemoteWorkers;
+	TMap<FCookWorkerServer*, TUniquePtr<FCookWorkerServer>> ShuttingDownWorkers;
+	TArray<FPendingConnection> PendingConnections;
+	FString WorkerConnectAuthority;
+	UCookOnTheFlyServer& COTFS;
+	FSocket* WorkerConnectSocket = nullptr;
+	double WorkersStalledStartTimeSeconds = 0.;
+	double WorkersStalledWarnTimeSeconds = 0.;
+	int32 DesiredNumRemoteWorkers = 4;
+	int32 WorkerConnectPort = 0;
+	bool bWorkersStalled = false;
+
+	friend class UE::Cook::FCookWorkerServer;
+};
+
+/** Parameters parsed from commandline for how a CookWorker connects to the CooKDirector. */
+struct FDirectorConnectionInfo
+{
+	bool TryParseCommandLine();
+
+	FString HostURI;
+	int32 RemoteIndex = 0;
+};
+
+/** Message sent from a CookWorker to the Director to report that it is ready for setup messages and cooking. */
+struct FWorkerConnectMessage : public UE::CompactBinaryTCP::IMessage
+{
+public:
+	virtual void Write(FCbWriter& Writer) const override;
+	virtual bool TryRead(FCbObjectView Object) override;
+	virtual FGuid GetMessageType() const override { return MessageType; }
+
+public:
+	int32 RemoteIndex = 0;
+	static FGuid MessageType;
 };
 
 }
\ No newline at end of file
diff --git a/Engine/Source/Editor/UnrealEd/Private/Cooker/CookPackageData.cpp b/Engine/Source/Editor/UnrealEd/Private/Cooker/CookPackageData.cpp
--- a/Engine/Source/Editor/UnrealEd/Private/Cooker/CookPackageData.cpp
+++ b/Engine/Source/Editor/UnrealEd/Private/Cooker/CookPackageData.cpp
@@ -200,35 +200,28 @@ void FPackageData::SetIsUrgent(bool Value)
 void FPackageData::UpdateRequestData(const TConstArrayView<const ITargetPlatform*> InRequestedPlatforms,
 	bool bInIsUrgent, FCompletionCallback&& InCompletionCallback, FInstigator&& InInstigator, bool bAllowUpdateUrgency)
 {
-	if (IsInProgress())
-	{
-		AddCompletionCallback(MoveTemp(InCompletionCallback));
+	check(IsInProgress());
+	AddCompletionCallback(MoveTemp(InCompletionCallback));
 
-		bool bUrgencyChanged = false;
-		if (bInIsUrgent && !GetIsUrgent())
-		{
-			bUrgencyChanged = true;
-			SetIsUrgent(true);
-		}
+	bool bUrgencyChanged = false;
+	if (bInIsUrgent && !GetIsUrgent())
+	{
+		bUrgencyChanged = true;
+		SetIsUrgent(true);
+	}
 
-		if (!HasAllRequestedPlatforms(InRequestedPlatforms))
-		{
-			// Send back to the Request state (canceling any current operations) and then add the new platforms
-			if (GetState() != EPackageState::Request)
-			{
-				SendToState(EPackageState::Request, ESendFlags::QueueAddAndRemove);
-			}
-			SetPlatformsRequested(InRequestedPlatforms, true);
-		}
-		else if (bUrgencyChanged && bAllowUpdateUrgency)
+	if (!HasAllRequestedPlatforms(InRequestedPlatforms))
+	{
+		// Send back to the Request state (canceling any current operations) and then add the new platforms
+		if (GetState() != EPackageState::Request)
 		{
-			SendToState(GetState(), ESendFlags::QueueAddAndRemove);
+			SendToState(EPackageState::Request, ESendFlags::QueueAddAndRemove);
 		}
+		SetPlatformsRequested(InRequestedPlatforms, true);
 	}
-	else if (InRequestedPlatforms.Num() > 0)
+	else if (bUrgencyChanged && bAllowUpdateUrgency)
 	{
-		SetRequestData(InRequestedPlatforms, bInIsUrgent, MoveTemp(InCompletionCallback), MoveTemp(InInstigator));
-		SendToState(EPackageState::Request, ESendFlags::QueueAddAndRemove);
+		SendToState(GetState(), ESendFlags::QueueAddAndRemove);
 	}
 }
 
@@ -471,6 +464,9 @@ struct FStateProperties
 		case EPackageState::Request:
 			Properties = EPackageStateProperty::InProgress;
 			break;
+		case EPackageState::AssignedToWorker:
+			Properties = EPackageStateProperty::InProgress;
+			break;
 		case EPackageState::LoadPrepare:
 			Properties = EPackageStateProperty::InProgress | EPackageStateProperty::Loading;
 			break;
@@ -504,6 +500,13 @@ void FPackageData::SendToState(EPackageState NextState, ESendFlags SendFlags)
 		}
 		OnExitRequest();
 		break;
+	case EPackageState::AssignedToWorker:
+		if (!!(SendFlags & ESendFlags::QueueRemove))
+		{
+			ensure(PackageDatas.GetAssignedToWorkerSet().Remove(this) == 1);
+		}
+		OnExitAssignedToWorker();
+		break;
 	case EPackageState::LoadPrepare:
 		if (!!(SendFlags & ESendFlags::QueueRemove))
 		{
@@ -598,6 +601,13 @@ void FPackageData::SendToState(EPackageState NextState, ESendFlags SendFlags)
 			PackageDatas.GetRequestQueue().AddRequest(this);
 		}
 		break;
+	case EPackageState::AssignedToWorker:
+		OnEnterAssignedToWorker();
+		if (((SendFlags & ESendFlags::QueueAdd) != ESendFlags::QueueNone))
+		{
+			PackageDatas.GetAssignedToWorkerSet().Add(this);
+		}
+		break;
 	case EPackageState::LoadPrepare:
 		OnEnterLoadPrepare();
 		if ((SendFlags & ESendFlags::QueueAdd) != ESendFlags::QueueNone)
@@ -657,6 +667,9 @@ void FPackageData::CheckInContainer() const
 	case EPackageState::Request:
 		check(PackageDatas.GetRequestQueue().Contains(this));
 		break;
+	case EPackageState::AssignedToWorker:
+		check(PackageDatas.GetAssignedToWorkerSet().Contains(this));
+		break;
 	case EPackageState::LoadPrepare:
 		check(PackageDatas.GetLoadPrepareQueue().Contains(this));
 		break;
@@ -707,6 +720,18 @@ void FPackageData::OnExitRequest()
 {
 }
 
+void FPackageData::OnEnterAssignedToWorker()
+{
+}
+
+void FPackageData::OnExitAssignedToWorker()
+{
+	if (GetWorkerAssignment().IsValid())
+	{
+		PackageDatas.GetCookOnTheFlyServer().NotifyRemovedFromWorker(*this);
+	}
+}
+
 void FPackageData::OnEnterLoadPrepare()
 {
 }
diff --git a/Engine/Source/Editor/UnrealEd/Private/Cooker/CookPackageData.h b/Engine/Source/Editor/UnrealEd/Private/Cooker/CookPackageData.h
--- a/Engine/Source/Editor/UnrealEd/Private/Cooker/CookPackageData.h
+++ b/Engine/Source/Editor/UnrealEd/Private/Cooker/CookPackageData.h
@@ -43,7 +43,7 @@ struct FPendingCookedPlatformDataCancelManager;
 extern const TCHAR* GeneratedPackageSubPath;
 
 /** Flags specifying the behavior of FPackageData::SendToState */
-enum class ESendFlags
+enum class ESendFlags : uint8
 {
 	/**
 	 * PackageData will not be removed from queue for its old state and will not be added to queue for its new state.
@@ -464,6 +464,10 @@ public:
 	/** Set whether COTFS is keeping this package referenced referenced during GC. */
 	void SetKeepReferencedDuringGC(bool Value) { bKeepReferencedDuringGC = Value != 0; }
 
+	/** For MultiProcessCooks, Get the id of the worker this Package is assigned to. Returns Invalid if not assigned to a worker. */
+	FWorkerId GetWorkerAssignment() const { return WorkerAssignment; }
+	/** Set the id of the worker this Package is assigned to. */
+	void SetWorkerAssignment(FWorkerId InWorkerAssignment) { WorkerAssignment = InWorkerAssignment; }
 
 private:
 	friend struct UE::Cook::FPackageDatas;
@@ -516,6 +520,8 @@ private:
 	void OnExitIdle();
 	void OnEnterRequest();
 	void OnExitRequest();
+	void OnEnterAssignedToWorker();
+	void OnExitAssignedToWorker();
 	void OnEnterLoadPrepare();
 	void OnExitLoadPrepare();
 	void OnEnterLoadReady();
@@ -562,6 +568,7 @@ private:
 	FOpenPackageResult PreloadableFileOpenResult;
 	FInstigator Instigator;
 
+	FWorkerId WorkerAssignment = FWorkerId::Invalid();
 	uint32 State : int32(EPackageState::BitCount);
 	uint32 bIsUrgent : 1;
 	uint32 bIsVisited : 1;
@@ -981,6 +988,10 @@ public:
 	 * PackageData that need to be cooked.
 	 */
 	FRequestQueue& GetRequestQueue();
+
+	/** Return the Set that holds unordered all PackageDatas that are in the AssignedToWorker state. */
+	TFastPointerSet<FPackageData*>& GetAssignedToWorkerSet() { return AssignedToWorkerSet; }
+
 	/**
 	 * Return the LoadPrepareQueue used by the CookOnTheFlyServer. The LoadPrepareQueue is the dependency-ordered
 	 * list of FPackageData that need to be preloaded before they can be loaded.
@@ -1177,6 +1188,7 @@ private:
 	TMap<FName, FPackageData*> FileNameToPackageData;
 	TArray<FPendingCookedPlatformData> PendingCookedPlatformDatas;
 	FRequestQueue RequestQueue;
+	TFastPointerSet<FPackageData*> AssignedToWorkerSet;
 	FLoadPrepareQueue LoadPrepareQueue;
 	FPackageDataQueue LoadReadyQueue;
 	FPackageDataQueue SaveQueue;
diff --git a/Engine/Source/Editor/UnrealEd/Private/Cooker/CookRequestCluster.cpp b/Engine/Source/Editor/UnrealEd/Private/Cooker/CookRequestCluster.cpp
--- a/Engine/Source/Editor/UnrealEd/Private/Cooker/CookRequestCluster.cpp
+++ b/Engine/Source/Editor/UnrealEd/Private/Cooker/CookRequestCluster.cpp
@@ -28,25 +28,27 @@
 namespace UE::Cook
 {
 
-FRequestCluster::FRequestCluster(UCookOnTheFlyServer& COTFS, TArray<const ITargetPlatform*>&& InPlatforms)
+FRequestCluster::FRequestCluster(UCookOnTheFlyServer& InCOTFS, TArray<const ITargetPlatform*>&& InPlatforms)
 	: Platforms(MoveTemp(InPlatforms))
-	, PackageDatas(*COTFS.PackageDatas)
+	, COTFS(InCOTFS)
+	, PackageDatas(*InCOTFS.PackageDatas)
 	, AssetRegistry(*IAssetRegistry::Get())
-	, PackageTracker(*COTFS.PackageTracker)
-	, BuildDefinitions(*COTFS.BuildDefinitions)
+	, PackageTracker(*InCOTFS.PackageTracker)
+	, BuildDefinitions(*InCOTFS.BuildDefinitions)
 {
 }
 
-FRequestCluster::FRequestCluster(UCookOnTheFlyServer& COTFS, TConstArrayView<const ITargetPlatform*> InPlatforms)
+FRequestCluster::FRequestCluster(UCookOnTheFlyServer& InCOTFS, TConstArrayView<const ITargetPlatform*> InPlatforms)
 	: Platforms(InPlatforms)
-	, PackageDatas(*COTFS.PackageDatas)
+	, COTFS(InCOTFS)
+	, PackageDatas(*InCOTFS.PackageDatas)
 	, AssetRegistry(*IAssetRegistry::Get())
-	, PackageTracker(*COTFS.PackageTracker)
-	, BuildDefinitions(*COTFS.BuildDefinitions)
+	, PackageTracker(*InCOTFS.PackageTracker)
+	, BuildDefinitions(*InCOTFS.BuildDefinitions)
 {
 }
 
-void FRequestCluster::AddClusters(UCookOnTheFlyServer& COTFS, TArray<FFilePlatformRequest>&& InRequests,
+void FRequestCluster::AddClusters(UCookOnTheFlyServer& InCOTFS, TArray<FFilePlatformRequest>&& InRequests,
 	bool bRequestsAreUrgent, TRingBuffer<FRequestCluster>& OutClusters)
 {
 	if (InRequests.Num() == 0)
@@ -56,7 +58,7 @@ void FRequestCluster::AddClusters(UCookOnTheFlyServer& COTFS, TArray<FFilePlatfo
 
 	TArray<FRequestCluster, TInlineAllocator<1>> AddedClusters;
 	auto FindOrAddClusterForPlatforms =
-		[&AddedClusters, &COTFS](TArray<const ITargetPlatform*>&& InPlatforms)
+		[&AddedClusters, &InCOTFS](TArray<const ITargetPlatform*>&& InPlatforms)
 	{
 		for (FRequestCluster& Existing : AddedClusters)
 		{
@@ -66,7 +68,7 @@ void FRequestCluster::AddClusters(UCookOnTheFlyServer& COTFS, TArray<FFilePlatfo
 				return &Existing;
 			}
 		}
-		return &AddedClusters.Emplace_GetRef(COTFS, MoveTemp(InPlatforms));
+		return &AddedClusters.Emplace_GetRef(InCOTFS, MoveTemp(InPlatforms));
 	};
 
 	UE::Cook::FRequestCluster* MRUCluster = FindOrAddClusterForPlatforms(MoveTemp(InRequests[0].GetPlatforms()));
@@ -92,14 +94,14 @@ void FRequestCluster::AddClusters(UCookOnTheFlyServer& COTFS, TArray<FFilePlatfo
 
 	for (UE::Cook::FRequestCluster& AddedCluster : AddedClusters)
 	{
-		AddedCluster.Initialize(COTFS);
+		AddedCluster.Initialize(InCOTFS);
 		OutClusters.Add(MoveTemp(AddedCluster));
 	}
 }
 
 FName GInstigatorRequestCluster(TEXT("RequestCluster"));
 
-void FRequestCluster::AddClusters(UCookOnTheFlyServer& COTFS, FPackageDataSet& UnclusteredRequests,
+void FRequestCluster::AddClusters(UCookOnTheFlyServer& InCOTFS, FPackageDataSet& UnclusteredRequests,
 	TRingBuffer<FRequestCluster>& OutClusters, FRequestQueue& QueueForReadyRequests)
 {
 	if (UnclusteredRequests.Num() == 0)
@@ -109,11 +111,11 @@ void FRequestCluster::AddClusters(UCookOnTheFlyServer& COTFS, FPackageDataSet& U
 
 	TArray<FRequestCluster, TInlineAllocator<1>> AddedClusters;
 	TArray<const ITargetPlatform*, TInlineAllocator<ExpectedMaxNumPlatforms>> RequestedPlatforms;
-	auto FindOrAddClusterForPlatforms = [&AddedClusters, &COTFS, &RequestedPlatforms, &UnclusteredRequests]()
+	auto FindOrAddClusterForPlatforms = [&AddedClusters, &InCOTFS, &RequestedPlatforms, &UnclusteredRequests]()
 	{
 		if (AddedClusters.Num() == 0)
 		{
-			FRequestCluster& Cluster = AddedClusters.Emplace_GetRef(COTFS, RequestedPlatforms);
+			FRequestCluster& Cluster = AddedClusters.Emplace_GetRef(InCOTFS, RequestedPlatforms);
 			// The usual case is all platforms are the same, so reserve the first Cluster's size assuming it will get all requests
 			Cluster.Requests.Reserve(UnclusteredRequests.Num());
 			return &Cluster;
@@ -125,18 +127,18 @@ void FRequestCluster::AddClusters(UCookOnTheFlyServer& COTFS, FPackageDataSet& U
 				return &Existing;
 			}
 		}
-		return &AddedClusters.Emplace_GetRef(COTFS, RequestedPlatforms);
+		return &AddedClusters.Emplace_GetRef(InCOTFS, RequestedPlatforms);
 	};
 
 	bool bErrorOnEngineContentUse = false;
 	bool bAllowUncookedAssetReferences = false;
-	UE::Cook::FCookByTheBookOptions& Options = *COTFS.CookByTheBookOptions;
+	UE::Cook::FCookByTheBookOptions& Options = *InCOTFS.CookByTheBookOptions;
 	FString DLCPath;
 	bErrorOnEngineContentUse = Options.bErrorOnEngineContentUse;
 	bAllowUncookedAssetReferences = Options.bAllowUncookedAssetReferences;
 	if (bErrorOnEngineContentUse)
 	{
-		DLCPath = FPaths::Combine(*COTFS.GetBaseDirectoryForDLC(), TEXT("Content"));
+		DLCPath = FPaths::Combine(*InCOTFS.GetBaseDirectoryForDLC(), TEXT("Content"));
 		FPaths::MakeStandardFilename(DLCPath);
 	}
 
@@ -154,10 +156,11 @@ void FRequestCluster::AddClusters(UCookOnTheFlyServer& COTFS, FPackageDataSet& U
 		// by checking non-cookable and sending the package to idle instead of adding to
 		// a cluster
 		PackageData->GetRequestedPlatforms(RequestedPlatforms);
-		if (!IsRequestCookable(PackageData->GetPackageName(), PackageData, *COTFS.PackageDatas, *COTFS.PackageTracker,
-			DLCPath, bErrorOnEngineContentUse, bAllowUncookedAssetReferences, RequestedPlatforms))
+		ESuppressCookReason SuppressCookReason;
+		if (!IsRequestCookable(PackageData->GetPackageName(), PackageData, *InCOTFS.PackageDatas, *InCOTFS.PackageTracker,
+			DLCPath, bErrorOnEngineContentUse, bAllowUncookedAssetReferences, RequestedPlatforms, SuppressCookReason))
 		{
-			PackageData->SendToState(EPackageState::Idle, ESendFlags::QueueAdd);
+			InCOTFS.DemoteToIdle(*PackageData, ESendFlags::QueueAdd, SuppressCookReason);
 			continue;
 		}
 
@@ -173,7 +176,7 @@ void FRequestCluster::AddClusters(UCookOnTheFlyServer& COTFS, FPackageDataSet& U
 
 	for (UE::Cook::FRequestCluster& AddedCluster : AddedClusters)
 	{
-		AddedCluster.Initialize(COTFS);
+		AddedCluster.Initialize(InCOTFS);
 		OutClusters.Add(MoveTemp(AddedCluster));
 	}
 }
@@ -186,7 +189,7 @@ FRequestCluster::FFileNameRequest::FFileNameRequest(FFilePlatformRequest&& FileR
 {
 }
 
-void FRequestCluster::Initialize(UCookOnTheFlyServer& COTFS)
+void FRequestCluster::Initialize(UCookOnTheFlyServer& InCOTFS)
 {
 	if (!COTFS.IsCookOnTheFlyMode())
 	{
@@ -338,6 +341,7 @@ bool FRequestCluster::TryTakeOwnership(FPackageData& PackageData, bool bUrgent,
 {
 	if (!PackageData.IsInProgress())
 	{
+		check(!COTFS.IsCookWorkerMode()); // CookWorkers skip dependency traversal and should only call TryTakeOwnership on incoming requests, which are already in progress
 		check(GetPlatforms().Num() != 0); // This is required for SetRequestData
 		if (PackageData.HasAllExploredPlatforms(GetPlatforms()))
 		{
@@ -434,7 +438,7 @@ void FRequestCluster::RemovePackageData(FPackageData* PackageData)
 		}
 	}
 	Requests.Remove(PackageData);
-	RequestsToDemote.Remove(PackageData);
+	RequestsToDemote.RemoveAll([PackageData](const TPair<FPackageData*, ESuppressCookReason>& Pair) { return Pair.Key == PackageData; });
 }
 
 bool FRequestCluster::Contains(FPackageData* PackageData) const
@@ -442,7 +446,7 @@ bool FRequestCluster::Contains(FPackageData* PackageData) const
 	return OwnedPackageDatas.Contains(PackageData);
 }
 
-void FRequestCluster::ClearAndDetachOwnedPackageDatas(TArray<FPackageData*>& OutRequestsToLoad, TArray<FPackageData*>& OutRequestsToDemote)
+void FRequestCluster::ClearAndDetachOwnedPackageDatas(TArray<FPackageData*>& OutRequestsToLoad, TArray<TPair<FPackageData*, ESuppressCookReason>>& OutRequestsToDemote)
 {
 	if (bStartAsyncComplete)
 	{
@@ -471,6 +475,12 @@ void FRequestCluster::FetchDependencies(const FCookerTimer& CookerTimer, bool& b
 		return;
 	}
 
+	if (COTFS.IsCookWorkerMode())
+	{
+		bDependenciesComplete = true;
+		return;
+	}
+
 	if (!GraphSearch)
 	{
 		GraphSearch.Reset(new FGraphSearch(*this));
@@ -820,7 +830,8 @@ void FRequestCluster::FGraphSearch::VisitVertex(const FVertexData& VertexData)
 	}
 	else
 	{
-		Cluster.RequestsToDemote.Add(PackageData);
+		ESuppressCookReason SuppressCookReason = VertexData.bCookable ? ESuppressCookReason::AlreadyCooked : VertexData.SuppressCookReason;
+		Cluster.RequestsToDemote.Emplace(PackageData, SuppressCookReason);
 	}
 }
 
@@ -835,7 +846,7 @@ FRequestCluster::FVertexData::FVertexData(ESkipDependenciesType, FPackageData& I
 	PackageName = InPackageData.GetPackageName();
 	PackageData = &InPackageData;
 	bInitialRequest = true;
-	bCookable = Cluster.IsRequestCookable(PackageData->GetPackageName(), PackageData);
+	bCookable = Cluster.IsRequestCookable(PackageData->GetPackageName(), PackageData, SuppressCookReason);
 	bExploreDependencies = false;
 }
 
@@ -850,6 +861,7 @@ void FRequestCluster::FVertexData::Reset()
 	bInitialRequest = false;
 	bCookable = false;
 	bExploreDependencies = false;
+	SuppressCookReason = ESuppressCookReason::InvalidSuppressCookReason;
 }
 
 TUniquePtr<FRequestCluster::FVertexData> FRequestCluster::FGraphSearch::AllocateVertex()
@@ -882,7 +894,8 @@ FPackageData* FRequestCluster::FGraphSearch::FindOrAddVertex(FName PackageName,
 	}
 	VisitStatus.bVisited = true;
 
-	bool bCookable = Cluster.IsRequestCookable(PackageName, PackageData);
+	ESuppressCookReason SuppressCookReason;
+	bool bCookable = Cluster.IsRequestCookable(PackageName, PackageData, SuppressCookReason);
 	if (!bInitialRequest)
 	{
 		if (!bCookable)
@@ -903,6 +916,7 @@ FPackageData* FRequestCluster::FGraphSearch::FindOrAddVertex(FName PackageName,
 	OutNewVertex->PackageData = PackageData;
 	OutNewVertex->bInitialRequest = bInitialRequest;
 	OutNewVertex->bCookable = bCookable;
+	OutNewVertex->SuppressCookReason = SuppressCookReason;
 	OutNewVertex->bExploreDependencies = bCookable;
 
 	return PackageData;
@@ -1093,15 +1107,17 @@ TMap<FPackageData*, TArray<FPackageData*>>& FRequestCluster::FGraphSearch::GetGr
 	return GraphEdges;
 }
 
-bool FRequestCluster::IsRequestCookable(FName PackageName, FPackageData*& InOutPackageData)
+bool FRequestCluster::IsRequestCookable(FName PackageName, FPackageData*& InOutPackageData,
+	ESuppressCookReason& OutReason)
 {
 	return IsRequestCookable(PackageName, InOutPackageData, PackageDatas, PackageTracker,
-		DLCPath, bErrorOnEngineContentUse, bAllowUncookedAssetReferences, GetPlatforms());
+		DLCPath, bErrorOnEngineContentUse, bAllowUncookedAssetReferences, GetPlatforms(), OutReason);
 }
 
 bool FRequestCluster::IsRequestCookable(FName PackageName, FPackageData*& InOutPackageData,
 	FPackageDatas& InPackageDatas, FPackageTracker& InPackageTracker,
-	FStringView InDLCPath, bool bInErrorOnEngineContentUse, bool bInAllowUncookedAssetReferences, TConstArrayView<const ITargetPlatform*> RequestPlatforms)
+	FStringView InDLCPath, bool bInErrorOnEngineContentUse, bool bInAllowUncookedAssetReferences,
+	TConstArrayView<const ITargetPlatform*> RequestPlatforms, ESuppressCookReason& OutReason)
 {
 	TStringBuilder<256> NameBuffer;
 	// We need to reject packagenames from adding themselves or their transitive dependencies using all the same rules that
@@ -1109,6 +1125,7 @@ bool FRequestCluster::IsRequestCookable(FName PackageName, FPackageData*& InOutP
 	PackageName.ToString(NameBuffer);
 	if (FPackageName::IsScriptPackage(NameBuffer))
 	{
+		OutReason = ESuppressCookReason::ScriptPackage;
 		return false;
 	}
 
@@ -1118,6 +1135,7 @@ bool FRequestCluster::IsRequestCookable(FName PackageName, FPackageData*& InOutP
 		if (!InOutPackageData)
 		{
 			// Package does not exist on disk
+			OutReason = ESuppressCookReason::DoesNotExistInWorkspaceDomain;
 			return false;
 		}
 	}
@@ -1126,6 +1144,7 @@ bool FRequestCluster::IsRequestCookable(FName PackageName, FPackageData*& InOutP
 	if (InPackageTracker.NeverCookPackageList.Contains(FileName))
 	{
 		UE_LOG(LogCook, Verbose, TEXT("Package %s is referenced but is in the never cook package list, discarding request"), *NameBuffer);
+		OutReason = ESuppressCookReason::NeverCook;
 		return false;
 	}
 
@@ -1144,10 +1163,12 @@ bool FRequestCluster::IsRequestCookable(FName PackageName, FPackageData*& InOutP
 					UE_LOG(LogCook, Error, TEXT("Uncooked Engine or Game content %s is being referenced by DLC!"), *NameBuffer);
 				}
 			}
+			OutReason = ESuppressCookReason::NotInCurrentPlugin;
 			return false;
 		}
 	}
 
+	OutReason = ESuppressCookReason::InvalidSuppressCookReason;
 	return true;
 }
 
diff --git a/Engine/Source/Editor/UnrealEd/Private/Cooker/CookRequestCluster.h b/Engine/Source/Editor/UnrealEd/Private/Cooker/CookRequestCluster.h
--- a/Engine/Source/Editor/UnrealEd/Private/Cooker/CookRequestCluster.h
+++ b/Engine/Source/Editor/UnrealEd/Private/Cooker/CookRequestCluster.h
@@ -67,7 +67,7 @@ public:
 	 * If called before Process sets bOutComplete=true, all packages are put in OutRequestToLoad and are unsorted.
 	 */
 	void ClearAndDetachOwnedPackageDatas(TArray<FPackageData*>& OutRequestsToLoad,
-		TArray<FPackageData*>& OutRequestsToDemote);
+		TArray<TPair<FPackageData*, ESuppressCookReason>>& OutRequestsToDemote);
 
 	/**
 	 * Create clusters(s) for all the given name or packagedata requests and append them to OutClusters.
@@ -130,6 +130,8 @@ private:
 		bool bCookable;
 		/** True if dependencies should be explored, based on other properties of the VertexData. */
 		bool bExploreDependencies;
+		/** The reason it is not cookable if !bCookable, or ESuppressCookReason::InvalidSuppressCookReason if bCookable. */
+		ESuppressCookReason SuppressCookReason;
 	};
 
 	/**
@@ -295,19 +297,21 @@ private:
 	void StartAsync(const FCookerTimer& CookerTimer, bool& bOutComplete);
 	bool TryTakeOwnership(FPackageData& PackageData, bool bUrgent, FCompletionCallback&& CompletionCallback,
 		const FInstigator& InInstigator);
-	bool IsRequestCookable(FName PackageName, FPackageData*& InOutPackageData);
+	bool IsRequestCookable(FName PackageName, FPackageData*& InOutPackageData, ESuppressCookReason& OutReason);
 	static bool IsRequestCookable(FName PackageName, FPackageData*& InOutPackageData,
 		FPackageDatas& InPackageDatas, FPackageTracker& InPackageTracker,
-		FStringView InDLCPath, bool bInErrorOnEngineContentUse, bool bInAllowUncookedAssetReferences, TConstArrayView<const ITargetPlatform*> RequestPlatforms);
+		FStringView InDLCPath, bool bInErrorOnEngineContentUse, bool bInAllowUncookedAssetReferences,
+		TConstArrayView<const ITargetPlatform*> RequestPlatforms, ESuppressCookReason& OutReason);
 
 	TArray<FFileNameRequest> InRequests;
 	TArray<FPackageData*> Requests;
-	TArray<FPackageData*> RequestsToDemote;
+	TArray<TPair<FPackageData*, ESuppressCookReason>> RequestsToDemote;
 	TArray<const ITargetPlatform*> Platforms;
 	TArray<ICookedPackageWriter*> PackageWriters;
 	FPackageDataSet OwnedPackageDatas;
 	FString DLCPath;
 	TUniquePtr<FGraphSearch> GraphSearch; // Needs to be dynamic-allocated because of large alignment
+	UCookOnTheFlyServer& COTFS;
 	FPackageDatas& PackageDatas;
 	IAssetRegistry& AssetRegistry;
 	FPackageTracker& PackageTracker;
diff --git a/Engine/Source/Editor/UnrealEd/Private/Cooker/CookTypes.cpp b/Engine/Source/Editor/UnrealEd/Private/Cooker/CookTypes.cpp
--- a/Engine/Source/Editor/UnrealEd/Private/Cooker/CookTypes.cpp
+++ b/Engine/Source/Editor/UnrealEd/Private/Cooker/CookTypes.cpp
@@ -6,6 +6,7 @@
 #include "HAL/PlatformTime.h"
 #include "Math/NumericLimits.h"
 #include "Misc/CommandLine.h"
+#include "Misc/Parse.h"
 
 namespace UE::Cook
 {
@@ -155,4 +156,10 @@ void FBuildDefinitions::Cancel()
 	PendingBuilds.Empty();
 }
 
+bool IsCookIgnoreTimeouts()
+{
+	static bool bIsIgnoreCookTimeouts = FParse::Param(FCommandLine::Get(), TEXT("CookIgnoreTimeouts"));
+	return bIsIgnoreCookTimeouts;
+}
+
 }
diff --git a/Engine/Source/Editor/UnrealEd/Private/Cooker/CookTypes.h b/Engine/Source/Editor/UnrealEd/Private/Cooker/CookTypes.h
--- a/Engine/Source/Editor/UnrealEd/Private/Cooker/CookTypes.h
+++ b/Engine/Source/Editor/UnrealEd/Private/Cooker/CookTypes.h
@@ -11,6 +11,7 @@
 #include "HAL/LowLevelMemTracker.h"
 #include "HAL/Platform.h"
 #include "Logging/TokenizedMessage.h"
+#include "Misc/AssertionMacros.h"
 #include "Serialization/PackageWriter.h"
 #include "Templates/Function.h"
 #include "UObject/NameTypes.h"
@@ -113,6 +114,22 @@ namespace UE::Cook
 	};
 	const TCHAR* LexToString(UE::Cook::EReleaseSaveReason Reason);
 
+	enum class ESuppressCookReason : uint8
+	{
+		InvalidSuppressCookReason,
+		AlreadyCooked,
+		NeverCook,
+		DoesNotExistInWorkspaceDomain,
+		ScriptPackage,
+		NotInCurrentPlugin,
+		Redirected,
+		OrphanedGenerated,
+		LoadError,
+		SaveError,
+		OnlyEditorOnly,
+		CookCanceled,
+	};
+
 	/** The type of callback for External Requests that needs to be executed within the Scheduler's lock. */
 	typedef TUniqueFunction<void()> FSchedulerCallback;
 
@@ -120,7 +137,8 @@ namespace UE::Cook
 	enum class EPackageState
 	{
 		Idle = 0,	  /* The Package is not being operated on by the cooker, and is not in any queues.  This is the state both for packages that have never been requested and for packages that have finished cooking. */
-		Request,	  /* The Package is in the RequestQueue; it is known to the cooker but has not had any operations performed on it. */
+		Request,	  /* The Package is in the RequestQueue; it is requested for cooking but has not had any operations performed on it. */
+		AssignedToWorker, /* The Package is in the AssignedToWorkerSet; it has been sent a remote CookWorker for cooking and has not had any operations performed on it locally. */
 		LoadPrepare,  /* The Package is in the LoadPrepareQueue. Preloading is in progress. */
 		LoadReady,	  /* The package is in the LoadReadyQueue. Preloading is complete and it will be loaded when its turn comes up. */
 		Save,		  /* The Package is in the SaveQueue; it has been fully loaded and some target data may have been calculated. */
@@ -268,6 +286,40 @@ namespace UE::Cook
 		TArray<FName> NeverCookPackageList;
 		TFastPointerMap<const ITargetPlatform*, TSet<FName>> PlatformSpecificNeverCookPackages;
 	};
+
+	/**
+	 *  Identifier for a CookWorker process launched from a Director process, or for the local process.
+	 *  A director can have multiple CookWorkers.
+	 */
+	struct FWorkerId
+	{
+	public:
+		FWorkerId() { Id = InvalidId; }
+		constexpr static FWorkerId Invalid() { return FWorkerId(InvalidId); }
+		constexpr static FWorkerId Local() { return FWorkerId(LocalId); }
+		static FWorkerId FromRemoteIndex(uint8 Index) { check(Index < InvalidId-1U);  return FWorkerId(Index + 1U); }
+		static FWorkerId FromLocalOrRemoteIndex(uint8 Index) { check(Index < InvalidId);  return FWorkerId(Index); }
+
+		bool IsValid() const { return Id != InvalidId; }
+		bool IsInvalid() const { return Id == InvalidId; }
+		bool IsLocal() const { return Id == LocalId; }
+		bool IsRemote() const { return Id != InvalidId && Id != LocalId; }
+		uint8 GetRemoteIndex() const { check(IsRemote()); return Id - 1U; }
+		uint8 GetLocalOrRemoteIndex() const { check(IsValid()); return Id; }
+		bool operator<(const FWorkerId& Other) const { return Id < Other.Id; }
+
+	private:
+		constexpr explicit FWorkerId(uint8 InId) : Id(InId) {}
+
+	private:
+		uint8 Id;
+
+		constexpr static uint8 InvalidId = 255;
+		constexpr static uint8 LocalId = 0;
+	};
+
+	/** Report whether commandline/config has disabled use of timeouts throughout the cooker, useful for debugging. */
+	bool IsCookIgnoreTimeouts();
 }
 
 inline void RouteBeginCacheForCookedPlatformData(UObject* Obj, const ITargetPlatform* TargetPlatform)
@@ -394,3 +446,4 @@ LLM_DECLARE_TAG(Cooker);
 
 constexpr uint32 ExpectedMaxNumPlatforms = 32;
 #define REMAPPED_PLUGINS TEXT("RemappedPlugins")
+extern float GCookProgressWarnBusyTime;
diff --git a/Engine/Source/Editor/UnrealEd/Private/Cooker/CookWorkerClient.cpp b/Engine/Source/Editor/UnrealEd/Private/Cooker/CookWorkerClient.cpp
new file mode 100644
--- /dev/null
+++ b/Engine/Source/Editor/UnrealEd/Private/Cooker/CookWorkerClient.cpp
@@ -0,0 +1,173 @@
+// Copyright Epic Games, Inc. All Rights Reserved.
+
+#include "CookWorkerClient.h"
+
+#include "CompactBinaryTCP.h"
+#include "CookDirector.h"
+#include "CookTypes.h"
+#include "CoreGlobals.h"
+#include "HAL/PlatformTime.h"
+#include "Serialization/CompactBinaryWriter.h"
+#include "Sockets.h"
+#include "SocketSubsystem.h"
+#include "WorkerRequestsRemote.h"
+
+namespace UE::Cook
+{
+
+FCookWorkerClient::~FCookWorkerClient()
+{
+	Sockets::CloseSocket(ServerSocket);
+}
+
+bool FCookWorkerClient::TryConnect(FDirectorConnectionInfo&& ConnectInfo)
+{
+	EPollStatus Status;
+	for (;;)
+	{
+		Status = PollTryConnect(ConnectInfo);
+		if (Status != EPollStatus::Incomplete)
+		{
+			break;
+		}
+		constexpr float SleepTime = 0.01f; // 10 ms
+		FPlatformProcess::Sleep(SleepTime);
+	}
+	return Status == EPollStatus::Success;
+}
+
+EPollStatus FCookWorkerClient::PollTryConnect(const FDirectorConnectionInfo& ConnectInfo)
+{
+	for (;;)
+	{
+		switch (ConnectStatus)
+		{
+		case EConnectStatus::Connected:
+			return EPollStatus::Success;
+		case EConnectStatus::Uninitialized:
+			CreateServerSocket(ConnectInfo);
+			break;
+		case EConnectStatus::PollWriteConnectMessage:
+			PollWriteConnectMessage();
+			if (ConnectStatus == EConnectStatus::PollWriteConnectMessage)
+			{
+				return EPollStatus::Incomplete;
+			}
+			break;
+		case EConnectStatus::PollReceiveConfigMessage:
+			PollReceiveConfigMessage();
+			if (ConnectStatus == EConnectStatus::PollReceiveConfigMessage)
+			{
+				return EPollStatus::Incomplete;
+			}
+			break;
+		case EConnectStatus::Failed:
+			return EPollStatus::Error;
+		default:
+			checkNoEntry();
+			return EPollStatus::Error;
+		}
+	}
+}
+
+void FCookWorkerClient::CreateServerSocket(const FDirectorConnectionInfo& ConnectInfo)
+{
+	using namespace CompactBinaryTCP;
+
+	ConnectStartTimeSeconds = FPlatformTime::Seconds();
+	DirectorURI = ConnectInfo.HostURI;
+
+	ISocketSubsystem* SocketSubsystem = ISocketSubsystem::Get();
+	if (!SocketSubsystem)
+	{
+		UE_LOG(LogCook, Error, TEXT("CookWorker initialization failure: platform does not support network sockets, cannot connect to CookDirector."));
+		ConnectStatus = EConnectStatus::Failed;
+		return;
+	}
+
+	DirectorAddr = Sockets::GetAddressFromStringWithPort(DirectorURI);
+	if (!DirectorAddr)
+	{
+		UE_LOG(LogCook, Error, TEXT("CookWorker initialization failure: could not convert -CookDirectorHost=%s into an address, cannot connect to CookDirector."),
+			*DirectorURI);
+		ConnectStatus = EConnectStatus::Failed;
+		return;
+	}
+
+	UE_LOG(LogCook, Display, TEXT("Connecting to CookDirector at %s..."), *DirectorURI);
+
+	ServerSocket = Sockets::ConnectToHost(*DirectorAddr, TEXT("FCookWorkerClient-WorkerConnect"));
+	if (!ServerSocket)
+	{
+		UE_LOG(LogCook, Error, TEXT("CookWorker initialization failure: Could not connect to CookDirector."));
+		ConnectStatus = EConnectStatus::Failed;
+		return;
+	}
+
+	constexpr float WaitForConnectTimeout = 60.f * 10;
+	float ConditionalTimeoutSeconds = IsCookIgnoreTimeouts() ? MAX_flt : WaitForConnectTimeout;
+	bool bServerSocketReady = ServerSocket->Wait(ESocketWaitConditions::WaitForWrite, FTimespan::FromSeconds(ConditionalTimeoutSeconds));
+	if (!bServerSocketReady)
+	{
+		UE_LOG(LogCook, Error, TEXT("CookWorker initialization failure: Timed out after %.0f seconds trying to connect to CookDirector."),
+			ConditionalTimeoutSeconds);
+		ConnectStatus = EConnectStatus::Failed;
+		return;
+	}
+
+	FWorkerConnectMessage ConnectMessage;
+	ConnectMessage.RemoteIndex = ConnectInfo.RemoteIndex;
+	EConnectionStatus Status = TryWritePacket(ServerSocket, SendBuffer, ConnectMessage);
+	if (Status == EConnectionStatus::Incomplete)
+	{
+		ConnectStatus = EConnectStatus::PollWriteConnectMessage;
+		return;
+	}
+	else if (Status != EConnectionStatus::Okay)
+	{
+		UE_LOG(LogCook, Error, TEXT("CookWorker initialization failure: could not send ConnectMessage."));
+		ConnectStatus = EConnectStatus::Failed;
+		return;
+	}
+	LogConnected();
+
+	ConnectStatus = EConnectStatus::PollReceiveConfigMessage;
+}
+
+void FCookWorkerClient::PollWriteConnectMessage()
+{
+	using namespace CompactBinaryTCP;
+	constexpr float WaitForConnectTimeout = 60.f;
+
+	EConnectionStatus Status = TryFlushBuffer(ServerSocket, SendBuffer);
+	if (Status == EConnectionStatus::Incomplete)
+	{
+		if (FPlatformTime::Seconds() - ConnectStartTimeSeconds > WaitForConnectTimeout && !IsCookIgnoreTimeouts())
+		{
+			UE_LOG(LogCook, Error, TEXT("CookWorker initialization failure: timed out waiting for %fs to send ConnectMessage."), WaitForConnectTimeout);
+			ConnectStatus = EConnectStatus::Failed;
+		}
+		return;
+	}
+	else if (Status != EConnectionStatus::Okay)
+	{
+		UE_LOG(LogCook, Error, TEXT("CookWorker initialization failure: could not send ConnectMessage."));
+		ConnectStatus = EConnectStatus::Failed;
+		return;
+	}
+	LogConnected();
+	ConnectStatus = EConnectStatus::PollReceiveConfigMessage;
+}
+
+void FCookWorkerClient::PollReceiveConfigMessage()
+{
+	UE_LOG(LogCook, Error, TEXT("CookWorker initialization failure: reading Settings information is not yet implemented."));
+	ConnectStatus = EConnectStatus::Failed;
+}
+
+void FCookWorkerClient::LogConnected()
+{
+	UE_LOG(LogCook, Display, TEXT("Connection to CookDirector successful."));
+}
+
+}
\ No newline at end of file
diff --git a/Engine/Source/Editor/UnrealEd/Private/Cooker/CookWorkerClient.h b/Engine/Source/Editor/UnrealEd/Private/Cooker/CookWorkerClient.h
new file mode 100644
--- /dev/null
+++ b/Engine/Source/Editor/UnrealEd/Private/Cooker/CookWorkerClient.h
@@ -0,0 +1,52 @@
+// Copyright Epic Games, Inc. All Rights Reserved.
+
+#pragma once
+
+#include "CompactBinaryTCP.h"
+#include "CookTypes.h"
+#include "IPAddress.h"
+
+namespace UE::Cook { struct FDirectorConnectionInfo; }
+
+namespace UE::Cook
+{
+
+/** Class in a CookWorker process that communicates over a Socket with FCookWorkerServer in a Director process. */
+class FCookWorkerClient
+{
+public:
+	~FCookWorkerClient();
+	/** Blocking operation: open the socket to the Director, send the Connected message, receive the setup message. */
+	bool TryConnect(FDirectorConnectionInfo&& ConnectInfo);
+
+private:
+	enum class EConnectStatus
+	{
+		Uninitialized,
+		PollWriteConnectMessage,
+		PollReceiveConfigMessage,
+		Failed,
+		Connected,
+	};
+
+private:
+	/** Reentrant helper for TryConnect which early exits if currently blocked. */
+	EPollStatus PollTryConnect(const FDirectorConnectionInfo& ConnectInfo);
+	/** Helper for PollTryConnect: create the ServerSocket */
+	void CreateServerSocket(const FDirectorConnectionInfo& ConnectInfo);
+	/** Try to send the Connect Message, switch state when it succeeds or fails. */
+	void PollWriteConnectMessage();
+	/** Wait for the Config Message, switch state when it succeeds or fails. */
+	void PollReceiveConfigMessage();
+	void LogConnected();
+
+private:
+	TSharedPtr<FInternetAddr> DirectorAddr;
+	UE::CompactBinaryTCP::FSendBuffer SendBuffer;
+	FString DirectorURI;
+	FSocket* ServerSocket = nullptr;
+	double ConnectStartTimeSeconds = 0.;
+	EConnectStatus ConnectStatus = EConnectStatus::Uninitialized;
+};
+
+}
\ No newline at end of file
diff --git a/Engine/Source/Editor/UnrealEd/Private/Cooker/CookWorkerServer.cpp b/Engine/Source/Editor/UnrealEd/Private/Cooker/CookWorkerServer.cpp
new file mode 100644
--- /dev/null
+++ b/Engine/Source/Editor/UnrealEd/Private/Cooker/CookWorkerServer.cpp
@@ -0,0 +1,406 @@
+// Copyright Epic Games, Inc. All Rights Reserved.
+
+#include "CookWorkerServer.h"
+
+#include "CookDirector.h"
+#include "CookPackageData.h"
+#include "HAL/PlatformProcess.h"
+#include "Misc/AssertionMacros.h"
+#include "Serialization/CompactBinaryWriter.h"
+#include "UnrealEdMisc.h"
+
+namespace UE::Cook
+{
+
+FCookWorkerServer::FCookWorkerServer(FCookDirector& InDirector, FWorkerId InWorkerId)
+	: Director(InDirector)
+	, COTFS(InDirector.COTFS)
+	, WorkerId(InWorkerId)
+{
+}
+
+FCookWorkerServer::~FCookWorkerServer()
+{
+	PendingPackages.Append(PackagesToAssign);
+	PackagesToAssign.Empty();
+	for (FPackageData* PackageData : PendingPackages)
+	{
+		check(PackageData->IsInProgress()); // Packages that were assigned to a worker should be in the AssignedToWorker state
+		PackageData->SetWorkerAssignment(FWorkerId::Invalid());
+		PackageData->SendToState(UE::Cook::EPackageState::Request, ESendFlags::QueueAddAndRemove);
+	}
+
+	if (ConnectStatus == EConnectStatus::Connected || ConnectStatus == EConnectStatus::WaitForDisconnect)
+	{
+		UE_LOG(LogCook, Error, TEXT("CookWorkerServer %d was destroyed before it finished Disconnect. The remote process may linger and may interfere with writes of future packages."),
+			WorkerId.GetLocalOrRemoteIndex());
+	}
+	DetachFromRemoteProcess();
+}
+
+void FCookWorkerServer::DetachFromRemoteProcess()
+{
+	Sockets::CloseSocket(Socket);
+	CookWorkerHandle = FProcHandle();
+	CookWorkerProcessId = 0;
+	bTerminateImmediately = false;
+	SendBuffer.Reset();
+	ReceiveBuffer.Reset();
+}
+
+void FCookWorkerServer::ShutdownRemoteProcess()
+{
+	Sockets::CloseSocket(Socket);
+	if (CookWorkerHandle.IsValid())
+	{
+		FPlatformProcess::TerminateProc(CookWorkerHandle, /* bKillTree */true);
+	}
+	DetachFromRemoteProcess();
+}
+
+void FCookWorkerServer::AppendAssignments(TArrayView<FPackageData*> Assignments)
+{
+	PackagesToAssign.Append(Assignments);
+}
+
+void FCookWorkerServer::AbortAssignments(TSet<FPackageData*>& OutPendingPackages)
+{
+	if (ConnectStatus == EConnectStatus::Connected)
+	{
+		if (PendingPackages.Num())
+		{
+			TArray<FName> PackageNames;
+			PackageNames.Reserve(PendingPackages.Num());
+			for (FPackageData* PackageData : PendingPackages)
+			{
+				PackageNames.Add(PackageData->GetPackageName());
+			}
+			SendMessage(FAbortPackagesMessage(MoveTemp(PackageNames)));
+			OutPendingPackages.Append(MoveTemp(PendingPackages));
+			PendingPackages.Reset();
+		}
+	}
+	else
+	{
+		check(PendingPackages.IsEmpty());
+	}
+	OutPendingPackages.Append(PackagesToAssign);
+	PackagesToAssign.Empty();
+}
+
+void FCookWorkerServer::AbortAssignment(FPackageData& PackageData)
+{
+	if (ConnectStatus == EConnectStatus::Connected)
+	{
+		if (PendingPackages.Remove(&PackageData))
+		{
+			TArray<FName> PackageNames;
+			PackageNames.Add(PackageData.GetPackageName());
+			SendMessage(FAbortPackagesMessage(MoveTemp(PackageNames)));
+		}
+	}
+	else
+	{
+		check(PendingPackages.IsEmpty());
+	}
+
+	PackagesToAssign.Remove(&PackageData);
+}
+
+void FCookWorkerServer::AbortWorker(TSet<FPackageData*>& OutPendingPackages)
+{
+	AbortAssignments(OutPendingPackages);
+	if (ConnectStatus == EConnectStatus::Connected)
+	{
+		SendMessage(FAbortWorkerMessage());
+		SendToState(EConnectStatus::WaitForDisconnect);
+	}
+}
+
+void FCookWorkerServer::SendToState(EConnectStatus TargetStatus)
+{
+	switch (TargetStatus)
+	{
+	case EConnectStatus::WaitForConnect:
+		ConnectStartTimeSeconds = FPlatformTime::Seconds();
+		ConnectTestStartTimeSeconds = ConnectStartTimeSeconds;
+		break;
+	case EConnectStatus::WaitForDisconnect:
+		ConnectStartTimeSeconds = FPlatformTime::Seconds();
+		ConnectTestStartTimeSeconds = ConnectStartTimeSeconds;
+		break;
+	case EConnectStatus::LostConnection:
+		DetachFromRemoteProcess();
+		break;
+	default:
+		break;
+	}
+	ConnectStatus = TargetStatus;
+}
+
+bool FCookWorkerServer::IsShuttingDown() const
+{
+	return ConnectStatus == EConnectStatus::WaitForDisconnect || ConnectStatus == EConnectStatus::LostConnection;
+}
+
+bool FCookWorkerServer::IsShutdownComplete() const
+{
+	return ConnectStatus == EConnectStatus::LostConnection;
+}
+
+bool FCookWorkerServer::TryHandleConnectMessage(FWorkerConnectMessage& Message, FSocket* InSocket, TArray<UE::CompactBinaryTCP::FMarshalledMessage>&& OtherPacketMessages)
+{
+	if (ConnectStatus != EConnectStatus::WaitForConnect)
+	{
+		return false;
+	}
+	check(!Socket);
+	Socket = InSocket;
+
+	UE_LOG(LogCook, Error, TEXT("CookWorkerServer failure: writing Settings information is not yet implemented."));
+
+	SendToState(EConnectStatus::Connected);
+	HandleReceiveMessages(MoveTemp(OtherPacketMessages));
+	return true;
+}
+
+void FCookWorkerServer::TickFromSchedulerThread()
+{
+	PumpConnect();
+	if (ConnectStatus == EConnectStatus::Connected)
+	{
+		PumpReceiveMessages();
+		PumpSendMessages();
+	}
+}
+
+void FCookWorkerServer::PumpConnect()
+{
+	for (;;)
+	{
+		switch (ConnectStatus)
+		{
+		case EConnectStatus::Connected:
+			return; // Nothing further to do
+		case EConnectStatus::LostConnection:
+			return; // Nothing further to do
+		case EConnectStatus::Uninitialized:
+			LaunchProcess();
+			break;
+		case EConnectStatus::WaitForConnect:
+			TickWaitForConnect();
+			if (ConnectStatus == EConnectStatus::WaitForConnect)
+			{
+				return; // Try again later
+			}
+			break;
+		case EConnectStatus::WaitForDisconnect:
+			TickWaitForDisconnect();
+			if (ConnectStatus == EConnectStatus::WaitForDisconnect)
+			{
+				return; // Try again later
+			}
+			break;
+		default:
+			checkNoEntry();
+			return;
+		}
+	}
+}
+
+void FCookWorkerServer::LaunchProcess()
+{
+	FString CommandletExecutable = FUnrealEdMisc::Get().GetProjectEditorBinaryPath();
+	FString CommandLine = Director.GetWorkerCommandLine(WorkerId);
+	CookWorkerHandle = FPlatformProcess::CreateProc(*CommandletExecutable, *CommandLine,
+		true /* bLaunchDetached */, true /* bLaunchHidden */, true /* bLaunchReallyHidden */,
+		&CookWorkerProcessId, 0 /* PriorityModifier */, *FPaths::GetPath(CommandletExecutable),
+		nullptr /* PipeWriteChild */);
+	if (CookWorkerHandle.IsValid())
+	{
+		SendToState(EConnectStatus::WaitForConnect);
+	}
+	else
+	{
+		// GetLastError information was logged by CreateProc
+		UE_LOG(LogCook, Error, TEXT("CookWorkerServer %d failed to create CookWorker process. Assigned packages will be returned to the director."),
+			WorkerId.GetLocalOrRemoteIndex());
+		SendToState(EConnectStatus::LostConnection);
+	}
+}
+
+void FCookWorkerServer::TickWaitForConnect()
+{
+	constexpr float TestProcessExistencePeriod = 1.f;
+	constexpr float WaitForConnectTimeout = 60.f * 10;
+
+	check(!Socket); // When the Socket is assigned we leave the WaitForConnect state, and we set it to null before entering
+
+	double CurrentTime = FPlatformTime::Seconds();
+	if (CurrentTime - ConnectTestStartTimeSeconds > TestProcessExistencePeriod)
+	{
+		if (!FPlatformProcess::IsProcRunning(CookWorkerHandle))
+		{
+			UE_LOG(LogCook, Error, TEXT("CookWorkerServer %d process terminated before connecting. Assigned packages will be returned to the director."),
+				WorkerId.GetLocalOrRemoteIndex());
+			SendToState(EConnectStatus::LostConnection);
+			return;
+		}
+		ConnectTestStartTimeSeconds = FPlatformTime::Seconds();
+	}
+
+	if (CurrentTime - ConnectStartTimeSeconds > WaitForConnectTimeout && !IsCookIgnoreTimeouts())
+	{
+		UE_LOG(LogCook, Error, TEXT("CookWorkerServer %d process failed to connect within %.0f seconds. Assigned packages will be returned to the director."),
+			WorkerId.GetLocalOrRemoteIndex(), WaitForConnectTimeout);
+		ShutdownRemoteProcess();
+		SendToState(EConnectStatus::LostConnection);
+		return;
+	}
+}
+
+void FCookWorkerServer::TickWaitForDisconnect()
+{
+	constexpr float TestProcessExistencePeriod = 1.f;
+	constexpr float WaitForDisconnectTimeout = 60.f * 10;
+
+	double CurrentTime = FPlatformTime::Seconds();
+	if (CurrentTime - ConnectTestStartTimeSeconds > TestProcessExistencePeriod)
+	{
+		if (!FPlatformProcess::IsProcRunning(CookWorkerHandle))
+		{
+			SendToState(EConnectStatus::LostConnection);
+			return;
+		}
+		ConnectTestStartTimeSeconds = FPlatformTime::Seconds();
+	}
+
+	// We might have been blocked from sending the disconnect, so keep trying to flush the buffer
+	UE::CompactBinaryTCP::TryFlushBuffer(Socket, SendBuffer);
+	TArray<UE::CompactBinaryTCP::FMarshalledMessage> Messages;
+	TryReadPacket(Socket, ReceiveBuffer, Messages);
+	for (UE::CompactBinaryTCP::FMarshalledMessage& Message : Messages)
+	{
+		if (Message.MessageType == FAbortWorkerMessage::MessageType)
+		{
+			SendToState(EConnectStatus::LostConnection);
+			return;
+		}
+	}
+
+	if ((bTerminateImmediately || CurrentTime - ConnectStartTimeSeconds > WaitForDisconnectTimeout) && !IsCookIgnoreTimeouts())
+	{
+		UE_LOG(LogCook, Warning, TEXT("CookWorkerServer %d's remote process failed to disconnect within %.0f seconds; we will terminate process."),
+			WorkerId.GetLocalOrRemoteIndex(), WaitForDisconnectTimeout);
+		ShutdownRemoteProcess();
+		SendToState(EConnectStatus::LostConnection);
+	}
+}
+
+void FCookWorkerServer::PumpSendMessages()
+{
+	UE::CompactBinaryTCP::EConnectionStatus Status = UE::CompactBinaryTCP::TryFlushBuffer(Socket, SendBuffer);
+	if (Status == UE::CompactBinaryTCP::Failed)
+	{
+		UE_LOG(LogCook, Error, TEXT("CookWorkerServer %d failed to write to socket, we will shutdown the remote process. Assigned packages will be returned to the director."),
+			WorkerId.GetLocalOrRemoteIndex());
+		SendToState(EConnectStatus::WaitForDisconnect);
+		bTerminateImmediately = true;
+	}
+}
+
+void FCookWorkerServer::PumpReceiveMessages()
+{
+	using namespace UE::CompactBinaryTCP;
+	TArray<FMarshalledMessage> Messages;
+	TryReadPacket(Socket, ReceiveBuffer, Messages);
+	HandleReceiveMessages(MoveTemp(Messages));
+}
+
+void FCookWorkerServer::HandleReceiveMessages(TArray<UE::CompactBinaryTCP::FMarshalledMessage>&& Messages)
+{
+	for (UE::CompactBinaryTCP::FMarshalledMessage& Message : Messages)
+	{
+		if (Message.MessageType == FAbortWorkerMessage::MessageType)
+		{
+			UE_LOG(LogCook, Error, TEXT("CookWorkerServer %d remote process shut down unexpectedly. Assigned packages will be returned to the director."),
+				WorkerId.GetLocalOrRemoteIndex());
+			SendToState(EConnectStatus::LostConnection);
+			break;
+		}
+	}
+}
+
+void FCookWorkerServer::SendMessage(const UE::CompactBinaryTCP::IMessage& Message)
+{
+	UE::CompactBinaryTCP::TryWritePacket(Socket, SendBuffer, Message);
+}
+
+FAssignPackagesMessage::FAssignPackagesMessage(TArray<FName>&& InPackageNames)
+	: PackageNames(MoveTemp(InPackageNames))
+{
+}
+
+void FAssignPackagesMessage::Write(FCbWriter& Writer) const
+{
+	WriteArrayOfNames(Writer, "PackageNames", PackageNames);
+}
+
+bool FAssignPackagesMessage::TryRead(FCbObjectView Object)
+{
+	return TryReadArrayOfNames(Object, "PackageNames", PackageNames);
+}
+
+FGuid FAssignPackagesMessage::MessageType(TEXT("B7B1542B73254B679319D73F753DB6F8"));
+
+FAbortPackagesMessage::FAbortPackagesMessage(TArray<FName>&& InPackageNames)
+	: PackageNames(MoveTemp(InPackageNames))
+{
+}
+
+void FAbortPackagesMessage::Write(FCbWriter& Writer) const
+{
+	WriteArrayOfNames(Writer, "PackageNames", PackageNames);
+}
+
+bool FAbortPackagesMessage::TryRead(FCbObjectView Object)
+{
+	return TryReadArrayOfNames(Object, "PackageNames", PackageNames);
+}
+
+FGuid FAbortPackagesMessage::MessageType(TEXT("D769F1BFF2F34978868D70E3CAEE94E7"));
+
+FGuid FAbortWorkerMessage::MessageType(TEXT("83FD99DFE8DB4A9A8E71684C121BE6F3"));
+
+void WriteArrayOfNames(FCbWriter& Writer, const char* ArrayName, TConstArrayView<FName> Names)
+{
+	Writer.BeginArray(ArrayName);
+	for (FName Name : Names)
+	{
+		Writer.AddString(WriteToString<FName::StringBufferSize>(Name).ToView());
+	}
+	Writer.EndArray();
+}
+
+bool TryReadArrayOfNames(FCbObjectView Object, const char* ArrayName, TArray<FName>& OutNames)
+{
+	FCbFieldView ArrayField = Object["ArrayName"];
+	FCbArrayView ArrayView = ArrayField.AsArrayView();
+	if (ArrayField.HasError())
+	{
+		return false;
+	}
+	OutNames.Reserve(OutNames.Num() + ArrayView.Num());
+	for (FCbFieldView ElementView : ArrayView)
+	{
+		FUtf8StringView StringView = ElementView.AsString();
+		if (ElementView.HasError())
+		{
+			return false;
+		}
+		OutNames.Add(FName(StringView));
+	}
+	return true;
+}
+
+}
\ No newline at end of file
diff --git a/Engine/Source/Editor/UnrealEd/Private/Cooker/CookWorkerServer.h b/Engine/Source/Editor/UnrealEd/Private/Cooker/CookWorkerServer.h
new file mode 100644
--- /dev/null
+++ b/Engine/Source/Editor/UnrealEd/Private/Cooker/CookWorkerServer.h
@@ -0,0 +1,157 @@
+// Copyright Epic Games, Inc. All Rights Reserved.
+
+#pragma once
+
+#include "CompactBinaryTCP.h"
+#include "Containers/Array.h"
+#include "Containers/ArrayView.h"
+#include "Containers/Set.h"
+#include "CookSockets.h"
+#include "CookTypes.h"
+#include "Misc/Guid.h"
+
+class FSocket;
+class UCookOnTheFlyServer;
+struct FProcHandle;
+namespace UE::CompactBinaryTCP { struct FMarshalledMessage; }
+namespace UE::Cook { class FCookDirector; }
+namespace UE::Cook { struct FPackageData; }
+namespace UE::Cook { struct FWorkerConnectMessage; }
+
+namespace UE::Cook
+{
+
+/** Class in a Director process that communicates over a Socket with FCookWorkerClient in a CookWorker process. */
+class FCookWorkerServer
+{
+public:
+	FCookWorkerServer(FCookDirector& InDirector, FWorkerId InWorkerId);
+	~FCookWorkerServer();
+
+	FWorkerId GetWorkerId() const { return WorkerId; }
+
+	/** Add the given assignments for the CookWorker. They will be sent during Tick */
+	void AppendAssignments(TArrayView<FPackageData*> Assignments);
+	/** Remove assignment of the package from local state and from the connected Client. */
+	void AbortAssignment(FPackageData& PackageData);
+	/**
+	 * Remove assignment of the all assigned packages from local state and from the connected Client.
+	 * Report all packages that were unassigned.
+	 */
+	void AbortAssignments(TSet<FPackageData*>& OutPendingPackages);
+	/** AbortAssignments and tell the connected Client to gracefully terminate. Report all packages that were unassigned. */
+	void AbortWorker(TSet<FPackageData*>& OutPendingPackages);
+	/** Take over the Socket for a CookWorker that has just connected. */
+	bool TryHandleConnectMessage(FWorkerConnectMessage& Message, FSocket* InSocket, TArray<UE::CompactBinaryTCP::FMarshalledMessage>&& OtherPacketMessages);
+	/** Periodic Tick function to send and receive messages to the Client. */
+	void TickFromSchedulerThread();
+
+	/** Is this either shutting down or completed shutdown of its remote Client? */
+	bool IsShuttingDown() const;
+	/** Is this not yet or no longer connected to a remote Client? */
+	bool IsShutdownComplete() const;
+
+private:
+	enum class EConnectStatus
+	{
+		Uninitialized,
+		WaitForConnect,
+		Connected,
+		WaitForDisconnect,
+		LostConnection,
+	};
+
+private:
+	/** Helper for Tick, Pump messages when connecting or disconnecting. */
+	void PumpConnect();
+	/** Helper for PumpConnect, launch the remote Client process. */
+	void LaunchProcess();
+	/** Helper for PumpConnect, wait for connect message from Client, set state to LostConnection if we timeout. */
+	void TickWaitForConnect();
+	/** Helper for PumpConnect, wait for disconnect message from Client, set state to LostConnection if we timeout. */
+	void TickWaitForDisconnect();
+	/** Helper for Tick, pump send messages to a connected Client. */
+	void PumpSendMessages();
+	/** Helper for Tick, pump receive messages to a connected Client. */
+	void PumpReceiveMessages();
+	/** Send the message immediately to the Socket. If cannot complete immediately, it will be finished during Tick. */
+	void SendMessage(const UE::CompactBinaryTCP::IMessage& Message);
+	/** Send this into the given state. Update any state-dependent variables. */
+	void SendToState(EConnectStatus TargetStatus);
+	/** Close the connection and connection resources to the remote process. Does not kill the process. */
+	void DetachFromRemoteProcess();
+	/** Kill the Client process (non-graceful termination), and close the connection resources. */
+	void ShutdownRemoteProcess();
+	/** Helper for PumpReceiveMessages: dispatch the messages received from the socket. */
+	void HandleReceiveMessages(TArray<UE::CompactBinaryTCP::FMarshalledMessage>&& Messages);
+
+	TArray<FPackageData*> PackagesToAssign;
+	TSet<FPackageData*> PendingPackages;
+	UE::CompactBinaryTCP::FSendBuffer SendBuffer;
+	UE::CompactBinaryTCP::FReceiveBuffer ReceiveBuffer;
+	FCookDirector& Director;
+	UCookOnTheFlyServer& COTFS;
+	FSocket* Socket = nullptr;
+	FProcHandle CookWorkerHandle;
+	uint32 CookWorkerProcessId = 0;
+	double ConnectStartTimeSeconds = 0.;
+	double ConnectTestStartTimeSeconds = 0.;
+	FWorkerId WorkerId = FWorkerId::Invalid();
+	EConnectStatus ConnectStatus = EConnectStatus::Uninitialized;
+	bool bTerminateImmediately = false;
+};
+
+/** Message from Server to Client to cook the given packages. */
+struct FAssignPackagesMessage : public UE::CompactBinaryTCP::IMessage
+{
+public:
+	FAssignPackagesMessage() = default;
+	FAssignPackagesMessage(TArray<FName>&& InPackageNames);
+
+	virtual void Write(FCbWriter& Writer) const override;
+	virtual bool TryRead(FCbObjectView Object) override;
+	virtual FGuid GetMessageType() const override { return MessageType; }
+
+public:
+	TArray<FName> PackageNames;
+	static FGuid MessageType;
+};
+
+/** Message from Server to Client to cancel the cook of the given packages. */
+struct FAbortPackagesMessage : public UE::CompactBinaryTCP::IMessage
+{
+public:
+	FAbortPackagesMessage() = default;
+	FAbortPackagesMessage(TArray<FName>&& InPackageNames);
+
+	virtual void Write(FCbWriter& Writer) const override;
+	virtual bool TryRead(FCbObjectView Object) override;
+	virtual FGuid GetMessageType() const override { return MessageType; }
+
+public:
+	TArray<FName> PackageNames;
+	static FGuid MessageType;
+};
+
+/**
+ * Message from either Server to Client.
+ * If from Server, request that Client shutdown.
+ * If from Client, notify Server it is shutting down.
+ */
+struct FAbortWorkerMessage : public UE::CompactBinaryTCP::IMessage
+{
+public:
+	virtual void Write(FCbWriter& Writer) const override { }
+	virtual bool TryRead(FCbObjectView Object) override { return true; }
+	virtual FGuid GetMessageType() const override { return MessageType; }
+
+public:
+	static FGuid MessageType;
+};
+
+/** IMessage helper: write as strings to compact binary. */
+void WriteArrayOfNames(FCbWriter& Writer, const char* ArrayName, TConstArrayView<FName> Names);
+/** IMessage helper: read FNames as strings from compact binary. */
+bool TryReadArrayOfNames(FCbObjectView Object, const char* ArrayName, TArray<FName>& OutNames);
+
+}
\ No newline at end of file
diff --git a/Engine/Source/Editor/UnrealEd/Private/Cooker/IWorkerRequests.h b/Engine/Source/Editor/UnrealEd/Private/Cooker/IWorkerRequests.h
--- a/Engine/Source/Editor/UnrealEd/Private/Cooker/IWorkerRequests.h
+++ b/Engine/Source/Editor/UnrealEd/Private/Cooker/IWorkerRequests.h
@@ -49,10 +49,7 @@ public:
 	virtual void DequeueAllExternal(TArray<FSchedulerCallback>& OutCallbacks, TArray<FFilePlatformRequest>& OutCookRequests) = 0;
 
 	// Writing Packages during Cook
-	virtual void AddDiscoveredPackage(FPackageData& PackageData, FInstigator& Instigator, bool bLoadReady, bool& bOutShouldAddToQueue)
-	{
-		bOutShouldAddToQueue = true;
-	}
+	virtual void AddDiscoveredPackage(FPackageData& PackageData, FInstigator& Instigator, bool bLoadReady, bool& bOutShouldAddToQueue) = 0;
 
 	// Writing Packages from CookByTheBook
 	virtual void AddStartCookByTheBookRequest(FFilePlatformRequest&& Request) = 0;
@@ -75,6 +72,8 @@ public:
 
 	// Functions that are written locally for Directors but are sent to Director for Workers
 	virtual void ReportAccessedIniSettings(UCookOnTheFlyServer& COTFS, const FConfigFile& Config) = 0;
+	virtual void ReportDemoteToIdle(UE::Cook::FPackageData& PackageData, ESuppressCookReason Reason) = 0;
+	virtual void ReportPromoteToSaveComplete(UE::Cook::FPackageData& PackageData) = 0;
 
 	// Functions that read from local for Directors but read from Director for Workers
 	virtual void GetInitializeConfigSettings(UCookOnTheFlyServer& COTFS, const FString& OutputDirectoryOverride, UE::Cook::FInitializeConfigSettings& Settings) = 0;
diff --git a/Engine/Source/Editor/UnrealEd/Private/Cooker/WorkerRequestsLocal.cpp b/Engine/Source/Editor/UnrealEd/Private/Cooker/WorkerRequestsLocal.cpp
--- a/Engine/Source/Editor/UnrealEd/Private/Cooker/WorkerRequestsLocal.cpp
+++ b/Engine/Source/Editor/UnrealEd/Private/Cooker/WorkerRequestsLocal.cpp
@@ -98,6 +98,14 @@ void FWorkerRequestsLocal::ReportAccessedIniSettings(UCookOnTheFlyServer& COTFS,
 	COTFS.ProcessAccessedIniSettings(&Config, COTFS.AccessedIniStrings);
 }
 
+void FWorkerRequestsLocal::ReportDemoteToIdle(UE::Cook::FPackageData& PackageData, ESuppressCookReason Reason)
+{
+}
+
+void FWorkerRequestsLocal::ReportPromoteToSaveComplete(UE::Cook::FPackageData& PackageData)
+{
+}
+
 void FWorkerRequestsLocal::GetInitializeConfigSettings(UCookOnTheFlyServer& COTFS,
 	const FString& OutputDirectoryOverride, UE::Cook::FInitializeConfigSettings& Settings)
 {
diff --git a/Engine/Source/Editor/UnrealEd/Private/Cooker/WorkerRequestsLocal.h b/Engine/Source/Editor/UnrealEd/Private/Cooker/WorkerRequestsLocal.h
--- a/Engine/Source/Editor/UnrealEd/Private/Cooker/WorkerRequestsLocal.h
+++ b/Engine/Source/Editor/UnrealEd/Private/Cooker/WorkerRequestsLocal.h
@@ -29,6 +29,8 @@ public:
 	virtual void AddPublicInterfaceRequest(FFilePlatformRequest&& Request, bool bForceFrontOfQueue) override;
 
 	virtual void ReportAccessedIniSettings(UCookOnTheFlyServer& COTFS, const FConfigFile& Config) override;
+	virtual void ReportDemoteToIdle(UE::Cook::FPackageData& PackageData, ESuppressCookReason Reason) override;
+	virtual void ReportPromoteToSaveComplete(UE::Cook::FPackageData& PackageData) override;
 	virtual void GetInitializeConfigSettings(UCookOnTheFlyServer& COTFS, const FString& OutputOverrideDirectory, UE::Cook::FInitializeConfigSettings& Settings) override;
 	virtual void GetBeginCookConfigSettings(UCookOnTheFlyServer& COTFS, FBeginCookContext& BeginContext, UE::Cook::FBeginCookConfigSettings& Settings) override;
 	virtual void GetBeginCookIterativeFlags(UCookOnTheFlyServer& COTFS, FBeginCookContext& BeginContext) override;
diff --git a/Engine/Source/Editor/UnrealEd/Private/Cooker/WorkerRequestsRemote.cpp b/Engine/Source/Editor/UnrealEd/Private/Cooker/WorkerRequestsRemote.cpp
--- a/Engine/Source/Editor/UnrealEd/Private/Cooker/WorkerRequestsRemote.cpp
+++ b/Engine/Source/Editor/UnrealEd/Private/Cooker/WorkerRequestsRemote.cpp
@@ -16,7 +16,7 @@ FWorkerRequestsRemote::FWorkerRequestsRemote(UCookOnTheFlyServer& InCOTFS)
 
 bool FWorkerRequestsRemote::TryConnect(FDirectorConnectionInfo&& ConnectInfo, ECookInitializationFlags& OutCookFlags)
 {
-	return false;
+	return CookWorkerClient.TryConnect(MoveTemp(ConnectInfo));
 }
 
 bool FWorkerRequestsRemote::TryGetInitializeSettings(ECookMode::Type& OutCookMode, ECookInitializationFlags& OutCookFlags,
@@ -110,6 +110,16 @@ void FWorkerRequestsRemote::ReportAccessedIniSettings(UCookOnTheFlyServer& COTFS
 	// MPCOOKTODO: Not yet implemented
 }
 
+void FWorkerRequestsRemote::ReportDemoteToIdle(UE::Cook::FPackageData& PackageData, ESuppressCookReason Reason)
+{
+	// MPCOOKTODO: Not yet implemented
+}
+
+void FWorkerRequestsRemote::ReportPromoteToSaveComplete(UE::Cook::FPackageData& PackageData)
+{
+	// MPCOOKTODO: Not yet implemented
+}
+
 void FWorkerRequestsRemote::GetInitializeConfigSettings(UCookOnTheFlyServer& COTFS,
 	const FString& OutputDirectoryOverride, UE::Cook::FInitializeConfigSettings& Settings)
 {
diff --git a/Engine/Source/Editor/UnrealEd/Private/Cooker/WorkerRequestsRemote.h b/Engine/Source/Editor/UnrealEd/Private/Cooker/WorkerRequestsRemote.h
--- a/Engine/Source/Editor/UnrealEd/Private/Cooker/WorkerRequestsRemote.h
+++ b/Engine/Source/Editor/UnrealEd/Private/Cooker/WorkerRequestsRemote.h
@@ -3,18 +3,16 @@
 #pragma once
 
 #include "CookRequests.h"
+#include "CookWorkerClient.h"
 #include "Containers/StringView.h"
 #include "IWorkerRequests.h"
 
+namespace UE::Cook { struct FDirectorConnectionInfo; }
+
 namespace UE::Cook
 {
 
 /** An IWorkerRequests for CookWorkers in MultiProcess cooks: functions are implemented as interprocess messages to/from the CookDirector. */
-struct FDirectorConnectionInfo
-{
-	FString HostURI;
-};
-
 class FWorkerRequestsRemote : public IWorkerRequests
 {
 public:
@@ -41,6 +39,8 @@ public:
 	virtual void AddPublicInterfaceRequest(FFilePlatformRequest&& Request, bool bForceFrontOfQueue) override;
 
 	virtual void ReportAccessedIniSettings(UCookOnTheFlyServer& COTFS, const FConfigFile& Config) override;
+	virtual void ReportDemoteToIdle(UE::Cook::FPackageData& PackageData, ESuppressCookReason Reason) override;
+	virtual void ReportPromoteToSaveComplete(UE::Cook::FPackageData& PackageData) override;
 	virtual void GetInitializeConfigSettings(UCookOnTheFlyServer& COTFS, const FString& OutputOverrideDirectory, UE::Cook::FInitializeConfigSettings& Settings) override;
 	virtual void GetBeginCookConfigSettings(UCookOnTheFlyServer& COTFS, FBeginCookContext& BeginContext, UE::Cook::FBeginCookConfigSettings& Settings) override;
 	virtual void GetBeginCookIterativeFlags(UCookOnTheFlyServer& COTFS, FBeginCookContext& BeginContext) override;
@@ -54,6 +54,7 @@ private:
 
 private:
 	FExternalRequests ExternalRequests;
+	FCookWorkerClient CookWorkerClient;
 };
 
 }
\ No newline at end of file
