Add more information to the DumpMaterialShaderTypes commandlet output.

- Keep track of a map from static parameter hash to a list of all mat instances that share the same hash. (currently this data isn't being printed out.)
- Print the full path to the parent material.  Some assets have the same name making it hard to track down which asset we were referring to.
- Generate and print out the static parameter hash for each material instance.
- Add material instance function to generate a string which contains a comma separated list of base property overrides for material instances.  These cause shaders to be stored with material instances.
- Generate and print out the base property override string to see which material instances also cause shaders.  This output looks like the following:

Base Property Overrides: True
	bOverride_OpacityMaskClipValue_0, bOverride_BlendMode_1, bOverride_ShadingModel_0, bOverride_TwoSided_0, bOverride_DitheredLODTransition_0, bOverride_CastDynamicShadowAsMasked_0, bOverride_OutputTranslucentVelocity_0

diff --git a/Engine/Source/Editor/UnrealEd/Private/Commandlets/DumpMaterialShaderTypes.cpp b/Engine/Source/Editor/UnrealEd/Private/Commandlets/DumpMaterialShaderTypes.cpp
--- a/Engine/Source/Editor/UnrealEd/Private/Commandlets/DumpMaterialShaderTypes.cpp
+++ b/Engine/Source/Editor/UnrealEd/Private/Commandlets/DumpMaterialShaderTypes.cpp
@@ -201,13 +201,50 @@ public:
 		}
 	}
 
+	void PrintUniqueMaterialInstances()
+	{
+		TRACE_CPUPROFILER_EVENT_SCOPE(PrintUniqueMaterialInstances);
+
+		// Sort by number of mat instances.
+		struct TArraySizeSort
+		{
+			bool operator()(const TArray<FString>& A, const TArray<FString>& B) const
+			{
+				return A.Num() > B.Num();
+			}
+		};
+		UniqueMaterialInstances.ValueSort(TArraySizeSort());
+
+		// Write out header.
+		const char MatInstHeader[] = "\nUnique Material Instances\n";
+		DebugWriter->Serialize(const_cast<char*>(MatInstHeader), sizeof(MatInstHeader) - 1);
+
+		// Print each item.
+		for (TPair<FSHAHash, TArray<FString>> UniqueMaterialInstance : UniqueMaterialInstances)
+		{
+			FString DuplicateLine = FString::Printf(TEXT("Duplicates: %d\n"), UniqueMaterialInstance.Value.Num());
+			DebugWriter->Serialize(const_cast<ANSICHAR*>(StringCast<ANSICHAR>(*DuplicateLine).Get()), DuplicateLine.Len());
+
+			for (const FString& InstanceName : UniqueMaterialInstance.Value)
+			{
+				FString OuputLine = FString::Printf(TEXT("\t%s\n"), *InstanceName);
+				DebugWriter->Serialize(const_cast<ANSICHAR*>(StringCast<ANSICHAR>(*OuputLine).Get()), OuputLine.Len());
+			}
+		}
+	}
+
 	void Log(const FString& OutString)
 	{
-		//UE_LOG(LogDumpMaterialShaderTypesCommandlet, Log, TEXT("%s"), *OutString);
 		FString OuputLine = OutString + TEXT("\n");
 		DebugWriter->Serialize(const_cast<ANSICHAR*>(StringCast<ANSICHAR>(*OuputLine).Get()), OuputLine.Len());
 	}
 
+	void AddMaterialInstance(const FString& MaterialInstanceName, const FSHAHash& StaticParameterHash)
+	{
+		TArray<FString>& InstanceList = UniqueMaterialInstances.FindOrAdd(StaticParameterHash);
+		InstanceList.Add(MaterialInstanceName);
+	}
+
 private:
 	FArchive* DebugWriter = nullptr;
 
@@ -223,6 +260,9 @@ private:
 	/** Map of global shader type display names to their counts. */
 	TMap<FString, int32>	GlobalShaderTypeHistogram;
 
+	/** Unique material instances. */
+	TMap<FSHAHash, TArray<FString>> UniqueMaterialInstances;
+
 	/** Store a copy of the the filename. */
 	FString FileName;
 
@@ -369,26 +409,57 @@ int ProcessMaterialInstances(const ITargetPlatform* TargetPlatform, const EShade
 			const int TotalShadersForMaterial = GetTotalShaders(OutShaderInfo);
 			TotalShaders += TotalShadersForMaterial;
 
-			FString StaticParameterString(TEXT(""));
+			// Find the root parent that is a material.
+			UMaterialInterface* Top = MaterialInstance->Parent;
+			while (Top)
+			{
+				if (UMaterial* Material = Cast<UMaterial>(Top))
+				{
+					break;
+				}
+				else if (UMaterialInstance* MatInst = Cast<UMaterialInstance>(Top))
+				{
+					Top = MatInst->Parent;
+				}
+				else
+				{
+					Top = nullptr;
+				}
+			}
+
+			Output.Log(TEXT(""));
+			Output.Log(FString::Printf(TEXT("Material Instance: %s - %d shaders"), *AssetData.AssetName.ToString(), TotalShadersForMaterial));
+			Output.Log(FString::Printf(TEXT("Parent: %s"), Top ? *Top->GetFullName() : TEXT("NO PARENT")));
+			Output.Log(FString::Printf(TEXT("Static Parameters: %d"), MaterialInstance->bHasStaticPermutationResource ? MaterialInstance->GetStaticParameters().EditorOnly.StaticSwitchParameters.Num() : 0));
+
+			FSHAHash OutHash;
 
 			if (MaterialInstance->bHasStaticPermutationResource)
 			{
+				FSHA1 Hasher;
 				const FStaticParameterSet& ParameterSet = MaterialInstance->GetStaticParameters();
 				for (int32 StaticSwitchIndex = 0; StaticSwitchIndex < ParameterSet.EditorOnly.StaticSwitchParameters.Num(); ++StaticSwitchIndex)
 				{
 					const FStaticSwitchParameter& StaticSwitchParameter = ParameterSet.EditorOnly.StaticSwitchParameters[StaticSwitchIndex];
-					StaticParameterString += FString::Printf(
-						TEXT(", StaticSwitch'%s'=%s"),
-						*StaticSwitchParameter.ParameterInfo.ToString(),
-						StaticSwitchParameter.Value ? TEXT("True") : TEXT("False")
-					);
+
+					StaticSwitchParameter.UpdateHash(Hasher);
+
+					Output.Log(FString::Printf(TEXT("\t%s : %s"), *StaticSwitchParameter.ParameterInfo.ToString(), StaticSwitchParameter.Value ? TEXT("True") : TEXT("False")));
 				}
+
+				Hasher.Final();
+				Hasher.GetHash(&OutHash.Hash[0]);
+
+				Output.AddMaterialInstance(AssetData.AssetName.ToString(), OutHash);
+				Output.Log(FString::Printf(TEXT("Static Parameter Hash: %s"), *OutHash.ToString()));
 			}
 
-			Output.Log(TEXT(""));
-			Output.Log(FString::Printf(TEXT("Material Instance: %s - %d shaders"), *AssetData.ObjectPath.ToString(), TotalShadersForMaterial));
-			Output.Log(FString::Printf(TEXT("Static Parameter %s"), *StaticParameterString));
-			Output.Log(FString::Printf(TEXT("Parent: %s"), MaterialInstance->Parent ? *MaterialInstance->Parent->GetName() : TEXT("NO PARENT")));
+			Output.Log(FString::Printf(TEXT("Base Property Overrides: %s"), MaterialInstance->HasOverridenBaseProperties() ? TEXT("True") : TEXT("False")));
+
+			if (MaterialInstance->HasOverridenBaseProperties())
+			{
+				Output.Log(FString::Printf(TEXT("\t%s"), *MaterialInstance->GetBasePropertyOverrideString()));
+			}
 
 			PrintDebugShaderInfo(Output, OutShaderInfo);
 
diff --git a/Engine/Source/Runtime/Engine/Classes/Materials/MaterialInstance.h b/Engine/Source/Runtime/Engine/Classes/Materials/MaterialInstance.h
--- a/Engine/Source/Runtime/Engine/Classes/Materials/MaterialInstance.h
+++ b/Engine/Source/Runtime/Engine/Classes/Materials/MaterialInstance.h
@@ -844,6 +844,10 @@ public:
 	void GetBasePropertyOverridesHash(FSHAHash& OutHash)const;
 	ENGINE_API virtual bool HasOverridenBaseProperties()const;
 
+#if WITH_EDITOR
+	ENGINE_API FString GetBasePropertyOverrideString() const;
+#endif
+
 	// For all materials instances, UMaterialInstance::CacheResourceShadersForRendering
 	ENGINE_API static void AllMaterialsCacheResourceShadersForRendering(bool bUpdateProgressDialog = false, bool bCacheAllRemainingShaders = true);
 
diff --git a/Engine/Source/Runtime/Engine/Private/Materials/MaterialInstance.cpp b/Engine/Source/Runtime/Engine/Private/Materials/MaterialInstance.cpp
--- a/Engine/Source/Runtime/Engine/Private/Materials/MaterialInstance.cpp
+++ b/Engine/Source/Runtime/Engine/Private/Materials/MaterialInstance.cpp
@@ -3984,6 +3984,24 @@ bool UMaterialInstance::HasOverridenBaseProperties()const
 	return false;
 }
 
+#if WITH_EDITOR
+FString UMaterialInstance::GetBasePropertyOverrideString() const
+{
+	FString BasePropString;
+	if (HasOverridenBaseProperties())
+	{
+		BasePropString += FString::Printf(TEXT("bOverride_OpacityMaskClipValue_%d, "), ((FMath::Abs(GetOpacityMaskClipValue() - Parent->GetOpacityMaskClipValue()) > UE_SMALL_NUMBER)));
+		BasePropString += FString::Printf(TEXT("bOverride_BlendMode_%d, "), (GetBlendMode() != Parent->GetBlendMode()));
+		BasePropString += FString::Printf(TEXT("bOverride_ShadingModel_%d, "), (GetShadingModels() != Parent->GetShadingModels()));
+		BasePropString += FString::Printf(TEXT("bOverride_TwoSided_%d, "), (IsTwoSided() != Parent->IsTwoSided()));
+		BasePropString += FString::Printf(TEXT("bOverride_DitheredLODTransition_%d, "), (IsDitheredLODTransition() != Parent->IsDitheredLODTransition()));
+		BasePropString += FString::Printf(TEXT("bOverride_CastDynamicShadowAsMasked_%d, "), (GetCastDynamicShadowAsMasked() != Parent->GetCastDynamicShadowAsMasked()));
+		BasePropString += FString::Printf(TEXT("bOverride_OutputTranslucentVelocity_%d "), (IsTranslucencyWritingVelocity() != Parent->IsTranslucencyWritingVelocity()));
+	}
+	return BasePropString;
+}
+#endif
+
 float UMaterialInstance::GetOpacityMaskClipValue() const
 {
 	return OpacityMaskClipValue;
