Fix crash when rendering skeletalmesh thumbnail. Make sure we do not try to render a skeletalmesh thumbnail if the render data is not ready.
#rnx

diff --git a/Engine/Source/Editor/UnrealEd/Classes/ThumbnailRendering/SkeletalMeshThumbnailRenderer.h b/Engine/Source/Editor/UnrealEd/Classes/ThumbnailRendering/SkeletalMeshThumbnailRenderer.h
--- a/Engine/Source/Editor/UnrealEd/Classes/ThumbnailRendering/SkeletalMeshThumbnailRenderer.h
+++ b/Engine/Source/Editor/UnrealEd/Classes/ThumbnailRendering/SkeletalMeshThumbnailRenderer.h
@@ -25,6 +25,7 @@ class UNREALED_API USkeletalMeshThumbnailRenderer : public UDefaultSizedThumbnai
 	// Begin UThumbnailRenderer Object
 	virtual void Draw(UObject* Object, int32 X, int32 Y, uint32 Width, uint32 Height, FRenderTarget*, FCanvas* Canvas, bool bAdditionalViewFamily) override;
 	virtual bool AllowsRealtimeThumbnails(UObject* Object) const override;
+	virtual bool CanVisualizeAsset(UObject* Object) override;
 	// End UThumbnailRenderer Object
 
 	// UObject implementation
diff --git a/Engine/Source/Editor/UnrealEd/Private/ThumbnailRendering/SkeletalMeshThumbnailRenderer.cpp b/Engine/Source/Editor/UnrealEd/Private/ThumbnailRendering/SkeletalMeshThumbnailRenderer.cpp
--- a/Engine/Source/Editor/UnrealEd/Private/ThumbnailRendering/SkeletalMeshThumbnailRenderer.cpp
+++ b/Engine/Source/Editor/UnrealEd/Private/ThumbnailRendering/SkeletalMeshThumbnailRenderer.cpp
@@ -45,6 +45,17 @@ bool USkeletalMeshThumbnailRenderer::AllowsRealtimeThumbnails(UObject* Object) c
 	return SkeletalMesh && SkeletalMesh->GetResourceForRendering();
 }
 
+bool USkeletalMeshThumbnailRenderer::CanVisualizeAsset(UObject* Object)
+{
+	if (!Super::CanVisualizeAsset(Object))
+	{
+		return false;
+	}
+	USkeletalMesh* SkeletalMesh = Cast<USkeletalMesh>(Object);
+	const bool bValidRenderData = SkeletalMesh && SkeletalMesh->GetResourceForRendering();
+	return bValidRenderData;
+}
+
 void USkeletalMeshThumbnailRenderer::BeginDestroy()
 {
 	ThumbnailSceneCache.Clear();
