Fixed static analysis warning in summarize trace

diff --git a/Engine/Source/Editor/UnrealEd/Private/Commandlets/SummarizeTraceCommandlet.cpp b/Engine/Source/Editor/UnrealEd/Private/Commandlets/SummarizeTraceCommandlet.cpp
--- a/Engine/Source/Editor/UnrealEd/Private/Commandlets/SummarizeTraceCommandlet.cpp
+++ b/Engine/Source/Editor/UnrealEd/Private/Commandlets/SummarizeTraceCommandlet.cpp
@@ -955,7 +955,11 @@ struct FSummarizeCounter
 
 	FSummarizeCounter()
 	{
-		SetIsFloatingPoint(false);
+		Type = TraceCounterType_Int;
+		First.Int = 0;
+		Last.Int = 0;
+		Minimum.Int = TNumericLimits<int64>::Max();
+		Maximum.Int = TNumericLimits<int64>::Min();
 	}
 
 	virtual void SetName(const TCHAR* InName) override
@@ -1877,6 +1881,7 @@ int32 USummarizeTraceCommandlet::Main(const FString& CmdLineParams)
 	if (!DataStream->Open(*TraceFileName))
 	{
 		UE_LOG(LogSummarizeTrace, Error, TEXT("Unable to open trace file '%s' for read"), *TraceFileName);
+		delete DataStream;
 		return 1;
 	}
 
