Added support for preview versions of the localization manifests to be generated and deleted during a preview localization run. This addresses false positives for loclaization warnings that can arise from comparing gathered strings to a stale manifest file.
- Added a preview parameter to the Localization UAT command which allows the command to run in preview mode.
- Localization commandlets now generate *_Preview.manifest files during a run into the destination path specified by the config files.
- The localization automation script will delete all files that end with _Preview.manifest at the end of the gather run.
- Improved checks for manifest filename extensions to ensure no malformed manifest filenames are passed into the commandlet

diff --git a/Engine/Source/Editor/UnrealEd/Classes/Commandlets/GenerateGatherManifestCommandlet.h b/Engine/Source/Editor/UnrealEd/Classes/Commandlets/GenerateGatherManifestCommandlet.h
--- a/Engine/Source/Editor/UnrealEd/Classes/Commandlets/GenerateGatherManifestCommandlet.h
+++ b/Engine/Source/Editor/UnrealEd/Classes/Commandlets/GenerateGatherManifestCommandlet.h
@@ -19,5 +19,8 @@ public:
 	//~ Begin UCommandlet Interface
 	virtual int32 Main( const FString& Params ) override;
 	//~ End UCommandlet Interface
+		//~ Begin UGatherTextCommandletBase  Interface
+	virtual bool ShouldRunInPreview(const TArray<FString>& Switches, const TMap<FString, FString>& ParamVals) const override;
+	//~ End UGatherTextCommandletBase  Interface
 #endif
 };
diff --git a/Engine/Source/Editor/UnrealEd/Private/Commandlets/GenerateGatherManifestCommandlet.cpp b/Engine/Source/Editor/UnrealEd/Private/Commandlets/GenerateGatherManifestCommandlet.cpp
--- a/Engine/Source/Editor/UnrealEd/Private/Commandlets/GenerateGatherManifestCommandlet.cpp
+++ b/Engine/Source/Editor/UnrealEd/Private/Commandlets/GenerateGatherManifestCommandlet.cpp
@@ -2,9 +2,51 @@
 
 #include "Commandlets/GenerateGatherManifestCommandlet.h"
 #include "Misc/Paths.h"
+#include "Misc/FileHelper.h"
 
 DEFINE_LOG_CATEGORY_STATIC(LogGenerateManifestCommandlet, Log, All);
 
+namespace GenerateManifestHelper
+{
+	FString GetManifestFileExtension()
+	{
+		static const FString ManifestFileExtension= TEXT(".manifest");
+		return ManifestFileExtension;
+	}
+
+	/** Returns true if the passed in file has a .manifest extension. Else returns false. */
+	bool IsManifestFileExtensionValid(const FString& InManifestFilename)
+	{
+		return InManifestFilename.EndsWith(GenerateManifestHelper::GetManifestFileExtension());
+	}
+} // namespace GenerateManifestHelper
+
+namespace GeneratePreviewManifestHelper
+{
+	FString GetPreviewManifestSuffix()
+	{
+		// Note: This suffix is also hardcoded in Localisation.automation.cs
+		// If you decide to change the suffix for preview manifest files here, also update the automation file. 
+		static const FString PreviewSuffix = TEXT("_Preview");
+		return PreviewSuffix;
+	}
+
+	/**
+	* Given a manifest filename, provide the preview version of the manifest filename.
+	* @param InOriginalManifestFilename The original manifest filename (e.g MyManifest.manifest)
+	* @return The preview version of the manifest filename (e.g MyManifest_Preview.manifest with current implementation)
+	*/
+	FString GetPreviewManifestFilename(const FString& InOriginalManifestFilename)
+	{
+		// This check should already be performed outside of this function. The passed in file is assumed to have the correct extension. 
+		check(GenerateManifestHelper::IsManifestFileExtensionValid(InOriginalManifestFilename));
+		FString PreviewFilename= InOriginalManifestFilename.LeftChop(GenerateManifestHelper::GetManifestFileExtension().Len());
+		PreviewFilename.Append(GeneratePreviewManifestHelper::GetPreviewManifestSuffix());
+		PreviewFilename.Append(GenerateManifestHelper::GetManifestFileExtension());
+		return PreviewFilename;
+	}
+} // namespace GeneratePreviewManifestHelper
+
 /**
  *	UGenerateGatherManifestCommandlet
  */
@@ -20,6 +62,16 @@ int32 UGenerateGatherManifestCommandlet::Main( const FString& Params )
 	TArray<FString> Switches;
 	TMap<FString, FString> ParamVals;
 	UCommandlet::ParseCommandLine( *Params, Tokens, Switches, ParamVals );
+	// we will have different behavior when this commandlet is running in preview
+	// We want to generate temporary preview manifest files.
+	// The temp manifest files will be loaded as manifest dependencies to avoid localization duplicate key warnings from surfacing due to stale manifest dependencies. 
+	// The Localisation.automation.cs file is set up to delete these temp manifest files with a run of the Localize UAT command with the preview switch.
+	// Note that these temp files will persist if running this commandlet via the command line. 
+	const bool bRunningInPreview = Switches.Contains(TEXT("Preview"));
+	if (bRunningInPreview)
+	{
+		UE_LOG(LogGenerateManifestCommandlet, Log, TEXT("Commandlet is running in preview mode. Preview versions of manifests will be saved and loaded."));
+	}
 
 	// Set config file.
 	const FString* ParamVal = ParamVals.Find( FString( TEXT("Config") ) );
@@ -65,16 +117,59 @@ int32 UGenerateGatherManifestCommandlet::Main( const FString& Params )
 		return -1;
 	}
 
+	if (!GenerateManifestHelper::IsManifestFileExtensionValid(ManifestName))
+	{
+		UE_LOG(LogGenerateManifestCommandlet, Error, TEXT("Found manifest file %s is malformed. All manifest files should have a %s extension."), *ManifestName, *GenerateManifestHelper::GetManifestFileExtension());
+		return -1;
+	}
+
+	if (bRunningInPreview)
+	{
+		// we change the manifest filename to reflect the preview filename 
+		ManifestName = GeneratePreviewManifestHelper::GetPreviewManifestFilename(ManifestName);
+	}
+
 	//Grab any manifest dependencies
 	TArray<FString> ManifestDependenciesList;
 	GetPathArrayFromConfig(*SectionName, TEXT("ManifestDependencies"), ManifestDependenciesList, GatherTextConfigPath);
 
+	// Check that all the dependent manifest files have valid file extensions
+	for (const FString& ManifestDependency : ManifestDependenciesList)
+	{
+		if (!GenerateManifestHelper::IsManifestFileExtensionValid(ManifestDependency))
+		{
+			UE_LOG(LogGenerateManifestCommandlet, Error, TEXT("Found manifest dependency %s is malformed. All manifest files should have a %s extension."), *ManifestDependency, *GenerateManifestHelper::GetManifestFileExtension());
+			return -1;
+		}
+	}
+	if (bRunningInPreview)
+	{
+		// Overwrite all the manifest dependency filenames with their preview counterparts 
+		for (FString& ManifestDependency : ManifestDependenciesList)
+		{
+			FString PreviewManifestDependency= GeneratePreviewManifestHelper::GetPreviewManifestFilename(ManifestDependency);
+			if (!FPaths::FileExists(PreviewManifestDependency))
+			{
+				UE_LOG(LogGenerateManifestCommandlet, Warning, TEXT("Preview manifest dependency %s does not exist. Make sure to generate all preview manifest dependencies of this localization target before trying again."), *PreviewManifestDependency);
+				continue;
+			}
+			// make sure the preview manifest file is newer than the regular manifest file
+			// we could be dealing with a scenario where the preview manifest file is from a previous run and isn't up to date or failed to be deleted
+			else if (IFileManager::Get().GetTimeStamp(*PreviewManifestDependency) < IFileManager::Get().GetTimeStamp(*ManifestDependency))
+			{
+				UE_LOG(LogGenerateManifestCommandlet, Warning, TEXT("Preview manifest dependency %s is older than the original manifest %s. Preview manifest is out of date and should be regenerated as a dependency."), *PreviewManifestDependency, *ManifestDependency);
+				continue;
+			}
+			ManifestDependency = MoveTemp(PreviewManifestDependency);
+		}
+	}
+
 	for (const FString& ManifestDependency : ManifestDependenciesList)
 	{
 		FText OutError;
 		if (!GatherManifestHelper->AddDependency(ManifestDependency, &OutError))
 		{
-			UE_LOG(LogGenerateManifestCommandlet, Error, TEXT("Failed to add a manifest dependency: %s"), *OutError.ToString());
+			UE_LOG(LogGenerateManifestCommandlet, Error, TEXT("Failed to add manifest dependency %s. Failure reason : %s"), *ManifestDependency, *OutError.ToString());
 			return -1;
 		}
 	}
@@ -92,3 +187,10 @@ int32 UGenerateGatherManifestCommandlet::Main( const FString& Params )
 
 	return 0;
 }
+
+bool UGenerateGatherManifestCommandlet::ShouldRunInPreview(const TArray<FString>& Switches, const TMap<FString, FString>& ParamVals) const
+{
+	// we need the commandlet to run to generate the preview manifests to avoid false positives from duplicate detection during preview runs 
+	return true;
+}
+
diff --git a/Engine/Source/Programs/AutomationTool/Scripts/Localisation.Automation.cs b/Engine/Source/Programs/AutomationTool/Scripts/Localisation.Automation.cs
--- a/Engine/Source/Programs/AutomationTool/Scripts/Localisation.Automation.cs
+++ b/Engine/Source/Programs/AutomationTool/Scripts/Localisation.Automation.cs
@@ -25,6 +25,7 @@ using System.Threading.Tasks;
 [Help("IncludePlatforms", "Optional flag to include platforms from within the given UEProjectDirectory as part of the gather.")]
 [Help("AdditionalCommandletArguments", "Optional arguments to pass to the gather process.")]
 [Help("ParallelGather", "Run the gather processes for a single batch in parallel rather than sequence.")]
+[Help("Preview", "Run the localization command in preview mode. This passes the -Preview flag along to all commandlets as an additional argument and deletes all temporary files generated by the commandlets in preview mode. Primarily used for build farm automation where localization warnings from localization gathers can be previewed without checking out any files under SCC.")]
 class Localize : BuildCommand
 {
 	private class LocalizationBatch
@@ -160,7 +161,16 @@ class Localize : BuildCommand
 
 		var EnableParallelGather = ParseParam("ParallelGather");
 
-		var StartTime = DateTime.UtcNow;
+		var IsRunningInPreview = ParseParam("Preview");
+		// We pass the preview switch along to have the gather text commandlets exhibit different behaviors. See UGatherTextCommandlet
+		if (IsRunningInPreview)
+		{
+			LogInformation("Running in preview mode. Preview switch will be passed along to all localization commandlets to be run.");
+			AdditionalCommandletArguments += " -Preview";
+		}
+			
+
+			var StartTime = DateTime.UtcNow;
 
 		var LocalizationBatches = new List<LocalizationBatch>();
 
@@ -231,7 +241,7 @@ class Localize : BuildCommand
 
 		// Create a single changelist to use for all changes
 		int PendingChangeList = 0;
-		if (P4Enabled)
+		if (P4Enabled && !IsRunningInPreview)
 		{
 			var ChangeListCommitMessage = String.Format("Localization Automation using CL {0}", P4Env.Changelist);
 			if (File.Exists(CombinePaths(CmdEnv.LocalRoot, @"Engine/Restricted/NotForLicensees/Build/EpicInternal.txt")))
@@ -250,6 +260,7 @@ class Localize : BuildCommand
 			LocalizationTasks.Add(LocalizationTask);
 
 			// Make sure the Localization configs and content is up-to-date to ensure we don't get errors later on
+			// we still run this in preview mode to bring all the gather configs up to date
 			if (P4Enabled)
 			{
 				LogInformation("Sync necessary content to head revision");
@@ -266,7 +277,7 @@ class Localize : BuildCommand
 
 		// Hash the current PO files on disk so we can work out whether they actually change
 		Dictionary<string, byte[]> InitalPOFileHashes = null;
-		if (P4Enabled)
+		if (P4Enabled && !IsRunningInPreview)
 		{
 			InitalPOFileHashes = GetPOFileHashes(LocalizationBatches, UEProjectRoot);
 		}
@@ -303,7 +314,6 @@ class Localize : BuildCommand
 		// These can run in parallel when ParallelGather is enabled
 		{
 			var EditorExe = CombinePaths(CmdEnv.LocalRoot, @"Engine/Binaries/Win64/UnrealEditor-Cmd.exe");
-
 			// Set the common basic editor arguments
 			var EditorArguments = P4Enabled 
 				? String.Format("-SCCProvider=Perforce -P4Port={0} -P4User={1} -P4Client={2} -P4Passwd={3} -P4Changelist={4} -EnableSCC -DisableSCCSubmit", P4Env.ServerAndPort, P4Env.User, P4Env.Client, P4.GetAuthenticationToken(), PendingChangeList)
@@ -387,6 +397,24 @@ class Localize : BuildCommand
 			}
 		}
 
+		// If we are running in preview, we can go ahead and delete all generated preview files after the gather step is complete 
+		if (IsRunningInPreview)
+        {
+			var PreviewManifestFiles= GetPreviewManifestFilesToDelete(LocalizationBatches, UEProjectRoot);
+			foreach (var PreviewManifestFile in PreviewManifestFiles) 
+			{
+				LogInformation("Deleting preview manifest file {0}.", PreviewManifestFile);
+				try
+				{
+					File.Delete(PreviewManifestFile);
+				}
+				catch (Exception Ex)
+				{
+					LogInformation("[FAILED] Deleting preview file: '{0}' - {1}", PreviewManifestFile, Ex);
+				}
+			}
+		}
+
 		// Upload the latest sources to our localization provider
 		if (LocalizationStepNames.Contains("Upload"))
 		{
@@ -417,7 +445,7 @@ class Localize : BuildCommand
 		}
 
 		// Clean-up the changelist so it only contains the changed files, and then submit it (if we were asked to)
-		if (P4Enabled)
+		if (P4Enabled && !IsRunningInPreview)
 		{
 			// Revert any PO files that haven't changed aside from their header
 			{
@@ -649,6 +677,26 @@ class Localize : BuildCommand
 
 		return AllFiles;
 	}
+
+	private List<string> GetPreviewManifestFilesToDelete(IReadOnlyList<LocalizationBatch> LocalizationBatches, string UEProjectRoot)
+	{
+		var AllPreviewManifestFiles = new List<string>();
+		foreach (var LocalizationBatch in LocalizationBatches)
+		{
+			var LocalizationPath = CombinePaths(UEProjectRoot, LocalizationBatch.LocalizationTargetDirectory, "Content", "Localization");
+			if (!Directory.Exists(LocalizationPath))
+			{
+				continue;
+			}
+			string[] PreviewManifestFilenames= Directory.GetFiles(LocalizationPath, "*_Preview.manifest", SearchOption.AllDirectories);
+			foreach (var ManifestFilename in PreviewManifestFilenames)
+			{
+				AllPreviewManifestFiles.Add(ManifestFilename);
+			}
+		}
+		return AllPreviewManifestFiles;
+	}
+
 }
 
 [Help("OnlyLoc", "Optional.  Only submit generated loc files, do not submit any other generated file.")]
