Added support for loading analytics to the Cook Commandlet

diff --git a/Engine/Source/Editor/UnrealEd/Private/Commandlets/CookCommandlet.cpp b/Engine/Source/Editor/UnrealEd/Private/Commandlets/CookCommandlet.cpp
--- a/Engine/Source/Editor/UnrealEd/Private/Commandlets/CookCommandlet.cpp
+++ b/Engine/Source/Editor/UnrealEd/Private/Commandlets/CookCommandlet.cpp
@@ -47,7 +47,6 @@
 #include "UObject/Package.h"
 #include "UObject/SavePackage.h"
 #include "UObject/UObjectIterator.h"
-
 DEFINE_LOG_CATEGORY_STATIC(LogCookCommandlet, Log, All);
 
 #if ENABLE_COOK_STATS
@@ -56,6 +55,7 @@ DEFINE_LOG_CATEGORY_STATIC(LogCookCommandlet, Log, All);
 #include "IAnalyticsProviderET.h"
 #include "AnalyticsET.h"
 #include "Virtualization/VirtualizationSystem.h"
+#include "StudioAnalytics.h"
 
 namespace DetailedCookStats
 {
@@ -606,9 +606,10 @@ int32 UCookCommandlet::Main(const FString& CmdLineParams)
 			DetailedCookStats::CookWallTimeSec = Now - GStartTime;
 			DetailedCookStats::StartupWallTimeSec = CookStartTime - GStartTime;
 			DetailedCookStats::LogCookStats(CmdLineParams);
+
+			FStudioAnalytics::FireEvent_Loading(TEXT("CookByTheBook"), DetailedCookStats::CookWallTimeSec);
 		});
 	}
-	
 	return 0;
 }
 
