Fix to support backward compatability with localization conflict files. This is to ensure localization config files that still use .txt extensions for their conflict reports will still be supported.
- Detects the extension used for the specified ConflictReportName in the vairous localization config files and generates a formatted report based on the extension.
- Currently only .txt and .csv file extensions are supported.
- We still support .txt for backward compatability.
- If an unsupported file extension is used, the conflict report will default to be generated in .csv
#rnx

diff --git a/Engine/Config/Localization/Category.ini b/Engine/Config/Localization/Category.ini
--- a/Engine/Config/Localization/Category.ini
+++ b/Engine/Config/Localization/Category.ini
@@ -61,6 +61,5 @@ DestinationPath=\\epicgames.net\root\UE3\Localization\WordCounts
 bWordCountReport=true
 WordCountReportName=Category.csv
 bConflictReport=true
-ConflictReportFormat=EConflictReportFormat::CSV
-; Do not add an extension to the filename. EConflictReportFormat will be used to append the correct extension
-ConflictReportName=CategoryConflicts
\ No newline at end of file
+; Different file extensions will output the report in different formats. Currently supported extensions are .txt and .csv
+ConflictReportName=CategoryConflicts.csv
\ No newline at end of file
diff --git a/Engine/Config/Localization/Editor.ini b/Engine/Config/Localization/Editor.ini
--- a/Engine/Config/Localization/Editor.ini
+++ b/Engine/Config/Localization/Editor.ini
@@ -106,6 +106,5 @@ DestinationPath=\\epicgames.net\root\UE3\Localization\WordCounts
 bWordCountReport=true
 WordCountReportName=Editor.csv
 bConflictReport=true
-ConflictReportFormat=EConflictReportFormat::CSV
-; Do not add an extension to the filename. EConflictReportFormat will be used to append the correct extension
-ConflictReportName=EditorConflicts
\ No newline at end of file
+; Different file extensions will output the report in different formats. Currently supported extensions are .txt and .csv
+ConflictReportName=EditorConflicts.csv
\ No newline at end of file
diff --git a/Engine/Config/Localization/EditorTutorials.ini b/Engine/Config/Localization/EditorTutorials.ini
--- a/Engine/Config/Localization/EditorTutorials.ini
+++ b/Engine/Config/Localization/EditorTutorials.ini
@@ -62,6 +62,5 @@ DestinationPath=\\epicgames.net\root\UE3\Localization\WordCounts
 bWordCountReport=true
 WordCountReportName=EditorTutorials.csv
 bConflictReport=true
-ConflictReportFormat=EConflictReportFormat::CSV
-; Do not add an extension to the filename. EConflictReportFormat will be used to append the correct extension
-ConflictReportName=EditorTutorialsConflicts
\ No newline at end of file
+; Different file extensions will output the report in different formats. Currently supported extensions are .txt and .csv
+ConflictReportName=EditorTutorialsConflicts.csv
\ No newline at end of file
diff --git a/Engine/Config/Localization/Engine.ini b/Engine/Config/Localization/Engine.ini
--- a/Engine/Config/Localization/Engine.ini
+++ b/Engine/Config/Localization/Engine.ini
@@ -84,6 +84,5 @@ DestinationPath=\\epicgames.net\root\UE3\Localization\WordCounts
 bWordCountReport=true
 WordCountReportName=Engine.csv
 bConflictReport=true
-ConflictReportFormat=EConflictReportFormat::CSV
-; Do not add an extension to the filename. EConflictReportFormat will be used to append the correct extension
-ConflictReportName=EngineConflicts
\ No newline at end of file
+; Different file extensions will output the report in different formats. Currently supported extensions are .txt and .csv
+ConflictReportName=EngineConflicts.csv
\ No newline at end of file
diff --git a/Engine/Config/Localization/Keywords.ini b/Engine/Config/Localization/Keywords.ini
--- a/Engine/Config/Localization/Keywords.ini
+++ b/Engine/Config/Localization/Keywords.ini
@@ -57,6 +57,5 @@ DestinationPath=\\epicgames.net\root\UE3\Localization\WordCounts
 bWordCountReport=true
 WordCountReportName=Keywords.csv
 bConflictReport=true
-ConflictReportFormat=EConflictReportFormat::CSV
-; Do not add an extension to the filename. EConflictReportFormat will be used to append the correct extension
-ConflictReportName=KeywordsConflicts
\ No newline at end of file
+; Different file extensions will output the report in different formats. Currently supported extensions are .txt and .csv
+ConflictReportName=KeywordsConflicts.csv
\ No newline at end of file
diff --git a/Engine/Config/Localization/PropertyNames.ini b/Engine/Config/Localization/PropertyNames.ini
--- a/Engine/Config/Localization/PropertyNames.ini
+++ b/Engine/Config/Localization/PropertyNames.ini
@@ -61,6 +61,5 @@ DestinationPath=\\epicgames.net\root\UE3\Localization\WordCounts
 bWordCountReport=true
 WordCountReportName=PropertyNames.csv
 bConflictReport=true
-ConflictReportFormat=EConflictReportFormat::CSV
-; Do not add an extension to the filename. EConflictReportFormat will be used to append the correct extension
-ConflictReportName=PropertyNamesConflicts
\ No newline at end of file
+; Different file extensions will output the report in different formats. Currently supported extensions are .txt and .csv
+ConflictReportName=PropertyNamesConflicts.csv
\ No newline at end of file
diff --git a/Engine/Config/Localization/ToolTips.ini b/Engine/Config/Localization/ToolTips.ini
--- a/Engine/Config/Localization/ToolTips.ini
+++ b/Engine/Config/Localization/ToolTips.ini
@@ -64,6 +64,5 @@ DestinationPath=\\epicgames.net\root\UE3\Localization\WordCounts
 bWordCountReport=true
 WordCountReportName=ToolTips.csv
 bConflictReport=true
-ConflictReportFormat=EConflictReportFormat::CSV
-; Do not add an extension to the filename. EConflictReportFormat will be used to append the correct extension
-ConflictReportName=ToolTipsConflicts
\ No newline at end of file
+; Different file extensions will output the report in different formats. Currently supported extensions are .txt and .csv
+ConflictReportName=ToolTipsConflicts.csv
\ No newline at end of file
diff --git a/Engine/Source/Editor/UnrealEd/Private/Commandlets/GenerateTextLocalizationReportCommandlet.cpp b/Engine/Source/Editor/UnrealEd/Private/Commandlets/GenerateTextLocalizationReportCommandlet.cpp
--- a/Engine/Source/Editor/UnrealEd/Private/Commandlets/GenerateTextLocalizationReportCommandlet.cpp
+++ b/Engine/Source/Editor/UnrealEd/Private/Commandlets/GenerateTextLocalizationReportCommandlet.cpp
@@ -178,59 +178,34 @@ bool UGenerateTextLocalizationReportCommandlet::ProcessConflictReport(const FStr
 		UE_LOG(LogGenerateTextLocalizationReportCommandlet, Error, TEXT("No conflict report name specified."));
 		return false;
 	}
-	if (ConflictReportName.Contains(TEXT("."), ESearchCase::IgnoreCase, ESearchDir::FromEnd))
+	
+	EConflictReportFormat ConflictReportFormat = EConflictReportFormat::None;
+	static const FString TxtExtension = TEXT(".txt");
+	static const FString CSVExtension = TEXT(".csv");
+	FString FileExtension = FPaths::GetExtension(ConflictReportName, true);
+	if (FileExtension == TxtExtension)
 	{
-		UE_LOG(LogGenerateTextLocalizationReportCommandlet, Error, TEXT("Conflict report name has a file extension. Please remove the file extension from the name. EConflictReportFormat is used to determine the file extension."));
-		return false;
+		ConflictReportFormat = EConflictReportFormat::Txt;
 	}
-	EConflictReportFormat ConflictReportFormat = EConflictReportFormat::None;
-	FString ConflictReportFormatString;
-	if (!GConfig->GetString(*SectionName, TEXT("ConflictReportFormat"), ConflictReportFormatString, GatherTextConfigPath))
+	else if (FileExtension == CSVExtension)
 	{
-		UE_LOG(LogGenerateTextLocalizationReportCommandlet, Display, TEXT("Conflict report format not specified. Conflict report will default to CSV format."));
 		ConflictReportFormat = EConflictReportFormat::CSV;
 	}
+	// This is an unsupported extension or an empty extension
 	else
 	{
-		ConflictReportFormatString.TrimStartAndEndInline();
-		static const TCHAR* TxtEnumString = TEXT("EConflictReportFormat::Txt");
-		static const TCHAR* CSVEnumString = TEXT("EConflictReportFormat::CSV");
-		if (ConflictReportFormatString.Equals(TxtEnumString))
-		{
-			ConflictReportFormat = EConflictReportFormat::Txt;
-		}
-		else if (ConflictReportFormatString.Equals(CSVEnumString))
+		// We default to csv 
+		ConflictReportFormat = EConflictReportFormat::CSV;
+		// We found a file extension somewhere in the specified name. 
+		if (!FileExtension.IsEmpty())
 		{
-			ConflictReportFormat = EConflictReportFormat::CSV;
+			UE_LOG(LogGenerateTextLocalizationReportCommandlet, Warning, TEXT("The conflict report filename %s has an unsupported extension. Only .txt and .csv is supported at this time."), *ConflictReportName);
 		}
 		else
 		{
-			// @TODOLocalization: Consider alternatives to defaulting to CSV.
-			UE_LOG(LogGenerateTextLocalizationReportCommandlet, Warning, TEXT("Specified conflict report format %s in \'%s\' not supported. Defaulting to CSV format."), *ConflictReportFormatString, *GatherTextConfigPath);
-			ConflictReportFormat = EConflictReportFormat::CSV;
-		}
-	}
-	switch (ConflictReportFormat)
-	{
-		case EConflictReportFormat::CSV:
-		{
-			UE_LOG(LogGenerateTextLocalizationReportCommandlet, Display, TEXT("Conflict report format will be in CSV."));
-			static const TCHAR* CSVExtension = TEXT(".csv");
-			ConflictReportName += CSVExtension;
-			break;
-		}
-		case EConflictReportFormat::Txt:
-		{
-			UE_LOG(LogGenerateTextLocalizationReportCommandlet, Display, TEXT("Conflict report formatwill be in in txt format."));
-			static const TCHAR* TxtExtension = TEXT(".txt");
-			ConflictReportName += TxtExtension;
-			break;
-		}
-		default:
-		{
-			UE_LOG(LogGenerateTextLocalizationReportCommandlet, Error, TEXT("Unsupported conflict report format %s detected in %s. Unable to create the appropriate report name and extension."), *ConflictReportFormatString, *GatherTextConfigPath);
-			return false;
+			UE_LOG(LogGenerateTextLocalizationReportCommandlet, Warning, TEXT("The conflict report filename %s has no extension. Defaulting the report to be generated as a .csv file."), *ConflictReportName);
 		}
+		ConflictReportName = FPaths::SetExtension(ConflictReportName, CSVExtension);
 	}
 	const FString ReportFilePath = (DestinationPath / ConflictReportName);
 
