Add the ability to specific material and material instances to cook via the cmd line to the CompileShaders commandlet.

- Adds the ability to specify which -targetplatform we should compile shaders for.  You can specify multiple platforms by adding them together: -targetplatform=Windows+MyNewPlatform
- Call BeginCacheForCookedPlatformData to compile shaders instead of CacheShaders since the former allows you to specify platform.

# Usage
## Specify collection
EngineTest -run=CompileShadersTestBed -targetplatform=Windows -collection=MyCollectionName

## Specify material list
EngineTest -run=CompileShadersTestBed -targetplatform=Windows -materials=/Game/MyNewMaterial.MyNewMaterial+/Game/Tests/Materials/CoolMaterial.CoolMaterial

## Specify both a collection and material list
EngineTest -run=CompileShadersTestBed -targetplatform=Windows -collection=MyCollectionName -materials=/Game/MyNewMaterial.MyNewMaterial+/Game/Tests/Materials/CoolMaterial.CoolMaterial

diff --git a/Engine/Source/Editor/UnrealEd/Private/Commandlets/CompileShadersTestBedCommandlet.cpp b/Engine/Source/Editor/UnrealEd/Private/Commandlets/CompileShadersTestBedCommandlet.cpp
--- a/Engine/Source/Editor/UnrealEd/Private/Commandlets/CompileShadersTestBedCommandlet.cpp
+++ b/Engine/Source/Editor/UnrealEd/Private/Commandlets/CompileShadersTestBedCommandlet.cpp
@@ -1,8 +1,17 @@
 // Copyright Epic Games, Inc. All Rights Reserved.
 
 #include "Commandlets/CompileShadersTestBedCommandlet.h"
+#include "AssetRegistry/AssetData.h"
+#include "AssetRegistry/AssetRegistryModule.h"
+#include "CollectionManagerModule.h"
+#include "CollectionManagerTypes.h"
 #include "GlobalShader.h"
+#include "ICollectionManager.h"
+#include "Interfaces/ITargetPlatform.h"
+#include "Interfaces/ITargetPlatformManagerModule.h"
 #include "Materials/Material.h"
+#include "Materials/MaterialInstance.h"
+#include "Materials/MaterialInstanceConstant.h"
 #include "ShaderCompiler.h"
 
 DEFINE_LOG_CATEGORY_STATIC(LogCompileShadersTestBedCommandlet, Log, All);
@@ -26,29 +35,143 @@ int32 UCompileShadersTestBedCommandlet::Main(const FString& Params)
 	{
 		UE_LOG(LogCompileShadersTestBedCommandlet, Log, TEXT("CompileShadersTestBed"));
 		UE_LOG(LogCompileShadersTestBedCommandlet, Log, TEXT("This commandlet compiles global and default material shaders.  Used to profile and test shader compilation."));
+		UE_LOG(LogCompileShadersTestBedCommandlet, Log, TEXT(" Optional: -collection=<name>                (You can also specify a collection of assets to narrow down the results e.g. if you maintain a collection that represents the actually used in-game assets)."));
+		UE_LOG(LogCompileShadersTestBedCommandlet, Log, TEXT(" Optional: -materials=<path1>+<path2>        (You can also specify a list of material asset paths separated by a '+' to narrow down the results."));
 		return 0;
 	}
 
 	PRIVATE_GAllowCommandletRendering = true;
 
+	IAssetRegistry& AssetRegistry = FModuleManager::Get().LoadModuleChecked<FAssetRegistryModule>("AssetRegistry").Get();
+	AssetRegistry.SearchAllAssets(true);
+
+	// Optional list of materials to compile.
+	TArray<FAssetData> MaterialList;
+
+	FARFilter Filter;
+
+	// Parse collection
+	FString CollectionName;
+	if (FParse::Value(*Params, TEXT("collection="), CollectionName, true))
+	{
+		if (!CollectionName.IsEmpty())
+		{
+			// Get the list of materials from a collection
+			Filter.PackagePaths.Add(FName(TEXT("/Game")));
+			Filter.bRecursivePaths = true;
+			Filter.ClassPaths.Add(UMaterial::StaticClass()->GetClassPathName());
+
+			FCollectionManagerModule& CollectionManagerModule = FCollectionManagerModule::GetModule();
+			CollectionManagerModule.Get().GetObjectsInCollection(FName(*CollectionName), ECollectionShareType::CST_All, Filter.ObjectPaths, ECollectionRecursionFlags::SelfAndChildren);
+
+			AssetRegistry.GetAssets(Filter, MaterialList);
+
+			Filter.ClassPaths.Empty();
+			Filter.ClassPaths.Add(UMaterialInstance::StaticClass()->GetClassPathName());
+			Filter.ClassPaths.Add(UMaterialInstanceConstant::StaticClass()->GetClassPathName());
+
+			AssetRegistry.GetAssets(Filter, MaterialList);
+		}
+	}
+
+	// Process -materials= switches separated by a '+'
+	TArray<FString> CmdLineMaterialEntries;
+	const TCHAR* MaterialsSwitchName = TEXT("Materials");
+	if (const FString* MaterialsSwitches = ParamVals.Find(MaterialsSwitchName))
 	{
-		TRACE_CPUPROFILER_EVENT_SCOPE(DefaultMaterials);
+		MaterialsSwitches->ParseIntoArray(CmdLineMaterialEntries, TEXT("+"));
+	}
 
-		for (int32 Domain = 0; Domain < MD_MAX; ++Domain)
+	if (CmdLineMaterialEntries.Num())
+	{
+		// re-use the filter and only filter based on the passed in objects.
+		Filter.ClassPaths.Empty();
+		Filter.ObjectPaths.Empty();
+		for (const FString& MaterialPath : CmdLineMaterialEntries)
 		{
-			UMaterial::GetDefaultMaterial(static_cast<EMaterialDomain>(Domain))->CacheShaders(EMaterialShaderPrecompileMode::Background);
+			const FName MaterialPathFName(MaterialPath);
+			if (!Filter.ObjectPaths.Contains(MaterialPathFName))
+			{
+				Filter.ObjectPaths.Add(MaterialPathFName);
+			}
 		}
-		GShaderCompilingManager->FinishAllCompilation();
+
+		AssetRegistry.GetAssets(Filter, MaterialList);
 	}
 
+	// For all active platforms
+	ITargetPlatformManagerModule* TPM = GetTargetPlatformManager();
+	const TArray<ITargetPlatform*>& Platforms = TPM->GetActiveTargetPlatforms();
+
+	for (ITargetPlatform* Platform : Platforms)
 	{
-		TRACE_CPUPROFILER_EVENT_SCOPE(GlobalShaders);
+		// Compile default materials
+		{
+			TRACE_CPUPROFILER_EVENT_SCOPE(DefaultMaterials);
 
-		CompileGlobalShaderMap(true);
+			for (int32 Domain = 0; Domain < MD_MAX; ++Domain)
+			{
+				UMaterial::GetDefaultMaterial(static_cast<EMaterialDomain>(Domain))->BeginCacheForCookedPlatformData(Platform);
+			}
+		}
 
-		TArray<int32> ShaderMapIds;
-		ShaderMapIds.Add(GlobalShaderMapId);
-		GShaderCompilingManager->FinishCompilation(TEXT("Global"), ShaderMapIds);
+		// Compile global shaders
+		{
+			TRACE_CPUPROFILER_EVENT_SCOPE(GlobalShaders);
+
+			TArray<FName> DesiredShaderFormats;
+			Platform->GetAllTargetedShaderFormats(DesiredShaderFormats);
+
+			for (int32 FormatIndex = 0; FormatIndex < DesiredShaderFormats.Num(); FormatIndex++)
+			{
+				const EShaderPlatform ShaderPlatform = ShaderFormatToLegacyShaderPlatform(DesiredShaderFormats[FormatIndex]);
+				CompileGlobalShaderMap(ShaderPlatform, Platform, false);
+			}
+		}
+
+		// Compile material shaders specified on the command line
+		{
+			TRACE_CPUPROFILER_EVENT_SCOPE(MaterialShaders);
+
+			// Sort the material lists by name so the order is stable.
+			Algo::SortBy(MaterialList, [](const FAssetData& AssetData) { return AssetData.ObjectPath; }, FNameLexicalLess());
+
+			for (const FAssetData& AssetData : MaterialList)
+			{
+				if (UMaterial* Material = Cast<UMaterial>(AssetData.GetAsset()))
+				{
+					Material->BeginCacheForCookedPlatformData(Platform);
+				}
+				else if (UMaterialInstance* MaterialInstance = Cast<UMaterialInstance>(AssetData.GetAsset()))
+				{
+					MaterialInstance->BeginCacheForCookedPlatformData(Platform);
+				}
+			}
+		}
+	}
+
+	// Block on all the jobs submitted above.
+	{
+		TRACE_CPUPROFILER_EVENT_SCOPE(BlockOnShaderCompiles);
+
+		GShaderCompilingManager->FinishAllCompilation();
+	}
+
+	// Perform cleanup and clear cached data for cooking.
+	{
+		TRACE_CPUPROFILER_EVENT_SCOPE(ClearCachedCookedPlatformData);
+
+		for (const FAssetData& AssetData : MaterialList)
+		{
+			if (UMaterial* Material = Cast<UMaterial>(AssetData.GetAsset()))
+			{
+				Material->ClearAllCachedCookedPlatformData();
+			}
+			else if (UMaterialInstance* MaterialInstance = Cast<UMaterialInstance>(AssetData.GetAsset()))
+			{
+				MaterialInstance->ClearAllCachedCookedPlatformData();
+			}
+		}
 	}
 
 	return 0;
