[Modes] Add a scroll box around toolkit's inline content in the main mode tab

diff --git a/Engine/Source/Editor/UnrealEd/Private/Toolkits/BaseToolkit.cpp b/Engine/Source/Editor/UnrealEd/Private/Toolkits/BaseToolkit.cpp
--- a/Engine/Source/Editor/UnrealEd/Private/Toolkits/BaseToolkit.cpp
+++ b/Engine/Source/Editor/UnrealEd/Private/Toolkits/BaseToolkit.cpp
@@ -586,9 +586,17 @@ void FModeToolkit::UpdatePrimaryModePanel()
 			}
 		}
 
-		if (GetInlineContent().IsValid() && InlineContentHolder.IsValid())
+		if (InlineContentHolder.IsValid())
 		{
-			InlineContentHolder->SetContent(GetInlineContent().ToSharedRef());
+			if (TSharedPtr<SWidget> InlineContent = GetInlineContent())
+			{
+				InlineContentHolder->SetContent(
+					SNew(SScrollBox)
+					+ SScrollBox::Slot()
+					[
+						InlineContent.ToSharedRef()
+					]);
+			}
 		}
 
 	}
