FBX import, added support for inbetween morph targets. They are imported as standalone blendshapes and their curves are baked to preserve the original animation

diff --git a/Engine/Source/Editor/UnrealEd/Private/SkeletalMeshEdit.cpp b/Engine/Source/Editor/UnrealEd/Private/SkeletalMeshEdit.cpp
--- a/Engine/Source/Editor/UnrealEd/Private/SkeletalMeshEdit.cpp
+++ b/Engine/Source/Editor/UnrealEd/Private/SkeletalMeshEdit.cpp
@@ -1377,6 +1377,170 @@ bool UnFbx::FFbxImporter::ImportCurveToAnimSequence(class UAnimSequence * Target
 	return false;
 }
 
+
+bool UnFbx::FFbxImporter::ImportRichCurvesToAnimSequence(UAnimSequence* TargetSequence, const TArray<FString>& CurveNames, const TArray<FRichCurve> RichCurves, int32 CurveFlags) const
+{
+	if (TargetSequence && CurveNames.Num() > 0 && CurveNames.Num() == RichCurves.Num())
+	{
+		for (int32 CurveIndex = 0; CurveIndex < CurveNames.Num(); ++CurveIndex)
+		{
+			FName Name = *CurveNames[CurveIndex];
+			USkeleton* Skeleton = TargetSequence->GetSkeleton();
+			const FSmartNameMapping* NameMapping = Skeleton->GetSmartNameContainer(USkeleton::AnimCurveMappingName);
+
+			// Add or retrieve curve
+			if (!NameMapping->Exists(Name))
+			{
+				// mark skeleton dirty
+				Skeleton->Modify();
+			}
+
+			FSmartName NewName;
+			Skeleton->AddSmartNameAndModify(USkeleton::AnimCurveMappingName, Name, NewName);
+
+			FAnimationCurveIdentifier FloatCurveId(NewName, ERawCurveTrackTypes::RCT_Float);
+
+
+			UAnimDataModel* DataModel = TargetSequence->GetDataModel();
+			IAnimationDataController& Controller = TargetSequence->GetController();
+
+			const FFloatCurve* TargetCurve = DataModel->FindFloatCurve(FloatCurveId);
+			
+			if (TargetCurve == nullptr)
+			{
+				// Need to add the curve first
+				Controller.AddCurve(FloatCurveId, AACF_DefaultCurve | CurveFlags);
+				TargetCurve = DataModel->FindFloatCurve(FloatCurveId);
+			}
+			else
+			{
+				// Need to update any of the flags
+				Controller.SetCurveFlags(FloatCurveId, CurveFlags | TargetCurve->GetCurveTypeFlags());
+			}
+		
+			// Should be valid at this point
+			ensure(TargetCurve);
+
+			Controller.UpdateCurveNamesFromSkeleton(Skeleton, ERawCurveTrackTypes::RCT_Float);
+
+			// Set actual keys on curve within the model
+			Controller.SetCurveKeys(FloatCurveId, RichCurves[CurveIndex].GetConstRefOfKeys());
+
+		}
+		
+		return true;
+	}
+
+	return false;
+}
+
+TArray<FRichCurve> UnFbx::FFbxImporter::ResolveWeightsForBlendShapeCurve(FRichCurve& ChannelWeightCurve, const TArray<float>& InbetweenFullWeights) const
+{
+	int32 NumInbetweens = InbetweenFullWeights.Num();
+	if (NumInbetweens == 0)
+	{
+		return { ChannelWeightCurve };
+	}
+
+	TArray<FRichCurve> Result;
+	Result.SetNum( NumInbetweens + 1 );
+
+	TArray<float> ResolvedInbetweenWeightsSample;
+	ResolvedInbetweenWeightsSample.SetNum( NumInbetweens );
+
+	for ( const FRichCurveKey& SourceKey : ChannelWeightCurve.Keys )
+	{
+		const float SourceTime = SourceKey.Time;
+		const float SourceValue = SourceKey.Value;
+
+		float ResolvedPrimarySample;
+
+		ResolveWeightsForBlendShape(InbetweenFullWeights,SourceValue, ResolvedPrimarySample, ResolvedInbetweenWeightsSample);
+
+		FRichCurve& PrimaryCurve = Result[ 0 ];
+		FKeyHandle PrimaryHandle = PrimaryCurve.AddKey( SourceTime, ResolvedPrimarySample );
+		PrimaryCurve.SetKeyInterpMode( PrimaryHandle, SourceKey.InterpMode );
+
+		for ( int32 InbetweenIndex = 0; InbetweenIndex < NumInbetweens; ++InbetweenIndex )
+		{
+			FRichCurve& InbetweenCurve = Result[ InbetweenIndex + 1 ];
+			FKeyHandle InbetweenHandle = InbetweenCurve.AddKey( SourceTime, ResolvedInbetweenWeightsSample[ InbetweenIndex ] );
+			InbetweenCurve.SetKeyInterpMode( InbetweenHandle, SourceKey.InterpMode );
+		}
+	}
+
+	return Result;
+}
+
+void UnFbx::FFbxImporter::ResolveWeightsForBlendShape(const TArray<float>& InbetweenFullWeights , float InWeight, float& OutMainWeight, TArray<float>& OutInbetweenWeights) const
+{
+	int32 NumInbetweens = InbetweenFullWeights.Num();
+	if ( NumInbetweens == 0 )
+	{
+		OutMainWeight = InWeight;
+		return;
+	}
+
+	OutInbetweenWeights.SetNumUninitialized( NumInbetweens );
+	for ( float& OutInbetweenWeight : OutInbetweenWeights )
+	{
+		OutInbetweenWeight = 0.0f;
+	}
+
+	if ( FMath::IsNearlyEqual( InWeight, 0.0f ) )
+	{
+		OutMainWeight = 0.0f;
+		return;
+	}
+	else if ( FMath::IsNearlyEqual( InWeight, 1.0f ) )
+	{
+		OutMainWeight = 1.0f;
+		return;
+	}
+
+	// Note how we don't care if UpperIndex/LowerIndex are beyond the bounds of the array here,
+	// as that signals when we're above/below all inbetweens
+	int32 UpperIndex = Algo::UpperBoundBy( InbetweenFullWeights, InWeight, []( const double& InbetweenWeight )
+	{
+		return InbetweenWeight;
+	} );
+	int32 LowerIndex = UpperIndex - 1;
+
+	float UpperWeight = 1.0f;
+	if ( UpperIndex <= NumInbetweens - 1 )
+	{
+		UpperWeight = InbetweenFullWeights[ UpperIndex ];
+	}
+
+	float LowerWeight = 0.0f;
+	if ( LowerIndex >= 0 )
+	{
+		LowerWeight = InbetweenFullWeights[ LowerIndex ];
+	}
+
+	UpperWeight = ( InWeight - LowerWeight ) / ( UpperWeight - LowerWeight );
+	LowerWeight = ( 1.0f - UpperWeight );
+
+	// We're between upper inbetween and the 1.0 weight
+	if ( UpperIndex > NumInbetweens - 1 )
+	{
+		OutMainWeight = UpperWeight;
+		OutInbetweenWeights[ NumInbetweens - 1 ] = LowerWeight;
+	}
+	// We're between 0.0 and the first inbetween weight
+	else if ( LowerIndex < 0 )
+	{
+		OutMainWeight = 0;
+		OutInbetweenWeights[ 0 ] = UpperWeight;
+	}
+	// We're between two inbetweens
+	else
+	{
+		OutInbetweenWeights[ UpperIndex ] = UpperWeight;
+		OutInbetweenWeights[ LowerIndex ] = LowerWeight;
+	}
+}
+
 template<typename AttributeType>
 void FillCurveAttributeToBone(TArray<float>& OutFrameTimes, TArray<AttributeType>& OutFrameValues, const FbxAnimCurve* FbxCurve, const FbxTimeSpan& AnimTimeSpan, TFunctionRef<AttributeType(const FbxAnimCurveKey*, const FbxTime*)> EvaluationFunction)
 {
@@ -1753,15 +1917,87 @@ void UnFbx::FFbxImporter::ImportBlendShapeCurves(FAnimCurveImportSettings& AnimI
 							SlowTaskNode.FrameMessage = CurrentExportMessage;
 							SlowTaskChannel.EnterProgressFrame(1);
 							bUpdatedProgress = true;
-							// now see if we have one already exists. If so, just overwrite that. if not, add new one. 
 
-							if (ImportCurveToAnimSequence(AnimImportSettings.DestSeq, *ChannelName, Curve, 0, AnimImportSettings.AnimTimeSpan, 0.01f /** for some reason blend shape values are coming as 100 scaled **/))
+							const int32 TargetShapeCount = Channel->GetTargetShapeCount();
+
+							if (ensure(TargetShapeCount > 0))
 							{
-								OutKeyCount = FMath::Max(OutKeyCount, Curve->KeyGetCount());
-								// this one doesn't reset Material curve to false, it just accumulate if true. 
-								MySkeleton->AccumulateCurveMetaData(*ChannelName, false, true);
-							}
+								if (TargetShapeCount == 1)
+								{
+									// now see if we have one already exists. If so, just overwrite that. if not, add new one. 
 
+									if (ImportCurveToAnimSequence(AnimImportSettings.DestSeq, *ChannelName, Curve, 0, AnimImportSettings.AnimTimeSpan, 0.01f /** for some reason blend shape values are coming as 100 scaled **/))
+									{
+										OutKeyCount = FMath::Max(OutKeyCount, Curve->KeyGetCount());
+										// this one doesn't reset Material curve to false, it just accumulate if true. 
+										MySkeleton->AccumulateCurveMetaData(*ChannelName, false, true);
+									}
+								}
+								else
+								{
+									// the blend shape channel can have multiple inbetween target shapes
+									// since the engine does not directly support inbetween morphs at runtime,
+									// each inbetween is imported as a standalone blendshape.
+									// as a result we have to create a curve for each one of those inbetweens
+									// and modify the primary channel curve such that their combined effect preserves
+									// the original animation
+							
+									TArray<FString> CurveNames;
+									CurveNames.Reserve(TargetShapeCount);
+
+									// in fbx the primary shape is the last shape, however to make
+									// the code more similar to usd importer, we deal with the primary shape separately
+									CurveNames.Add(ChannelName);
+
+									// ignoring the last shape because it is not a inbetween, i.e. it is the primary shape
+									int32 InbetweenCount = TargetShapeCount - 1;
+							
+									TArrayView<double> FbxInbetweenFullWeights = {Channel->GetTargetShapeFullWeights(), InbetweenCount};
+
+									TArray<float> InbetweenFullWeights;
+									InbetweenFullWeights.Reserve(InbetweenCount);
+									/** for some reason blend shape values are coming as 100 scaled, so a transform is needed to scale it to 0-1 **/
+									Algo::Transform(FbxInbetweenFullWeights, InbetweenFullWeights, [](double Input){ return Input * 0.01f; });
+
+									// collect inbetween shape names
+									for (int32 InbetweenIndex = 0; InbetweenIndex < InbetweenCount; ++InbetweenIndex)
+									{
+										FbxShape* Shape = Channel->GetTargetShape(InbetweenIndex);
+										CurveNames.Add(MakeName(Shape->GetName()));
+									}
+						
+									// first convert fbx curve to rich curve
+									FRichCurve ChannelWeightCurve;
+									ImportCurve(Curve, ChannelWeightCurve,  AnimImportSettings.AnimTimeSpan, false, 0.01f /** for some reason blend shape values are coming as 100 scaled **/);
+									if (ensure(AnimImportSettings.DestSeq))
+									{
+#if WITH_EDITORONLY_DATA
+										ChannelWeightCurve.BakeCurve(1.0f / AnimImportSettings.DestSeq->ImportResampleFramerate);
+#endif
+									}
+									
+									// use the primary curve to generate inbetween shape curves + a modified primary curve
+									TArray<FRichCurve> Results = ResolveWeightsForBlendShapeCurve(ChannelWeightCurve, InbetweenFullWeights);
+
+									for (FRichCurve& Result : Results)
+									{
+										if (ImportOptions->bRemoveRedundantKeys)
+										{
+											Result.RemoveRedundantAutoTangentKeys(SMALL_NUMBER);
+										}
+									}
+
+									if (ImportRichCurvesToAnimSequence(AnimImportSettings.DestSeq, CurveNames, Results, 0))
+									{
+										OutKeyCount = FMath::Max(OutKeyCount, Curve->KeyGetCount());
+										for (const FString& CurveName : CurveNames)
+										{
+											// this one doesn't reset Material curve to false, it just accumulate if true. 
+											MySkeleton->AccumulateCurveMetaData(*CurveName, false, true);
+										}	
+									}	
+								}
+							}
 						}
 						else
 						{
diff --git a/Engine/Source/Editor/UnrealEd/Public/FbxImporter.h b/Engine/Source/Editor/UnrealEd/Public/FbxImporter.h
--- a/Engine/Source/Editor/UnrealEd/Public/FbxImporter.h
+++ b/Engine/Source/Editor/UnrealEd/Public/FbxImporter.h
@@ -1911,6 +1911,25 @@ private:
 	 */
 	bool ImportCurveToAnimSequence(class UAnimSequence * TargetSequence, const FString& CurveName, const FbxAnimCurve* FbxCurve, int32 CurveFlags,const FbxTimeSpan& AnimTimeSpan, float ValueScale = 1.f) const;
 
+	/**
+	 * Import rich Curves to anim sequence
+	 */
+	bool ImportRichCurvesToAnimSequence(class UAnimSequence * TargetSequence, const TArray<FString>& CurveNames, const TArray<FRichCurve> RichCurves, int32 CurveFlags) const;
+	
+	/**
+	 * Given a primary blend shape channel curve and inbetween target full weights,
+	 * generate curves for each target as if they are standalone blend shapes
+	 * while preserving the animation
+	 */
+	TArray<FRichCurve> ResolveWeightsForBlendShapeCurve(FRichCurve& ChannelWeightCurve, const TArray<float>& InbetweenFullWeights) const;
+	
+	/**
+	 * Given a primary blend shape channel curve value and inbetween target full weights,
+	 * calculate the curve value for each target as if they are standalone blend shapes
+	 * while preserving the animation
+	 */
+	void ResolveWeightsForBlendShape(const TArray<float>& InbetweenFullWeights , float InWeight, float& OutMainWeight, TArray<float>& OutInbetweenWeights) const;
+	
 	/**
 	 * Import custom attribute (curve or not) to the associated bone.
 	 *
