Removed ensure that was firing inappropriately

#rnx

diff --git a/Engine/Source/Editor/UnrealEd/Private/ComponentTypeRegistry.cpp b/Engine/Source/Editor/UnrealEd/Private/ComponentTypeRegistry.cpp
--- a/Engine/Source/Editor/UnrealEd/Private/ComponentTypeRegistry.cpp
+++ b/Engine/Source/Editor/UnrealEd/Private/ComponentTypeRegistry.cpp
@@ -492,7 +492,7 @@ void FComponentTypeRegistryData::ForceRefreshComponentList()
 		FComponentClassComboEntryPtr NewBPClass = MakeShareable(new FComponentClassComboEntry(NewComponentsHeading, UActorComponent::StaticClass(), true, EComponentCreateAction::CreateNewBlueprintClass));
 		ComponentClassList.Add(NewBPClass);
 
-		if (ensure(GUnrealEd) && GUnrealEd->GetUnrealEdOptions()->IsCPPAllowed())
+		if (GUnrealEd && GUnrealEd->GetUnrealEdOptions()->IsCPPAllowed())
 		{
 			FComponentClassComboEntryPtr NewCPPClass = MakeShareable(new FComponentClassComboEntry(NewComponentsHeading, UActorComponent::StaticClass(), true, EComponentCreateAction::CreateNewCPPClass));
 			ComponentClassList.Add(NewCPPClass);
