Fixed a conflicting localization text key/value.

diff --git a/Engine/Source/Editor/UnrealEd/Private/PlayLevel.cpp b/Engine/Source/Editor/UnrealEd/Private/PlayLevel.cpp
--- a/Engine/Source/Editor/UnrealEd/Private/PlayLevel.cpp
+++ b/Engine/Source/Editor/UnrealEd/Private/PlayLevel.cpp
@@ -457,18 +457,19 @@ void UEditorEngine::EndPlayMap()
 
 		FFormatNamedArguments Arguments;
 		Arguments.Add(TEXT("Path"), FText::FromString(Paths[i]));
+		Arguments.Add(TEXT("Object"), FText::FromString(Object->GetFullName()));
 
 		// We cannot safely recover from this.
 		if (UObjectBaseUtility::IsPendingKillEnabled())
 		{
-			Arguments.Add(TEXT("Object"), FText::FromString(Object->GetFullName()));
 			checkf(false, *FText::Format(
 				LOCTEXT("PIEObjectStillReferenced", "Object '{Object}' from PIE level still referenced. Shortest path from root: {Path}"), Arguments).ToString());
 		}
 		else
 		{
 			// Nonfatal error, we will rename objects to try and recover. 
-			FText ErrorMessage = FText::Format(LOCTEXT("PIEObjectStillReferenced", "Object from PIE level still referenced. Shortest path from root: {Path}"), Arguments);
+			FText ErrorMessage = FText::Format(
+				LOCTEXT("PIEAnObjectStillReferenced", "Object '{Object}' from PIE level still referenced. Shortest path from root: {Path}"), Arguments);
 			FMessageLog(NAME_CategoryPIE).Error()
 				->AddToken(FUObjectToken::Create(Object, FText::FromString(Object->GetFullName())))
 				->AddToken(FTextToken::Create(ErrorMessage));
