Fixed issue with Color Calibration cube spawning on any button hit.

diff --git a/Engine/Source/Editor/UnrealEd/Private/EditorViewportClient.cpp b/Engine/Source/Editor/UnrealEd/Private/EditorViewportClient.cpp
--- a/Engine/Source/Editor/UnrealEd/Private/EditorViewportClient.cpp
+++ b/Engine/Source/Editor/UnrealEd/Private/EditorViewportClient.cpp
@@ -3204,11 +3204,12 @@ bool FEditorViewportClient::IsCommandChordPressed(const TSharedPtr<FUICommandInf
 		EMultipleKeyBindingIndex ChordIndex = static_cast<EMultipleKeyBindingIndex> (i);
 		const FInputChord& Chord = *InCommand->GetActiveChord(ChordIndex);
 
-		bIsChordPressed |= (Chord.NeedsControl() == IsCtrlPressed())
+		bIsChordPressed |= Chord.IsValidChord()
+			&& (Chord.NeedsControl() == IsCtrlPressed())
 			&& (Chord.NeedsAlt() == IsAltPressed())
 			&& (Chord.NeedsShift() == IsShiftPressed())
 			&& (Chord.NeedsCommand() == IsCmdPressed())
-			&& Viewport->KeyState(InOptionalKey.IsValid() ? InOptionalKey : Chord.Key) == true;
+			&& (InOptionalKey.IsValid() ? (Chord.Key == InOptionalKey && Viewport->KeyState(InOptionalKey)) : Viewport->KeyState(Chord.Key));
 	}
 	return bIsChordPressed;
 }
