Cooker Multiprocess: Add CookWorker messages for config and package assignment.
#rnx

diff --git a/Engine/Source/Editor/UnrealEd/Classes/CookOnTheSide/CookOnTheFlyServer.h b/Engine/Source/Editor/UnrealEd/Classes/CookOnTheSide/CookOnTheFlyServer.h
--- a/Engine/Source/Editor/UnrealEd/Classes/CookOnTheSide/CookOnTheFlyServer.h
+++ b/Engine/Source/Editor/UnrealEd/Classes/CookOnTheSide/CookOnTheFlyServer.h
@@ -126,6 +126,8 @@ namespace UE::Cook
 {
 	class FBuildDefinitions;
 	class FCookDirector;
+	class FCookWorkerClient;
+	class FCookWorkerServer;
 	class FRequestCluster;
 	class FRequestQueue;
 	class FSaveCookedPackageContext;
@@ -552,6 +554,7 @@ public:
 	void Tick(float DeltaTime) override;
 	bool IsTickable() const override;
 	ECookMode::Type GetCookMode() const { return CurrentCookMode; }
+	ECookInitializationFlags GetCookFlags() const { return CookFlags; }
 
 
 	/**
@@ -1312,6 +1315,7 @@ private:
 	TUniquePtr<UE::Cook::IWorkerRequests> WorkerRequests;
 	TUniquePtr<UE::Cook::FBuildDefinitions> BuildDefinitions;
 	TUniquePtr<UE::Cook::FCookDirector> CookDirector;
+	TUniquePtr<UE::Cook::FCookWorkerClient> CookWorkerClient;
 
 	TArray<UE::Cook::FCookSavePackageContext*> SavePackageContexts;
 	/** Objects that were collected during the single-threaded PreGarbageCollect callback and that should be reported as referenced in CookerAddReferencedObjects. */
@@ -1345,7 +1349,10 @@ private:
 	EIdleStatus IdleStatus = EIdleStatus::Done;
 
 	friend UE::Cook::FBeginCookConfigSettings;
+	friend UE::Cook::FInitializeConfigSettings;
 	friend UE::Cook::FCookDirector;
+	friend UE::Cook::FCookWorkerClient;
+	friend UE::Cook::FCookWorkerServer;
 	friend UE::Cook::FGeneratorPackage;
 	friend UE::Cook::FPackageData;
 	friend UE::Cook::FPendingCookedPlatformData;
diff --git a/Engine/Source/Editor/UnrealEd/Private/CookOnTheFlyServer.cpp b/Engine/Source/Editor/UnrealEd/Private/CookOnTheFlyServer.cpp
--- a/Engine/Source/Editor/UnrealEd/Private/CookOnTheFlyServer.cpp
+++ b/Engine/Source/Editor/UnrealEd/Private/CookOnTheFlyServer.cpp
@@ -28,6 +28,7 @@
 #include "Cooker/CookRequestCluster.h"
 #include "Cooker/CookRequests.h"
 #include "Cooker/CookTypes.h"
+#include "Cooker/CookWorkerClient.h"
 #include "Cooker/DiffPackageWriter.h"
 #include "Cooker/IoStoreCookOnTheFlyRequestManager.h"
 #include "Cooker/LooseCookedPackageWriter.h"
@@ -1293,6 +1294,12 @@ uint32 UCookOnTheFlyServer::TickCookWorker()
 	UE::Cook::FTickStackData StackData(MAX_flt, ECookTickFlags::None);
 
 	TickMainCookLoop(StackData);
+	if (StackData.bCookCancelled)
+	{
+		CancelAllQueues();
+		ShutdownCookSession();
+		SetIdleStatus(StackData, EIdleStatus::Done);
+	}
 
 	return StackData.ResultFlags;
 }
@@ -2077,6 +2084,10 @@ void UCookOnTheFlyServer::InitializePollables()
 		DirectorPollable = new FPollable(TEXT("CookDirector"), 1.0f, 1.0f, [this](FTickStackData& StackData) { CookDirector->TickFromSchedulerThread(); });
 		Pollables.Add(FPollableQueueKey(DirectorPollable));
 	}
+	if (CookWorkerClient)
+	{
+		Pollables.Emplace(new FPollable(TEXT("CookWorkerClient"), 1.0f, 1.0f, [this](FTickStackData& StackData) { CookWorkerClient->TickFromSchedulerThread(StackData); }));
+	}
 	Pollables.Heapify();
 
 	PollNextTimeSeconds = 0.;
@@ -2222,13 +2233,44 @@ UCookOnTheFlyServer::ECookAction UCookOnTheFlyServer::DecideNextCookAction(UE::C
 		}
 	}
 
-	if (IsCookOnTheFlyMode())
+	if (IsCookOnTheFlyMode() || IsCookWorkerMode())
 	{
+		// These modes are not done until a manual trigger, so continue polling idle
 		if (CurrentTime >= PollNextTimeIdleSeconds)
 		{
 			// Polling does not impact idle status
 			return ECookAction::PollIdle;
 		}
+		if (IsCookOnTheFlyMode())
+		{
+			SetIdleStatus(StackData, EIdleStatus::Done);
+			return ECookAction::Done;
+		}
+		else
+		{
+			SetIdleStatus(StackData, EIdleStatus::Idle);
+			return ECookAction::WaitForAsync;
+		}
+	}
+
+	// We're in the CookComplete phase, pump the special cases in this phase
+	// and return WaitForAsync until they are complete
+	if (CookDirector)
+	{
+		bool bCompleted;
+		CookDirector->PumpCookComplete(bCompleted);
+		if (!bCompleted)
+		{
+			// Continue polling idle
+			if (CurrentTime >= PollNextTimeIdleSeconds)
+			{
+				// Polling does not impact idle status
+				return ECookAction::PollIdle;
+			}
+
+			SetIdleStatus(StackData, EIdleStatus::Idle);
+			return ECookAction::WaitForAsync;
+		}
 	}
 
 	SetIdleStatus(StackData, EIdleStatus::Done);
@@ -5803,21 +5845,7 @@ void FInitializeConfigSettings::LoadLocal(const FString& InOutputDirectoryOverri
 
 void UCookOnTheFlyServer::SetInitializeConfigSettings(UE::Cook::FInitializeConfigSettings&& Settings)
 {
-	OutputDirectoryOverride = MoveTemp(Settings.OutputDirectoryOverride);
-	MaxPrecacheShaderJobs = Settings.MaxPrecacheShaderJobs;
-	MaxConcurrentShaderJobs = Settings.MaxConcurrentShaderJobs;
-	PackagesPerGC = Settings.PackagesPerGC;
-	MemoryExpectedFreedToSpreadRatio = Settings.MemoryExpectedFreedToSpreadRatio;
-	IdleTimeToGC = Settings.IdleTimeToGC;
-	MemoryMaxUsedVirtual = Settings.MemoryMaxUsedVirtual;
-	MemoryMaxUsedPhysical = Settings.MemoryMaxUsedPhysical;
-	MemoryMinFreeVirtual = Settings.MemoryMinFreeVirtual;
-	MemoryMinFreePhysical = Settings.MemoryMinFreePhysical;
-	MinFreeUObjectIndicesBeforeGC = Settings.MinFreeUObjectIndicesBeforeGC;
-	MaxNumPackagesBeforePartialGC = Settings.MaxNumPackagesBeforePartialGC;
-	ConfigSettingDenyList = MoveTemp(Settings.ConfigSettingDenyList);
-	MaxAsyncCacheForType = MoveTemp(Settings.MaxAsyncCacheForType);
-	bHybridIterativeDebug = Settings.bHybridIterativeDebug;
+	Settings.MoveToLocal(*this);
 
 	MaxPreloadAllocated = 16;
 	DesiredSaveQueueLength = 8;
@@ -5863,14 +5891,14 @@ bool UCookOnTheFlyServer::TryInitializeCookWorker()
 	{
 		return false;
 	}
+	CookWorkerClient = MakeUnique<FCookWorkerClient>(*this);
 	TUniquePtr<FWorkerRequestsRemote> RemoteTasks = MakeUnique<FWorkerRequestsRemote>(*this);
-	ECookInitializationFlags LocalCookFlags;
-	if (!RemoteTasks->TryConnect(MoveTemp(ConnectInfo), LocalCookFlags))
+	if (!CookWorkerClient->TryConnect(MoveTemp(ConnectInfo)))
 	{
 		return false;
 	}
 	WorkerRequests.Reset(RemoteTasks.Release());
-	Initialize(ECookMode::CookWorker, LocalCookFlags, FString());
+	Initialize(ECookMode::CookWorker, CookWorkerClient->GetCookInitializationFlags(), FString());
 	StartCookAsCookWorker();
 	return true;
 }
@@ -7808,9 +7836,9 @@ void UCookOnTheFlyServer::BeginCookStartShaderCodeLibrary(FBeginCookContext& Beg
 
 void UCookOnTheFlyServer::BeginCookFinishShaderCodeLibrary(FBeginCookContext& BeginContext)
 {
-	check(BeginContext.StartupOptions && IsDirectorCookByTheBook()); // CookByTheBook only for now
+	check(IsDirectorCookByTheBook()); // CookByTheBook only for now
 	// don't resave the global shader map files in dlc
-	if (!IsCookingDLC() && !EnumHasAnyFlags(BeginContext.StartupOptions->CookOptions, ECookByTheBookOptions::ForceDisableSaveGlobalShaders))
+	if (!IsCookWorkerMode() && !IsCookingDLC() && !EnumHasAnyFlags(CookByTheBookOptions->StartupOptions, ECookByTheBookOptions::ForceDisableSaveGlobalShaders))
 	{
 		OpenGlobalShaderLibrary();
 
@@ -8348,6 +8376,10 @@ void UCookOnTheFlyServer::CookByTheBookFinished()
 
 void UCookOnTheFlyServer::ShutdownCookSession()
 {
+	if (CookDirector)
+	{
+		CookDirector->ShutdownCookSession();
+	}
 	for (UE::Cook::FPackageData* PackageData : *PackageDatas)
 	{
 		PackageData->DestroyGeneratorPackage();
@@ -9131,6 +9163,7 @@ FBeginCookContext UCookOnTheFlyServer::CreateBeginCookByTheBookContext(const FCo
 	BeginContext.StartupOptions = &StartupOptions;
 	const ECookByTheBookOptions& CookOptions = StartupOptions.CookOptions;
 	bZenStore = !!(CookOptions & ECookByTheBookOptions::ZenStore);
+	CookByTheBookOptions->StartupOptions = CookOptions;
 	CookByTheBookOptions->CookTime = 0.0f;
 	CookByTheBookOptions->CookStartTime = FPlatformTime::Seconds();
 	CookByTheBookOptions->bGenerateStreamingInstallManifests = StartupOptions.bGenerateStreamingInstallManifests;
@@ -9206,6 +9239,7 @@ void UCookOnTheFlyServer::StartCookAsCookWorker()
 	{
 		GShaderCompilingManager->SkipShaderCompilation(true);
 	}
+	CookWorkerClient->DoneWithInitialSettings();
 
 	// Initialize systems referenced by later stages or that need to start early for async performance
 	// MPCOOKTODO: Need to send information from ShaderCodeLibrary to director
@@ -9254,12 +9288,24 @@ void UCookOnTheFlyServer::ShutdownCookAsCookWorker()
 
 FBeginCookContext UCookOnTheFlyServer::CreateCookWorkerContext()
 {
-	// MPCOOKTODO: Copy CookByTheBookOptions and CookOnTheFlyOptions from CookDirector
-	// MPCOOKTODO: Copy TargetPlatforms from CookDirector
-	// MPCOOKTODO: GenerateLocalizationReferences(BeginContext.StartupOptions->CookCultures);
-	// MPCOOKTODO: RecordDLCPackagesFromBaseGame(BeginContext);
+	FBeginCookContext BeginContext(*this);
+	*CookByTheBookOptions = CookWorkerClient->ConsumeCookByTheBookOptions();
+	*CookOnTheFlyOptions = CookWorkerClient->ConsumeCookOnTheFlyOptions();
+	BeginContext.TargetPlatforms = CookWorkerClient->GetTargetPlatforms();
 
-	return FBeginCookContext(*this);
+	TArray<ITargetPlatform*> UniqueTargetPlatforms = BeginContext.TargetPlatforms;
+	Algo::Sort(UniqueTargetPlatforms);
+	UniqueTargetPlatforms.SetNum(Algo::Unique(UniqueTargetPlatforms));
+	checkf(UniqueTargetPlatforms.Num() == BeginContext.TargetPlatforms.Num(), TEXT("List of TargetPlatforms received from Director was not unique."));
+
+	BeginContext.PlatformContexts.SetNum(BeginContext.TargetPlatforms.Num());
+	for (int32 Index = 0; Index < BeginContext.TargetPlatforms.Num(); ++Index)
+	{
+		BeginContext.PlatformContexts[Index].TargetPlatform = BeginContext.TargetPlatforms[Index];
+		// PlatformContext.PlatformData is currently null and is set in SelectSessionPlatforms
+	}
+
+	return BeginContext;
 }
 
 bool UCookOnTheFlyServer::IsCookWorkerMode() const
@@ -9300,7 +9346,7 @@ void UCookOnTheFlyServer::GenerateInitialRequests(FBeginCookContext& BeginContex
 	const TArray<FString>& CookMaps = BeginContext.StartupOptions->CookMaps;
 	const TArray<FString>& CookDirectories = BeginContext.StartupOptions->CookDirectories;
 	const TArray<FString>& IniMapSections = BeginContext.StartupOptions->IniMapSections;
-	ECookByTheBookOptions CookOptions = BeginContext.StartupOptions->CookOptions;
+	ECookByTheBookOptions CookOptions = CookByTheBookOptions->StartupOptions;
 	CollectFilesToCook(FilesInPath, FilesInPathInstigators, CookMaps, CookDirectories, IniMapSections, CookOptions, TargetPlatforms, GameDefaultObjects);
 
 	// Add soft/hard startup references after collecting requested files and handling empty requests
@@ -9410,7 +9456,7 @@ void UCookOnTheFlyServer::RecordDLCPackagesFromBaseGame(FBeginCookContext& Begin
 		return;
 	}
 
-	const ECookByTheBookOptions& CookOptions = BeginContext.StartupOptions->CookOptions;
+	const ECookByTheBookOptions& CookOptions = CookByTheBookOptions->StartupOptions;
 	const FString& BasedOnReleaseVersion = BeginContext.StartupOptions->BasedOnReleaseVersion;
 
 	// If we're cooking against a fixed base, we don't need to verify the packages exist on disk, we simply want to use the Release Data 
diff --git a/Engine/Source/Editor/UnrealEd/Private/Cooker/CookDirector.cpp b/Engine/Source/Editor/UnrealEd/Private/Cooker/CookDirector.cpp
--- a/Engine/Source/Editor/UnrealEd/Private/Cooker/CookDirector.cpp
+++ b/Engine/Source/Editor/UnrealEd/Private/Cooker/CookDirector.cpp
@@ -25,11 +25,10 @@ FCookDirector::FCookDirector(UCookOnTheFlyServer& InCOTFS)
 	WorkersStalledStartTimeSeconds = MAX_flt;
 	WorkersStalledWarnTimeSeconds = MAX_flt;
 
-	const TCHAR* CommandLine = FCommandLine::Get();
-	GConfig->GetInt(TEXT("CookSettings"), TEXT("CookWorkerCount"), DesiredNumRemoteWorkers, GEditorIni);
-	FParse::Value(CommandLine, TEXT("-CookWorkerCount="), DesiredNumRemoteWorkers);
+	ParseDesiredNumRemoteWorkers();
 	WorkerConnectPort = Sockets::COOKDIRECTOR_DEFAULT_REQUEST_CONNECTION_PORT;
-	FParse::Value(CommandLine, TEXT("-CookDirectorListenPort="), WorkerConnectPort);
+	FParse::Value(FCommandLine::Get(), TEXT("-CookDirectorListenPort="), WorkerConnectPort);
+	ParseShowWorkerOption();
 
 	if (DesiredNumRemoteWorkers > 0)
 	{
@@ -45,6 +44,39 @@ FCookDirector::FCookDirector(UCookOnTheFlyServer& InCOTFS)
 	{
 		TryCreateWorkerConnectSocket();
 	}
+	UE_LOG(LogCook, Display, TEXT("MultiprocessCook is enabled with %d CookWorker processes."), DesiredNumRemoteWorkers);
+}
+
+void FCookDirector::ParseDesiredNumRemoteWorkers()
+{
+	DesiredNumRemoteWorkers = 4;
+	GConfig->GetInt(TEXT("CookSettings"), TEXT("CookWorkerCount"), DesiredNumRemoteWorkers, GEditorIni);
+	FParse::Value(FCommandLine::Get(), TEXT("-CookWorkerCount="), DesiredNumRemoteWorkers);
+}
+
+void FCookDirector::ParseShowWorkerOption()
+{
+	FString Text;
+	const TCHAR* CommandLine = FCommandLine::Get();
+	if (!FParse::Value(CommandLine, TEXT("-ShowCookWorker="), Text))
+	{
+		if (FParse::Param(CommandLine, TEXT("ShowCookWorker")))
+		{
+			Text = TEXT("SeparateWindows");
+		}
+	}
+
+	if (Text == TEXT("CombinedLogs")) { ShowWorkerOption = EShowWorker::CombinedLogs; }
+	else if (Text == TEXT("SeparateLogs")) { ShowWorkerOption = EShowWorker::SeparateLogs; }
+	else if (Text == TEXT("SeparateWindows")) { ShowWorkerOption = EShowWorker::SeparateWindows; }
+	else
+	{
+		if (!Text.IsEmpty())
+		{
+			UE_LOG(LogCook, Warning, TEXT("Invalid selection \"%s\" for -ShowCookWorker."), *Text);
+		}
+		ShowWorkerOption = EShowWorker::CombinedLogs;
+	}
 }
 
 FCookDirector::~FCookDirector()
@@ -155,6 +187,46 @@ void FCookDirector::TickFromSchedulerThread()
 	SetWorkersStalled(bIsStalled);
 }
 
+void FCookDirector::PumpCookComplete(bool& bCompleted)
+{
+	TickWorkerConnects();
+	for (TPair<int32, TUniquePtr<FCookWorkerServer>>& Pair : RemoteWorkers)
+	{
+		FCookWorkerServer& RemoteWorker = *Pair.Value;
+		RemoteWorker.PumpCookComplete();
+		if (RemoteWorker.IsShuttingDown())
+		{
+			TUniquePtr<FCookWorkerServer>& Existing = ShuttingDownWorkers.FindOrAdd(&RemoteWorker);
+			check(!Existing); // We should not be able to send the same pointer into ShuttingDown twice
+		}
+	}
+	TickWorkerShutdowns();
+	bCompleted = RemoteWorkers.Num() == 0;
+	SetWorkersStalled(!bCompleted);
+}
+
+void FCookDirector::ShutdownCookSession()
+{
+	while (!RemoteWorkers.IsEmpty())
+	{
+		AbortWorker(TMap<int32, TUniquePtr<FCookWorkerServer>>::TIterator(RemoteWorkers).Value()->GetWorkerId());
+	}
+	for (;;)
+	{
+		TickWorkerShutdowns();
+		if (ShuttingDownWorkers.IsEmpty())
+		{
+			break;
+		}
+		constexpr double SleepSeconds = 0.010;
+		FPlatformProcess::Sleep(SleepSeconds);
+	}
+	PendingConnections.Reset();
+
+	// Restore the FCookDirector to its original state so that it is ready for a new session
+	ParseDesiredNumRemoteWorkers();
+}
+
 void FCookDirector::SetWorkersStalled(bool bInWorkersStalled)
 {
 	if (bInWorkersStalled != bWorkersStalled)
@@ -376,6 +448,7 @@ void FCookDirector::TickWorkerShutdowns()
 	{
 		FCookWorkerServer& Worker = *Iter.Key();
 		check(Iter.Value()); // All non-yet-moved workers should have been moved by the AbortWorker calls above
+		Worker.TickFromSchedulerThread();
 		if (Worker.IsShutdownComplete())
 		{
 			// Worker is now deleted, do not access
@@ -402,7 +475,8 @@ FString FCookDirector::GetWorkerCommandLine(FWorkerId WorkerId)
 				Token.StartsWith(TEXT("-CookCultures")) ||
 				Token.StartsWith(TEXT("-CookDirectorCount=")) ||
 				Token.StartsWith(TEXT("-CookDirectorHost=")) ||
-				Token.StartsWith(TEXT("-CookWorkerId="))
+				Token.StartsWith(TEXT("-CookWorkerId=")) ||
+				Token.StartsWith(TEXT("-ShowCookWorker"))
 				)
 			{
 				return;
diff --git a/Engine/Source/Editor/UnrealEd/Private/Cooker/CookDirector.h b/Engine/Source/Editor/UnrealEd/Private/Cooker/CookDirector.h
--- a/Engine/Source/Editor/UnrealEd/Private/Cooker/CookDirector.h
+++ b/Engine/Source/Editor/UnrealEd/Private/Cooker/CookDirector.h
@@ -11,9 +11,9 @@
 #include "CookTypes.h"
 #include "Memory/SharedBuffer.h"
 #include "Misc/Guid.h"
-#include "Serialization/CompactBinary.h"
 #include "Templates/UniquePtr.h"
 
+class FCbObjectView;
 class FCbWriter;
 class UCookOnTheFlyServer;
 namespace UE::Cook { class FCookWorkerServer; }
@@ -39,6 +39,19 @@ public:
 	void RemoveFromWorker(FPackageData& PackageData);
 	/** Periodic tick function. Sends/Receives messages to CookWorkers. */
 	void TickFromSchedulerThread();
+	/** Called when the COTFS Server has detected all packages are complete. Tells the CookWorkers to flush messages and exit. */
+	void PumpCookComplete(bool& bOutCompleted);
+	/** Called when a session ends. The Director blocks on shutdown of all CookWorkers and returns state to before session started. */
+	void ShutdownCookSession();
+
+	/** Enum specifying how CookWorker log output should be shown. */
+	enum EShowWorker
+	{
+		CombinedLogs,
+		SeparateLogs,
+		SeparateWindows, // Implies SeparateLogs as well
+	};
+	EShowWorker GetShowWorkerOption() const { return ShowWorkerOption; }
 
 private:
 	/** CookWorker connections that have not yet identified which CookWorker they are. */
@@ -57,6 +70,10 @@ private:
 		FSocket* Socket = nullptr;
 		UE::CompactBinaryTCP::FReceiveBuffer Buffer;
 	};
+	/** Helper for constructor parsing. */
+	void ParseDesiredNumRemoteWorkers();
+	/** Helper for constructor parsing. */
+	void ParseShowWorkerOption();
 	/** Initialization helper: create the listen socket. */
 	bool TryCreateWorkerConnectSocket();
 	/** Initialization helper: add the local Server for a remote worker, worker process not yet created. */
@@ -87,9 +104,10 @@ private:
 	FSocket* WorkerConnectSocket = nullptr;
 	double WorkersStalledStartTimeSeconds = 0.;
 	double WorkersStalledWarnTimeSeconds = 0.;
-	int32 DesiredNumRemoteWorkers = 4;
+	int32 DesiredNumRemoteWorkers = 0;
 	int32 WorkerConnectPort = 0;
 	bool bWorkersStalled = false;
+	EShowWorker ShowWorkerOption = EShowWorker::CombinedLogs;
 
 	friend class UE::Cook::FCookWorkerServer;
 };
diff --git a/Engine/Source/Editor/UnrealEd/Private/Cooker/CookPackageData.cpp b/Engine/Source/Editor/UnrealEd/Private/Cooker/CookPackageData.cpp
--- a/Engine/Source/Editor/UnrealEd/Private/Cooker/CookPackageData.cpp
+++ b/Engine/Source/Editor/UnrealEd/Private/Cooker/CookPackageData.cpp
@@ -9,6 +9,7 @@
 #include "AssetCompilingManager.h"
 #include "AssetRegistry/IAssetRegistry.h"
 #include "Async/ParallelFor.h"
+#include "CompactBinaryTCP.h"
 #include "Cooker/CookPlatformManager.h"
 #include "Cooker/CookRequestCluster.h"
 #include "CookOnTheSide/CookOnTheFlyServer.h"
@@ -1266,7 +1267,36 @@ UE::Cook::FGeneratorPackage* FPackageData::CreateGeneratorPackage(const UObject*
 	}
 	return GetGeneratorPackage();
 }
-	
+
+FConstructPackageData FPackageData::CreateConstructData()
+{
+	FConstructPackageData ConstructData;
+	ConstructData.PackageName = PackageName;
+	ConstructData.NormalizedFileName = FileName;
+	return ConstructData;
+}
+
+}
+
+FCbWriter& operator<<(FCbWriter& Writer, const UE::Cook::FConstructPackageData& PackageData)
+{
+	Writer.BeginObject();
+	Writer << "P" << PackageData.PackageName;
+	Writer << "F" << PackageData.NormalizedFileName;
+	Writer.EndObject();
+	return Writer;
+}
+
+bool LoadFromCompactBinary(FCbFieldView Field, UE::Cook::FConstructPackageData& PackageData)
+{
+	LoadFromCompactBinary(Field["P"], PackageData.PackageName);
+	LoadFromCompactBinary(Field["F"], PackageData.NormalizedFileName);
+	return !PackageData.PackageName.IsNone() && !PackageData.NormalizedFileName.IsNone();
+}
+
+namespace UE::Cook
+{
+
 //////////////////////////////////////////////////////////////////////////
 // FGeneratorPackage
 
diff --git a/Engine/Source/Editor/UnrealEd/Private/Cooker/CookPackageData.h b/Engine/Source/Editor/UnrealEd/Private/Cooker/CookPackageData.h
--- a/Engine/Source/Editor/UnrealEd/Private/Cooker/CookPackageData.h
+++ b/Engine/Source/Editor/UnrealEd/Private/Cooker/CookPackageData.h
@@ -27,9 +27,15 @@
 class FPreloadableFile;
 class FReferenceCollector;
 class ITargetPlatform;
+class FCbFieldView;
+class FCbWriter;
 class UCookOnTheFlyServer;
 class UObject;
 class UPackage;
+namespace UE::Cook { struct FConstructPackageData; }
+
+FCbWriter& operator<<(FCbWriter& Writer, const UE::Cook::FConstructPackageData& PackageData);
+bool LoadFromCompactBinary(FCbFieldView Field, UE::Cook::FConstructPackageData& PackageData);
 
 namespace UE::Cook
 {
@@ -68,6 +74,16 @@ enum class ESendFlags : uint8
 };
 ENUM_CLASS_FLAGS(ESendFlags);
 
+/** Data necessary to create an FPackageData without checking the disk, for e.g. AddExistingPackageDatasForPlatform */
+struct FConstructPackageData
+{
+	friend FCbWriter& ::operator<<(FCbWriter& Writer, const FConstructPackageData& PackageData);
+	friend bool ::LoadFromCompactBinary(FCbFieldView Field, FConstructPackageData& PackageData);
+
+	FName PackageName;
+	FName NormalizedFileName;
+};
+
 /**
  * Contains all the information the cooker uses for a package, during request, load, or save.  Once allocated, this
  * structure is never deallocated or moved for a given package; it is deallocated only when the CookOnTheFlyServer
@@ -469,6 +485,9 @@ public:
 	/** Set the id of the worker this Package is assigned to. */
 	void SetWorkerAssignment(FWorkerId InWorkerAssignment) { WorkerAssignment = InWorkerAssignment; }
 
+	/** Marshall this PackageData to a ConstructData that can be used later or on a remote machine to reconstruct it. */
+	FConstructPackageData CreateConstructData();
+
 private:
 	friend struct UE::Cook::FPackageDatas;
 
@@ -945,13 +964,6 @@ public:
 	FPackageDataQueue EntryQueue;
 };
 
-/** Data necessary to create an FPackageData without checking the disk, for e.g. AddExistingPackageDatasForPlatform */
-struct FConstructPackageData
-{
-	FName PackageName;
-	FName NormalizedFileName;
-};
-
 /*
  * Class that manages the list of all PackageDatas for a CookOnTheFlyServer. PackageDatas is an associative
  * array for extra data about a package (e.g. the requested platforms) that is needed by the CookOnTheFlyServer.
diff --git a/Engine/Source/Editor/UnrealEd/Private/Cooker/CookTypes.cpp b/Engine/Source/Editor/UnrealEd/Private/Cooker/CookTypes.cpp
--- a/Engine/Source/Editor/UnrealEd/Private/Cooker/CookTypes.cpp
+++ b/Engine/Source/Editor/UnrealEd/Private/Cooker/CookTypes.cpp
@@ -2,11 +2,15 @@
 
 #include "CookTypes.h"
 
+#include "CompactBinaryTCP.h"
 #include "Cooker/CookPackageData.h"
 #include "HAL/PlatformTime.h"
+#include "Interfaces/ITargetPlatform.h"
+#include "Interfaces/ITargetPlatformManagerModule.h"
 #include "Math/NumericLimits.h"
 #include "Misc/CommandLine.h"
 #include "Misc/Parse.h"
+#include "PackageTracker.h"
 
 namespace UE::Cook
 {
@@ -163,3 +167,231 @@ bool IsCookIgnoreTimeouts()
 }
 
 }
+
+FCbWriter& operator<<(FCbWriter& Writer, const UE::Cook::FInitializeConfigSettings& OutValue)
+{
+	Writer.BeginObject();
+	Writer << "OutputDirectoryOverride" << OutValue.OutputDirectoryOverride;
+	Writer << "MaxPrecacheShaderJobs" << OutValue.MaxPrecacheShaderJobs;
+	Writer << "MaxConcurrentShaderJobs" << OutValue.MaxConcurrentShaderJobs;
+	Writer << "PackagesPerGC" << OutValue.PackagesPerGC;
+	Writer << "MemoryExpectedFreedToSpreadRatio" << OutValue.MemoryExpectedFreedToSpreadRatio;
+	Writer << "IdleTimeToGC" << OutValue.IdleTimeToGC;
+	Writer << "MemoryMaxUsedVirtual" << OutValue.MemoryMaxUsedVirtual;
+	Writer << "MemoryMaxUsedPhysical" << OutValue.MemoryMaxUsedPhysical;
+	Writer << "MemoryMinFreeVirtual" << OutValue.MemoryMinFreeVirtual;
+	Writer << "MemoryMinFreePhysical" << OutValue.MemoryMinFreePhysical;
+	Writer << "MinFreeUObjectIndicesBeforeGC" << OutValue.MinFreeUObjectIndicesBeforeGC;
+	Writer << "MaxNumPackagesBeforePartialGC" << OutValue.MaxNumPackagesBeforePartialGC;
+	Writer << "ConfigSettingDenyList" << OutValue.ConfigSettingDenyList;
+	Writer << "MaxAsyncCacheForType" << OutValue.MaxAsyncCacheForType;
+	Writer << "bHybridIterativeDebug" << OutValue.bHybridIterativeDebug;
+	// Make sure new values are added to LoadFromCompactBinary and MoveOrCopy
+	Writer.EndObject();
+	return Writer;
+}
+
+bool LoadFromCompactBinary(FCbFieldView Field, UE::Cook::FInitializeConfigSettings& OutValue)
+{
+	bool bOk = Field.IsObject();
+	bOk = LoadFromCompactBinary(Field["OutputDirectoryOverride"], OutValue.OutputDirectoryOverride) & bOk;
+	bOk = LoadFromCompactBinary(Field["MaxPrecacheShaderJobs"], OutValue.MaxPrecacheShaderJobs) & bOk;
+	bOk = LoadFromCompactBinary(Field["MaxConcurrentShaderJobs"], OutValue.MaxConcurrentShaderJobs) & bOk;
+	bOk = LoadFromCompactBinary(Field["PackagesPerGC"], OutValue.PackagesPerGC) & bOk;
+	bOk = LoadFromCompactBinary(Field["MemoryExpectedFreedToSpreadRatio"], OutValue.MemoryExpectedFreedToSpreadRatio) & bOk;
+	bOk = LoadFromCompactBinary(Field["IdleTimeToGC"], OutValue.IdleTimeToGC) & bOk;
+	bOk = LoadFromCompactBinary(Field["MemoryMaxUsedVirtual"], OutValue.MemoryMaxUsedVirtual) & bOk;
+	bOk = LoadFromCompactBinary(Field["MemoryMaxUsedPhysical"], OutValue.MemoryMaxUsedPhysical) & bOk;
+	bOk = LoadFromCompactBinary(Field["MemoryMinFreeVirtual"], OutValue.MemoryMinFreeVirtual) & bOk;
+	bOk = LoadFromCompactBinary(Field["MemoryMinFreePhysical"], OutValue.MemoryMinFreePhysical) & bOk;
+	bOk = LoadFromCompactBinary(Field["MinFreeUObjectIndicesBeforeGC"], OutValue.MinFreeUObjectIndicesBeforeGC) & bOk;
+	bOk = LoadFromCompactBinary(Field["MaxNumPackagesBeforePartialGC"], OutValue.MaxNumPackagesBeforePartialGC) & bOk;
+	bOk = LoadFromCompactBinary(Field["ConfigSettingDenyList"], OutValue.ConfigSettingDenyList) & bOk;
+	bOk = LoadFromCompactBinary(Field["MaxAsyncCacheForType"], OutValue.MaxAsyncCacheForType) & bOk;
+	bOk = LoadFromCompactBinary(Field["bHybridIterativeDebug"], OutValue.bHybridIterativeDebug) & bOk;
+	// Make sure new values are added to MoveOrCopy and operator<<
+	return bOk;
+}
+
+namespace UE::Cook
+{
+
+template <typename SourceType, typename TargetType>
+void FInitializeConfigSettings::MoveOrCopy(SourceType&& Source, TargetType&& Target)
+{
+	Target.OutputDirectoryOverride = MoveTempIfPossible(Source.OutputDirectoryOverride);
+	Target.MaxPrecacheShaderJobs = Source.MaxPrecacheShaderJobs;
+	Target.MaxConcurrentShaderJobs = Source.MaxConcurrentShaderJobs;
+	Target.PackagesPerGC = Source.PackagesPerGC;
+	Target.MemoryExpectedFreedToSpreadRatio = Source.MemoryExpectedFreedToSpreadRatio;
+	Target.IdleTimeToGC = Source.IdleTimeToGC;
+	Target.MemoryMaxUsedVirtual = Source.MemoryMaxUsedVirtual;
+	Target.MemoryMaxUsedPhysical = Source.MemoryMaxUsedPhysical;
+	Target.MemoryMinFreeVirtual = Source.MemoryMinFreeVirtual;
+	Target.MemoryMinFreePhysical = Source.MemoryMinFreePhysical;
+	Target.MinFreeUObjectIndicesBeforeGC = Source.MinFreeUObjectIndicesBeforeGC;
+	Target.MaxNumPackagesBeforePartialGC = Source.MaxNumPackagesBeforePartialGC;
+	Target.ConfigSettingDenyList = MoveTempIfPossible(Source.ConfigSettingDenyList);
+	Target.MaxAsyncCacheForType = MoveTempIfPossible(Source.MaxAsyncCacheForType);
+	Target.bHybridIterativeDebug = Source.bHybridIterativeDebug;
+	// Make sure new values are added to operator<< and LoadFromCompactBinary
+}
+
+void FInitializeConfigSettings::CopyFromLocal(const UCookOnTheFlyServer& COTFS)
+{
+	MoveOrCopy(COTFS, *this);
+}
+
+void FInitializeConfigSettings::MoveToLocal(UCookOnTheFlyServer& COTFS)
+{
+	MoveOrCopy(MoveTemp(*this), COTFS);
+}
+
+void FBeginCookConfigSettings::CopyFromLocal(const UCookOnTheFlyServer& COTFS)
+{
+	bHybridIterativeEnabled = COTFS.bHybridIterativeEnabled;
+	FParse::Value(FCommandLine::Get(), TEXT("-CookShowInstigator="), CookShowInstigator); // We don't store this on COTFS, so reparse it from commandLine
+	TSet<FName> COTFSNeverCookPackageList;
+	COTFS.PackageTracker->NeverCookPackageList.GetValues(COTFSNeverCookPackageList);
+	NeverCookPackageList = COTFSNeverCookPackageList.Array();
+	PlatformSpecificNeverCookPackages = COTFS.PackageTracker->PlatformSpecificNeverCookPackages;
+	// Make sure new values are added to SetBeginCookConfigSettings, operator<<, and LoadFromCompactBinary
+}
+
+}
+
+FCbWriter& operator<<(FCbWriter& Writer, const UE::Cook::FBeginCookConfigSettings& OutValue)
+{
+	Writer.BeginObject();
+	Writer << "HybridIterativeEnabled" << OutValue.bHybridIterativeEnabled;
+	Writer << "CookShowInstigator" << OutValue.CookShowInstigator;
+	Writer << "NeverCookPackageList" << OutValue.NeverCookPackageList;
+	
+	Writer.BeginArray("PlatformSpecificNeverCookPackages");
+	for (const TPair<const ITargetPlatform*, TSet<FName>>& Pair : OutValue.PlatformSpecificNeverCookPackages)
+	{
+		Writer.BeginObject();
+		Writer << "K" << Pair.Key->PlatformName();
+		Writer << "V" << Pair.Value;
+		Writer.EndObject();
+	}
+	Writer.EndArray();
+	Writer.EndObject();
+	// Make sure new values are added to SetBeginCookConfigSettings, LoadFromCompactBinary, and CopyFromLocal
+	return Writer;
+}
+
+bool LoadFromCompactBinary(FCbFieldView Field, UE::Cook::FBeginCookConfigSettings& OutValue)
+{
+	bool bOk = Field.IsObject();
+	bOk = LoadFromCompactBinary(Field["HybridIterativeEnabled"], OutValue.bHybridIterativeEnabled) & bOk;
+	bOk = LoadFromCompactBinary(Field["CookShowInstigator"], OutValue.CookShowInstigator) & bOk;
+	bOk = LoadFromCompactBinary(Field["NeverCookPackageList"], OutValue.NeverCookPackageList) & bOk;
+
+	ITargetPlatformManagerModule& TPM(GetTargetPlatformManagerRef());
+	FCbFieldView PlatformNeverCookField = Field["PlatformSpecificNeverCookPackages"];
+	{
+		bOk = PlatformNeverCookField.IsArray() & bOk;
+		OutValue.PlatformSpecificNeverCookPackages.Reset();
+		OutValue.PlatformSpecificNeverCookPackages.Reserve(PlatformNeverCookField.AsArrayView().Num());
+		for (FCbFieldView PairField : PlatformNeverCookField)
+		{
+			bOk &= PairField.IsObject();
+			TStringBuilder<128> KeyName;
+			if (LoadFromCompactBinary(PairField["K"], KeyName))
+			{
+				const ITargetPlatform* TargetPlatform = TPM.FindTargetPlatform(KeyName.ToView());
+				if (TargetPlatform)
+				{
+					TSet<FName>& Value = OutValue.PlatformSpecificNeverCookPackages.FindOrAdd(TargetPlatform);
+					bOk = LoadFromCompactBinary(PairField["V"], Value) & bOk;
+				}
+				else
+				{
+					UE_LOG(LogCook, Error, TEXT("Could not find TargetPlatform \"%.*s\" received from CookDirector."),
+						KeyName.Len(), KeyName.GetData());
+					bOk = false;
+				}
+			}
+			else
+			{
+				bOk = false;
+			}
+		}
+	}
+	// Make sure new values are added to SetBeginCookConfigSettings, CopyFromLocal, and operator<<
+	return bOk;
+}
+
+FCbWriter& operator<<(FCbWriter& Writer, const UE::Cook::FCookByTheBookOptions& OutValue)
+{
+	Writer.BeginObject();
+	// StartupPackages are process-specific
+
+	Writer << "DlcName" << OutValue.DlcName;
+	Writer << "CreateReleaseVersion" << OutValue.CreateReleaseVersion;
+	Writer << "BasedOnReleaseCookedPackages" << OutValue.BasedOnReleaseCookedPackages;
+	Writer << "SourceToLocalizedPackageVariants" << OutValue.SourceToLocalizedPackageVariants;
+	Writer << "AllCulturesToCook" << OutValue.AllCulturesToCook;
+
+	// CookTime is process-specific
+	// CookStartTime is process-specific
+
+	Writer << "StartupOptions" << static_cast<int32>(OutValue.StartupOptions);
+	Writer << "GenerateStreamingInstallManifests" << OutValue.bGenerateStreamingInstallManifests;
+	Writer << "ErrorOnEngineContentUse" << OutValue.bErrorOnEngineContentUse;
+	Writer << "AllowUncookedAssetReferences" << OutValue.bAllowUncookedAssetReferences;
+	Writer << "SkipHardReferences" << OutValue.bSkipHardReferences;
+	Writer << "SkipSoftReferences" << OutValue.bSkipSoftReferences;
+	Writer << "FullLoadAndSave" << OutValue.bFullLoadAndSave;
+	Writer << "CookAgainstFixedBase" << OutValue.bCookAgainstFixedBase;
+	Writer << "DlcLoadMainAssetRegistry" << OutValue.bDlcLoadMainAssetRegistry;
+	Writer.EndObject();
+	return Writer;
+}
+
+bool LoadFromCompactBinary(FCbFieldView Field, UE::Cook::FCookByTheBookOptions& OutValue)
+{
+	bool bOk = Field.IsObject();
+	OutValue.StartupPackages.Reset();
+
+	bOk = LoadFromCompactBinary(Field["DlcName"], OutValue.DlcName) & bOk;
+	bOk = LoadFromCompactBinary(Field["CreateReleaseVersion"], OutValue.CreateReleaseVersion) & bOk;
+	bOk = LoadFromCompactBinary(Field["BasedOnReleaseCookedPackages"], OutValue.BasedOnReleaseCookedPackages) & bOk;
+	bOk = LoadFromCompactBinary(Field["SourceToLocalizedPackageVariants"], OutValue.SourceToLocalizedPackageVariants) & bOk;
+	bOk = LoadFromCompactBinary(Field["AllCulturesToCook"], OutValue.AllCulturesToCook) & bOk;
+
+	OutValue.CookTime = 0.;
+	OutValue.CookStartTime = 0.;
+
+	int32 LocalStartupOptions;
+	bOk = LoadFromCompactBinary(Field["StartupOptions"], LocalStartupOptions) & bOk;
+	OutValue.StartupOptions = static_cast<ECookByTheBookOptions>(LocalStartupOptions);
+	bOk = LoadFromCompactBinary(Field["GenerateStreamingInstallManifests"], OutValue.bGenerateStreamingInstallManifests) & bOk;
+	bOk = LoadFromCompactBinary(Field["ErrorOnEngineContentUse"], OutValue.bErrorOnEngineContentUse) & bOk;
+	bOk = LoadFromCompactBinary(Field["AllowUncookedAssetReferences"], OutValue.bAllowUncookedAssetReferences) & bOk;
+	bOk = LoadFromCompactBinary(Field["SkipHardReferences"], OutValue.bSkipHardReferences) & bOk;
+	bOk = LoadFromCompactBinary(Field["SkipSoftReferences"], OutValue.bSkipSoftReferences) & bOk;
+	bOk = LoadFromCompactBinary(Field["FullLoadAndSave"], OutValue.bFullLoadAndSave) & bOk;
+	bOk = LoadFromCompactBinary(Field["CookAgainstFixedBase"], OutValue.bCookAgainstFixedBase) & bOk;
+	bOk = LoadFromCompactBinary(Field["DlcLoadMainAssetRegistry"], OutValue.bDlcLoadMainAssetRegistry) & bOk;
+
+	return bOk;
+}
+
+FCbWriter& operator<<(FCbWriter& Writer, const UE::Cook::FCookOnTheFlyOptions& OutValue)
+{
+	Writer.BeginObject();
+	Writer << "BindAnyPort" << OutValue.bBindAnyPort;
+	Writer << "PlatformProtocol" << OutValue.bPlatformProtocol;
+	Writer.EndObject();
+	return Writer;
+}
+
+bool LoadFromCompactBinary(FCbFieldView Field, UE::Cook::FCookOnTheFlyOptions& OutValue)
+{
+	bool bOk = Field.IsObject();
+	bOk = LoadFromCompactBinary(Field["BindAnyPort"], OutValue.bBindAnyPort) & bOk;
+	bOk = LoadFromCompactBinary(Field["PlatformProtocol"], OutValue.bPlatformProtocol) & bOk;
+	return bOk;
+}
diff --git a/Engine/Source/Editor/UnrealEd/Private/Cooker/CookTypes.h b/Engine/Source/Editor/UnrealEd/Private/Cooker/CookTypes.h
--- a/Engine/Source/Editor/UnrealEd/Private/Cooker/CookTypes.h
+++ b/Engine/Source/Editor/UnrealEd/Private/Cooker/CookTypes.h
@@ -17,8 +17,23 @@
 #include "UObject/NameTypes.h"
 #include "UObject/SavePackage.h"
 
+class FCbFieldView;
+class FCbWriter;
 class ITargetPlatform;
 namespace UE::DerivedData { class FBuildDefinition; }
+namespace UE::Cook { struct FBeginCookConfigSettings; }
+namespace UE::Cook { struct FCookByTheBookOptions; }
+namespace UE::Cook { struct FCookOnTheFlyOptions; }
+namespace UE::Cook { struct FInitializeConfigSettings; }
+
+FCbWriter& operator<<(FCbWriter& Writer, const UE::Cook::FBeginCookConfigSettings& Value);
+bool LoadFromCompactBinary(FCbFieldView Field, UE::Cook::FBeginCookConfigSettings& Value);
+FCbWriter& operator<<(FCbWriter& Writer, const UE::Cook::FCookByTheBookOptions& Value);
+bool LoadFromCompactBinary(FCbFieldView Field, UE::Cook::FCookByTheBookOptions& Value);
+FCbWriter& operator<<(FCbWriter& Writer, const UE::Cook::FCookOnTheFlyOptions& Value);
+bool LoadFromCompactBinary(FCbFieldView Field, UE::Cook::FCookOnTheFlyOptions& Value);
+FCbWriter& operator<<(FCbWriter& Writer, const UE::Cook::FInitializeConfigSettings& Value);
+bool LoadFromCompactBinary(FCbFieldView Field, UE::Cook::FInitializeConfigSettings& Value);
 
 #define COOK_CHECKSLOW_PACKAGEDATA 0
 #define DEBUG_COOKONTHEFLY 0
@@ -256,8 +271,15 @@ namespace UE::Cook
 
 	struct FInitializeConfigSettings
 	{
+	public:
 		void LoadLocal(const FString& InOutputDirectoryOverride);
+		void CopyFromLocal(const UCookOnTheFlyServer& COTFS);
+		void MoveToLocal(UCookOnTheFlyServer& COTFS);
+	private:
+		template <typename SourceType, typename TargetType>
+		void MoveOrCopy(SourceType&& Source, TargetType&& Target);
 
+	public:
 		FString OutputDirectoryOverride;
 		int32 MaxPrecacheShaderJobs = 0;
 		int32 MaxConcurrentShaderJobs = 0;
@@ -273,6 +295,9 @@ namespace UE::Cook
 		TArray<FString> ConfigSettingDenyList;
 		TMap<FName, int32> MaxAsyncCacheForType; // max number of objects of a specific type which are allowed to async cache at once
 		bool bHybridIterativeDebug = false;
+
+		friend FCbWriter& ::operator<<(FCbWriter& Writer, const UE::Cook::FInitializeConfigSettings& Value);
+		friend bool ::LoadFromCompactBinary(FCbFieldView Field, UE::Cook::FInitializeConfigSettings& Value);
 	};
 
 	struct FBeginCookConfigSettings
@@ -280,11 +305,15 @@ namespace UE::Cook
 		/** Initialize NeverCookPackageList from packaging settings and platform-specific sources. */
 		void LoadLocal(FBeginCookContext& BeginContext);
 		void LoadNeverCookLocal(FBeginCookContext& BeginContext);
+		void CopyFromLocal(const UCookOnTheFlyServer& COTFS);
 
 		FString CookShowInstigator;
 		bool bHybridIterativeEnabled = true;
 		TArray<FName> NeverCookPackageList;
 		TFastPointerMap<const ITargetPlatform*, TSet<FName>> PlatformSpecificNeverCookPackages;
+
+		friend FCbWriter& ::operator<<(FCbWriter& Writer, const UE::Cook::FBeginCookConfigSettings& Value);
+		friend bool ::LoadFromCompactBinary(FCbFieldView Field, UE::Cook::FBeginCookConfigSettings& Value);
 	};
 
 	/**
@@ -367,6 +396,8 @@ public:
 	double							CookTime = 0.0;
 	double							CookStartTime = 0.0;
 
+	ECookByTheBookOptions			StartupOptions = ECookByTheBookOptions::None;
+
 	/** Should we generate streaming install manifests (only valid option in cook by the book) */
 	bool							bGenerateStreamingInstallManifests = false;
 
@@ -389,6 +420,9 @@ public:
 		EmptyOptions.StartupPackages = MoveTemp(StartupPackages);
 		*this = MoveTemp(EmptyOptions);
 	}
+
+	friend FCbWriter& ::operator<<(FCbWriter& Writer, const UE::Cook::FCookByTheBookOptions& Value);
+	friend bool ::LoadFromCompactBinary(FCbFieldView Field, UE::Cook::FCookByTheBookOptions& Value);
 };
 
 //////////////////////////////////////////////////////////////////////////
@@ -399,6 +433,9 @@ struct FCookOnTheFlyOptions
 	bool bBindAnyPort = false;
 	/** Whether the network file server should use a platform-specific communication protocol instead of TCP (used when bZenStore == false) */
 	bool bPlatformProtocol = false;
+
+	friend FCbWriter& ::operator<<(FCbWriter& Writer, const UE::Cook::FCookOnTheFlyOptions& Value);
+	friend bool ::LoadFromCompactBinary(FCbFieldView Field, UE::Cook::FCookOnTheFlyOptions& Value);
 };
 
 }
diff --git a/Engine/Source/Editor/UnrealEd/Private/Cooker/CookWorkerClient.cpp b/Engine/Source/Editor/UnrealEd/Private/Cooker/CookWorkerClient.cpp
--- a/Engine/Source/Editor/UnrealEd/Private/Cooker/CookWorkerClient.cpp
+++ b/Engine/Source/Editor/UnrealEd/Private/Cooker/CookWorkerClient.cpp
@@ -4,10 +4,12 @@
 
 #include "CompactBinaryTCP.h"
 #include "CookDirector.h"
+#include "CookPackageData.h"
+#include "CookPlatformManager.h"
 #include "CookTypes.h"
+#include "CookWorkerServer.h"
 #include "CoreGlobals.h"
 #include "HAL/PlatformTime.h"
-#include "Serialization/CompactBinaryWriter.h"
 #include "Sockets.h"
 #include "SocketSubsystem.h"
 #include "WorkerRequestsRemote.h"
@@ -15,8 +17,23 @@
 namespace UE::Cook
 {
 
+namespace CookWorkerClient
+{
+constexpr float WaitForConnectReplyTimeout = 60.f;
+}
+
+FCookWorkerClient::FCookWorkerClient(UCookOnTheFlyServer& InCOTFS)
+	: COTFS(InCOTFS)
+{
+}
+
 FCookWorkerClient::~FCookWorkerClient()
 {
+	if (ConnectStatus == EConnectStatus::Connected ||
+		(EConnectStatus::WaitForDisconnectFirst <= ConnectStatus && ConnectStatus <= EConnectStatus::WaitForDisconnectLast))
+	{
+		UE_LOG(LogCook, Warning, TEXT("CookWorkerServer %d was destroyed before it finished Disconnect. The CookDirector may be missing some information."));
+	}
 	Sockets::CloseSocket(ServerSocket);
 }
 
@@ -36,6 +53,89 @@ bool FCookWorkerClient::TryConnect(FDirectorConnectionInfo&& ConnectInfo)
 	return Status == EPollStatus::Success;
 }
 
+void FCookWorkerClient::TickFromSchedulerThread(FTickStackData& StackData)
+{
+	if (ConnectStatus == EConnectStatus::Connected)
+	{
+		PumpReceiveMessages();
+		if (ConnectStatus == EConnectStatus::Connected)
+		{
+			SendPendingResults();
+			PumpSendMessages();
+		}
+	}
+	else
+	{
+		PumpDisconnect(StackData);
+	}
+}
+
+bool FCookWorkerClient::IsDisconnecting() const
+{
+	return ConnectStatus == EConnectStatus::LostConnection ||
+		(EConnectStatus::WaitForDisconnectFirst <= ConnectStatus && ConnectStatus <= EConnectStatus::WaitForDisconnectLast);
+}
+
+bool FCookWorkerClient::IsDisconnectComplete() const
+{
+	return ConnectStatus == EConnectStatus::LostConnection;
+}
+
+ECookInitializationFlags FCookWorkerClient::GetCookInitializationFlags()
+{
+	check(InitialConfigMessage); // Should only be called after TryConnect and before DoneWithInitialSettings
+	return InitialConfigMessage->GetCookInitializationFlags();
+}
+FInitializeConfigSettings&& FCookWorkerClient::ConsumeInitializeConfigSettings()
+{
+	check(InitialConfigMessage); // Should only be called after TryConnect and before DoneWithInitialSettings
+	return InitialConfigMessage->ConsumeInitializeConfigSettings();
+}
+FBeginCookConfigSettings&& FCookWorkerClient::ConsumeBeginCookConfigSettings()
+{
+	check(InitialConfigMessage); // Should only be called after TryConnect and before DoneWithInitialSettings
+	return InitialConfigMessage->ConsumeBeginCookConfigSettings();
+}
+FCookByTheBookOptions&& FCookWorkerClient::ConsumeCookByTheBookOptions()
+{
+	check(InitialConfigMessage); // Should only be called after TryConnect and before DoneWithInitialSettings
+	return InitialConfigMessage->ConsumeCookByTheBookOptions();
+}
+FCookOnTheFlyOptions&& FCookWorkerClient::ConsumeCookOnTheFlyOptions()
+{
+	check(InitialConfigMessage); // Should only be called after TryConnect and before DoneWithInitialSettings
+	return InitialConfigMessage->ConsumeCookOnTheFlyOptions();
+}
+const TArray<ITargetPlatform*>& FCookWorkerClient::GetTargetPlatforms() const
+{
+	return OrderedSessionPlatforms;
+}
+void FCookWorkerClient::DoneWithInitialSettings()
+{
+	InitialConfigMessage.Reset();
+}
+
+void FCookWorkerClient::ReportDemoteToIdle(FPackageData& PackageData, ESuppressCookReason Reason)
+{
+	FPackageResult& Result = PendingResults.Emplace_GetRef();
+	Result.PackageName = PackageData.GetPackageName();
+	Result.SuppressCookReason = Reason;
+}
+
+void FCookWorkerClient::ReportPromoteToSaveComplete(FPackageData& PackageData)
+{
+	FPackageResult& Result = PendingResults.Emplace_GetRef();
+	Result.PackageName = PackageData.GetPackageName();
+	Result.SuppressCookReason = ESuppressCookReason::InvalidSuppressCookReason;
+	for (ITargetPlatform* TargetPlatform : OrderedSessionPlatforms)
+	{
+		FPackagePlatformResult& PlatformResult = Result.Platforms.Emplace_GetRef();
+		FPackageData::FPlatformData& PackagePlatformData = PackageData.FindOrAddPlatformData(TargetPlatform);
+		PlatformResult.bSuccessful = PackagePlatformData.bCookSucceeded;
+		// MPCOOKTODO: Accumulate results for saved information so we can add the other information about the Platform
+	}
+}
+
 EPollStatus FCookWorkerClient::PollTryConnect(const FDirectorConnectionInfo& ConnectInfo)
 {
 	for (;;)
@@ -61,10 +161,9 @@ EPollStatus FCookWorkerClient::PollTryConnect(const FDirectorConnectionInfo& Con
 				return EPollStatus::Incomplete;
 			}
 			break;
-		case EConnectStatus::Failed:
+		case EConnectStatus::LostConnection:
 			return EPollStatus::Error;
 		default:
-			checkNoEntry();
 			return EPollStatus::Error;
 		}
 	}
@@ -81,7 +180,7 @@ void FCookWorkerClient::CreateServerSocket(const FDirectorConnectionInfo& Connec
 	if (!SocketSubsystem)
 	{
 		UE_LOG(LogCook, Error, TEXT("CookWorker initialization failure: platform does not support network sockets, cannot connect to CookDirector."));
-		ConnectStatus = EConnectStatus::Failed;
+		SendToState(EConnectStatus::LostConnection);
 		return;
 	}
 
@@ -90,7 +189,7 @@ void FCookWorkerClient::CreateServerSocket(const FDirectorConnectionInfo& Connec
 	{
 		UE_LOG(LogCook, Error, TEXT("CookWorker initialization failure: could not convert -CookDirectorHost=%s into an address, cannot connect to CookDirector."),
 			*DirectorURI);
-		ConnectStatus = EConnectStatus::Failed;
+		SendToState(EConnectStatus::LostConnection);
 		return;
 	}
 
@@ -100,7 +199,7 @@ void FCookWorkerClient::CreateServerSocket(const FDirectorConnectionInfo& Connec
 	if (!ServerSocket)
 	{
 		UE_LOG(LogCook, Error, TEXT("CookWorker initialization failure: Could not connect to CookDirector."));
-		ConnectStatus = EConnectStatus::Failed;
+		SendToState(EConnectStatus::LostConnection);
 		return;
 	}
 
@@ -111,7 +210,7 @@ void FCookWorkerClient::CreateServerSocket(const FDirectorConnectionInfo& Connec
 	{
 		UE_LOG(LogCook, Error, TEXT("CookWorker initialization failure: Timed out after %.0f seconds trying to connect to CookDirector."),
 			ConditionalTimeoutSeconds);
-		ConnectStatus = EConnectStatus::Failed;
+		SendToState(EConnectStatus::LostConnection);
 		return;
 	}
 
@@ -120,49 +219,91 @@ void FCookWorkerClient::CreateServerSocket(const FDirectorConnectionInfo& Connec
 	EConnectionStatus Status = TryWritePacket(ServerSocket, SendBuffer, ConnectMessage);
 	if (Status == EConnectionStatus::Incomplete)
 	{
-		ConnectStatus = EConnectStatus::PollWriteConnectMessage;
+		SendToState(EConnectStatus::PollWriteConnectMessage);
 		return;
 	}
 	else if (Status != EConnectionStatus::Okay)
 	{
 		UE_LOG(LogCook, Error, TEXT("CookWorker initialization failure: could not send ConnectMessage."));
-		ConnectStatus = EConnectStatus::Failed;
+		SendToState(EConnectStatus::LostConnection);
 		return;
 	}
 	LogConnected();
 
-	ConnectStatus = EConnectStatus::PollReceiveConfigMessage;
+	SendToState(EConnectStatus::PollReceiveConfigMessage);
 }
 
 void FCookWorkerClient::PollWriteConnectMessage()
 {
 	using namespace CompactBinaryTCP;
-	constexpr float WaitForConnectTimeout = 60.f;
 
 	EConnectionStatus Status = TryFlushBuffer(ServerSocket, SendBuffer);
 	if (Status == EConnectionStatus::Incomplete)
 	{
-		if (FPlatformTime::Seconds() - ConnectStartTimeSeconds > WaitForConnectTimeout && !IsCookIgnoreTimeouts())
+		if (FPlatformTime::Seconds() - ConnectStartTimeSeconds > CookWorkerClient::WaitForConnectReplyTimeout &&
+			!IsCookIgnoreTimeouts())
 		{
-			UE_LOG(LogCook, Error, TEXT("CookWorker initialization failure: timed out waiting for %fs to send ConnectMessage."), WaitForConnectTimeout);
-			ConnectStatus = EConnectStatus::Failed;
+			UE_LOG(LogCook, Error, TEXT("CookWorker initialization failure: timed out waiting for %fs to send ConnectMessage."),
+				CookWorkerClient::WaitForConnectReplyTimeout);
+			SendToState(EConnectStatus::LostConnection);
 		}
 		return;
 	}
 	else if (Status != EConnectionStatus::Okay)
 	{
 		UE_LOG(LogCook, Error, TEXT("CookWorker initialization failure: could not send ConnectMessage."));
-		ConnectStatus = EConnectStatus::Failed;
+		SendToState(EConnectStatus::LostConnection);
 		return;
 	}
 	LogConnected();
-	ConnectStatus = EConnectStatus::PollReceiveConfigMessage;
+	SendToState(EConnectStatus::PollReceiveConfigMessage);
 }
 
 void FCookWorkerClient::PollReceiveConfigMessage()
 {
-	UE_LOG(LogCook, Error, TEXT("CookWorker initialization failure: reading Settings information is not yet implemented."));
-	ConnectStatus = EConnectStatus::Failed;
+	using namespace UE::CompactBinaryTCP;
+	TArray<FMarshalledMessage> Messages;
+	EConnectionStatus SocketStatus = TryReadPacket(ServerSocket, ReceiveBuffer, Messages);
+	if (SocketStatus != EConnectionStatus::Okay && SocketStatus != EConnectionStatus::Incomplete)
+	{
+		UE_LOG(LogCook, Error, TEXT("CookWorker initialization failure: failed to read from socket."));
+		SendToState(EConnectStatus::LostConnection);
+		return;
+	}
+	if (Messages.Num() == 0)
+	{
+		if (FPlatformTime::Seconds() - ConnectStartTimeSeconds > CookWorkerClient::WaitForConnectReplyTimeout &&
+			!IsCookIgnoreTimeouts())
+		{
+			UE_LOG(LogCook, Error, TEXT("CookWorker initialization failure: timed out waiting for %fs to receive InitialConfigMessage."),
+				CookWorkerClient::WaitForConnectReplyTimeout);
+			SendToState(EConnectStatus::LostConnection);
+		}
+		return;
+	}
+	
+	if (Messages[0].MessageType != FInitialConfigMessage::MessageType)
+	{
+		UE_LOG(LogCook, Warning, TEXT("CookWorker initialization failure: Director sent a different message before sending an InitialConfigMessage. MessageType: %s."),
+			*Messages[0].MessageType.ToString());
+		SendToState(EConnectStatus::LostConnection);
+		return;
+	}
+	check(!InitialConfigMessage);
+	InitialConfigMessage = MakeUnique<FInitialConfigMessage>();
+	if (!InitialConfigMessage->TryRead(Messages[0].Object))
+	{
+		UE_LOG(LogCook, Warning, TEXT("CookWorker initialization failure: Director sent an invalid InitialConfigMessage."));
+		SendToState(EConnectStatus::LostConnection);
+		return;
+	}
+	DirectorCookMode = InitialConfigMessage->GetDirectorCookMode();
+	OrderedSessionPlatforms = InitialConfigMessage->GetOrderedSessionPlatforms();
+
+	UE_LOG(LogCook, Display, TEXT("Initialization from CookDirector complete."));
+	SendToState(EConnectStatus::Connected);
+	Messages.RemoveAt(0);
+	HandleReceiveMessages(MoveTemp(Messages));
 }
 
 void FCookWorkerClient::LogConnected()
@@ -170,4 +311,174 @@ void FCookWorkerClient::LogConnected()
 	UE_LOG(LogCook, Display, TEXT("Connection to CookDirector successful."));
 }
 
+void FCookWorkerClient::PumpSendMessages()
+{
+	UE::CompactBinaryTCP::EConnectionStatus Status = UE::CompactBinaryTCP::TryFlushBuffer(ServerSocket, SendBuffer);
+	if (Status == UE::CompactBinaryTCP::Failed)
+	{
+		UE_LOG(LogCook, Error, TEXT("CookWorkerClient failed to write message to Director. We will abort the CookAsCookWorker commandlet."));
+		SendToState(EConnectStatus::LostConnection);
+	}
+}
+
+void FCookWorkerClient::SendPendingResults()
+{
+	if (PendingResults.IsEmpty())
+	{
+		return;
+	}
+	FPackageResultsMessage Message;
+	Message.Results = MoveTemp(PendingResults);
+	SendMessage(Message);
+	PendingResults.Reset();
+}
+
+void FCookWorkerClient::PumpReceiveMessages()
+{
+	using namespace UE::CompactBinaryTCP;
+	TArray<FMarshalledMessage> Messages;
+	EConnectionStatus SocketStatus = TryReadPacket(ServerSocket, ReceiveBuffer, Messages);
+	if (SocketStatus != EConnectionStatus::Okay && SocketStatus != EConnectionStatus::Incomplete)
+	{
+		UE_LOG(LogCook, Error, TEXT("CookWorkerClient failed to read from Director. We will abort the CookAsCookWorker commandlet."));
+		SendToState(EConnectStatus::LostConnection);
+		return;
+	}
+	HandleReceiveMessages(MoveTemp(Messages));
+}
+
+void FCookWorkerClient::HandleReceiveMessages(TArray<UE::CompactBinaryTCP::FMarshalledMessage>&& Messages)
+{
+	for (UE::CompactBinaryTCP::FMarshalledMessage& Message : Messages)
+	{
+		if (Message.MessageType == FAbortWorkerMessage::MessageType)
+		{
+			FAbortWorkerMessage AbortMessage;
+			AbortMessage.TryRead(Message.Object);
+			if (AbortMessage.Type == FAbortWorkerMessage::EType::CookComplete)
+			{
+				UE_LOG(LogCook, Display, TEXT("CookWorkerClient received CookComplete message from Director. Flushing messages and shutting down."));
+				// MPCOOKTODO: Add synchronous flush of messages here
+			}
+			else
+			{
+				UE_LOG(LogCook, Display, TEXT("CookWorkerClient received AbortWorker message from Director. Shutting down."));
+			}
+			SendToState(EConnectStatus::WaitForDisconnect);
+			break;
+		}
+		else if (Message.MessageType == FInitialConfigMessage::MessageType)
+		{
+			UE_LOG(LogCook, Warning, TEXT("CookWorkerClient received unexpected repeat of InitialConfigMessage. Ignoring it."));
+		}
+		else if (Message.MessageType == FAssignPackagesMessage::MessageType)
+		{
+			FAssignPackagesMessage AssignPackagesMessage;
+			if (!AssignPackagesMessage.TryRead(Message.Object))
+			{
+				LogInvalidMessage(TEXT("FAssignPackagesMessage"));
+			}
+			else
+			{
+				AssignPackages(AssignPackagesMessage);
+			}
+		}
+	}
+}
+
+void FCookWorkerClient::PumpDisconnect(FTickStackData& StackData)
+{
+	for (;;)
+	{
+		switch (ConnectStatus)
+		{
+		case EConnectStatus::WaitForDisconnect:
+		{
+			// Add code here for any waiting we need to do for the local CookOnTheFlyServer to gracefully shutdown
+			SendMessage(FAbortWorkerMessage(FAbortWorkerMessage::EType::Abort));
+			SendToState(EConnectStatus::WaitForDisconnectSocketFlush);
+			break;
+		}
+		case EConnectStatus::WaitForDisconnectSocketFlush:
+		{
+			using namespace UE::CompactBinaryTCP;
+			EConnectionStatus SocketStatus = TryFlushBuffer(ServerSocket, SendBuffer);
+			if (SocketStatus == EConnectionStatus::Incomplete)
+			{
+				constexpr float WaitForDisconnectTimeout = 60.f;
+				if (FPlatformTime::Seconds() - ConnectStartTimeSeconds > WaitForDisconnectTimeout && !IsCookIgnoreTimeouts())
+				{
+					UE_LOG(LogCook, Warning, TEXT("Timedout after %.0fs waiting to send disconnect message to CookDirector."),
+						WaitForDisconnectTimeout);
+					SendToState(EConnectStatus::LostConnection);
+					return;
+				}
+				return; // Exit the Pump loop for now and keep waiting
+			}
+			SendToState(EConnectStatus::LostConnection);
+			break;
+		}
+		case EConnectStatus::LostConnection:
+		{
+			StackData.bCookCancelled = true;
+			StackData.ResultFlags |= UCookOnTheFlyServer::COSR_YieldTick;
+			return;
+		}
+		default:
+			return;
+		}
+	}
+}
+
+void FCookWorkerClient::SendMessage(const UE::CompactBinaryTCP::IMessage& Message)
+{
+	UE::CompactBinaryTCP::TryWritePacket(ServerSocket, SendBuffer, Message);
+}
+
+void FCookWorkerClient::SendToState(EConnectStatus TargetStatus)
+{
+	switch (TargetStatus)
+	{
+	case EConnectStatus::WaitForDisconnect:
+		ConnectStartTimeSeconds = FPlatformTime::Seconds();
+		break;
+	case EConnectStatus::LostConnection:
+		Sockets::CloseSocket(ServerSocket);
+		break;
+	}
+	ConnectStatus = TargetStatus;
+}
+
+void FCookWorkerClient::LogInvalidMessage(const TCHAR* MessageTypeName)
+{
+	UE_LOG(LogCook, Warning, TEXT("CookWorkerClient received invalidly formatted message for type %s from CookDirector. Ignoring it."),
+		MessageTypeName);
+}
+
+void FCookWorkerClient::AssignPackages(FAssignPackagesMessage& Message)
+{
+	for (FConstructPackageData& ConstructPackageData : Message.PackageDatas)
+	{
+		FPackageData& PackageData = COTFS.PackageDatas->FindOrAddPackageData(ConstructPackageData.PackageName,
+			ConstructPackageData.NormalizedFileName);
+		// If already InProgress, ignore the duplicate package silently
+		if (PackageData.IsInProgress())
+		{
+			return;
+		}
+
+		// We do not want CookWorkers to explore dependencies in CookRequestCluster because the Director did it already.
+		// Mmark the PackageDatas we get from the Director as already explored.
+		for (const ITargetPlatform* TargetPlatform : OrderedSessionPlatforms)
+		{
+			FPackageData::FPlatformData& PlatformData = PackageData.FindOrAddPlatformData(TargetPlatform);
+			PlatformData.bExplored = true;
+		}
+
+		PackageData.SetRequestData(OrderedSessionPlatforms , false /* bInIsUrgent */, FCompletionCallback(),
+			FInstigator(EInstigator::CookDirector));
+		PackageData.SendToState(EPackageState::Request, ESendFlags::QueueAddAndRemove);
+	}
+}
+
 }
\ No newline at end of file
diff --git a/Engine/Source/Editor/UnrealEd/Private/Cooker/CookWorkerClient.h b/Engine/Source/Editor/UnrealEd/Private/Cooker/CookWorkerClient.h
--- a/Engine/Source/Editor/UnrealEd/Private/Cooker/CookWorkerClient.h
+++ b/Engine/Source/Editor/UnrealEd/Private/Cooker/CookWorkerClient.h
@@ -3,10 +3,15 @@
 #pragma once
 
 #include "CompactBinaryTCP.h"
+#include "Containers/Array.h"
+#include "CookOnTheSide/CookOnTheFlyServer.h"
 #include "CookTypes.h"
 #include "IPAddress.h"
 
+namespace UE::Cook { struct FAssignPackagesMessage; }
 namespace UE::Cook { struct FDirectorConnectionInfo; }
+namespace UE::Cook { struct FInitialConfigMessage; }
+namespace UE::Cook { struct FPackageResult; }
 
 namespace UE::Cook
 {
@@ -15,9 +20,35 @@ namespace UE::Cook
 class FCookWorkerClient
 {
 public:
+	FCookWorkerClient(UCookOnTheFlyServer& COTFS);
 	~FCookWorkerClient();
+
 	/** Blocking operation: open the socket to the Director, send the Connected message, receive the setup message. */
 	bool TryConnect(FDirectorConnectionInfo&& ConnectInfo);
+	/** Periodic Tick function to send and receive messages to the Server. */
+	void TickFromSchedulerThread(FTickStackData& StackData);
+	/** Is this either shutting down or completed shutdown of communications with the Director? */
+	bool IsDisconnecting() const;
+	/** Is this not yet or no longer connected to the Director? */
+	bool IsDisconnectComplete() const;
+
+	/** Reads the cookmode that was sent from the Director. */
+	ECookMode::Type GetDirectorCookMode() const { return DirectorCookMode; }
+	/** Reads the OrderedSessionPlatforms received from the Director. */
+	const TArray<ITargetPlatform*>& GetTargetPlatforms() const;
+	/** Consumes the initialization settings from the Director. Only available during initialization. */
+	ECookInitializationFlags GetCookInitializationFlags();
+	FInitializeConfigSettings&& ConsumeInitializeConfigSettings();
+	FBeginCookConfigSettings&& ConsumeBeginCookConfigSettings();
+	FCookByTheBookOptions&& ConsumeCookByTheBookOptions();
+	FCookOnTheFlyOptions&& ConsumeCookOnTheFlyOptions();
+	/** Mark that initialization is complete and we can free the memory for initialization settings. */
+	void DoneWithInitialSettings();
+
+	/** Queue a message to the server that the Package was cook-suppressed. Will be sent during Tick. */
+	void ReportDemoteToIdle(FPackageData& PackageData, ESuppressCookReason Reason);
+	/** Queue a message to the server that the Package was saved. Will be sent during Tick. */
+	void ReportPromoteToSaveComplete(FPackageData& PackageData);
 
 private:
 	enum class EConnectStatus
@@ -25,8 +56,12 @@ private:
 		Uninitialized,
 		PollWriteConnectMessage,
 		PollReceiveConfigMessage,
-		Failed,
 		Connected,
+		WaitForDisconnectFirst,
+		WaitForDisconnect = WaitForDisconnectFirst,
+		WaitForDisconnectSocketFlush,
+		WaitForDisconnectLast = WaitForDisconnectSocketFlush,
+		LostConnection,
 	};
 
 private:
@@ -39,14 +74,38 @@ private:
 	/** Wait for the Config Message, switch state when it succeeds or fails. */
 	void PollReceiveConfigMessage();
 	void LogConnected();
+	/** Helper for Tick, pump send messages to the Server. */
+	void PumpSendMessages();
+	/** Helper for Tick, send a message for any pending package results that we have. */
+	void SendPendingResults();
+	/** Helper for Tick, pump receive messages from the Server. */
+	void PumpReceiveMessages();
+	/** Helper for PumpReceiveMessages: dispatch the messages received from the socket. */
+	void HandleReceiveMessages(TArray<UE::CompactBinaryTCP::FMarshalledMessage>&& Messages);
+	/** Helper for Tick, pump Send/Receive and check for whether we are done shutting down. */
+	void PumpDisconnect(FTickStackData& StackData);
+	/** Send the message immediately to the Socket. If cannot complete immediately, it will be finished during Tick. */
+	void SendMessage(const UE::CompactBinaryTCP::IMessage& Message);
+	/** Send this into the given state. Update any state-dependent variables. */
+	void SendToState(EConnectStatus TargetStatus);
+	void LogInvalidMessage(const TCHAR* MessageTypeName);
+	/** Send packages assigned from the server into the request state. */
+	void AssignPackages(FAssignPackagesMessage& Message);
+
 
 private:
 	TSharedPtr<FInternetAddr> DirectorAddr;
+	TUniquePtr<FInitialConfigMessage> InitialConfigMessage;
+	TArray<ITargetPlatform*> OrderedSessionPlatforms;
+	TArray<FPackageResult> PendingResults;
 	UE::CompactBinaryTCP::FSendBuffer SendBuffer;
+	UE::CompactBinaryTCP::FReceiveBuffer ReceiveBuffer;
 	FString DirectorURI;
+	UCookOnTheFlyServer& COTFS;
 	FSocket* ServerSocket = nullptr;
 	double ConnectStartTimeSeconds = 0.;
 	EConnectStatus ConnectStatus = EConnectStatus::Uninitialized;
+	ECookMode::Type DirectorCookMode = ECookMode::CookByTheBook;
 };
 
 }
\ No newline at end of file
diff --git a/Engine/Source/Editor/UnrealEd/Private/Cooker/CookWorkerServer.cpp b/Engine/Source/Editor/UnrealEd/Private/Cooker/CookWorkerServer.cpp
--- a/Engine/Source/Editor/UnrealEd/Private/Cooker/CookWorkerServer.cpp
+++ b/Engine/Source/Editor/UnrealEd/Private/Cooker/CookWorkerServer.cpp
@@ -2,11 +2,16 @@
 
 #include "CookWorkerServer.h"
 
+#include "CompactBinaryTCP.h"
 #include "CookDirector.h"
 #include "CookPackageData.h"
+#include "CookPlatformManager.h"
 #include "HAL/PlatformProcess.h"
+#include "Interfaces/ITargetPlatform.h"
+#include "Interfaces/ITargetPlatformManagerModule.h"
+#include "Math/NumericLimits.h"
 #include "Misc/AssertionMacros.h"
-#include "Serialization/CompactBinaryWriter.h"
+#include "PackageTracker.h"
 #include "UnrealEdMisc.h"
 
 namespace UE::Cook
@@ -30,10 +35,10 @@ FCookWorkerServer::~FCookWorkerServer()
 		PackageData->SendToState(UE::Cook::EPackageState::Request, ESendFlags::QueueAddAndRemove);
 	}
 
-	if (ConnectStatus == EConnectStatus::Connected || ConnectStatus == EConnectStatus::WaitForDisconnect)
+	if (IsConnected() || ConnectStatus == EConnectStatus::WaitForDisconnect)
 	{
 		UE_LOG(LogCook, Error, TEXT("CookWorkerServer %d was destroyed before it finished Disconnect. The remote process may linger and may interfere with writes of future packages."),
-			WorkerId.GetLocalOrRemoteIndex());
+			WorkerId.GetRemoteIndex());
 	}
 	DetachFromRemoteProcess();
 }
@@ -65,9 +70,9 @@ void FCookWorkerServer::AppendAssignments(TArrayView<FPackageData*> Assignments)
 
 void FCookWorkerServer::AbortAssignments(TSet<FPackageData*>& OutPendingPackages)
 {
-	if (ConnectStatus == EConnectStatus::Connected)
+	if (PendingPackages.Num())
 	{
-		if (PendingPackages.Num())
+		if (IsConnected())
 		{
 			TArray<FName> PackageNames;
 			PackageNames.Reserve(PendingPackages.Num());
@@ -76,13 +81,9 @@ void FCookWorkerServer::AbortAssignments(TSet<FPackageData*>& OutPendingPackages
 				PackageNames.Add(PackageData->GetPackageName());
 			}
 			SendMessage(FAbortPackagesMessage(MoveTemp(PackageNames)));
-			OutPendingPackages.Append(MoveTemp(PendingPackages));
-			PendingPackages.Reset();
 		}
-	}
-	else
-	{
-		check(PendingPackages.IsEmpty());
+		OutPendingPackages.Append(MoveTemp(PendingPackages));
+		PendingPackages.Empty();
 	}
 	OutPendingPackages.Append(PackagesToAssign);
 	PackagesToAssign.Empty();
@@ -90,19 +91,15 @@ void FCookWorkerServer::AbortAssignments(TSet<FPackageData*>& OutPendingPackages
 
 void FCookWorkerServer::AbortAssignment(FPackageData& PackageData)
 {
-	if (ConnectStatus == EConnectStatus::Connected)
+	if (PendingPackages.Remove(&PackageData))
 	{
-		if (PendingPackages.Remove(&PackageData))
+		if (IsConnected())
 		{
 			TArray<FName> PackageNames;
 			PackageNames.Add(PackageData.GetPackageName());
 			SendMessage(FAbortPackagesMessage(MoveTemp(PackageNames)));
 		}
 	}
-	else
-	{
-		check(PendingPackages.IsEmpty());
-	}
 
 	PackagesToAssign.Remove(&PackageData);
 }
@@ -110,9 +107,9 @@ void FCookWorkerServer::AbortAssignment(FPackageData& PackageData)
 void FCookWorkerServer::AbortWorker(TSet<FPackageData*>& OutPendingPackages)
 {
 	AbortAssignments(OutPendingPackages);
-	if (ConnectStatus == EConnectStatus::Connected)
+	if (IsConnected())
 	{
-		SendMessage(FAbortWorkerMessage());
+		SendMessage(FAbortWorkerMessage(FAbortWorkerMessage::EType::Abort));
 		SendToState(EConnectStatus::WaitForDisconnect);
 	}
 }
@@ -129,6 +126,10 @@ void FCookWorkerServer::SendToState(EConnectStatus TargetStatus)
 		ConnectStartTimeSeconds = FPlatformTime::Seconds();
 		ConnectTestStartTimeSeconds = ConnectStartTimeSeconds;
 		break;
+	case EConnectStatus::PumpingCookComplete:
+		ConnectStartTimeSeconds = FPlatformTime::Seconds();
+		ConnectTestStartTimeSeconds = ConnectStartTimeSeconds;
+		break;
 	case EConnectStatus::LostConnection:
 		DetachFromRemoteProcess();
 		break;
@@ -138,6 +139,11 @@ void FCookWorkerServer::SendToState(EConnectStatus TargetStatus)
 	ConnectStatus = TargetStatus;
 }
 
+bool FCookWorkerServer::IsConnected() const
+{
+	return EConnectStatus::ConnectedFirst <= ConnectStatus && ConnectStatus <= EConnectStatus::ConnectedLast;
+}
+
 bool FCookWorkerServer::IsShuttingDown() const
 {
 	return ConnectStatus == EConnectStatus::WaitForDisconnect || ConnectStatus == EConnectStatus::LostConnection;
@@ -157,33 +163,78 @@ bool FCookWorkerServer::TryHandleConnectMessage(FWorkerConnectMessage& Message,
 	check(!Socket);
 	Socket = InSocket;
 
-	UE_LOG(LogCook, Error, TEXT("CookWorkerServer failure: writing Settings information is not yet implemented."));
-
 	SendToState(EConnectStatus::Connected);
 	HandleReceiveMessages(MoveTemp(OtherPacketMessages));
+	FInitialConfigMessage ConfigMessage;
+	const TArray<const ITargetPlatform*>& SessionPlatforms = COTFS.PlatformManager->GetSessionPlatforms();
+	OrderedSessionPlatforms.Reset(SessionPlatforms.Num());
+	for (const ITargetPlatform* TargetPlatform : SessionPlatforms)
+	{
+		OrderedSessionPlatforms.Add(const_cast<ITargetPlatform*>(TargetPlatform));
+	}
+	ConfigMessage.ReadFromLocal(COTFS, OrderedSessionPlatforms,
+		*COTFS.CookByTheBookOptions, *COTFS.CookOnTheFlyOptions);
+	SendMessage(ConfigMessage);
 	return true;
 }
 
 void FCookWorkerServer::TickFromSchedulerThread()
 {
-	PumpConnect();
-	if (ConnectStatus == EConnectStatus::Connected)
+	if (IsConnected())
 	{
 		PumpReceiveMessages();
-		PumpSendMessages();
+		if (IsConnected())
+		{
+			SendPendingPackages();
+			PumpSendMessages();
+		}
+	}
+	else
+	{
+		PumpConnect();
+		if (IsConnected())
+		{
+			// Recursively call this function to call PumpReceive and PumpSend
+			TickFromSchedulerThread();
+		}
 	}
 }
 
+void FCookWorkerServer::PumpCookComplete()
+{
+	if (ConnectStatus == EConnectStatus::Connected)
+	{
+		SendMessage(FAbortWorkerMessage(FAbortWorkerMessage::EType::CookComplete));
+		SendToState(EConnectStatus::PumpingCookComplete);
+	}
+	else if (ConnectStatus == EConnectStatus::PumpingCookComplete)
+	{
+		TickFromSchedulerThread();
+		if (IsConnected())
+		{
+			constexpr float WaitForPumpCompleteTimeout = 10.f * 60;
+			if (FPlatformTime::Seconds()  - ConnectStartTimeSeconds > WaitForPumpCompleteTimeout && !IsCookIgnoreTimeouts())
+			{
+				UE_LOG(LogCook, Error, TEXT("CookWorker process of CookWorkerServer %d failed to finalize its cook within %.0f seconds; we will tell it to shutdown."),
+					WorkerId.GetRemoteIndex(), WaitForPumpCompleteTimeout);
+				SendMessage(FAbortWorkerMessage(FAbortWorkerMessage::EType::Abort));
+				SendToState(EConnectStatus::WaitForDisconnect);
+				return;
+			}
+		}
+	}
+}
 void FCookWorkerServer::PumpConnect()
 {
 	for (;;)
 	{
+		if (IsConnected())
+		{
+			// Nothing further to do
+			return;
+		}
 		switch (ConnectStatus)
 		{
-		case EConnectStatus::Connected:
-			return; // Nothing further to do
-		case EConnectStatus::LostConnection:
-			return; // Nothing further to do
 		case EConnectStatus::Uninitialized:
 			LaunchProcess();
 			break;
@@ -201,6 +252,8 @@ void FCookWorkerServer::PumpConnect()
 				return; // Try again later
 			}
 			break;
+		case EConnectStatus::LostConnection:
+			return; // Nothing further to do
 		default:
 			checkNoEntry();
 			return;
@@ -210,21 +263,25 @@ void FCookWorkerServer::PumpConnect()
 
 void FCookWorkerServer::LaunchProcess()
 {
+	bool bShowCookWorkers = Director.GetShowWorkerOption() == FCookDirector::EShowWorker::SeparateWindows;
+
 	FString CommandletExecutable = FUnrealEdMisc::Get().GetProjectEditorBinaryPath();
 	FString CommandLine = Director.GetWorkerCommandLine(WorkerId);
 	CookWorkerHandle = FPlatformProcess::CreateProc(*CommandletExecutable, *CommandLine,
-		true /* bLaunchDetached */, true /* bLaunchHidden */, true /* bLaunchReallyHidden */,
+		true /* bLaunchDetached */, !bShowCookWorkers /* bLaunchHidden */, !bShowCookWorkers /* bLaunchReallyHidden */,
 		&CookWorkerProcessId, 0 /* PriorityModifier */, *FPaths::GetPath(CommandletExecutable),
 		nullptr /* PipeWriteChild */);
 	if (CookWorkerHandle.IsValid())
 	{
+		UE_LOG(LogCook, Display, TEXT("CookWorkerServer %d launched CookWorker as PID %u with commandline \"%s\"."),
+			WorkerId.GetRemoteIndex(), CookWorkerProcessId, *CommandLine);
 		SendToState(EConnectStatus::WaitForConnect);
 	}
 	else
 	{
 		// GetLastError information was logged by CreateProc
 		UE_LOG(LogCook, Error, TEXT("CookWorkerServer %d failed to create CookWorker process. Assigned packages will be returned to the director."),
-			WorkerId.GetLocalOrRemoteIndex());
+			WorkerId.GetRemoteIndex());
 		SendToState(EConnectStatus::LostConnection);
 	}
 }
@@ -242,7 +299,7 @@ void FCookWorkerServer::TickWaitForConnect()
 		if (!FPlatformProcess::IsProcRunning(CookWorkerHandle))
 		{
 			UE_LOG(LogCook, Error, TEXT("CookWorkerServer %d process terminated before connecting. Assigned packages will be returned to the director."),
-				WorkerId.GetLocalOrRemoteIndex());
+				WorkerId.GetRemoteIndex());
 			SendToState(EConnectStatus::LostConnection);
 			return;
 		}
@@ -252,7 +309,7 @@ void FCookWorkerServer::TickWaitForConnect()
 	if (CurrentTime - ConnectStartTimeSeconds > WaitForConnectTimeout && !IsCookIgnoreTimeouts())
 	{
 		UE_LOG(LogCook, Error, TEXT("CookWorkerServer %d process failed to connect within %.0f seconds. Assigned packages will be returned to the director."),
-			WorkerId.GetLocalOrRemoteIndex(), WaitForConnectTimeout);
+			WorkerId.GetRemoteIndex(), WaitForConnectTimeout);
 		ShutdownRemoteProcess();
 		SendToState(EConnectStatus::LostConnection);
 		return;
@@ -288,10 +345,11 @@ void FCookWorkerServer::TickWaitForDisconnect()
 		}
 	}
 
-	if ((bTerminateImmediately || CurrentTime - ConnectStartTimeSeconds > WaitForDisconnectTimeout) && !IsCookIgnoreTimeouts())
+	if (bTerminateImmediately || (CurrentTime - ConnectStartTimeSeconds > WaitForDisconnectTimeout && !IsCookIgnoreTimeouts()))
 	{
-		UE_LOG(LogCook, Warning, TEXT("CookWorkerServer %d's remote process failed to disconnect within %.0f seconds; we will terminate process."),
-			WorkerId.GetLocalOrRemoteIndex(), WaitForDisconnectTimeout);
+		UE_CLOG(!bTerminateImmediately, LogCook, Warning,
+			TEXT("CookWorker process of CookWorkerServer %d failed to disconnect within %.0f seconds; we will terminate it."),
+			WorkerId.GetRemoteIndex(), WaitForDisconnectTimeout);
 		ShutdownRemoteProcess();
 		SendToState(EConnectStatus::LostConnection);
 	}
@@ -303,17 +361,43 @@ void FCookWorkerServer::PumpSendMessages()
 	if (Status == UE::CompactBinaryTCP::Failed)
 	{
 		UE_LOG(LogCook, Error, TEXT("CookWorkerServer %d failed to write to socket, we will shutdown the remote process. Assigned packages will be returned to the director."),
-			WorkerId.GetLocalOrRemoteIndex());
+			WorkerId.GetRemoteIndex());
 		SendToState(EConnectStatus::WaitForDisconnect);
 		bTerminateImmediately = true;
 	}
 }
 
+void FCookWorkerServer::SendPendingPackages()
+{
+	if (PackagesToAssign.IsEmpty())
+	{
+		return;
+	}
+
+	TArray<FConstructPackageData> ConstructDatas;
+	ConstructDatas.Reserve(PackagesToAssign.Num());
+	for (FPackageData* PackageData : PackagesToAssign)
+	{
+		ConstructDatas.Add(PackageData->CreateConstructData());
+	}
+	PendingPackages.Append(PackagesToAssign);
+	PackagesToAssign.Empty();
+	SendMessage(FAssignPackagesMessage(MoveTemp(ConstructDatas)));
+}
+
 void FCookWorkerServer::PumpReceiveMessages()
 {
 	using namespace UE::CompactBinaryTCP;
 	TArray<FMarshalledMessage> Messages;
-	TryReadPacket(Socket, ReceiveBuffer, Messages);
+	EConnectionStatus SocketStatus = TryReadPacket(Socket, ReceiveBuffer, Messages);
+	if (SocketStatus != EConnectionStatus::Okay && SocketStatus != EConnectionStatus::Incomplete)
+	{
+		UE_LOG(LogCook, Error, TEXT("CookWorkerServer %d failed to read from socket, we will shutdown the remote process. Assigned packages will be returned to the director."),
+			WorkerId.GetRemoteIndex());
+		SendToState(EConnectStatus::WaitForDisconnect);
+		bTerminateImmediately = true;
+		return;
+	}
 	HandleReceiveMessages(MoveTemp(Messages));
 }
 
@@ -323,11 +407,24 @@ void FCookWorkerServer::HandleReceiveMessages(TArray<UE::CompactBinaryTCP::FMars
 	{
 		if (Message.MessageType == FAbortWorkerMessage::MessageType)
 		{
-			UE_LOG(LogCook, Error, TEXT("CookWorkerServer %d remote process shut down unexpectedly. Assigned packages will be returned to the director."),
-				WorkerId.GetLocalOrRemoteIndex());
+			UE_CLOG(ConnectStatus != EConnectStatus::PumpingCookComplete && ConnectStatus != EConnectStatus::WaitForDisconnect,
+				LogCook, Error, TEXT("CookWorkerServer %d remote process shut down unexpectedly. Assigned packages will be returned to the director."),
+				WorkerId.GetRemoteIndex());
 			SendToState(EConnectStatus::LostConnection);
 			break;
 		}
+		else if (Message.MessageType == FPackageResultsMessage::MessageType)
+		{
+			FPackageResultsMessage ResultsMessage;
+			if (!ResultsMessage.TryRead(Message.Object))
+			{
+				LogInvalidMessage(TEXT("FPackageResultsMessage"));
+			}
+			else
+			{
+				RecordResults(ResultsMessage);
+			}
+		}
 	}
 }
 
@@ -336,19 +433,77 @@ void FCookWorkerServer::SendMessage(const UE::CompactBinaryTCP::IMessage& Messag
 	UE::CompactBinaryTCP::TryWritePacket(Socket, SendBuffer, Message);
 }
 
-FAssignPackagesMessage::FAssignPackagesMessage(TArray<FName>&& InPackageNames)
-	: PackageNames(MoveTemp(InPackageNames))
+void FCookWorkerServer::RecordResults(FPackageResultsMessage& Message)
+{
+	for (FPackageResult& Result : Message.Results)
+	{
+		FPackageData* PackageData = COTFS.PackageDatas->FindPackageDataByPackageName(Result.PackageName);
+		if (!PackageData)
+		{
+			UE_LOG(LogCook, Warning, TEXT("CookWorkerServer %d received FPackageResultsMessage for invalid package %s. Ignoring it."),
+				WorkerId.GetRemoteIndex() , *Result.PackageName.ToString());
+			continue;
+		}
+		if (PendingPackages.Remove(PackageData) != 1)
+		{
+			UE_LOG(LogCook, Warning, TEXT("CookWorkerServer %d received FPackageResultsMessage for package %s which is not a pending package. Ignoring it."),
+				WorkerId.GetRemoteIndex() , *Result.PackageName.ToString());
+			continue;
+		}
+		PackageData->SetWorkerAssignment(FWorkerId::Invalid());
+
+		// MPCOOKTODO: Refactor FSaveCookedPackageContext::FinishPlatform and ::FinishPackage so we can call them from here
+		// to reduce duplication
+		if (Result.SuppressCookReason == ESuppressCookReason::InvalidSuppressCookReason)
+		{
+			int32 NumPlatforms = OrderedSessionPlatforms.Num();
+			if (Result.Platforms.Num() != NumPlatforms)
+			{
+				UE_LOG(LogCook, Warning, TEXT("CookWorkerServer %d received FPackageResultsMessage for package %s with an invalid number of platform results: expected %d, actual %d. Ignoring it."),
+					WorkerId.GetRemoteIndex(), *Result.PackageName.ToString(), NumPlatforms, Result.Platforms.Num());
+				continue;
+			}
+			for (int32 PlatformIndex = 0; PlatformIndex < NumPlatforms; ++PlatformIndex)
+			{
+				ITargetPlatform* TargetPlatform = OrderedSessionPlatforms[PlatformIndex];
+				FPackagePlatformResult& PlatformResult = Result.Platforms[PlatformIndex];
+				PackageData->SetPlatformCooked(TargetPlatform, PlatformResult.bSuccessful);
+				// MPCOOKTODO: Call CommitRemotePackage on the PackageWriter
+				// MPCOOKTODO: Add the AssetRegistry information here
+				// MPCOOKTODO: Copy the Hash
+			}
+			if (Result.bReferencedOnlyByEditorOnlyData)
+			{
+				COTFS.PackageTracker->UncookedEditorOnlyPackages.AddUnique(Result.PackageName);
+			}
+			COTFS.PromoteToSaveComplete(*PackageData, ESendFlags::QueueAddAndRemove);
+		}
+		else
+		{
+			COTFS.DemoteToIdle(*PackageData, ESendFlags::QueueAddAndRemove, Result.SuppressCookReason);
+		}
+	}
+}
+
+void FCookWorkerServer::LogInvalidMessage(const TCHAR* MessageTypeName)
+{
+	UE_LOG(LogCook, Warning, TEXT("CookWorkerClient received invalidly formatted message for type %s from CookDirector. Ignoring it."),
+		MessageTypeName);
+}
+
+FAssignPackagesMessage::FAssignPackagesMessage(TArray<FConstructPackageData>&& InPackageDatas)
+	: PackageDatas(MoveTemp(InPackageDatas))
 {
 }
 
 void FAssignPackagesMessage::Write(FCbWriter& Writer) const
 {
-	WriteArrayOfNames(Writer, "PackageNames", PackageNames);
+	Writer << "P" << PackageDatas;
 }
 
 bool FAssignPackagesMessage::TryRead(FCbObjectView Object)
 {
-	return TryReadArrayOfNames(Object, "PackageNames", PackageNames);
+	return LoadFromCompactBinary(Object["P"], PackageDatas);
 }
 
 FGuid FAssignPackagesMessage::MessageType(TEXT("B7B1542B73254B679319D73F753DB6F8"));
@@ -370,8 +525,175 @@ bool FAbortPackagesMessage::TryRead(FCbObjectView Object)
 
 FGuid FAbortPackagesMessage::MessageType(TEXT("D769F1BFF2F34978868D70E3CAEE94E7"));
 
+FAbortWorkerMessage::FAbortWorkerMessage(EType InType)
+	: Type(InType)
+{
+}
+
+void FAbortWorkerMessage::Write(FCbWriter& Writer) const
+{
+	Writer << "Type" << (uint8)Type;
+}
+
+bool FAbortWorkerMessage::TryRead(FCbObjectView Object)
+{
+	Type = static_cast<EType>(Object["Type"].AsUInt8((uint8)EType::Abort));
+	return true;
+}
+
 FGuid FAbortWorkerMessage::MessageType(TEXT("83FD99DFE8DB4A9A8E71684C121BE6F3"));
 
+void FInitialConfigMessage::ReadFromLocal(const UCookOnTheFlyServer& COTFS, const TArray<ITargetPlatform*>& InOrderedSessionPlatforms,
+	const FCookByTheBookOptions& InCookByTheBookOptions, const FCookOnTheFlyOptions& InCookOnTheFlyOptions)
+{
+	InitialSettings.CopyFromLocal(COTFS);
+	BeginCookSettings.CopyFromLocal(COTFS);
+	OrderedSessionPlatforms = InOrderedSessionPlatforms;
+	DirectorCookMode = COTFS.GetCookMode();
+	CookInitializationFlags = COTFS.GetCookFlags();
+	CookByTheBookOptions = InCookByTheBookOptions;
+	CookOnTheFlyOptions = InCookOnTheFlyOptions;
+	bZenStore = COTFS.IsUsingZenStore();
+}
+
+void FInitialConfigMessage::Write(FCbWriter& Writer) const
+{
+	int32 LocalCookMode = static_cast<int32>(DirectorCookMode);
+	Writer << "DirectorCookMode" << LocalCookMode;
+	int32 LocalCookFlags = static_cast<int32>(CookInitializationFlags);
+	Writer << "CookInitializationFlags" << LocalCookFlags;
+	Writer << "ZenStore" << bZenStore;
+
+	Writer.BeginArray("TargetPlatforms");
+	for (const ITargetPlatform* TargetPlatform : OrderedSessionPlatforms)
+	{
+		Writer << TargetPlatform->PlatformName();
+	}
+	Writer.EndArray();
+	Writer << "InitialSettings" << InitialSettings;
+	Writer << "BeginCookSettings" << BeginCookSettings;
+	Writer << "CookByTheBookOptions" << CookByTheBookOptions;
+	Writer << "CookOnTheFlyOptions" << CookOnTheFlyOptions;
+}
+
+bool FInitialConfigMessage::TryRead(FCbObjectView Object)
+{
+	bool bOk = true;
+	int32 LocalCookMode;
+	bOk = LoadFromCompactBinary(Object["DirectorCookMode"], LocalCookMode) & bOk;
+	DirectorCookMode = static_cast<ECookMode::Type>(LocalCookMode);
+	int32 LocalCookFlags;
+	bOk = LoadFromCompactBinary(Object["CookInitializationFlags"], LocalCookFlags) & bOk;
+	CookInitializationFlags = static_cast<ECookInitializationFlags>(LocalCookFlags);
+	bOk = LoadFromCompactBinary(Object["ZenStore"], bZenStore) & bOk;
+
+	ITargetPlatformManagerModule& TPM(GetTargetPlatformManagerRef());
+	FCbFieldView TargetPlatformsField = Object["TargetPlatforms"];
+	{
+		bOk = TargetPlatformsField.IsArray() & bOk;
+		OrderedSessionPlatforms.Reset(TargetPlatformsField.AsArrayView().Num());
+		for (FCbFieldView ElementField : TargetPlatformsField)
+		{
+			TStringBuilder<128> KeyName;
+			if (LoadFromCompactBinary(ElementField, KeyName))
+			{
+				ITargetPlatform* TargetPlatform = TPM.FindTargetPlatform(KeyName.ToView());
+				if (TargetPlatform)
+				{
+					OrderedSessionPlatforms.Add(TargetPlatform);
+				}
+				else
+				{
+					UE_LOG(LogCook, Error, TEXT("Could not find TargetPlatform \"%.*s\" received from CookDirector."),
+						KeyName.Len(), KeyName.GetData());
+					bOk = false;
+				}
+
+			}
+			else
+			{
+				bOk = false;
+			}
+		}
+	}
+
+	bOk = LoadFromCompactBinary(Object["InitialSettings"], InitialSettings) & bOk;
+	bOk = LoadFromCompactBinary(Object["BeginCookSettings"], BeginCookSettings) & bOk;
+	bOk = LoadFromCompactBinary(Object["CookByTheBookOptions"], CookByTheBookOptions) & bOk;
+	bOk = LoadFromCompactBinary(Object["CookOnTheFlyOptions"], CookOnTheFlyOptions) & bOk;
+	return bOk;
+}
+
+FGuid FInitialConfigMessage::MessageType(TEXT("340CDCB927304CEB9C0A66B5F707FC2B"));
+
+void FPackageResultsMessage::Write(FCbWriter& Writer) const
+{
+	Writer.BeginArray("R");
+	for (const FPackageResult& Result : Results)
+	{
+		Writer.BeginObject();
+		{
+			Writer << "N" << Result.PackageName;
+			Writer << "S" << (uint8) Result.SuppressCookReason;
+			Writer << "E" << Result.bReferencedOnlyByEditorOnlyData;
+			Writer.BeginArray("P");
+			for (const FPackagePlatformResult& PlatformResult : Result.Platforms)
+			{
+				Writer.BeginObject();
+				Writer << "S" << PlatformResult.bSuccessful;
+				Writer << "G" << PlatformResult.PackageGuid;
+				Writer << "D" << PlatformResult.TargetDomainDependencies;
+				Writer.EndObject();
+			}
+			Writer.EndArray();
+		}
+		Writer.EndObject();
+	}
+	Writer.EndArray();
+}
+
+bool FPackageResultsMessage::TryRead(FCbObjectView Object)
+{
+	Results.Reset();
+	for (FCbFieldView ResultField : Object["R"])
+	{
+		FCbObjectView ResultObject = ResultField.AsObjectView();
+		FPackageResult& Result = Results.Emplace_GetRef();
+		LoadFromCompactBinary(ResultObject["N"], Result.PackageName);
+		if (Result.PackageName.IsNone())
+		{
+			return false;
+		}
+		int32 LocalSuppressCookReason = ResultObject["S"].AsUInt8(MAX_uint8);
+		if (LocalSuppressCookReason == MAX_uint8)
+		{
+			return false;
+		}
+		Result.SuppressCookReason = static_cast<ESuppressCookReason>(LocalSuppressCookReason);
+		Result.bReferencedOnlyByEditorOnlyData = ResultObject["E"].AsBool();
+		Result.Platforms.Reset();
+		for (FCbFieldView PlatformField : ResultObject["P"])
+		{
+			FPackagePlatformResult& PlatformResult = Result.Platforms.Emplace_GetRef();
+			FCbObjectView PlatformObject = PlatformField.AsObjectView();
+
+			PlatformResult.bSuccessful = PlatformObject["S"].AsBool();
+			PlatformResult.PackageGuid = PlatformObject["G"].AsUuid();
+
+			FCbObjectView TargetDepsView = PlatformObject["D"].AsObjectView();
+			if (TargetDepsView.GetSize() > 0)
+			{
+				FUniqueBuffer TargetDepsCopy = FUniqueBuffer::Alloc(TargetDepsView.GetSize());
+				TargetDepsView.CopyTo(TargetDepsCopy.GetView());
+				PlatformResult.TargetDomainDependencies = FCbObject(TargetDepsCopy.MoveToShared());
+			}
+		}
+	}
+	return true;
+}
+
+FGuid FPackageResultsMessage::MessageType(TEXT("4631C6C0F6DC4CEFB2B09D3FB0B524DB"));
+
 void WriteArrayOfNames(FCbWriter& Writer, const char* ArrayName, TConstArrayView<FName> Names)
 {
 	Writer.BeginArray(ArrayName);
diff --git a/Engine/Source/Editor/UnrealEd/Private/Cooker/CookWorkerServer.h b/Engine/Source/Editor/UnrealEd/Private/Cooker/CookWorkerServer.h
--- a/Engine/Source/Editor/UnrealEd/Private/Cooker/CookWorkerServer.h
+++ b/Engine/Source/Editor/UnrealEd/Private/Cooker/CookWorkerServer.h
@@ -6,16 +6,19 @@
 #include "Containers/Array.h"
 #include "Containers/ArrayView.h"
 #include "Containers/Set.h"
+#include "CookOnTheSide/CookOnTheFlyServer.h"
 #include "CookSockets.h"
 #include "CookTypes.h"
 #include "Misc/Guid.h"
 
 class FSocket;
-class UCookOnTheFlyServer;
+class ITargetPlatform;
 struct FProcHandle;
 namespace UE::CompactBinaryTCP { struct FMarshalledMessage; }
 namespace UE::Cook { class FCookDirector; }
+namespace UE::Cook { struct FConstructPackageData; }
 namespace UE::Cook { struct FPackageData; }
+namespace UE::Cook { struct FPackageResultsMessage; }
 namespace UE::Cook { struct FWorkerConnectMessage; }
 
 namespace UE::Cook
@@ -45,7 +48,11 @@ public:
 	bool TryHandleConnectMessage(FWorkerConnectMessage& Message, FSocket* InSocket, TArray<UE::CompactBinaryTCP::FMarshalledMessage>&& OtherPacketMessages);
 	/** Periodic Tick function to send and receive messages to the Client. */
 	void TickFromSchedulerThread();
+	/** Called when the COTFS Server has detected all packages are complete. Tell the CookWorker to flush messages and exit. */
+	void PumpCookComplete();
 
+	/** Is this in one of the connected states (wherein the CookWorker is working and can send arbitrary messages)? */
+	bool IsConnected() const;
 	/** Is this either shutting down or completed shutdown of its remote Client? */
 	bool IsShuttingDown() const;
 	/** Is this not yet or no longer connected to a remote Client? */
@@ -56,7 +63,10 @@ private:
 	{
 		Uninitialized,
 		WaitForConnect,
-		Connected,
+		ConnectedFirst,
+		Connected = ConnectedFirst,
+		PumpingCookComplete,
+		ConnectedLast = PumpingCookComplete,
 		WaitForDisconnect,
 		LostConnection,
 	};
@@ -72,7 +82,9 @@ private:
 	void TickWaitForDisconnect();
 	/** Helper for Tick, pump send messages to a connected Client. */
 	void PumpSendMessages();
-	/** Helper for Tick, pump receive messages to a connected Client. */
+	/** Helper for PumpSendMessages; send a message for any PackagesToAssign we have. */
+	void SendPendingPackages();
+	/** Helper for Tick, pump receive messages from a connected Client. */
 	void PumpReceiveMessages();
 	/** Send the message immediately to the Socket. If cannot complete immediately, it will be finished during Tick. */
 	void SendMessage(const UE::CompactBinaryTCP::IMessage& Message);
@@ -84,9 +96,13 @@ private:
 	void ShutdownRemoteProcess();
 	/** Helper for PumpReceiveMessages: dispatch the messages received from the socket. */
 	void HandleReceiveMessages(TArray<UE::CompactBinaryTCP::FMarshalledMessage>&& Messages);
+	/** Add results from the client to the local CookOnTheFlyServer. */
+	void RecordResults(FPackageResultsMessage& Message);
+	void LogInvalidMessage(const TCHAR* MessageTypeName);
 
 	TArray<FPackageData*> PackagesToAssign;
 	TSet<FPackageData*> PendingPackages;
+	TArray<ITargetPlatform*> OrderedSessionPlatforms;
 	UE::CompactBinaryTCP::FSendBuffer SendBuffer;
 	UE::CompactBinaryTCP::FReceiveBuffer ReceiveBuffer;
 	FCookDirector& Director;
@@ -106,14 +122,14 @@ struct FAssignPackagesMessage : public UE::CompactBinaryTCP::IMessage
 {
 public:
 	FAssignPackagesMessage() = default;
-	FAssignPackagesMessage(TArray<FName>&& InPackageNames);
+	FAssignPackagesMessage(TArray<FConstructPackageData>&& InPackageDatas);
 
 	virtual void Write(FCbWriter& Writer) const override;
 	virtual bool TryRead(FCbObjectView Object) override;
 	virtual FGuid GetMessageType() const override { return MessageType; }
 
 public:
-	TArray<FName> PackageNames;
+	TArray<FConstructPackageData> PackageDatas;
 	static FGuid MessageType;
 };
 
@@ -141,11 +157,84 @@ public:
 struct FAbortWorkerMessage : public UE::CompactBinaryTCP::IMessage
 {
 public:
-	virtual void Write(FCbWriter& Writer) const override { }
-	virtual bool TryRead(FCbObjectView Object) override { return true; }
+	enum EType
+	{
+		CookComplete,
+		Abort,
+	};
+	FAbortWorkerMessage(EType InType = EType::Abort);
+	virtual void Write(FCbWriter& Writer) const override;
+	virtual bool TryRead(FCbObjectView Object) override;
+	virtual FGuid GetMessageType() const override { return MessageType; }
+
+public:
+	EType Type;
+	static FGuid MessageType;
+};
+
+/** Message From Server to Client giving all of the COTFS settings the client needs. */
+struct FInitialConfigMessage : public UE::CompactBinaryTCP::IMessage
+{
+public:
+	virtual void Write(FCbWriter& Writer) const override;
+	virtual bool TryRead(FCbObjectView Object) override;
 	virtual FGuid GetMessageType() const override { return MessageType; }
 
+	void ReadFromLocal(const UCookOnTheFlyServer& COTFS, const TArray<ITargetPlatform*>& InOrderedSessionPlatforms,
+		const FCookByTheBookOptions& InCookByTheBookOptions, const FCookOnTheFlyOptions& InCookOnTheFlyOptions);
+
+	ECookMode::Type GetDirectorCookMode() const { return DirectorCookMode; }
+	ECookInitializationFlags GetCookInitializationFlags() const { return CookInitializationFlags; }
+	FInitializeConfigSettings&& ConsumeInitializeConfigSettings() { return MoveTemp(InitialSettings); }
+	FBeginCookConfigSettings&& ConsumeBeginCookConfigSettings() { return MoveTemp(BeginCookSettings); }
+	FCookByTheBookOptions&& ConsumeCookByTheBookOptions() { return MoveTemp(CookByTheBookOptions); }
+	FCookOnTheFlyOptions&& ConsumeCookOnTheFlyOptions() { return MoveTemp(CookOnTheFlyOptions); }
+	const TArray<ITargetPlatform*>& GetOrderedSessionPlatforms() { return OrderedSessionPlatforms; }
+	bool IsZenStore() const { return bZenStore; }
+
 public:
+	static FGuid MessageType;
+private:
+	FInitializeConfigSettings InitialSettings;
+	FBeginCookConfigSettings BeginCookSettings;
+	FCookByTheBookOptions CookByTheBookOptions;
+	FCookOnTheFlyOptions CookOnTheFlyOptions;
+	TArray<ITargetPlatform*> OrderedSessionPlatforms;
+	ECookMode::Type DirectorCookMode = ECookMode::CookByTheBook;
+	ECookInitializationFlags CookInitializationFlags = ECookInitializationFlags::None;
+	bool bZenStore = false;
+};
+
+/** Helper struct for FPackageResultsMessage: holds per-package-per-platform data. */
+struct FPackagePlatformResult
+{
+	FGuid PackageGuid;
+	FCbObject TargetDomainDependencies;
+	bool bSuccessful = true;
+};
+
+/** Helper struct for FPackageResultsMessage: holds per-package data. */
+struct FPackageResult
+{
+	FName PackageName;
+	/** If failure reason is InvalidSuppressCookReason, it was saved. Otherwise, holds the suppression reason */
+	ESuppressCookReason SuppressCookReason;
+	bool bReferencedOnlyByEditorOnlyData = false;
+	// The elements of Platforms correspond to the TargetPlatforms present in OrderedSessionPlatforms
+	TArray<FPackagePlatformResult, TInlineAllocator<1>> Platforms;
+};
+
+/** Message from Client to Server giving the results for saved or refused-to-cook packages. */
+struct FPackageResultsMessage : public UE::CompactBinaryTCP::IMessage
+{
+public:
+	virtual void Write(FCbWriter& Writer) const override;
+	virtual bool TryRead(FCbObjectView Object) override;
+	virtual FGuid GetMessageType() const override { return MessageType; }
+
+public:
+	TArray<FPackageResult> Results;
+
 	static FGuid MessageType;
 };
 
diff --git a/Engine/Source/Editor/UnrealEd/Private/Cooker/WorkerRequestsRemote.cpp b/Engine/Source/Editor/UnrealEd/Private/Cooker/WorkerRequestsRemote.cpp
--- a/Engine/Source/Editor/UnrealEd/Private/Cooker/WorkerRequestsRemote.cpp
+++ b/Engine/Source/Editor/UnrealEd/Private/Cooker/WorkerRequestsRemote.cpp
@@ -11,20 +11,10 @@ namespace UE::Cook
 {
 
 FWorkerRequestsRemote::FWorkerRequestsRemote(UCookOnTheFlyServer& InCOTFS)
+	: CookWorkerClient(*InCOTFS.CookWorkerClient)
 {
 }
 
-bool FWorkerRequestsRemote::TryConnect(FDirectorConnectionInfo&& ConnectInfo, ECookInitializationFlags& OutCookFlags)
-{
-	return CookWorkerClient.TryConnect(MoveTemp(ConnectInfo));
-}
-
-bool FWorkerRequestsRemote::TryGetInitializeSettings(ECookMode::Type& OutCookMode, ECookInitializationFlags& OutCookFlags,
-	FString& OutOutputDirectoryOverride)
-{
-	return false;
-}
-
 bool FWorkerRequestsRemote::HasExternalRequests() const
 {
 	return ExternalRequests.HasRequests();
@@ -112,25 +102,23 @@ void FWorkerRequestsRemote::ReportAccessedIniSettings(UCookOnTheFlyServer& COTFS
 
 void FWorkerRequestsRemote::ReportDemoteToIdle(UE::Cook::FPackageData& PackageData, ESuppressCookReason Reason)
 {
-	// MPCOOKTODO: Not yet implemented
+	CookWorkerClient.ReportDemoteToIdle(PackageData, Reason);
 }
 
 void FWorkerRequestsRemote::ReportPromoteToSaveComplete(UE::Cook::FPackageData& PackageData)
 {
-	// MPCOOKTODO: Not yet implemented
+	CookWorkerClient.ReportPromoteToSaveComplete(PackageData);
 }
 
 void FWorkerRequestsRemote::GetInitializeConfigSettings(UCookOnTheFlyServer& COTFS,
 	const FString& OutputDirectoryOverride, UE::Cook::FInitializeConfigSettings& Settings)
 {
-	// MPCOOKTODO: Not yet implemented
-	Settings.LoadLocal(OutputDirectoryOverride);
+	Settings = CookWorkerClient.ConsumeInitializeConfigSettings();
 }
 
 void FWorkerRequestsRemote::GetBeginCookConfigSettings(UCookOnTheFlyServer& COTFS, FBeginCookContext& BeginContext, UE::Cook::FBeginCookConfigSettings& Settings)
 {
-	// MPCOOKTODO: Not yet implemented
-	Settings.LoadLocal(BeginContext);
+	Settings = CookWorkerClient.ConsumeBeginCookConfigSettings();
 }
 
 void FWorkerRequestsRemote::GetBeginCookIterativeFlags(UCookOnTheFlyServer& COTFS, FBeginCookContext& BeginContext)
@@ -152,8 +140,7 @@ void FWorkerRequestsRemote::GetBeginCookIterativeFlags(UCookOnTheFlyServer& COTF
 
 ECookMode::Type FWorkerRequestsRemote::GetDirectorCookMode(UCookOnTheFlyServer& COTFS)
 {
-	// MPCOOKTODO: Not yet implemented
-	return ECookMode::CookByTheBook;
+	return CookWorkerClient.GetDirectorCookMode();
 }
 
 void FWorkerRequestsRemote::LogCalledCookByTheBookError(const TCHAR* FunctionName) const
diff --git a/Engine/Source/Editor/UnrealEd/Private/Cooker/WorkerRequestsRemote.h b/Engine/Source/Editor/UnrealEd/Private/Cooker/WorkerRequestsRemote.h
--- a/Engine/Source/Editor/UnrealEd/Private/Cooker/WorkerRequestsRemote.h
+++ b/Engine/Source/Editor/UnrealEd/Private/Cooker/WorkerRequestsRemote.h
@@ -18,10 +18,6 @@ class FWorkerRequestsRemote : public IWorkerRequests
 public:
 	FWorkerRequestsRemote(UCookOnTheFlyServer& InCOTFS);
 
-	bool TryConnect(FDirectorConnectionInfo&& ConnectInfo, ECookInitializationFlags& OutCookFlags);
-	bool TryGetInitializeSettings(ECookMode::Type& OutCookMode, ECookInitializationFlags& OutCookFlags,
-		FString& OutOutputDirectoryOverride);
-
 	virtual bool HasExternalRequests() const override;
 	virtual int32 GetNumExternalRequests() const override;
 	virtual EExternalRequestType DequeueNextCluster(TArray<FSchedulerCallback>& OutCallbacks, TArray<FFilePlatformRequest>& OutBuildRequests) override;
@@ -54,7 +50,7 @@ private:
 
 private:
 	FExternalRequests ExternalRequests;
-	FCookWorkerClient CookWorkerClient;
+	FCookWorkerClient& CookWorkerClient;
 };
 
 }
\ No newline at end of file
diff --git a/Engine/Source/Runtime/Engine/Public/Engine/ICookInfo.h b/Engine/Source/Runtime/Engine/Public/Engine/ICookInfo.h
--- a/Engine/Source/Runtime/Engine/Public/Engine/ICookInfo.h
+++ b/Engine/Source/Runtime/Engine/Public/Engine/ICookInfo.h
@@ -41,6 +41,7 @@ namespace UE::Cook
 	CallbackMacro(AssetManagerModifyDLCCook, true) \
 	CallbackMacro(TargetPlatformExtraPackagesToCook, true) \
 	CallbackMacro(ConsoleCommand, true) \
+	CallbackMacro(CookDirector, true) \
 	CallbackMacro(CookOnTheFly, true) \
 	CallbackMacro(IterativeCook, true) \
 	CallbackMacro(PreviousAssetRegistry, true) \
