World Partition: fix recent regression with resave actors builder.

#rnx

diff --git a/Engine/Source/Editor/UnrealEd/Private/WorldPartition/WorldPartitionResaveActorsBuilder.cpp b/Engine/Source/Editor/UnrealEd/Private/WorldPartition/WorldPartitionResaveActorsBuilder.cpp
--- a/Engine/Source/Editor/UnrealEd/Private/WorldPartition/WorldPartitionResaveActorsBuilder.cpp
+++ b/Engine/Source/Editor/UnrealEd/Private/WorldPartition/WorldPartitionResaveActorsBuilder.cpp
@@ -304,7 +304,7 @@ bool UWorldPartitionResaveActorsBuilder::RunInternal(UWorld* World, const FCellI
 				}
 			}
 
-			return false;
+			return ActorTags.IsEmpty() && ActorProperties.IsEmpty();
 		};
 
 		ForEachActorWithLoadingParams.OnPreGarbageCollect = [&PackagesToSave, &PackageHelper]()
