World Partition:
- Dont explicitly reset loaders when saving a map using "Save As".
- Reset package loaders when changing an actor packaging from external to internal (required if the package contains bulk data).
- Propagate renaming flags to dependant objects when renaming a world.

#rnx

diff --git a/Engine/Source/Editor/UnrealEd/Private/FileHelpers.cpp b/Engine/Source/Editor/UnrealEd/Private/FileHelpers.cpp
--- a/Engine/Source/Editor/UnrealEd/Private/FileHelpers.cpp
+++ b/Engine/Source/Editor/UnrealEd/Private/FileHelpers.cpp
@@ -803,10 +803,8 @@ static bool SaveWorld(UWorld* World,
 
 				if (!DuplicatedWorld)
 				{
-					// Explict Reset Loaders of Package here because we want to avoid resetting of all loaders which is the current behavior of UObject::Rename when passing in a UPackage
-					ResetLoaders(Package);
 					// Duplicate failed or not needed. Just do a rename.
-					Package->Rename(*NewPackageName, NULL, REN_NonTransactional | REN_DontCreateRedirectors | REN_ForceNoResetLoaders);
+					Package->Rename(*NewPackageName, NULL, REN_NonTransactional | REN_DontCreateRedirectors);
 					
 					if (bWorldNeedsRename)
 					{
@@ -823,7 +821,7 @@ static bool SaveWorld(UWorld* World,
 							}
 						}
 
-						World->Rename(*NewWorldAssetName, NULL, REN_NonTransactional | REN_DontCreateRedirectors | REN_ForceNoResetLoaders);
+						World->Rename(*NewWorldAssetName, NULL, REN_NonTransactional | REN_DontCreateRedirectors);
 					}
 
 					// We're changing the world path, add a path redirector so that soft object paths get fixed on save
diff --git a/Engine/Source/Runtime/Engine/Private/ActorEditor.cpp b/Engine/Source/Runtime/Engine/Private/ActorEditor.cpp
--- a/Engine/Source/Runtime/Engine/Private/ActorEditor.cpp
+++ b/Engine/Source/Runtime/Engine/Private/ActorEditor.cpp
@@ -838,8 +838,8 @@ void AActor::SetPackageExternal(bool bExternal, bool bShouldDirty)
 	else 
 	{
 		UPackage* ActorPackage = GetExternalPackage();
-		// Detach the linker exports so it doesn't resolve to this actor anymore
-		ResetLinkerExports(ActorPackage);
+		// Detach the loaders so it doesn't resolve to this actor anymore
+		ResetLoaders(ActorPackage);
 		SetExternalPackage(nullptr);
 	}
 
diff --git a/Engine/Source/Runtime/Engine/Private/World.cpp b/Engine/Source/Runtime/Engine/Private/World.cpp
--- a/Engine/Source/Runtime/Engine/Private/World.cpp
+++ b/Engine/Source/Runtime/Engine/Private/World.cpp
@@ -702,7 +702,7 @@ bool UWorld::Rename(const TCHAR* InName, UObject* NewOuter, ERenameFlags Flags)
 			// Move dependant objects into the new actor package
 			for (UObject* DependantObject : DependantObjects)
 			{
-				DependantObject->Rename(nullptr, Actor->GetExternalPackage(), REN_NonTransactional | REN_DontCreateRedirectors | REN_ForceNoResetLoaders | REN_DoNotDirty);
+				DependantObject->Rename(nullptr, Actor->GetExternalPackage(), Flags);
 			}
 		}
 
