NoPCH DisableUnity fixes.

diff --git a/Engine/Source/Editor/UnrealEd/Private/Cooker/CookPackageData.cpp b/Engine/Source/Editor/UnrealEd/Private/Cooker/CookPackageData.cpp
--- a/Engine/Source/Editor/UnrealEd/Private/Cooker/CookPackageData.cpp
+++ b/Engine/Source/Editor/UnrealEd/Private/Cooker/CookPackageData.cpp
@@ -27,6 +27,7 @@
 #include "Misc/PreloadableFile.h"
 #include "Misc/ScopeExit.h"
 #include "Misc/ScopeLock.h"
+#include "Serialization/CompactBinaryWriter.h"
 #include "ShaderCompiler.h"
 #include "UObject/Object.h"
 #include "UObject/Package.h"
