Avoid using FPackageSourceControlHelper as base class or class member. It contains a FScopedSourceControl that will initiate a connection to the SCM system when the class is constructed. This is known to cause failures for automated processes (cooking) when the connection to the SCM fails even if not relevant to the actual process.

diff --git a/Engine/Source/Editor/UnrealEd/Classes/Commandlets/ExternalActorsCommandlet.h b/Engine/Source/Editor/UnrealEd/Classes/Commandlets/ExternalActorsCommandlet.h
--- a/Engine/Source/Editor/UnrealEd/Classes/Commandlets/ExternalActorsCommandlet.h
+++ b/Engine/Source/Editor/UnrealEd/Classes/Commandlets/ExternalActorsCommandlet.h
@@ -20,6 +20,4 @@ public:
 
 protected:
 	UWorld* LoadWorld(const FString& LevelToLoad);
-
-	FPackageSourceControlHelper PackageHelper;
 };
diff --git a/Engine/Source/Editor/UnrealEd/Classes/Commandlets/WorldPartitionBuilderCommandlet.h b/Engine/Source/Editor/UnrealEd/Classes/Commandlets/WorldPartitionBuilderCommandlet.h
--- a/Engine/Source/Editor/UnrealEd/Classes/Commandlets/WorldPartitionBuilderCommandlet.h
+++ b/Engine/Source/Editor/UnrealEd/Classes/Commandlets/WorldPartitionBuilderCommandlet.h
@@ -9,7 +9,7 @@
 class UWorldPartitionBuilder;
 
 UCLASS(Config=Engine)
-class UWorldPartitionBuilderCommandlet : public UCommandlet, public FPackageSourceControlHelper
+class UWorldPartitionBuilderCommandlet : public UCommandlet
 {
 	GENERATED_UCLASS_BODY()
 
diff --git a/Engine/Source/Editor/UnrealEd/Classes/Commandlets/WorldPartitionConvertCommandlet.h b/Engine/Source/Editor/UnrealEd/Classes/Commandlets/WorldPartitionConvertCommandlet.h
--- a/Engine/Source/Editor/UnrealEd/Classes/Commandlets/WorldPartitionConvertCommandlet.h
+++ b/Engine/Source/Editor/UnrealEd/Classes/Commandlets/WorldPartitionConvertCommandlet.h
@@ -87,7 +87,6 @@ protected:
 	FString LevelConfigFilename;
 	TArray<UPackage*> PackagesToSave;
 	TArray<UPackage*> PackagesToDelete;
-	FPackageSourceControlHelper PackageHelper;
 
 	bool bOnlyMergeSubLevels;
 	bool bDeleteSourceLevels;
diff --git a/Engine/Source/Editor/UnrealEd/Classes/Commandlets/WorldPartitionDataLayerToAssetCommandLet.h b/Engine/Source/Editor/UnrealEd/Classes/Commandlets/WorldPartitionDataLayerToAssetCommandLet.h
--- a/Engine/Source/Editor/UnrealEd/Classes/Commandlets/WorldPartitionDataLayerToAssetCommandLet.h
+++ b/Engine/Source/Editor/UnrealEd/Classes/Commandlets/WorldPartitionDataLayerToAssetCommandLet.h
@@ -149,8 +149,8 @@ protected:
 private:
 	bool InitializeFromCommandLine(TArray<FString>& Tokens, TArray<FString> const& Switches, TMap<FString, FString> const& Params);
 
-	bool BuildConversionInfos(TStrongObjectPtr<UDataLayerToAssetCommandletContext>& CommandletContext);
-	bool CreateConversionFromDataLayer(TStrongObjectPtr<UDataLayerToAssetCommandletContext>& CommandletContext, const UDeprecatedDataLayerInstance* DataLayer);
+	bool BuildConversionInfos(TStrongObjectPtr<UDataLayerToAssetCommandletContext>& CommandletContext, FPackageSourceControlHelper& PackageHelper);
+	bool CreateConversionFromDataLayer(TStrongObjectPtr<UDataLayerToAssetCommandletContext>& CommandletContext, const UDeprecatedDataLayerInstance* DataLayer, FPackageSourceControlHelper& PackageHelper);
 	TObjectPtr<UDataLayerAsset> GetOrCreateDataLayerAssetForConversion(TStrongObjectPtr<UDataLayerToAssetCommandletContext>& CommandletContext, FName DataLayerName);
 
 	bool ResolvePreviousConversionsToCurrent(TStrongObjectPtr<UDataLayerToAssetCommandletContext>& CommandletContext);
@@ -158,13 +158,13 @@ private:
 	bool CreateDataLayerInstances(TStrongObjectPtr<UDataLayerToAssetCommandletContext>& CommandletContext);
 	bool RebuildDataLayerHierarchies(TStrongObjectPtr<UDataLayerToAssetCommandletContext>& CommandletContext);
 
-	bool RemapActorDataLayersToAssets(TStrongObjectPtr<UDataLayerToAssetCommandletContext>& CommandletContext);
+	bool RemapActorDataLayersToAssets(TStrongObjectPtr<UDataLayerToAssetCommandletContext>& CommandletContext, FPackageSourceControlHelper& PackageHelper);
 	uint32 RemapActorDataLayers(TStrongObjectPtr<UDataLayerToAssetCommandletContext>& CommandletContext, AActor* Actor);
 	uint32 RemapDataLayersAssetsFromPreviousConversions(TStrongObjectPtr<UDataLayerToAssetCommandletContext>& CommandletContext, AActor* Actor);
 	
-	bool DeletePreviousConversionsData(TStrongObjectPtr<UDataLayerToAssetCommandletContext>& CommandletContext);
+	bool DeletePreviousConversionsData(TStrongObjectPtr<UDataLayerToAssetCommandletContext>& CommandletContext, FPackageSourceControlHelper& PackageHelper);
 
-	bool CommitConversion(TStrongObjectPtr<UDataLayerToAssetCommandletContext>& CommandletContext);
+	bool CommitConversion(TStrongObjectPtr<UDataLayerToAssetCommandletContext>& CommandletContext, FPackageSourceControlHelper& PackageHelper);
 
 	FString GetConversionFolder() const;
 	bool IsAssetInConversionFolder(const TObjectPtr<UDataLayerAsset> DataLayerAsset);
@@ -186,6 +186,4 @@ private:
 
 	UPROPERTY(Transient)
 	TObjectPtr<UWorld> MainWorld;
-
-	FPackageSourceControlHelper PackageHelper;
 };
\ No newline at end of file
diff --git a/Engine/Source/Editor/UnrealEd/Private/Commandlets/ExternalActorsCommandlet.cpp b/Engine/Source/Editor/UnrealEd/Private/Commandlets/ExternalActorsCommandlet.cpp
--- a/Engine/Source/Editor/UnrealEd/Private/Commandlets/ExternalActorsCommandlet.cpp
+++ b/Engine/Source/Editor/UnrealEd/Private/Commandlets/ExternalActorsCommandlet.cpp
@@ -36,6 +36,8 @@ UWorld* UExternalActorsCommandlet::LoadWorld(const FString& LevelToLoad)
 
 int32 UExternalActorsCommandlet::Main(const FString& Params)
 {
+	FPackageSourceControlHelper PackageHelper;
+
 	TArray<FString> Tokens, Switches;
 	ParseCommandLine(*Params, Tokens, Switches);
 
diff --git a/Engine/Source/Editor/UnrealEd/Private/Commandlets/WorldPartitionBuilderCommandlet.cpp b/Engine/Source/Editor/UnrealEd/Private/Commandlets/WorldPartitionBuilderCommandlet.cpp
--- a/Engine/Source/Editor/UnrealEd/Private/Commandlets/WorldPartitionBuilderCommandlet.cpp
+++ b/Engine/Source/Editor/UnrealEd/Private/Commandlets/WorldPartitionBuilderCommandlet.cpp
@@ -22,6 +22,8 @@ UWorldPartitionBuilderCommandlet::UWorldPartitionBuilderCommandlet(const FObject
 
 int32 UWorldPartitionBuilderCommandlet::Main(const FString& Params)
 {
+	FPackageSourceControlHelper PackageHelper;
+
 	TRACE_CPUPROFILER_EVENT_SCOPE(UWorldPartitionBuilderCommandlet::Main);
 
 	UE_SCOPED_TIMER(TEXT("Execution"), LogWorldPartitionBuilderCommandlet, Display);
diff --git a/Engine/Source/Editor/UnrealEd/Private/Commandlets/WorldPartitionConvertCommandlet.cpp b/Engine/Source/Editor/UnrealEd/Private/Commandlets/WorldPartitionConvertCommandlet.cpp
--- a/Engine/Source/Editor/UnrealEd/Private/Commandlets/WorldPartitionConvertCommandlet.cpp
+++ b/Engine/Source/Editor/UnrealEd/Private/Commandlets/WorldPartitionConvertCommandlet.cpp
@@ -690,6 +690,8 @@ int32 UWorldPartitionConvertCommandlet::Main(const FString& Params)
 
 	UE_SCOPED_TIMER(TEXT("Conversion"), LogWorldPartitionConvertCommandlet, Display);
 
+	FPackageSourceControlHelper PackageHelper;
+
 	TArray<FString> Tokens, Switches;
 	TMap<FString, FString> Arguments;
 	ParseCommandLine(*Params, Tokens, Switches, Arguments);
@@ -792,7 +794,7 @@ int32 UWorldPartitionConvertCommandlet::Main(const FString& Params)
 			FString ExternalObjectsFilePath = FPackageName::LongPackageNameToFilename(ExternalObjectsPath);
 			if (IFileManager::Get().DirectoryExists(*ExternalObjectsFilePath))
 			{
-				bool bResult = IFileManager::Get().IterateDirectoryRecursively(*ExternalObjectsFilePath, [this](const TCHAR* FilenameOrDirectory, bool bIsDirectory)
+				bool bResult = IFileManager::Get().IterateDirectoryRecursively(*ExternalObjectsFilePath, [this, &PackageHelper](const TCHAR* FilenameOrDirectory, bool bIsDirectory)
 				{
 					if (!bIsDirectory)
 					{
diff --git a/Engine/Source/Editor/UnrealEd/Private/Commandlets/WorldPartitionDataLayerToAssetCommandlet.cpp b/Engine/Source/Editor/UnrealEd/Private/Commandlets/WorldPartitionDataLayerToAssetCommandlet.cpp
--- a/Engine/Source/Editor/UnrealEd/Private/Commandlets/WorldPartitionDataLayerToAssetCommandlet.cpp
+++ b/Engine/Source/Editor/UnrealEd/Private/Commandlets/WorldPartitionDataLayerToAssetCommandlet.cpp
@@ -244,6 +244,8 @@ int32 UDataLayerToAssetCommandlet::Main(const FString& Params)
 
 	UE_SCOPED_TIMER(TEXT("Data Layer Conversion"), LogDataLayerToAssetCommandlet, Display);
 
+	FPackageSourceControlHelper PackageHelper;
+
 	ON_SCOPE_EXIT
 	{
 		if (MainWorld != nullptr)
@@ -265,7 +267,7 @@ int32 UDataLayerToAssetCommandlet::Main(const FString& Params)
 
 	TStrongObjectPtr<UDataLayerToAssetCommandletContext> Context(NewObject<UDataLayerToAssetCommandletContext>());
 
-	if (!BuildConversionInfos(Context))
+	if (!BuildConversionInfos(Context, PackageHelper))
 	{
 		return EReturnCode::DataLayerConversionError;
 	}
@@ -282,7 +284,7 @@ int32 UDataLayerToAssetCommandlet::Main(const FString& Params)
 		return EReturnCode::DataLayerConversionError;
 	}
 
-	if (!RemapActorDataLayersToAssets(Context))
+	if (!RemapActorDataLayersToAssets(Context, PackageHelper))
 	{
 		return EReturnCode::ActorDataLayerRemappingError;
 	}
@@ -292,12 +294,12 @@ int32 UDataLayerToAssetCommandlet::Main(const FString& Params)
 		return EReturnCode::ProjectSpecificConversionError;
 	}
 
-	if (!DeletePreviousConversionsData(Context))
+	if (!DeletePreviousConversionsData(Context, PackageHelper))
 	{
 		return EReturnCode::DataLayerConversionError;
 	}
 
-	if (!CommitConversion(Context))
+	if (!CommitConversion(Context, PackageHelper))
 	{
 		return EReturnCode::DataLayerConversionError;
 	}
@@ -337,7 +339,7 @@ bool UDataLayerToAssetCommandlet::InitializeFromCommandLine(TArray<FString>& Tok
 	return true;
 }
 
-bool UDataLayerToAssetCommandlet::BuildConversionInfos(TStrongObjectPtr<UDataLayerToAssetCommandletContext>& CommandletContext)
+bool UDataLayerToAssetCommandlet::BuildConversionInfos(TStrongObjectPtr<UDataLayerToAssetCommandletContext>& CommandletContext, FPackageSourceControlHelper& PackageHelper)
 {
 	UE_SCOPED_TIMER(TEXT("Retrieving Already Converted Data Layers"), LogDataLayerToAssetCommandlet, Display);
 
@@ -371,7 +373,7 @@ bool UDataLayerToAssetCommandlet::BuildConversionInfos(TStrongObjectPtr<UDataLay
 	{
 		if (const UDeprecatedDataLayerInstance* DataLayerToConvert = Cast<UDeprecatedDataLayerInstance>(DataLayerInstance))
 		{
-			if (!CreateConversionFromDataLayer(CommandletContext, DataLayerToConvert))
+			if (!CreateConversionFromDataLayer(CommandletContext, DataLayerToConvert, PackageHelper))
 			{
 				ErrorCount++;
 			}
@@ -381,7 +383,7 @@ bool UDataLayerToAssetCommandlet::BuildConversionInfos(TStrongObjectPtr<UDataLay
 	return ErrorCount == 0;
 }
 
-bool UDataLayerToAssetCommandlet::CreateConversionFromDataLayer(TStrongObjectPtr<UDataLayerToAssetCommandletContext>& CommandletContext, const UDeprecatedDataLayerInstance* DataLayer)
+bool UDataLayerToAssetCommandlet::CreateConversionFromDataLayer(TStrongObjectPtr<UDataLayerToAssetCommandletContext>& CommandletContext, const UDeprecatedDataLayerInstance* DataLayer, FPackageSourceControlHelper &PackageHelper)
 {
 	if(TObjectPtr<UDataLayerAsset> DataLayerAsset = GetOrCreateDataLayerAssetForConversion(CommandletContext, FName(DataLayer->GetDataLayerShortName())))
 	{
@@ -470,13 +472,13 @@ bool UDataLayerToAssetCommandlet::ResolvePreviousConversionsToCurrent(TStrongObj
 	return ErrorCount == 0;
 }
 
-bool UDataLayerToAssetCommandlet::RemapActorDataLayersToAssets(TStrongObjectPtr<UDataLayerToAssetCommandletContext>& CommandletContext)
+bool UDataLayerToAssetCommandlet::RemapActorDataLayersToAssets(TStrongObjectPtr<UDataLayerToAssetCommandletContext>& CommandletContext, FPackageSourceControlHelper& PackageHelper)
 {
 	UE_SCOPED_TIMER(TEXT("Remapping Actors Data Layers"), LogDataLayerToAssetCommandlet, Display);
 	UE_LOG(LogDataLayerToAssetCommandlet, Log, TEXT("Starting Actor Data Layer Remapping To Data Layer Asset. This can take a while."));
 
 	uint32 ErrorCount = 0;
-	FWorldPartitionHelpers::ForEachActorWithLoading(MainWorld->GetWorldPartition(), [&ErrorCount, &CommandletContext, this](const FWorldPartitionActorDesc* ActorDesc)
+	FWorldPartitionHelpers::ForEachActorWithLoading(MainWorld->GetWorldPartition(), [&ErrorCount, &CommandletContext, this, &PackageHelper](const FWorldPartitionActorDesc* ActorDesc)
 	{
 		uint32 ActorConversionErrors = 0;
 		if (AActor* Actor = ActorDesc->GetActor())
@@ -696,7 +698,7 @@ bool UDataLayerToAssetCommandlet::RebuildDataLayerHierarchies(TStrongObjectPtr<U
 	return ErrorCount == 0;
 }
 
-bool UDataLayerToAssetCommandlet::DeletePreviousConversionsData(TStrongObjectPtr<UDataLayerToAssetCommandletContext>& CommandletContext)
+bool UDataLayerToAssetCommandlet::DeletePreviousConversionsData(TStrongObjectPtr<UDataLayerToAssetCommandletContext>& CommandletContext, FPackageSourceControlHelper& PackageHelper)
 {
 	UE_SCOPED_TIMER(TEXT("Delete Conflicting Assets"), LogDataLayerToAssetCommandlet, Display);
 
@@ -724,7 +726,7 @@ bool UDataLayerToAssetCommandlet::DeletePreviousConversionsData(TStrongObjectPtr
 	return ErrorCount == 0;
 }
 
-bool UDataLayerToAssetCommandlet::CommitConversion(TStrongObjectPtr<UDataLayerToAssetCommandletContext>& CommandletContext)
+bool UDataLayerToAssetCommandlet::CommitConversion(TStrongObjectPtr<UDataLayerToAssetCommandletContext>& CommandletContext, FPackageSourceControlHelper& PackageHelper)
 {
 	uint32 ErrorCount = 0;
 	AWorldDataLayers* WorldDataLayers = MainWorld->GetWorldDataLayers();
diff --git a/Engine/Source/Editor/UnrealEd/Public/PackageSourceControlHelper.h b/Engine/Source/Editor/UnrealEd/Public/PackageSourceControlHelper.h
--- a/Engine/Source/Editor/UnrealEd/Public/PackageSourceControlHelper.h
+++ b/Engine/Source/Editor/UnrealEd/Public/PackageSourceControlHelper.h
@@ -6,7 +6,9 @@
 #include "UObject/ObjectMacros.h"
 #include "SourceControlHelpers.h"
 
-class UNREALED_API FPackageSourceControlHelper
+//This class should not be used as a base class because it can force a connection to the SCM through FScopedSourceControl. Connection errors can result in failing automation processes that are unrelated. 
+//The same can happen if the class is used as a member of another class. It should only be used as a local variable in a method so the connection is only attempted if required. 
+class UNREALED_API FPackageSourceControlHelper final
 {
 public:
 	bool UseSourceControl() const;
