[Backout] - CL21042585
Original CL Desc
-----------------------------------------------------------------
* Unified landscape brush command shortcuts with other modeling tools
* Added command to change both the brush size (left-right mouse move) and falloff (up-down mouse move). Default is ctrl+alt+right click drag
* Fixed the "change landscape brush size/falloff/strength" commands that couldn't work with default chord values nor secondary key bindings
* Added DoesCommandSupportInput function to editor viewport client in order to centralize copy/pasted logic
* Fixed "cmd" modifier not being taken into account all the time in said logic

diff --git a/Engine/Source/Editor/UnrealEd/Private/EditorViewportClient.cpp b/Engine/Source/Editor/UnrealEd/Private/EditorViewportClient.cpp
--- a/Engine/Source/Editor/UnrealEd/Private/EditorViewportClient.cpp
+++ b/Engine/Source/Editor/UnrealEd/Private/EditorViewportClient.cpp
@@ -2848,7 +2848,7 @@ bool FEditorViewportClient::Internal_InputKey(const FInputKeyEventArgs& EventArg
 	{
 		return true;
 	}
-	
+
 	FKey Key = EventArgs.Key;
 	EInputEvent Event = EventArgs.Event;
 	FViewport* InViewport = EventArgs.Viewport;
@@ -2867,7 +2867,7 @@ bool FEditorViewportClient::Internal_InputKey(const FInputKeyEventArgs& EventArg
 	}
 
 	FInputEventState InputState(EventArgs.Viewport, EventArgs.Key, EventArgs.Event);
-	
+
 	bool bHandled = false;
 
 	if ((IsOrtho() || InputState.IsAltButtonPressed()) && (Key == EKeys::Left || Key == EKeys::Right || Key == EKeys::Up || Key == EKeys::Down))
@@ -3195,24 +3195,6 @@ bool FEditorViewportClient::IsCmdPressed() const
 	return Viewport->KeyState(EKeys::LeftCommand) || Viewport->KeyState(EKeys::RightCommand);
 }
 
-bool FEditorViewportClient::IsCommandChordPressed(const TSharedPtr<FUICommandInfo> InCommand, FKey InOptionalKey) const
-{
-	bool bIsChordPressed = false;
-	// Check each bound chord
-	for (uint32 i = 0; i < static_cast<uint32>(EMultipleKeyBindingIndex::NumChords); ++i)
-	{
-		EMultipleKeyBindingIndex ChordIndex = static_cast<EMultipleKeyBindingIndex> (i);
-		const FInputChord& Chord = *InCommand->GetActiveChord(ChordIndex);
-
-		bIsChordPressed |= Chord.IsValidChord()
-			&& (Chord.NeedsControl() == IsCtrlPressed())
-			&& (Chord.NeedsAlt() == IsAltPressed())
-			&& (Chord.NeedsShift() == IsShiftPressed())
-			&& (Chord.NeedsCommand() == IsCmdPressed())
-			&& (InOptionalKey.IsValid() ? (Chord.Key == InOptionalKey && Viewport->KeyState(InOptionalKey)) : Viewport->KeyState(Chord.Key));
-	}
-	return bIsChordPressed;
-}
 
 void FEditorViewportClient::ProcessDoubleClickInViewport( const struct FInputEventState& InputState, FSceneView& View )
 {
