World Partition: fix convert commandlet invalidating soft object paths belonging to the persistent level being converted.

#rnx

diff --git a/Engine/Source/Editor/UnrealEd/Private/Commandlets/WorldPartitionConvertCommandlet.cpp b/Engine/Source/Editor/UnrealEd/Private/Commandlets/WorldPartitionConvertCommandlet.cpp
--- a/Engine/Source/Editor/UnrealEd/Private/Commandlets/WorldPartitionConvertCommandlet.cpp
+++ b/Engine/Source/Editor/UnrealEd/Private/Commandlets/WorldPartitionConvertCommandlet.cpp
@@ -433,7 +433,7 @@ void UWorldPartitionConvertCommandlet::FixupSoftObjectPaths(UPackage* OuterPacka
 				return *this;
 			}
 
-			FString OriginalValue = Value.ToString();
+			const FString OriginalValue = Value.ToString();
 
 			auto GetSourceString = [this]()
 			{
@@ -463,14 +463,6 @@ void UWorldPartitionConvertCommandlet::FixupSoftObjectPaths(UPackage* OuterPacka
 						Value.SetPath(NewPath);
 					}
 				}
-
-				FString NewValue = Value.ToString();
-				if (NewValue == OriginalValue)
-				{
-					Value.Reset();
-					UE_LOG(LogWorldPartitionConvertCommandlet, Warning, TEXT("Error remapping SoftObjectPath %s"), *OriginalValue);
-					UE_LOG(LogWorldPartitionConvertCommandlet, Warning, TEXT("  Source: %s"), *GetSourceString());
-				}
 			}
 
 			if (!Value.IsNull())
