Fix: Ctrl+S / Save Current Level regression when done on a new OFPA/WP level (should always save all actors)

diff --git a/Engine/Source/Editor/UnrealEd/Private/FileHelpers.cpp b/Engine/Source/Editor/UnrealEd/Private/FileHelpers.cpp
--- a/Engine/Source/Editor/UnrealEd/Private/FileHelpers.cpp
+++ b/Engine/Source/Editor/UnrealEd/Private/FileHelpers.cpp
@@ -3900,8 +3900,8 @@ bool FEditorFileUtils::SaveCurrentLevel()
 
 		if (PackagesToSave.Num())
 		{
-			// If Level gets saved we don't want it to save its external packages because we've already filtered out the ones that need saving and they are part of the PackagesToSave array
-			TGuardValue<bool> GuardValue(bSkipExternalObjectSave, true);
+			// If Level gets saved we don't want it to save its external packages because we've already filtered out the ones that need saving and they are part of the PackagesToSave array (unless level is PKG_NewlyCreated then we should save all actors)
+			TGuardValue<bool> GuardValue(bSkipExternalObjectSave, !LevelPackage->HasAnyPackageFlags(PKG_NewlyCreated));
 			bReturnCode &= InternalSavePackages(PackagesToSave, false, false, false);
 		}
 	}
