World Partition: fix issue during the conversion commandlet where an actor moved from a sublevel containing a level script to the main world would have its soft object paths to actors that remained in the sublevel improperly remapped.

#rnx

diff --git a/Engine/Source/Editor/UnrealEd/Private/Commandlets/WorldPartitionConvertCommandlet.cpp b/Engine/Source/Editor/UnrealEd/Private/Commandlets/WorldPartitionConvertCommandlet.cpp
--- a/Engine/Source/Editor/UnrealEd/Private/Commandlets/WorldPartitionConvertCommandlet.cpp
+++ b/Engine/Source/Editor/UnrealEd/Private/Commandlets/WorldPartitionConvertCommandlet.cpp
@@ -1233,6 +1233,10 @@ int32 UWorldPartitionConvertCommandlet::Main(const FString& Params)
 						{
 							ActorsToConvert.Add(Actor);
 						}
+						else
+						{
+							RemapSoftObjectPaths.Add(FSoftObjectPath(Actor).ToString(), FSoftObjectPath(Actor).ToString());
+						}
 					}
 				}
 			}
