AssetRegistryGenerator: Sort includes
#rnx

diff --git a/Engine/Source/Editor/UnrealEd/Private/Commandlets/AssetRegistryGenerator.cpp b/Engine/Source/Editor/UnrealEd/Private/Commandlets/AssetRegistryGenerator.cpp
--- a/Engine/Source/Editor/UnrealEd/Private/Commandlets/AssetRegistryGenerator.cpp
+++ b/Engine/Source/Editor/UnrealEd/Private/Commandlets/AssetRegistryGenerator.cpp
@@ -1,54 +1,53 @@
 // Copyright Epic Games, Inc. All Rights Reserved.
 
 #include "Commandlets/AssetRegistryGenerator.h"
-#include "HAL/FileManager.h"
-#include "HAL/PlatformFileManager.h"
-#include "Misc/FileHelper.h"
-#include "Serialization/ArrayReader.h"
-#include "Serialization/ArrayWriter.h"
-#include "Misc/App.h"
-#include "Serialization/JsonTypes.h"
-#include "Serialization/JsonReader.h"
-#include "Policies/PrettyJsonPrintPolicy.h"
-#include "Serialization/JsonSerializer.h"
+
+#include "Algo/Sort.h"
+#include "Algo/StableSort.h"
+#include "AssetRegistry/AssetRegistryModule.h"
+#include "AssetRegistry/CookTagList.h"
+#include "CollectionManagerModule.h"
+#include "CollectionManagerTypes.h"
+#include "Commandlets/ChunkDependencyInfo.h"
+#include "Commandlets/IChunkDataGenerator.h"
+#include "Engine/AssetManager.h"
 #include "Engine/Level.h"
 #include "Engine/World.h"
-#include "Settings/ProjectPackagingSettings.h"
-#include "CollectionManagerTypes.h"
+#include "GameDelegates.h"
+#include "HAL/FileManager.h"
+#include "HAL/PlatformFileManager.h"
 #include "ICollectionManager.h"
-#include "CollectionManagerModule.h"
 #include "Interfaces/ITargetPlatform.h"
-#include "AssetRegistry/AssetRegistryModule.h"
-#include "AssetRegistry/CookTagList.h"
-#include "GameDelegates.h"
-#include "Commandlets/IChunkDataGenerator.h"
-#include "Commandlets/ChunkDependencyInfo.h"
 #include "IPlatformFileSandboxWrapper.h"
+#include "Misc/App.h"
 #include "Misc/ConfigCacheIni.h"
-#include "Misc/PathViews.h"
-#include "Stats/StatsMisc.h"
-#include "Templates/UniquePtr.h"
-#include "Engine/AssetManager.h"
 #include "Misc/DataDrivenPlatformInfoRegistry.h"
-#include "Algo/Sort.h"
-#include "Algo/StableSort.h"
-#include "ProfilingDebugging/ScopedTimers.h"
-
+#include "Misc/FileHelper.h"
+#include "Misc/PathViews.h"
 #include "PakFileUtilities.h"
-
+#include "Policies/PrettyJsonPrintPolicy.h"
+#include "ProfilingDebugging/ScopedTimers.h"
+#include "Serialization/ArrayReader.h"
+#include "Serialization/ArrayWriter.h"
+#include "Serialization/JsonReader.h"
+#include "Serialization/JsonReader.h"
+#include "Serialization/JsonSerializer.h"
+#include "Serialization/JsonSerializer.h"
+#include "Serialization/JsonTypes.h"
 #include "Serialization/JsonWriter.h"
-#include "Serialization/JsonReader.h"
-#include "Serialization/JsonSerializer.h"
+#include "Settings/ProjectPackagingSettings.h"
+#include "Stats/StatsMisc.h"
 #include "String/ParseTokens.h"
+#include "Templates/UniquePtr.h"
 
-DEFINE_LOG_CATEGORY_STATIC(LogAssetRegistryGenerator, Log, All);
-
-#define LOCTEXT_NAMESPACE "AssetRegistryGenerator"
-
 #if WITH_EDITOR
 #include "HAL/ThreadHeartBeat.h"
 #endif
 
+DEFINE_LOG_CATEGORY_STATIC(LogAssetRegistryGenerator, Log, All);
+
+#define LOCTEXT_NAMESPACE "AssetRegistryGenerator"
+
 //////////////////////////////////////////////////////////////////////////
 // Static functions
 FName GetPackageNameFromDependencyPackageName(const FName RawPackageFName)
