Multiprocess cook: Copy AssetRegistry data for saved packages from CookWorkers to Director.
#rnx

diff --git a/Engine/Source/Editor/UnrealEd/Private/Commandlets/AssetRegistryGenerator.cpp b/Engine/Source/Editor/UnrealEd/Private/Commandlets/AssetRegistryGenerator.cpp
--- a/Engine/Source/Editor/UnrealEd/Private/Commandlets/AssetRegistryGenerator.cpp
+++ b/Engine/Source/Editor/UnrealEd/Private/Commandlets/AssetRegistryGenerator.cpp
@@ -8,6 +8,8 @@
 #include "AssetRegistry/CookTagList.h"
 #include "CollectionManagerModule.h"
 #include "CollectionManagerTypes.h"
+#include "Cooker/CookPackageData.h"
+#include "Cooker/CookWorkerClient.h"
 #include "Commandlets/ChunkDependencyInfo.h"
 #include "Commandlets/IChunkDataGenerator.h"
 #include "Engine/AssetManager.h"
@@ -1090,6 +1092,22 @@ void FAssetRegistryGenerator::UpdateKeptPackages()
 	}
 }
 
+static void AppendCookTagsToTagMap(TArray<TPair<FName, FString>>&& InTags, FAssetDataTagMap& OutTags)
+{
+	for (TPair<FName, FString>& Tag : InTags)
+	{
+		// Don't add empty tags
+		if (!Tag.Key.IsNone() && !Tag.Value.IsEmpty())
+		{
+			// Prepend Cook_
+			// Do the accumulation in a stack buffer to avoid a bit of work, but we still
+			// eat the FName generation.
+			FName TagName(WriteToString<256>(TEXTVIEW("Cook_"), Tag.Key));
+			OutTags.Add(TagName, MoveTemp(Tag.Value));
+		}
+	}
+}
+
 static void AddCookTagsToState(TMap<FSoftObjectPath, TArray<TPair<FName, FString>>>&& InCookTags, FAssetRegistryState& InState)
 {
 	TRACE_CPUPROFILER_EVENT_SCOPE(AddCookTagToState)
@@ -1103,21 +1121,9 @@ static void AddCookTagsToState(TMap<FSoftObjectPath, TArray<TPair<FName, FString
 		const FAssetData* AssetData = InState.GetAssetByObjectPath(FName(WriteToString<256>(ObjectToTags.Key)));
 
 		// Migrate to FAssetDataTagMap
-		FAssetDataTagMap NewTags;				
-		for (TPair<FName, FString>& Tag : ObjectToTags.Value)
-		{			
-			// Don't add empty tags
-			if (!Tag.Key.IsNone() && !Tag.Value.IsEmpty())
-			{
-				// Prepend Cook_
-				// Do the accumulation in a stack buffer to avoid a bit of work, but we still
-				// eat the FName generation.
-				FName TagName(WriteToString<256>(TEXTVIEW("Cook_"), Tag.Key));
-				NewTags.Add(TagName, MoveTemp(Tag.Value));
-				TagCount++;
-			}
-		}
-
+		FAssetDataTagMap NewTags;
+		AppendCookTagsToTagMap(MoveTemp(ObjectToTags.Value), NewTags);
+		TagCount += NewTags.Num();
 		InState.AddTagsToAssetData(ObjectToTags.Key, MoveTemp(NewTags));
 	} // end for each object
 
@@ -2579,6 +2585,129 @@ void FAssetRegistryGenerator::UpdateAssetRegistryPackageData(const UPackage& Pac
 		TEXT("Trying to update asset package flags in package '%s' that does not exist"), *PackageName.ToString());
 }
 
+void FAssetRegistryGenerator::UpdateAssetRegistryPackageData(UE::Cook::FPackageData& PackageData, UE::Cook::FAssetRegistryPackageMessage&& Message)
+{
+	const FName PackageName = PackageData.GetPackageName();
+	PreviousPackagesToUpdate.Remove(PackageName);
+
+	for (FAssetData& AssetData : Message.AssetDatas)
+	{
+		AssetData.PackageFlags = Message.PackageFlags;
+
+		// MPCOOKTODO: We are setting the AssetData onto this->State rather than into the global assetregistry,
+		// because the AR does not currently support updating an AssetData and it is better for performance to
+		// edit this->State. This will not work as is, because FAssetRegistryGenerator::UpdateAssetDatas later
+		// clobbers this->State with AssetDatas from the global AssetRegistry. Changing that for SingleProcess
+		// director cooks is something we want to do anyways because its confusing as is. Make that change so
+		// we don't clobber this function's results, or change this function to write to the global asset registry.
+		const FAssetData* ExistingAssetData = State.GetAssetByObjectPath(AssetData.ObjectPath);
+		if (!ExistingAssetData)
+		{
+			FAssetData* const NewAssetData = new FAssetData(MoveTemp(AssetData));
+			State.AddAssetData(NewAssetData);
+		}
+		else
+		{
+			State.UpdateAssetData(const_cast<FAssetData*>(ExistingAssetData), MoveTemp(AssetData));
+		}
+	}
+	FAssetPackageData* AssetPackageData = GetAssetPackageData(PackageName);
+	AssetPackageData->DiskSize = Message.DiskSize;
+}
+
+namespace UE::Cook
+{
+
+FAssetRegistryReporterRemote::FAssetRegistryReporterRemote(FCookWorkerClient& InClient, const ITargetPlatform* InTargetPlatform)
+	: Client(InClient)
+	, TargetPlatform(InTargetPlatform)
+{
+}
+
+void FAssetRegistryReporterRemote::UpdateAssetRegistryPackageData(FPackageData& PackageData, const UPackage& Package,
+	FSavePackageResultStruct& SavePackageResult, FCookTagList&& InArchiveCookTagList)
+{
+	uint32 NewPackageFlags = 0;
+	int64 DiskSize = -1;
+	bool bSaveSucceeded = SavePackageResult.IsSuccessful();
+	if (bSaveSucceeded)
+	{
+		NewPackageFlags = SavePackageResult.SerializedPackageFlags;
+		DiskSize = SavePackageResult.TotalFileSize;
+	}
+	else
+	{
+		// Set the package flags to zero to indicate that the package failed to save
+		NewPackageFlags = 0;
+		// Set DiskSize (previous value was disksize in the WorkspaceDomain) to -1 to indicate the cooked file does not exist
+		DiskSize = -1;
+	}
+
+	FAssetRegistryPackageMessage Message;
+	Message.PackageFlags = NewPackageFlags;
+	Message.DiskSize = DiskSize;
+	IAssetRegistry::Get()->GetAssetsByPackageName(PackageData.GetPackageName(), Message.AssetDatas);
+	for (FAssetData& AssetData : Message.AssetDatas)
+	{
+		TArray<FCookTagList::FTagNameValuePair>* ExtraCookTags = InArchiveCookTagList.ObjectToTags.Find(FSoftObjectPath(AssetData.ObjectPath));
+		if (ExtraCookTags)
+		{
+			FAssetDataTagMap NewTags = AssetData.TagsAndValues.CopyMap();
+			AppendCookTagsToTagMap(MoveTemp(*ExtraCookTags), NewTags);
+			AssetData.TagsAndValues = FAssetDataTagMapSharedView(MoveTemp(NewTags));
+		}
+	}
+
+	PackageData.GetOrAddPackageRemoteResult().AddMessage(PackageData, TargetPlatform, Message);
+}
+
+void FAssetRegistryPackageMessage::Write(FCbWriter& Writer, const FPackageData& PackageData, const ITargetPlatform* TargetPlatform) const
+{
+	Writer.BeginArray("A");
+	for (const FAssetData& AssetData : AssetDatas)
+	{
+		check(AssetData.PackageName == PackageData.GetPackageName() && !AssetData.AssetName.IsNone()); // We replicate only regular Assets of the form PackageName.AssetName
+		AssetData.NetworkWrite(Writer, false /* bWritePackageName */);
+
+	}
+	Writer.EndArray();
+	Writer << "F" << PackageFlags;
+	Writer << "S" << DiskSize;
+}
+
+bool FAssetRegistryPackageMessage::TryRead(FCbObject&& Object, FPackageData& PackageData, const ITargetPlatform* TargetPlatform)
+{
+	FCbFieldView AssetDatasField = Object["A"];
+	FCbArrayView AssetDatasArray = AssetDatasField.AsArrayView();
+	if (AssetDatasField.HasError())
+	{
+		return false;
+	}
+	AssetDatas.Reset(AssetDatasArray.Num());
+	for (FCbFieldView ElementField : AssetDatasArray)
+	{
+		FAssetData& AssetData = AssetDatas.Emplace_GetRef();
+		if (!AssetData.TryNetworkRead(ElementField, false /* bReadPackageName */, PackageData.GetPackageName()))
+		{
+			return false;
+		}
+	}
+
+	if (!LoadFromCompactBinary(Object["F"], PackageFlags))
+	{
+		return false;
+	}
+	if (!LoadFromCompactBinary(Object["S"], DiskSize))
+	{
+		return false;
+	}
+	return true;
+}
+
+FGuid FAssetRegistryPackageMessage::MessageType(TEXT("0588DCCEBF1742399EC1E011FC97E4DC"));
+
+}
+
 bool FAssetRegistryGenerator::UpdateAssetPackageFlags(const FName& PackageName, const uint32 PackageFlags)
 {
 	return State.UpdateAssetDataPackageFlags(PackageName, PackageFlags);
diff --git a/Engine/Source/Editor/UnrealEd/Private/CookOnTheFlyServer.cpp b/Engine/Source/Editor/UnrealEd/Private/CookOnTheFlyServer.cpp
--- a/Engine/Source/Editor/UnrealEd/Private/CookOnTheFlyServer.cpp
+++ b/Engine/Source/Editor/UnrealEd/Private/CookOnTheFlyServer.cpp
@@ -2471,10 +2471,8 @@ void UCookOnTheFlyServer::DemoteToIdle(UE::Cook::FPackageData& PackageData, UE::
 
 void UCookOnTheFlyServer::PromoteToSaveComplete(UE::Cook::FPackageData& PackageData, UE::Cook::ESendFlags SendFlags)
 {
-	if (PackageData.IsInProgress())
-	{
-		WorkerRequests->ReportPromoteToSaveComplete(PackageData);
-	}
+	check(PackageData.IsInProgress());
+	WorkerRequests->ReportPromoteToSaveComplete(PackageData);
 	PackageData.SendToState(UE::Cook::EPackageState::Idle, SendFlags);
 }
 
@@ -5478,10 +5476,9 @@ void FSaveCookedPackageContext::FinishPlatform()
 
 	bool bSuccessful = SavePackageResult.IsSuccessful();
 
-	FAssetRegistryGenerator& Generator = *(COTFS.PlatformManager->GetPlatformData(TargetPlatform)->RegistryGenerator);
 	if (bPlatformSetupSuccessful)
 	{
-		FAssetPackageData* AssetPackageData = Generator.GetAssetPackageData(Package->GetFName());
+		TOptional<FAssetPackageData> AssetPackageData = COTFS.AssetRegistry->GetAssetPackageDataCopy(Package->GetFName());
 		check(AssetPackageData);
 
 		// TODO_BuildDefinitionList: Calculate and store BuildDefinitionList on the PackageData, or collect it here from some other source.
@@ -5525,7 +5522,8 @@ void FSaveCookedPackageContext::FinishPlatform()
 	// Update asset registry
 	if (COTFS.IsDirectorCookByTheBook())
 	{
-		Generator.UpdateAssetRegistryPackageData(*Package, SavePackageResult, MoveTemp(*ArchiveCookContext.GetCookTagList()));
+		IAssetRegistryReporter& Reporter = *(COTFS.PlatformManager->GetPlatformData(TargetPlatform)->RegistryReporter);
+		Reporter.UpdateAssetRegistryPackageData(PackageData, *Package, SavePackageResult, MoveTemp(*ArchiveCookContext.GetCookTagList()));
 	}
 
 	// If not retrying, mark the package as cooked, either successfully or with failure
@@ -7140,17 +7138,28 @@ void UCookOnTheFlyServer::RefreshPlatformAssetRegistries(const TArrayView<const
 		FName PlatformName = FName(*TargetPlatform->PlatformName());
 
 		UE::Cook::FPlatformData* PlatformData = PlatformManager->GetPlatformData(TargetPlatform);
-		FAssetRegistryGenerator* RegistryGenerator = PlatformData->RegistryGenerator.Get();
-		if (!RegistryGenerator)
+		UE::Cook::IAssetRegistryReporter* RegistryReporter = PlatformData->RegistryReporter.Get();
+		if (!RegistryReporter)
 		{
-			RegistryGenerator = new FAssetRegistryGenerator(TargetPlatform);
-			PlatformData->RegistryGenerator = TUniquePtr<FAssetRegistryGenerator>(RegistryGenerator);
+			if (!IsCookWorkerMode())
+			{
+				PlatformData->RegistryGenerator = MakeUnique<FAssetRegistryGenerator>(TargetPlatform);
+				PlatformData->RegistryReporter = MakeUnique<UE::Cook::FAssetRegistryReporterLocal>(*PlatformData->RegistryGenerator);
+			}
+			else
+			{
+				PlatformData->RegistryReporter = MakeUnique<UE::Cook::FAssetRegistryReporterRemote>(*CookWorkerClient, TargetPlatform);
+			}
+			RegistryReporter = PlatformData->RegistryReporter.Get();
 		}
 
-		// if we are cooking DLC, we will just spend a lot of time removing the shipped packages from the AR,
-		// so we don't bother copying them over. can easily save 10 seconds on a large project
-		bool bInitalizeFromExisting = !IsCookingDLC();
-		RegistryGenerator->Initialize(CookByTheBookOptions->StartupPackages, bInitalizeFromExisting);
+		if (PlatformData->RegistryGenerator)
+		{
+			// if we are cooking DLC, we will just spend a lot of time removing the shipped packages from the AR,
+			// so we don't bother copying them over. can easily save 10 seconds on a large project
+			bool bInitalizeFromExisting = !IsCookingDLC();
+			PlatformData->RegistryGenerator->Initialize(CookByTheBookOptions->StartupPackages, bInitalizeFromExisting);
+		}
 	}
 }
 
@@ -7859,6 +7868,7 @@ void UCookOnTheFlyServer::BeginCookFinishShaderCodeLibrary(FBeginCookContext& Be
 
 void UCookOnTheFlyServer::RegisterShaderChunkDataGenerator()
 {
+	check(!IsCookWorkerMode());
 	// add shader library and PSO cache chunkers
 	const UProjectPackagingSettings* const PackagingSettings = GetDefault<UProjectPackagingSettings>();
 	FString LibraryName = GetProjectShaderLibraryName();
@@ -8267,6 +8277,7 @@ void UCookOnTheFlyServer::CookByTheBookFinished()
 				// Add the package hashes to the relevant AssetPackageDatas.
 				// PackageHashes are gated by requiring UPackage::WaitForAsyncFileWrites(), which is called above.
 				FCookSavePackageContext& SaveContext = FindOrCreateSaveContext(TargetPlatform);
+				// MPCOOKTODO: Need to replicate AllPackageHashes from CookWorkers
 				TMap<FName, TRefCountPtr<FPackageHashes>>& AllPackageHashes = SaveContext.PackageWriter->GetPackageHashes();
 				for (TPair<FName, TRefCountPtr<FPackageHashes>>& HashSet : AllPackageHashes)
 				{
@@ -8937,6 +8948,7 @@ void UCookOnTheFlyServer::BeginCookSandbox(FBeginCookContext& BeginContext)
 			PackageWriter.Initialize(CookInfo);
 			if (!PlatformContext.bWorkerOnSharedSandbox)
 			{
+				check(!IsCookWorkerMode());
 				// Clean the Manifest directory even on iterative builds; it is written from scratch each time
 				// But only do this if we own the output directory
 				PlatformData->RegistryGenerator->CleanManifestDirectories();
@@ -10761,6 +10773,8 @@ void UCookOnTheFlyServer::GenerateLocalizationReferences(TConstArrayView<FString
 
 void UCookOnTheFlyServer::RegisterLocalizationChunkDataGenerator()
 {
+	check(!IsCookWorkerMode());
+
 	// Get the list of localization targets to chunk, and remove any targets that we've been asked not to stage
 	const UProjectPackagingSettings* const PackagingSettings = GetDefault<UProjectPackagingSettings>();
 	TArray<FString> LocalizationTargetsToChunk = PackagingSettings->LocalizationTargetsToChunk;
diff --git a/Engine/Source/Editor/UnrealEd/Private/Cooker/CompactBinaryTCP.h b/Engine/Source/Editor/UnrealEd/Private/Cooker/CompactBinaryTCP.h
--- a/Engine/Source/Editor/UnrealEd/Private/Cooker/CompactBinaryTCP.h
+++ b/Engine/Source/Editor/UnrealEd/Private/Cooker/CompactBinaryTCP.h
@@ -44,7 +44,7 @@ public:
 	/** Marshall the message to a CompactBinaryObject. */
 	virtual void Write(FCbWriter& Writer) const = 0;
 	/** Unmarshall the message from a CompactBinaryObject. */
-	virtual bool TryRead(FCbObjectView Object) = 0;
+	virtual bool TryRead(FCbObject&& Object) = 0;
 	/** Return the Guid that identifies the message to the remote connection. */
 	virtual FGuid GetMessageType() const = 0;
 };
diff --git a/Engine/Source/Editor/UnrealEd/Private/Cooker/CookDirector.cpp b/Engine/Source/Editor/UnrealEd/Private/Cooker/CookDirector.cpp
--- a/Engine/Source/Editor/UnrealEd/Private/Cooker/CookDirector.cpp
+++ b/Engine/Source/Editor/UnrealEd/Private/Cooker/CookDirector.cpp
@@ -279,7 +279,7 @@ void FWorkerConnectMessage::Write(FCbWriter& Writer) const
 	Writer << "RemoteIndex" << RemoteIndex;
 }
 
-bool FWorkerConnectMessage::TryRead(FCbObjectView Object)
+bool FWorkerConnectMessage::TryRead(FCbObject&& Object)
 {
 	RemoteIndex = Object["RemoteIndex"].AsInt32(-1);
 	return RemoteIndex >= 0;
@@ -390,7 +390,7 @@ void FCookDirector::TickWorkerConnects()
 			continue;
 		}
 		FWorkerConnectMessage Message;
-		if (!Message.TryRead(Messages[0].Object))
+		if (!Message.TryRead(MoveTemp(Messages[0].Object)))
 		{
 			UE_LOG(LogCook, Warning, TEXT("Pending connection sent an invalid Connection Message. Connection will be ignored."));
 			continue;
diff --git a/Engine/Source/Editor/UnrealEd/Private/Cooker/CookDirector.h b/Engine/Source/Editor/UnrealEd/Private/Cooker/CookDirector.h
--- a/Engine/Source/Editor/UnrealEd/Private/Cooker/CookDirector.h
+++ b/Engine/Source/Editor/UnrealEd/Private/Cooker/CookDirector.h
@@ -13,7 +13,7 @@
 #include "Misc/Guid.h"
 #include "Templates/UniquePtr.h"
 
-class FCbObjectView;
+class FCbObject;
 class FCbWriter;
 class UCookOnTheFlyServer;
 namespace UE::Cook { class FCookWorkerServer; }
@@ -126,7 +126,7 @@ struct FWorkerConnectMessage : public UE::CompactBinaryTCP::IMessage
 {
 public:
 	virtual void Write(FCbWriter& Writer) const override;
-	virtual bool TryRead(FCbObjectView Object) override;
+	virtual bool TryRead(FCbObject&& Object) override;
 	virtual FGuid GetMessageType() const override { return MessageType; }
 
 public:
diff --git a/Engine/Source/Editor/UnrealEd/Private/Cooker/CookPackageData.cpp b/Engine/Source/Editor/UnrealEd/Private/Cooker/CookPackageData.cpp
--- a/Engine/Source/Editor/UnrealEd/Private/Cooker/CookPackageData.cpp
+++ b/Engine/Source/Editor/UnrealEd/Private/Cooker/CookPackageData.cpp
@@ -12,6 +12,7 @@
 #include "CompactBinaryTCP.h"
 #include "Cooker/CookPlatformManager.h"
 #include "Cooker/CookRequestCluster.h"
+#include "Cooker/CookWorkerClient.h"
 #include "CookOnTheSide/CookOnTheFlyServer.h"
 #include "Containers/StringView.h"
 #include "EditorDomain/EditorDomain.h"
@@ -757,6 +758,7 @@ void FPackageData::OnEnterSave()
 	check(!HasPrepareSaveFailed());
 	CheckObjectCacheEmpty();
 	CheckCookedPlatformDataEmpty();
+	check(!PackageRemoteResult);
 }
 
 void FPackageData::OnExitSave()
@@ -765,6 +767,22 @@ void FPackageData::OnExitSave()
 	ClearObjectCache();
 	SetHasPrepareSaveFailed(false);
 	SetIsPrepareSaveRequiresGC(false);
+	PackageRemoteResult.Reset();
+}
+
+FPackageRemoteResult& FPackageData::GetOrAddPackageRemoteResult()
+{
+	check(GetState() == EPackageState::Save);
+	if (!PackageRemoteResult)
+	{
+		PackageRemoteResult = MakeUnique<FPackageRemoteResult>();
+	}
+	return *PackageRemoteResult;
+}
+
+TUniquePtr<FPackageRemoteResult>& FPackageData::GetPackageRemoteResult()
+{
+	return PackageRemoteResult;
 }
 
 void FPackageData::OnEnterInProgress()
diff --git a/Engine/Source/Editor/UnrealEd/Private/Cooker/CookPackageData.h b/Engine/Source/Editor/UnrealEd/Private/Cooker/CookPackageData.h
--- a/Engine/Source/Editor/UnrealEd/Private/Cooker/CookPackageData.h
+++ b/Engine/Source/Editor/UnrealEd/Private/Cooker/CookPackageData.h
@@ -17,6 +17,7 @@
 #include "HAL/Platform.h"
 #include "Misc/EnumClassFlags.h"
 #include "Misc/Optional.h"
+#include "PackageResultsMessage.h"
 #include "Templates/SharedPointer.h"
 #include "UObject/GCObject.h"
 #include "UObject/NameTypes.h"
@@ -488,6 +489,8 @@ public:
 	/** Marshall this PackageData to a ConstructData that can be used later or on a remote machine to reconstruct it. */
 	FConstructPackageData CreateConstructData();
 
+	FPackageRemoteResult& GetOrAddPackageRemoteResult();
+	TUniquePtr<FPackageRemoteResult>& GetPackageRemoteResult();
 private:
 	friend struct UE::Cook::FPackageDatas;
 
@@ -558,6 +561,7 @@ private:
 	void OnPackageDataFirstRequested(FInstigator&& InInstigator);
 
 	TUniquePtr<FGeneratorPackage> GeneratorPackage;
+	TUniquePtr<FPackageRemoteResult> PackageRemoteResult;
 	FGeneratorPackage* GeneratedOwner;
 	/** Data for each platform that has been interacted with by *this. */
 	TSortedMap<const ITargetPlatform*, FPlatformData> PlatformDatas;
diff --git a/Engine/Source/Editor/UnrealEd/Private/Cooker/CookPlatformManager.h b/Engine/Source/Editor/UnrealEd/Private/Cooker/CookPlatformManager.h
--- a/Engine/Source/Editor/UnrealEd/Private/Cooker/CookPlatformManager.h
+++ b/Engine/Source/Editor/UnrealEd/Private/Cooker/CookPlatformManager.h
@@ -15,6 +15,7 @@
 
 class ITargetPlatform;
 class UCookOnTheFlyServer;
+namespace UE::Cook { class IAssetRegistryReporter; }
 
 namespace UE
 {
@@ -43,6 +44,7 @@ namespace Cook
 
 		/* Pointer to the platform-specific RegistryGenerator for this platform.  If already constructed we can take a faster refresh path on future sessions. */
 		TUniquePtr<FAssetRegistryGenerator> RegistryGenerator;
+		TUniquePtr<UE::Cook::IAssetRegistryReporter> RegistryReporter;
 
 		/* Whether BeginCookSandbox has been called for this platform.  If we have already initialized the sandbox we can take a faster refresh path on future sessions. */
 		bool bIsSandboxInitialized = false;
diff --git a/Engine/Source/Editor/UnrealEd/Private/Cooker/CookWorkerClient.cpp b/Engine/Source/Editor/UnrealEd/Private/Cooker/CookWorkerClient.cpp
--- a/Engine/Source/Editor/UnrealEd/Private/Cooker/CookWorkerClient.cpp
+++ b/Engine/Source/Editor/UnrealEd/Private/Cooker/CookWorkerClient.cpp
@@ -10,6 +10,7 @@
 #include "CookWorkerServer.h"
 #include "CoreGlobals.h"
 #include "HAL/PlatformTime.h"
+#include "PackageResultsMessage.h"
 #include "Sockets.h"
 #include "SocketSubsystem.h"
 #include "WorkerRequestsRemote.h"
@@ -117,22 +118,66 @@ void FCookWorkerClient::DoneWithInitialSettings()
 
 void FCookWorkerClient::ReportDemoteToIdle(FPackageData& PackageData, ESuppressCookReason Reason)
 {
-	FPackageResult& Result = PendingResults.Emplace_GetRef();
+	FPackageRemoteResult& Result = PendingResults.Emplace_GetRef();
 	Result.PackageName = PackageData.GetPackageName();
 	Result.SuppressCookReason = Reason;
 }
 
 void FCookWorkerClient::ReportPromoteToSaveComplete(FPackageData& PackageData)
 {
-	FPackageResult& Result = PendingResults.Emplace_GetRef();
-	Result.PackageName = PackageData.GetPackageName();
-	Result.SuppressCookReason = ESuppressCookReason::InvalidSuppressCookReason;
-	for (ITargetPlatform* TargetPlatform : OrderedSessionPlatforms)
+	TUniquePtr<FPackageRemoteResult> MovedResult = MoveTemp(PackageData.GetPackageRemoteResult());
+	FPackageRemoteResult* Result;
+	if (MovedResult)
+	{
+		Result = &PendingResults.Emplace_GetRef(MoveTemp(*MovedResult));
+	} 
+	else
+	{
+		Result = &PendingResults.Emplace_GetRef();
+	}
+
+	Result->PackageName = PackageData.GetPackageName();
+	Result->SuppressCookReason = ESuppressCookReason::InvalidSuppressCookReason;
+
+	// Sort the platforms to match the OrderedSessionPlatforms order, and add any missing platforms
+	int32 NumPlatforms = OrderedSessionPlatforms.Num();
+	bool bAlreadySorted = true;
+	for (int32 PlatformIndex = 0; PlatformIndex < NumPlatforms; ++PlatformIndex)
+	{
+		if (OrderedSessionPlatforms[PlatformIndex] != Result->Platforms[PlatformIndex].Platform)
+		{
+			bAlreadySorted = false;
+			break;
+		}
+	}
+	if (!bAlreadySorted)
+	{
+		TArray<FPackageRemoteResult::FPlatformResult, TInlineAllocator<1>> SortedPlatforms;
+		SortedPlatforms.SetNum(NumPlatforms);
+		for (FPackageRemoteResult::FPlatformResult& ExistingResult : Result->Platforms)
+		{
+			int32 PlatformIndex = OrderedSessionPlatforms.IndexOfByKey(ExistingResult.Platform);
+			check(PlatformIndex != INDEX_NONE && SortedPlatforms[PlatformIndex].Platform == nullptr); // Only platforms in the session platforms should have been added, and there should not be duplicates
+			SortedPlatforms[PlatformIndex] = MoveTemp(ExistingResult);
+		}
+		for (int32 PlatformIndex = 0; PlatformIndex < NumPlatforms; ++PlatformIndex)
+		{
+			FPackageRemoteResult::FPlatformResult& Sorted = SortedPlatforms[PlatformIndex];
+			if (!Sorted.Platform)
+			{
+				Sorted.Platform = OrderedSessionPlatforms[PlatformIndex];
+				Sorted.bSuccessful = false;
+			}
+		}
+		Swap(Result->Platforms, SortedPlatforms);
+	}
+
+	for (int32 PlatformIndex = 0; PlatformIndex < NumPlatforms; ++PlatformIndex)
 	{
-		FPackagePlatformResult& PlatformResult = Result.Platforms.Emplace_GetRef();
+		ITargetPlatform* TargetPlatform = OrderedSessionPlatforms[PlatformIndex];
+		FPackageRemoteResult::FPlatformResult& PlatformResults = Result->Platforms[PlatformIndex];
 		FPackageData::FPlatformData& PackagePlatformData = PackageData.FindOrAddPlatformData(TargetPlatform);
-		PlatformResult.bSuccessful = PackagePlatformData.bCookSucceeded;
-		// MPCOOKTODO: Accumulate results for saved information so we can add the other information about the Platform
+		PlatformResults.bSuccessful = PackagePlatformData.bCookSucceeded;
 	}
 }
 
@@ -291,7 +336,7 @@ void FCookWorkerClient::PollReceiveConfigMessage()
 	}
 	check(!InitialConfigMessage);
 	InitialConfigMessage = MakeUnique<FInitialConfigMessage>();
-	if (!InitialConfigMessage->TryRead(Messages[0].Object))
+	if (!InitialConfigMessage->TryRead(MoveTemp(Messages[0].Object)))
 	{
 		UE_LOG(LogCook, Warning, TEXT("CookWorker initialization failure: Director sent an invalid InitialConfigMessage."));
 		SendToState(EConnectStatus::LostConnection);
@@ -354,7 +399,7 @@ void FCookWorkerClient::HandleReceiveMessages(TArray<UE::CompactBinaryTCP::FMars
 		if (Message.MessageType == FAbortWorkerMessage::MessageType)
 		{
 			FAbortWorkerMessage AbortMessage;
-			AbortMessage.TryRead(Message.Object);
+			AbortMessage.TryRead(MoveTemp(Message.Object));
 			if (AbortMessage.Type == FAbortWorkerMessage::EType::CookComplete)
 			{
 				UE_LOG(LogCook, Display, TEXT("CookWorkerClient received CookComplete message from Director. Flushing messages and shutting down."));
@@ -374,7 +419,7 @@ void FCookWorkerClient::HandleReceiveMessages(TArray<UE::CompactBinaryTCP::FMars
 		else if (Message.MessageType == FAssignPackagesMessage::MessageType)
 		{
 			FAssignPackagesMessage AssignPackagesMessage;
-			if (!AssignPackagesMessage.TryRead(Message.Object))
+			if (!AssignPackagesMessage.TryRead(MoveTemp(Message.Object)))
 			{
 				LogInvalidMessage(TEXT("FAssignPackagesMessage"));
 			}
diff --git a/Engine/Source/Editor/UnrealEd/Private/Cooker/CookWorkerClient.h b/Engine/Source/Editor/UnrealEd/Private/Cooker/CookWorkerClient.h
--- a/Engine/Source/Editor/UnrealEd/Private/Cooker/CookWorkerClient.h
+++ b/Engine/Source/Editor/UnrealEd/Private/Cooker/CookWorkerClient.h
@@ -11,7 +11,7 @@
 namespace UE::Cook { struct FAssignPackagesMessage; }
 namespace UE::Cook { struct FDirectorConnectionInfo; }
 namespace UE::Cook { struct FInitialConfigMessage; }
-namespace UE::Cook { struct FPackageResult; }
+namespace UE::Cook { struct FPackageRemoteResult; }
 
 namespace UE::Cook
 {
@@ -97,7 +97,7 @@ private:
 	TSharedPtr<FInternetAddr> DirectorAddr;
 	TUniquePtr<FInitialConfigMessage> InitialConfigMessage;
 	TArray<ITargetPlatform*> OrderedSessionPlatforms;
-	TArray<FPackageResult> PendingResults;
+	TArray<FPackageRemoteResult> PendingResults;
 	UE::CompactBinaryTCP::FSendBuffer SendBuffer;
 	UE::CompactBinaryTCP::FReceiveBuffer ReceiveBuffer;
 	FString DirectorURI;
diff --git a/Engine/Source/Editor/UnrealEd/Private/Cooker/CookWorkerServer.cpp b/Engine/Source/Editor/UnrealEd/Private/Cooker/CookWorkerServer.cpp
--- a/Engine/Source/Editor/UnrealEd/Private/Cooker/CookWorkerServer.cpp
+++ b/Engine/Source/Editor/UnrealEd/Private/Cooker/CookWorkerServer.cpp
@@ -2,6 +2,7 @@
 
 #include "CookWorkerServer.h"
 
+#include "Commandlets/AssetRegistryGenerator.h"
 #include "CompactBinaryTCP.h"
 #include "CookDirector.h"
 #include "CookPackageData.h"
@@ -11,6 +12,7 @@
 #include "Interfaces/ITargetPlatformManagerModule.h"
 #include "Math/NumericLimits.h"
 #include "Misc/AssertionMacros.h"
+#include "PackageResultsMessage.h"
 #include "PackageTracker.h"
 #include "UnrealEdMisc.h"
 
@@ -416,7 +418,7 @@ void FCookWorkerServer::HandleReceiveMessages(TArray<UE::CompactBinaryTCP::FMars
 		else if (Message.MessageType == FPackageResultsMessage::MessageType)
 		{
 			FPackageResultsMessage ResultsMessage;
-			if (!ResultsMessage.TryRead(Message.Object))
+			if (!ResultsMessage.TryRead(MoveTemp(Message.Object)))
 			{
 				LogInvalidMessage(TEXT("FPackageResultsMessage"));
 			}
@@ -428,6 +430,27 @@ void FCookWorkerServer::HandleReceiveMessages(TArray<UE::CompactBinaryTCP::FMars
 	}
 }
 
+void FCookWorkerServer::HandleReceivedPackagePlatformMessages(FPackageData& PackageData, const ITargetPlatform* TargetPlatform, TArray<UE::CompactBinaryTCP::FMarshalledMessage>&& Messages)
+{
+	for (UE::CompactBinaryTCP::FMarshalledMessage& Message : Messages)
+	{
+		if (Message.MessageType == FAssetRegistryPackageMessage::MessageType)
+		{
+			FAssetRegistryPackageMessage ARMessage;
+			if (!ARMessage.TryRead(MoveTemp(Message.Object), PackageData, TargetPlatform))
+			{
+				LogInvalidMessage(TEXT("FAssetRegistryPackageMessage"));
+			}
+			else
+			{
+				FAssetRegistryGenerator* RegistryGenerator = COTFS.PlatformManager->GetPlatformData(TargetPlatform)->RegistryGenerator.Get();
+				check(RegistryGenerator); // The TargetPlatform came from OrderedSessionPlatforms, and the RegistryGenerator should exist for any of those platforms
+				RegistryGenerator->UpdateAssetRegistryPackageData(PackageData, MoveTemp(ARMessage));
+			}
+		}
+	}
+}
+
 void FCookWorkerServer::SendMessage(const UE::CompactBinaryTCP::IMessage& Message)
 {
 	UE::CompactBinaryTCP::TryWritePacket(Socket, SendBuffer, Message);
@@ -435,7 +458,7 @@ void FCookWorkerServer::SendMessage(const UE::CompactBinaryTCP::IMessage& Messag
 
 void FCookWorkerServer::RecordResults(FPackageResultsMessage& Message)
 {
-	for (FPackageResult& Result : Message.Results)
+	for (FPackageRemoteResult& Result : Message.Results)
 	{
 		FPackageData* PackageData = COTFS.PackageDatas->FindPackageDataByPackageName(Result.PackageName);
 		if (!PackageData)
@@ -466,11 +489,10 @@ void FCookWorkerServer::RecordResults(FPackageResultsMessage& Message)
 			for (int32 PlatformIndex = 0; PlatformIndex < NumPlatforms; ++PlatformIndex)
 			{
 				ITargetPlatform* TargetPlatform = OrderedSessionPlatforms[PlatformIndex];
-				FPackagePlatformResult& PlatformResult = Result.Platforms[PlatformIndex];
+				FPackageRemoteResult::FPlatformResult& PlatformResult = Result.Platforms[PlatformIndex];
 				PackageData->SetPlatformCooked(TargetPlatform, PlatformResult.bSuccessful);
 				// MPCOOKTODO: Call CommitRemotePackage on the PackageWriter
-				// MPCOOKTODO: Add the AssetRegistry information here
-				// MPCOOKTODO: Copy the Hash
+				HandleReceivedPackagePlatformMessages(*PackageData, TargetPlatform, MoveTemp(PlatformResult.Messages));
 			}
 			if (Result.bReferencedOnlyByEditorOnlyData)
 			{
@@ -501,7 +523,7 @@ void FAssignPackagesMessage::Write(FCbWriter& Writer) const
 	Writer << "P" << PackageDatas;
 }
 
-bool FAssignPackagesMessage::TryRead(FCbObjectView Object)
+bool FAssignPackagesMessage::TryRead(FCbObject&& Object)
 {
 	return LoadFromCompactBinary(Object["P"], PackageDatas);
 }
@@ -515,12 +537,12 @@ FAbortPackagesMessage::FAbortPackagesMessage(TArray<FName>&& InPackageNames)
 
 void FAbortPackagesMessage::Write(FCbWriter& Writer) const
 {
-	WriteArrayOfNames(Writer, "PackageNames", PackageNames);
+	Writer << "PackageNames" <<  PackageNames;
 }
 
-bool FAbortPackagesMessage::TryRead(FCbObjectView Object)
+bool FAbortPackagesMessage::TryRead(FCbObject&& Object)
 {
-	return TryReadArrayOfNames(Object, "PackageNames", PackageNames);
+	return LoadFromCompactBinary(Object["PackageNames"], PackageNames);
 }
 
 FGuid FAbortPackagesMessage::MessageType(TEXT("D769F1BFF2F34978868D70E3CAEE94E7"));
@@ -535,7 +557,7 @@ void FAbortWorkerMessage::Write(FCbWriter& Writer) const
 	Writer << "Type" << (uint8)Type;
 }
 
-bool FAbortWorkerMessage::TryRead(FCbObjectView Object)
+bool FAbortWorkerMessage::TryRead(FCbObject&& Object)
 {
 	Type = static_cast<EType>(Object["Type"].AsUInt8((uint8)EType::Abort));
 	return true;
@@ -576,7 +598,7 @@ void FInitialConfigMessage::Write(FCbWriter& Writer) const
 	Writer << "CookOnTheFlyOptions" << CookOnTheFlyOptions;
 }
 
-bool FInitialConfigMessage::TryRead(FCbObjectView Object)
+bool FInitialConfigMessage::TryRead(FCbObject&& Object)
 {
 	bool bOk = true;
 	int32 LocalCookMode;
@@ -626,103 +648,4 @@ bool FInitialConfigMessage::TryRead(FCbObjectView Object)
 
 FGuid FInitialConfigMessage::MessageType(TEXT("340CDCB927304CEB9C0A66B5F707FC2B"));
 
-void FPackageResultsMessage::Write(FCbWriter& Writer) const
-{
-	Writer.BeginArray("R");
-	for (const FPackageResult& Result : Results)
-	{
-		Writer.BeginObject();
-		{
-			Writer << "N" << Result.PackageName;
-			Writer << "S" << (uint8) Result.SuppressCookReason;
-			Writer << "E" << Result.bReferencedOnlyByEditorOnlyData;
-			Writer.BeginArray("P");
-			for (const FPackagePlatformResult& PlatformResult : Result.Platforms)
-			{
-				Writer.BeginObject();
-				Writer << "S" << PlatformResult.bSuccessful;
-				Writer << "G" << PlatformResult.PackageGuid;
-				Writer << "D" << PlatformResult.TargetDomainDependencies;
-				Writer.EndObject();
-			}
-			Writer.EndArray();
-		}
-		Writer.EndObject();
-	}
-	Writer.EndArray();
-}
-
-bool FPackageResultsMessage::TryRead(FCbObjectView Object)
-{
-	Results.Reset();
-	for (FCbFieldView ResultField : Object["R"])
-	{
-		FCbObjectView ResultObject = ResultField.AsObjectView();
-		FPackageResult& Result = Results.Emplace_GetRef();
-		LoadFromCompactBinary(ResultObject["N"], Result.PackageName);
-		if (Result.PackageName.IsNone())
-		{
-			return false;
-		}
-		int32 LocalSuppressCookReason = ResultObject["S"].AsUInt8(MAX_uint8);
-		if (LocalSuppressCookReason == MAX_uint8)
-		{
-			return false;
-		}
-		Result.SuppressCookReason = static_cast<ESuppressCookReason>(LocalSuppressCookReason);
-		Result.bReferencedOnlyByEditorOnlyData = ResultObject["E"].AsBool();
-		Result.Platforms.Reset();
-		for (FCbFieldView PlatformField : ResultObject["P"])
-		{
-			FPackagePlatformResult& PlatformResult = Result.Platforms.Emplace_GetRef();
-			FCbObjectView PlatformObject = PlatformField.AsObjectView();
-
-			PlatformResult.bSuccessful = PlatformObject["S"].AsBool();
-			PlatformResult.PackageGuid = PlatformObject["G"].AsUuid();
-
-			FCbObjectView TargetDepsView = PlatformObject["D"].AsObjectView();
-			if (TargetDepsView.GetSize() > 0)
-			{
-				FUniqueBuffer TargetDepsCopy = FUniqueBuffer::Alloc(TargetDepsView.GetSize());
-				TargetDepsView.CopyTo(TargetDepsCopy.GetView());
-				PlatformResult.TargetDomainDependencies = FCbObject(TargetDepsCopy.MoveToShared());
-			}
-		}
-	}
-	return true;
-}
-
-FGuid FPackageResultsMessage::MessageType(TEXT("4631C6C0F6DC4CEFB2B09D3FB0B524DB"));
-
-void WriteArrayOfNames(FCbWriter& Writer, const char* ArrayName, TConstArrayView<FName> Names)
-{
-	Writer.BeginArray(ArrayName);
-	for (FName Name : Names)
-	{
-		Writer.AddString(WriteToString<FName::StringBufferSize>(Name).ToView());
-	}
-	Writer.EndArray();
-}
-
-bool TryReadArrayOfNames(FCbObjectView Object, const char* ArrayName, TArray<FName>& OutNames)
-{
-	FCbFieldView ArrayField = Object["ArrayName"];
-	FCbArrayView ArrayView = ArrayField.AsArrayView();
-	if (ArrayField.HasError())
-	{
-		return false;
-	}
-	OutNames.Reserve(OutNames.Num() + ArrayView.Num());
-	for (FCbFieldView ElementView : ArrayView)
-	{
-		FUtf8StringView StringView = ElementView.AsString();
-		if (ElementView.HasError())
-		{
-			return false;
-		}
-		OutNames.Add(FName(StringView));
-	}
-	return true;
-}
-
 }
\ No newline at end of file
diff --git a/Engine/Source/Editor/UnrealEd/Private/Cooker/CookWorkerServer.h b/Engine/Source/Editor/UnrealEd/Private/Cooker/CookWorkerServer.h
--- a/Engine/Source/Editor/UnrealEd/Private/Cooker/CookWorkerServer.h
+++ b/Engine/Source/Editor/UnrealEd/Private/Cooker/CookWorkerServer.h
@@ -96,6 +96,7 @@ private:
 	void ShutdownRemoteProcess();
 	/** Helper for PumpReceiveMessages: dispatch the messages received from the socket. */
 	void HandleReceiveMessages(TArray<UE::CompactBinaryTCP::FMarshalledMessage>&& Messages);
+	void HandleReceivedPackagePlatformMessages(FPackageData& PackageData, const ITargetPlatform* TargetPlatform, TArray<UE::CompactBinaryTCP::FMarshalledMessage>&& Messages);
 	/** Add results from the client to the local CookOnTheFlyServer. */
 	void RecordResults(FPackageResultsMessage& Message);
 	void LogInvalidMessage(const TCHAR* MessageTypeName);
@@ -125,7 +126,7 @@ public:
 	FAssignPackagesMessage(TArray<FConstructPackageData>&& InPackageDatas);
 
 	virtual void Write(FCbWriter& Writer) const override;
-	virtual bool TryRead(FCbObjectView Object) override;
+	virtual bool TryRead(FCbObject&& Object) override;
 	virtual FGuid GetMessageType() const override { return MessageType; }
 
 public:
@@ -141,7 +142,7 @@ public:
 	FAbortPackagesMessage(TArray<FName>&& InPackageNames);
 
 	virtual void Write(FCbWriter& Writer) const override;
-	virtual bool TryRead(FCbObjectView Object) override;
+	virtual bool TryRead(FCbObject&& Object) override;
 	virtual FGuid GetMessageType() const override { return MessageType; }
 
 public:
@@ -164,7 +165,7 @@ public:
 	};
 	FAbortWorkerMessage(EType InType = EType::Abort);
 	virtual void Write(FCbWriter& Writer) const override;
-	virtual bool TryRead(FCbObjectView Object) override;
+	virtual bool TryRead(FCbObject&& Object) override;
 	virtual FGuid GetMessageType() const override { return MessageType; }
 
 public:
@@ -177,7 +178,7 @@ struct FInitialConfigMessage : public UE::CompactBinaryTCP::IMessage
 {
 public:
 	virtual void Write(FCbWriter& Writer) const override;
-	virtual bool TryRead(FCbObjectView Object) override;
+	virtual bool TryRead(FCbObject&& Object) override;
 	virtual FGuid GetMessageType() const override { return MessageType; }
 
 	void ReadFromLocal(const UCookOnTheFlyServer& COTFS, const TArray<ITargetPlatform*>& InOrderedSessionPlatforms,
@@ -205,42 +206,4 @@ private:
 	bool bZenStore = false;
 };
 
-/** Helper struct for FPackageResultsMessage: holds per-package-per-platform data. */
-struct FPackagePlatformResult
-{
-	FGuid PackageGuid;
-	FCbObject TargetDomainDependencies;
-	bool bSuccessful = true;
-};
-
-/** Helper struct for FPackageResultsMessage: holds per-package data. */
-struct FPackageResult
-{
-	FName PackageName;
-	/** If failure reason is InvalidSuppressCookReason, it was saved. Otherwise, holds the suppression reason */
-	ESuppressCookReason SuppressCookReason;
-	bool bReferencedOnlyByEditorOnlyData = false;
-	// The elements of Platforms correspond to the TargetPlatforms present in OrderedSessionPlatforms
-	TArray<FPackagePlatformResult, TInlineAllocator<1>> Platforms;
-};
-
-/** Message from Client to Server giving the results for saved or refused-to-cook packages. */
-struct FPackageResultsMessage : public UE::CompactBinaryTCP::IMessage
-{
-public:
-	virtual void Write(FCbWriter& Writer) const override;
-	virtual bool TryRead(FCbObjectView Object) override;
-	virtual FGuid GetMessageType() const override { return MessageType; }
-
-public:
-	TArray<FPackageResult> Results;
-
-	static FGuid MessageType;
-};
-
-/** IMessage helper: write as strings to compact binary. */
-void WriteArrayOfNames(FCbWriter& Writer, const char* ArrayName, TConstArrayView<FName> Names);
-/** IMessage helper: read FNames as strings from compact binary. */
-bool TryReadArrayOfNames(FCbObjectView Object, const char* ArrayName, TArray<FName>& OutNames);
-
 }
\ No newline at end of file
diff --git a/Engine/Source/Editor/UnrealEd/Private/Cooker/PackageResultsMessage.cpp b/Engine/Source/Editor/UnrealEd/Private/Cooker/PackageResultsMessage.cpp
new file mode 100644
--- /dev/null
+++ b/Engine/Source/Editor/UnrealEd/Private/Cooker/PackageResultsMessage.cpp
@@ -0,0 +1,128 @@
+// Copyright Epic Games, Inc. All Rights Reserved.
+
+#include "PackageResultsMessage.h"
+
+#include "CookPackageData.h"
+
+namespace UE::Cook
+{
+
+void FPackageRemoteResult::AddMessage(const FPackageData& PackageData, const ITargetPlatform* TargetPlatform, const IPackageMessage& Message)
+{
+	int32 PlatformIndex = Platforms.IndexOfByPredicate([TargetPlatform](const FPlatformResult& Result)
+		{ return Result.Platform == TargetPlatform; }
+	);
+	FPlatformResult* Result;
+	if (PlatformIndex != INDEX_NONE)
+	{
+		Result = &Platforms[PlatformIndex];
+	}
+	else
+	{
+		Result = &Platforms.Emplace_GetRef();
+		Result->Platform = TargetPlatform;
+	}
+
+	FCbWriter Writer;
+	Writer.BeginObject();
+	Message.Write(Writer, PackageData, TargetPlatform);
+	Writer.EndObject();
+
+	Result->Messages.Add(UE::CompactBinaryTCP::FMarshalledMessage{ Message.GetMessageType(), Writer.Save().AsObject() });
+}
+
+
+void FPackageResultsMessage::Write(FCbWriter& Writer) const
+{
+	Writer.BeginArray("R");
+	for (const FPackageRemoteResult& Result : Results)
+	{
+		Writer.BeginObject();
+		{
+			Writer << "N" << Result.PackageName;
+			Writer << "R" << (uint8) Result.SuppressCookReason;
+			Writer << "E" << Result.bReferencedOnlyByEditorOnlyData;
+			Writer.BeginArray("P");
+			for (const FPackageRemoteResult::FPlatformResult& PlatformResult : Result.Platforms)
+			{
+				Writer.BeginObject();
+				Writer << "S" << PlatformResult.bSuccessful;
+				Writer << "G" << PlatformResult.PackageGuid;
+				Writer << "D" << PlatformResult.TargetDomainDependencies;
+				if (!PlatformResult.Messages.IsEmpty())
+				{
+					// We write a nonhomogenous array of length 2N. 2N+0 is the Message type, 2N+1 is the message object.
+					Writer.BeginArray("M");
+					for (const UE::CompactBinaryTCP::FMarshalledMessage& Message : PlatformResult.Messages)
+					{
+						Writer << Message.MessageType;
+						Writer << Message.Object;
+					}
+					Writer.EndArray();
+				}
+				Writer.EndObject();
+			}
+			Writer.EndArray();
+		}
+		Writer.EndObject();
+	}
+	Writer.EndArray();
+}
+
+bool FPackageResultsMessage::TryRead(FCbObject&& Object)
+{
+	Results.Reset();
+	for (FCbField ResultField : Object["R"])
+	{
+		FCbObject ResultObject = ResultField.AsObject();
+		FPackageRemoteResult& Result = Results.Emplace_GetRef();
+		LoadFromCompactBinary(ResultObject["N"], Result.PackageName);
+		if (Result.PackageName.IsNone())
+		{
+			return false;
+		}
+		int32 LocalSuppressCookReason = ResultObject["R"].AsUInt8(MAX_uint8);
+		if (LocalSuppressCookReason == MAX_uint8)
+		{
+			return false;
+		}
+		Result.SuppressCookReason = static_cast<ESuppressCookReason>(LocalSuppressCookReason);
+		Result.bReferencedOnlyByEditorOnlyData = ResultObject["E"].AsBool();
+		Result.Platforms.Reset();
+		for (FCbField PlatformField : ResultObject["P"])
+		{
+			FPackageRemoteResult::FPlatformResult& PlatformResult = Result.Platforms.Emplace_GetRef();
+			FCbObject PlatformObject = PlatformField.AsObject();
+
+			PlatformResult.bSuccessful = PlatformObject["S"].AsBool();
+			PlatformResult.PackageGuid = PlatformObject["G"].AsUuid();
+			PlatformResult.TargetDomainDependencies = PlatformObject["D"].AsObject();
+
+			// We read a nonhomogenous array of length 2N. 2N+0 is the Message type, 2N+1 is the message object.
+			FCbArray MessagesArray = PlatformObject["M"].AsArray();
+			PlatformResult.Messages.Reserve(MessagesArray.Num() / 2);
+			FCbFieldIterator MessageField = MessagesArray.CreateIterator();
+			while (MessageField)
+			{
+				UE::CompactBinaryTCP::FMarshalledMessage& Message = PlatformResult.Messages.Emplace_GetRef();
+				Message.MessageType = MessageField->AsUuid();
+				if (MessageField.HasError())
+				{
+					return false;
+				}
+				++MessageField;
+				Message.Object = MessageField->AsObject();
+				if (MessageField.HasError())
+				{
+					return false;
+				}
+				++MessageField;
+			}
+		}
+	}
+	return true;
+}
+
+FGuid FPackageResultsMessage::MessageType(TEXT("4631C6C0F6DC4CEFB2B09D3FB0B524DB"));
+
+}
\ No newline at end of file
diff --git a/Engine/Source/Editor/UnrealEd/Private/Cooker/PackageResultsMessage.h b/Engine/Source/Editor/UnrealEd/Private/Cooker/PackageResultsMessage.h
new file mode 100644
--- /dev/null
+++ b/Engine/Source/Editor/UnrealEd/Private/Cooker/PackageResultsMessage.h
@@ -0,0 +1,77 @@
+// Copyright Epic Games, Inc. All Rights Reserved.
+
+#pragma once
+
+#include "CompactBinaryTCP.h"
+#include "Containers/Array.h"
+#include "CookTypes.h"
+#include "Templates/UniquePtr.h"
+#include "UObject/NameTypes.h"
+
+class ITargetPlatform;
+namespace UE::Cook { struct FPackageData; }
+
+namespace UE::Cook
+{
+
+/**
+ * The base class of messages that can be sent as platform-specific package-specific submessages
+ * in a FPackageResultMessage. Messages are identified to the remote connection by the Guid identifier
+ * from GetMessageType.
+ */
+class IPackageMessage
+{
+public:
+	virtual ~IPackageMessage() {}
+	/** Marshall the message to a CompactBinaryObject. */
+	virtual void Write(FCbWriter& Writer, const FPackageData& PackageData, const ITargetPlatform* TargetPlatform) const = 0;
+	/** Unmarshall the message from a CompactBinaryObject. */
+	virtual bool TryRead(FCbObject&& Object, FPackageData& PackageData, const ITargetPlatform* TargetPlatform) = 0;
+	/** Return the Guid that identifies the message to the remote connection. */
+	virtual FGuid GetMessageType() const = 0;
+};
+
+/**
+ * Helper struct for FPackageResultsMessage that can also be owned by other types like PackageData.
+ * Holds replication information about the result of the Package's save on each platform, including
+ * system-specific messages from other systems
+ */
+struct FPackageRemoteResult
+{
+public:
+	/** Information about the results for a single platform */
+	struct FPlatformResult
+	{
+		const ITargetPlatform* Platform = nullptr;
+		TArray<UE::CompactBinaryTCP::FMarshalledMessage> Messages;
+		FGuid PackageGuid;
+		FCbObject TargetDomainDependencies;
+		bool bSuccessful = true;
+	};
+
+	void AddMessage(const FPackageData& PackageData, const ITargetPlatform* TargetPlatform, const IPackageMessage& Message);
+
+public:
+	FName PackageName;
+	/** If failure reason is InvalidSuppressCookReason, it was saved. Otherwise, holds the suppression reason */
+	ESuppressCookReason SuppressCookReason;
+	bool bReferencedOnlyByEditorOnlyData = false;
+	TArray<FPlatformResult, TInlineAllocator<1>> Platforms;
+
+};
+
+/** Message from Client to Server giving the results for saved or refused-to-cook packages. */
+struct FPackageResultsMessage : public UE::CompactBinaryTCP::IMessage
+{
+public:
+	virtual void Write(FCbWriter& Writer) const override;
+	virtual bool TryRead(FCbObject&& Object) override;
+	virtual FGuid GetMessageType() const override { return MessageType; }
+
+public:
+	TArray<FPackageRemoteResult> Results;
+
+	static FGuid MessageType;
+};
+
+}
diff --git a/Engine/Source/Editor/UnrealEd/Public/Commandlets/AssetRegistryGenerator.h b/Engine/Source/Editor/UnrealEd/Public/Commandlets/AssetRegistryGenerator.h
--- a/Engine/Source/Editor/UnrealEd/Public/Commandlets/AssetRegistryGenerator.h
+++ b/Engine/Source/Editor/UnrealEd/Public/Commandlets/AssetRegistryGenerator.h
@@ -6,6 +6,7 @@
 
 #include "AssetRegistry/AssetData.h"
 #include "AssetRegistry/AssetRegistryState.h"
+#include "Cooker/PackageResultsMessage.h"
 #include "Misc/AssetRegistryInterface.h"
 #include "Misc/Paths.h"
 #include "Templates/UniquePtr.h"
@@ -19,6 +20,9 @@ class IChunkDataGenerator;
 class UChunkDependencyInfo;
 struct FChunkDependencyTreeNode;
 struct FCookTagList;
+namespace UE::Cook { class FAssetRegistryPackageMessage; }
+namespace UE::Cook { class FCookWorkerClient; }
+namespace UE::Cook { struct FPackageData; }
 
 /**
  * Helper class for generating streaming install manifests
@@ -200,6 +204,7 @@ public:
 	 * @param SavePackageResult The metadata to associate with the given package name
 	 */
 	void UpdateAssetRegistryPackageData(const UPackage& Package, FSavePackageResultStruct& SavePackageResult, FCookTagList&& InArchiveCookTagList);
+	void UpdateAssetRegistryPackageData(UE::Cook::FPackageData& PackageData, UE::Cook::FAssetRegistryPackageMessage&& Message);
 
 	/**
 	 * Check config to see whether chunk assignments use the AssetManager. If so, run the once-per-process construction
@@ -420,3 +425,66 @@ private:
 	 */
 	const FAssetData* CreateOrFindAssetData(UObject& Object);
 };
+
+namespace UE::Cook
+{
+
+class IAssetRegistryReporter
+{
+public:
+	virtual ~IAssetRegistryReporter() {}
+
+	virtual void UpdateAssetRegistryPackageData(FPackageData& PackageData, const UPackage& Package,
+		FSavePackageResultStruct& SavePackageResult, FCookTagList&& InArchiveCookTagList) = 0;
+
+};
+
+class FAssetRegistryReporterLocal : public IAssetRegistryReporter
+{
+public:
+	FAssetRegistryReporterLocal(FAssetRegistryGenerator& InGenerator)
+		: Generator(InGenerator)
+	{
+	}
+
+	virtual void UpdateAssetRegistryPackageData(FPackageData& PackageData, const UPackage& Package,
+		FSavePackageResultStruct& SavePackageResult, FCookTagList&& InArchiveCookTagList) override
+	{
+		Generator.UpdateAssetRegistryPackageData(Package, SavePackageResult, MoveTemp(InArchiveCookTagList));
+	}
+
+private:
+	FAssetRegistryGenerator& Generator;
+};
+
+class FAssetRegistryReporterRemote : public IAssetRegistryReporter
+{
+public:
+	FAssetRegistryReporterRemote(FCookWorkerClient& InClient, const ITargetPlatform* InTargetPlatform);
+
+	virtual void UpdateAssetRegistryPackageData(FPackageData& PackageData, const UPackage& Package,
+		FSavePackageResultStruct& SavePackageResult, FCookTagList&& InArchiveCookTagList) override;
+
+private:
+	FCookWorkerClient& Client;
+	const ITargetPlatform* TargetPlatform = nullptr;
+};
+
+class FAssetRegistryPackageMessage : public IPackageMessage
+{
+public:
+	virtual void Write(FCbWriter& Writer, const FPackageData& PackageData, const ITargetPlatform* TargetPlatform) const override;
+	virtual bool TryRead(FCbObject&& Object, FPackageData& PackageData, const ITargetPlatform* TargetPlatform) override;
+	virtual FGuid GetMessageType() const override { return MessageType; }
+
+	TArray<FAssetData> AssetDatas;
+	uint32 PackageFlags = 0;
+	int64 DiskSize = -1;
+
+public:
+	static FGuid MessageType;
+};
+
+
+
+}
\ No newline at end of file
diff --git a/Engine/Source/Runtime/CoreUObject/Private/AssetRegistry/AssetData.cpp b/Engine/Source/Runtime/CoreUObject/Private/AssetRegistry/AssetData.cpp
--- a/Engine/Source/Runtime/CoreUObject/Private/AssetRegistry/AssetData.cpp
+++ b/Engine/Source/Runtime/CoreUObject/Private/AssetRegistry/AssetData.cpp
@@ -7,7 +7,11 @@
 #include "AssetRegistry/ARFilter.h"
 #include "Containers/Set.h"
 #include "HAL/CriticalSection.h"
+#include "Misc/PathViews.h"
 #include "Misc/ScopeRWLock.h"
+#include "Serialization/CompactBinary.h"
+#include "Serialization/CompactBinarySerialization.h"
+#include "Serialization/CompactBinaryWriter.h"
 #include "Serialization/CustomVersion.h"
 #include "String/Find.h"
 #include "UObject/PropertyPortFlags.h"
@@ -21,10 +25,10 @@ UE_IMPLEMENT_STRUCT("/Script/CoreUObject", AssetData);
 const FGuid FAssetRegistryVersion::GUID(0x717F9EE7, 0xE9B0493A, 0x88B39132, 0x1B388107);
 FCustomVersionRegistration GRegisterAssetRegistryVersion(FAssetRegistryVersion::GUID, FAssetRegistryVersion::LatestVersion, TEXT("AssetRegistry"));
 
-namespace UE { namespace AssetData { namespace Private {
-
 const FName GAssetBundleDataName("AssetBundleData");
 
+namespace UE { namespace AssetData { namespace Private {
+
 static TSharedPtr<FAssetBundleData, ESPMode::ThreadSafe> ParseAssetBundles(const TCHAR* Text, const FAssetData& Context)
 {
 	// Register that the SoftObjectPaths we read in the FAssetBundleEntry::BundleAssets are non-package data and don't need to be tracked
@@ -240,6 +244,111 @@ PRAGMA_ENABLE_DEPRECATION_WARNINGS
 	Ar << PackageFlags;
 }
 
+void FAssetData::NetworkWrite(FCbWriter& Writer, bool bWritePackageName) const
+{
+	// Note we use single-letter field names to reduce network bandwidth
+	Writer.BeginObject();
+	if (bWritePackageName)
+	{
+		Writer << "O" << ObjectPath;
+		Writer << "P" << PackagePath;
+		Writer << "Q" << PackageName;
+	}
+	Writer << "N" << AssetName;
+	Writer << "C" << AssetClassPath.ToString();
+	if (TagsAndValues.Num() != 0 || TaggedAssetBundles)
+	{
+		Writer.BeginArray("T");
+		TagsAndValues.ForEach([&Writer](const TPair<FName, FAssetTagValueRef>& Pair)
+			{
+				Writer.BeginObject();
+				Writer << "K" << Pair.Key;
+				Writer << "V" << Pair.Value.GetStorageString();
+				Writer.EndObject();
+			});
+		if (TaggedAssetBundles)
+		{
+			FString ValueText;
+			TaggedAssetBundles->ExportTextItem(ValueText, FAssetBundleData(), nullptr, PPF_None, nullptr);
+
+			Writer.BeginObject();
+			Writer << "K" << GAssetBundleDataName;
+			Writer << "V" << ValueText;
+			Writer.EndObject();
+		}
+		Writer.EndArray();
+	}
+	if (!ChunkIDs.IsEmpty())
+	{
+		Writer << "I" << ChunkIDs;
+	}
+	Writer.EndObject();
+}
+
+bool FAssetData::TryNetworkRead(FCbFieldView Field, bool bReadPackageName, FName InPackageName)
+{
+	bool bOk = true;
+	FCbObjectView Object = Field.AsObjectView();
+	bOk &= !Field.HasError();
+
+	bool bHasAssetName = LoadFromCompactBinary(Object["N"], AssetName);
+	bOk &= bHasAssetName;
+	if (bReadPackageName)
+	{
+		bOk = LoadFromCompactBinary(Object["O"], ObjectPath) & bOk;
+		bOk = LoadFromCompactBinary(Object["P"], PackagePath) & bOk;
+		bOk = LoadFromCompactBinary(Object["Q"], PackageName) & bOk;
+	}
+	else
+	{
+		if (bHasAssetName)
+		{
+			WriteToString<256> PackageNameStr(InPackageName);
+			ObjectPath = FName(WriteToString<256>(PackageNameStr.ToView(), TEXT("."), AssetName));
+			PackagePath = FName(FPathViews::GetPath(PackageNameStr.ToView()));
+		}
+		else
+		{
+			ObjectPath = NAME_None;
+			PackagePath = NAME_None;
+		}
+		PackageName = InPackageName;
+	}
+	FString ClassPath;
+	if (LoadFromCompactBinary(Object["C"], ClassPath))
+	{
+		bOk = AssetClassPath.TrySetPath(ClassPath) & bOk;
+	}
+	else
+	{
+		AssetClassPath.Reset();
+		bOk = false;
+	}
+
+	FCbFieldView TagsField = Object["T"];
+	FCbArrayView TagsArray = TagsField.AsArrayView();
+	if (!TagsField.HasError()) // Ok if it does not exist
+	{
+		FAssetDataTagMap Tags;
+		Tags.Reserve(TagsArray.Num());
+		for (FCbFieldView TagField : TagsArray)
+		{
+			FName TagName;
+			bOk = LoadFromCompactBinary(TagField["K"], TagName) & bOk;
+			FString& TagValue = Tags.FindOrAdd(TagName);
+			bOk = LoadFromCompactBinary(TagField["V"], TagValue) & bOk;
+		}
+		SetTagsAndAssetBundles(MoveTemp(Tags));
+	}
+	else
+	{
+		SetTagsAndAssetBundles(FAssetDataTagMap());
+	}
+
+	LoadFromCompactBinary(Object["I"], ChunkIDs); // Ok if it does not exist
+	return bOk;
+}
+
 void FAssetData::SerializeForCacheWithTagsAndBundles(FArchive& Ar, void (*SerializeTagsAndBundles)(FArchive&, FAssetData&))
 {
 	SerializeForCacheInternal(Ar, FAssetRegistryVersion::LatestVersion, SerializeTagsAndBundles);
diff --git a/Engine/Source/Runtime/CoreUObject/Private/UObject/Obj.cpp b/Engine/Source/Runtime/CoreUObject/Private/UObject/Obj.cpp
--- a/Engine/Source/Runtime/CoreUObject/Private/UObject/Obj.cpp
+++ b/Engine/Source/Runtime/CoreUObject/Private/UObject/Obj.cpp
@@ -1871,8 +1871,6 @@ FString GetConfigFilename( UObject* SourceObject )
 
 namespace UE { namespace Object { namespace Private {
 
-const FName GAssetBundleDataName("AssetBundleData");
-
 // Thread local state to avoid UObject::GetAssetRegistryTags() API change
 thread_local FAssetBundleData const** TGetAssetRegistryTags_OutBundles = nullptr;
 
diff --git a/Engine/Source/Runtime/CoreUObject/Public/AssetRegistry/AssetData.h b/Engine/Source/Runtime/CoreUObject/Public/AssetRegistry/AssetData.h
--- a/Engine/Source/Runtime/CoreUObject/Public/AssetRegistry/AssetData.h
+++ b/Engine/Source/Runtime/CoreUObject/Public/AssetRegistry/AssetData.h
@@ -25,8 +25,11 @@
 #include "UObject/TopLevelAssetPath.h"
 
 struct FCustomVersion;
+class FCbWriter;
+class FCbFieldView;
 
 COREUOBJECT_API DECLARE_LOG_CATEGORY_EXTERN(LogAssetData, Log, All);
+COREUOBJECT_API extern const FName GAssetBundleDataName;
 
 /** Version used for serializing asset registry caches, both runtime and editor */
 struct COREUOBJECT_API FAssetRegistryVersion
@@ -540,6 +543,9 @@ PRAGMA_ENABLE_DEPRECATION_WARNINGS
 			static_cast<Archive&>(Ar).SerializeTagsAndBundles(Ad);
 			});
 	}
+
+	COREUOBJECT_API void NetworkWrite(FCbWriter& Writer, bool bWritePackageName) const;
+	COREUOBJECT_API bool TryNetworkRead(FCbFieldView Field, bool bReadPackageName, FName InPackageName);
 private:
 	/**
 	 * The actual implementation of SerializeForCache.
diff --git a/Engine/Source/Runtime/CoreUObject/Public/AssetRegistry/AssetDataTagMap.h b/Engine/Source/Runtime/CoreUObject/Public/AssetRegistry/AssetDataTagMap.h
--- a/Engine/Source/Runtime/CoreUObject/Public/AssetRegistry/AssetDataTagMap.h
+++ b/Engine/Source/Runtime/CoreUObject/Public/AssetRegistry/AssetDataTagMap.h
@@ -247,6 +247,8 @@ public:
 	bool						TryGetAsText(FText& Out) const; // @return false if value isn't a localized string
 
 	FString						GetValue() const { return AsString(); }
+	/** Coerce the type to a Complex String capable of representing the type */
+	FString						GetStorageString() const { return ToLoose(); }
 
 	bool						Equals(FStringView Str) const;
 
