[WP Navmesh] Fix WP navmesh not generating caused by wrong settings

diff --git a/Engine/Source/Editor/UnrealEd/Private/WorldPartition/WorldPartitionNavigationDataBuilder.cpp b/Engine/Source/Editor/UnrealEd/Private/WorldPartition/WorldPartitionNavigationDataBuilder.cpp
--- a/Engine/Source/Editor/UnrealEd/Private/WorldPartition/WorldPartitionNavigationDataBuilder.cpp
+++ b/Engine/Source/Editor/UnrealEd/Private/WorldPartition/WorldPartitionNavigationDataBuilder.cpp
@@ -51,7 +51,7 @@ bool UWorldPartitionNavigationDataBuilder::PreRun(UWorld* World, FPackageSourceC
 
 	// Size of loaded cell. Set as big as your hardware can afford.
 	const uint32 LoadingCellSizeSetting = World->GetWorldSettings()->NavigationDataBuilderLoadingCellSize;
-	IterativeCellSize = (LoadingCellSizeSetting / GridSize) * GridSize;
+	IterativeCellSize = GridSize * FMath::Max(LoadingCellSizeSetting / GridSize, 1u);
 	
 	// Extra padding around loaded cell.
 	// @todo: set value programatically.
diff --git a/Engine/Source/Runtime/Engine/Classes/GameFramework/WorldSettings.h b/Engine/Source/Runtime/Engine/Classes/GameFramework/WorldSettings.h
--- a/Engine/Source/Runtime/Engine/Classes/GameFramework/WorldSettings.h
+++ b/Engine/Source/Runtime/Engine/Classes/GameFramework/WorldSettings.h
@@ -575,8 +575,9 @@ public:
 	/**
 	 * Loading cell size used when building navigation data iteratively.
 	 * The actual cell size used will be rounded using the NavigationDataChunkGridSize.
+	 * It's recommended to use a value as high as the hardware memory allows to load. 
 	 */
-	UPROPERTY(EditAnywhere, Category = Navigation)
+	UPROPERTY(EditAnywhere, Category = Navigation, meta = (ClampMin = "5000"))
 	uint32 NavigationDataBuilderLoadingCellSize;
 	
 	/** Default size of the grid for placed elements from the editor */
