Cooker with WorldPartition: Fix assertion failure due to null PackageData when cooking generated packages.
#rnx

diff --git a/Engine/Source/Editor/UnrealEd/Private/CookOnTheFlyServer.cpp b/Engine/Source/Editor/UnrealEd/Private/CookOnTheFlyServer.cpp
--- a/Engine/Source/Editor/UnrealEd/Private/CookOnTheFlyServer.cpp
+++ b/Engine/Source/Editor/UnrealEd/Private/CookOnTheFlyServer.cpp
@@ -5478,8 +5478,8 @@ void FSaveCookedPackageContext::FinishPlatform()
 
 	if (bPlatformSetupSuccessful)
 	{
-		TOptional<FAssetPackageData> AssetPackageData = COTFS.AssetRegistry->GetAssetPackageDataCopy(Package->GetFName());
-		check(AssetPackageData);
+		FAssetRegistryGenerator& Generator = *(COTFS.PlatformManager->GetPlatformData(TargetPlatform)->RegistryGenerator);
+		FAssetPackageData* AssetPackageData = Generator.GetAssetPackageData(Package->GetFName());
 
 		// TODO_BuildDefinitionList: Calculate and store BuildDefinitionList on the PackageData, or collect it here from some other source.
 		TArray<UE::DerivedData::FBuildDefinition> BuildDefinitions;
