Fix SaveCurrentLevel no longer saving level without modifications

#rnx

diff --git a/Engine/Source/Editor/UnrealEd/Private/FileHelpers.cpp b/Engine/Source/Editor/UnrealEd/Private/FileHelpers.cpp
--- a/Engine/Source/Editor/UnrealEd/Private/FileHelpers.cpp
+++ b/Engine/Source/Editor/UnrealEd/Private/FileHelpers.cpp
@@ -3633,14 +3633,14 @@ static bool InternalSavePackagesFast(const TArray<UPackage*>& PackagesToSave, bo
  * @param	bPromptUserToSave			true if we should prompt the user to save dirty packages we found. false to assume all dirty packages should be saved.  Regardless of this setting the user will be prompted for checkout(if needed) unless bFastSave is set
  * @param	bFastSave					true if we should do a fast save. (I.E don't prompt the user to save, don't prompt for checkout, and only save packages that are currently writable).  Note: Still prompts for SaveAs if a package needs a filename
  * @param	bCanBeDeclined				true if the user prompt should contain a "Don't Save" button in addition to "Cancel", which won't result in a failure return code.
+ * @param	bCheckDirty					true if only dirty packages should be saved
  */
-static bool InternalSavePackages(const TArray<UPackage*>& PackagesToSave, bool bPromptUserToSave, bool bFastSave, bool bCanBeDeclined)
+static bool InternalSavePackages(const TArray<UPackage*>& PackagesToSave, bool bPromptUserToSave, bool bFastSave, bool bCanBeDeclined, bool bCheckDirty)
 {
 	bool bReturnCode = true;
 
 	if (!bFastSave)
 	{
-		const bool bCheckDirty = true;
 		const bool bAlreadyCheckedOut = false;
 		const FEditorFileUtils::EPromptReturnCode Return = FEditorFileUtils::PromptForCheckoutAndSave(PackagesToSave, bCheckDirty, bPromptUserToSave, nullptr, bAlreadyCheckedOut, bCanBeDeclined);
 		if (Return == FEditorFileUtils::EPromptReturnCode::PR_Cancelled)
@@ -3784,7 +3784,8 @@ bool FEditorFileUtils::SaveDirtyPackages(const bool bPromptUserToSave, const boo
 			*bOutPackagesNeededSaving = true;
 		}
 
-		bReturnCode = InternalSavePackages(PackagesToSave, bPromptUserToSave, bFastSave, bCanBeDeclined);
+		const bool bCheckDirty = true;
+		bReturnCode = InternalSavePackages(PackagesToSave, bPromptUserToSave, bFastSave, bCanBeDeclined, bCheckDirty);
 	}
 	else if (bNotifyNoPackagesSaved)
 	{
@@ -3851,7 +3852,8 @@ bool FEditorFileUtils::SaveDirtyContentPackages(TArray<UClass*>& SaveContentClas
 	bool bResult = false;
 	if (PackagesToSave.Num() > 0)
 	{
-		bResult = InternalSavePackages(PackagesToSave, bPromptUserToSave, bFastSave, bCanBeDeclined);
+		const bool bCheckDirty = true;
+		bResult = InternalSavePackages(PackagesToSave, bPromptUserToSave, bFastSave, bCanBeDeclined, bCheckDirty);
 	}
 	else if (bNotifyNoPackagesSaved)
 	{
@@ -3875,7 +3877,7 @@ bool FEditorFileUtils::SaveCurrentLevel()
 	if (Level)
 	{
 		// Check dirtiness if the level is using external objects, no need to save it needlessly
-		const bool bCheckDirty = Level->IsUsingExternalObjects();
+		bool bCheckDirty = Level->IsUsingExternalObjects();
 
 		TArray<UPackage*> PackagesToSave;
 		
@@ -3903,7 +3905,11 @@ bool FEditorFileUtils::SaveCurrentLevel()
 		{
 			// If Level gets saved we don't want it to save its external packages because we've already filtered out the ones that need saving and they are part of the PackagesToSave array (unless level is PKG_NewlyCreated then we should save all actors)
 			TGuardValue<bool> GuardValue(bSkipExternalObjectSave, !LevelPackage->HasAnyPackageFlags(PKG_NewlyCreated));
-			bReturnCode &= InternalSavePackages(PackagesToSave, false, false, false);
+			const bool bPromptUserToSave = false;
+			const bool bFastSave = false;
+			const bool bCanBeDeclined = false;
+			bCheckDirty = false; // force the flag back to false because we already checked conditions to add to PackagesToSave. Some Packages like newly created packages might not be dirty and we still want to save them.
+			bReturnCode &= InternalSavePackages(PackagesToSave, bPromptUserToSave, bFastSave, bCanBeDeclined, bCheckDirty);
 		}
 	}
 	return bReturnCode;
@@ -4848,7 +4854,8 @@ static bool InternalCheckoutAndSavePackages(const TArray<UPackage*>& PackagesToS
 			const bool bPromptUserToSave = true;
 			const bool bFastSave = false;
 			const bool bCanBeDeclined = true;
-			bResult = InternalSavePackages(PackagesToSave, bPromptUserToSave, bFastSave, bCanBeDeclined);
+			const bool bCheckDirty = true;
+			bResult = InternalSavePackages(PackagesToSave, bPromptUserToSave, bFastSave, bCanBeDeclined, bCheckDirty);
 		}
 		else
 		{
