Cooker Multiprocess with WorldPartition: On a CookWorker with no generator, assign an empty Guid for generated packages which have no PackageData in the global AssetRegistry; this is the same end result that Singleprocess cooks use.
#rnx

diff --git a/Engine/Source/Editor/UnrealEd/Private/CookOnTheFlyServer.cpp b/Engine/Source/Editor/UnrealEd/Private/CookOnTheFlyServer.cpp
--- a/Engine/Source/Editor/UnrealEd/Private/CookOnTheFlyServer.cpp
+++ b/Engine/Source/Editor/UnrealEd/Private/CookOnTheFlyServer.cpp
@@ -5478,8 +5478,7 @@ void FSaveCookedPackageContext::FinishPlatform()
 
 	if (bPlatformSetupSuccessful)
 	{
-		FAssetRegistryGenerator& Generator = *(COTFS.PlatformManager->GetPlatformData(TargetPlatform)->RegistryGenerator);
-		FAssetPackageData* AssetPackageData = Generator.GetAssetPackageData(Package->GetFName());
+		TOptional<FAssetPackageData> AssetPackageData = COTFS.AssetRegistry->GetAssetPackageDataCopy(Package->GetFName());
 
 		// TODO_BuildDefinitionList: Calculate and store BuildDefinitionList on the PackageData, or collect it here from some other source.
 		TArray<UE::DerivedData::FBuildDefinition> BuildDefinitions;
@@ -5505,7 +5504,7 @@ void FSaveCookedPackageContext::FinishPlatform()
 		}
 		Info.PackageName = Package->GetFName();
 		PRAGMA_DISABLE_DEPRECATION_WARNINGS
-		Info.PackageGuid = AssetPackageData->PackageGuid;
+		Info.PackageGuid = AssetPackageData ? AssetPackageData->PackageGuid : FGuid();
 		PRAGMA_ENABLE_DEPRECATION_WARNINGS
 		Info.Attachments.Add({ "Dependencies", TargetDomainDependencies });
 		// TODO: Reenable BuildDefinitionList once FCbPackage support for empty FCbObjects is in
