Check UActorEditorContextSubsystem when unbinding in SActorEditorContext destructor due to race condition where the subsystem may not exist by the time the destructor is called

diff --git a/Engine/Source/Editor/UnrealEd/Private/SActorEditorContext.cpp b/Engine/Source/Editor/UnrealEd/Private/SActorEditorContext.cpp
--- a/Engine/Source/Editor/UnrealEd/Private/SActorEditorContext.cpp
+++ b/Engine/Source/Editor/UnrealEd/Private/SActorEditorContext.cpp
@@ -26,7 +26,10 @@ SActorEditorContext::~SActorEditorContext()
 {
 	GEditor->GetEditorWorldContext().RemoveRef(World);
 	FEditorDelegates::MapChange.RemoveAll(this);
-	UActorEditorContextSubsystem::Get()->OnActorEditorContextSubsystemChanged().RemoveAll(this);
+	if (UActorEditorContextSubsystem* ActorEditorContextSubsystem = UActorEditorContextSubsystem::Get())
+	{
+		ActorEditorContextSubsystem->OnActorEditorContextSubsystemChanged().RemoveAll(this);
+	}
 }
 
 void SActorEditorContext::Rebuild()
