Prevent call to UObject once the UObject system are un-initialized.
#rnx

diff --git a/Engine/Source/Editor/UnrealEd/Private/SActorEditorContext.cpp b/Engine/Source/Editor/UnrealEd/Private/SActorEditorContext.cpp
--- a/Engine/Source/Editor/UnrealEd/Private/SActorEditorContext.cpp
+++ b/Engine/Source/Editor/UnrealEd/Private/SActorEditorContext.cpp
@@ -24,11 +24,14 @@ void SActorEditorContext::Construct(const FArguments& InArgs)
 
 SActorEditorContext::~SActorEditorContext()
 {
-	GEditor->GetEditorWorldContext().RemoveRef(World);
 	FEditorDelegates::MapChange.RemoveAll(this);
-	if (UActorEditorContextSubsystem* ActorEditorContextSubsystem = UActorEditorContextSubsystem::Get())
+	if (UObjectInitialized())
 	{
-		ActorEditorContextSubsystem->OnActorEditorContextSubsystemChanged().RemoveAll(this);
+		GEditor->GetEditorWorldContext().RemoveRef(World);
+		if (UActorEditorContextSubsystem* ActorEditorContextSubsystem = UActorEditorContextSubsystem::Get())
+		{
+			ActorEditorContextSubsystem->OnActorEditorContextSubsystemChanged().RemoveAll(this);
+		}
 	}
 }
 
