Mark TargetSkeleton and SourceAnimation as BlueprintReadWrite in UAnimMontageFactory for python scripting

diff --git a/Engine/Source/Editor/UnrealEd/Classes/Factories/AnimMontageFactory.h b/Engine/Source/Editor/UnrealEd/Classes/Factories/AnimMontageFactory.h
--- a/Engine/Source/Editor/UnrealEd/Classes/Factories/AnimMontageFactory.h
+++ b/Engine/Source/Editor/UnrealEd/Classes/Factories/AnimMontageFactory.h
@@ -20,11 +20,11 @@ class UAnimMontageFactory : public UFactory
 {
 	GENERATED_UCLASS_BODY()
 
-	UPROPERTY()
+	UPROPERTY(BlueprintReadWrite)
 	TObjectPtr<class USkeleton> TargetSkeleton;
 
 	/* Used when creating a montage from an AnimSequence, becomes the only AnimSequence contained */
-	UPROPERTY()
+	UPROPERTY(BlueprintReadWrite)
 	TObjectPtr<class UAnimSequence> SourceAnimation;
 
 	/** The preview mesh to use with this animation */
