Remove Matinee (1/2)

diff --git a/Engine/Source/Editor/UnrealEd/Private/Fbx/FbxAnimationExport.cpp b/Engine/Source/Editor/UnrealEd/Private/Fbx/FbxAnimationExport.cpp
--- a/Engine/Source/Editor/UnrealEd/Private/Fbx/FbxAnimationExport.cpp
+++ b/Engine/Source/Editor/UnrealEd/Private/Fbx/FbxAnimationExport.cpp
@@ -10,11 +10,8 @@
 #include "Misc/ScopedSlowTask.h"
 #include "Animation/AnimTypes.h"
 #include "Components/SkeletalMeshComponent.h"
-#include "Matinee/InterpData.h"
-#include "Matinee/InterpTrackAnimControl.h"
 #include "Animation/AnimSequence.h"
 #include "Editor/EditorPerProjectUserSettings.h"
-#include "Matinee/MatineeActor.h"
 #include "Animation/SkeletalMeshActor.h"
 #include "FbxExporter.h"
 #include "Exporters/FbxExportOption.h"
@@ -552,98 +549,6 @@ FbxNode* FFbxExporter::ExportAnimSequence( const UAnimSequence* AnimSeq, const U
 }
 
 
-void FFbxExporter::ExportAnimSequencesAsSingle( USkeletalMesh* SkelMesh, const ASkeletalMeshActor* SkelMeshActor, const FString& ExportName, const TArray<UAnimSequence*>& AnimSeqList, const TArray<struct FAnimControlTrackKey>& TrackKeys )
-{
-	if (Scene == NULL || SkelMesh == NULL || AnimSeqList.Num() == 0 || AnimSeqList.Num() != TrackKeys.Num()) return;
-
-	FbxNode* BaseNode = FbxNode::Create(Scene, Converter.ConvertToFbxString(ExportName));
-	Scene->GetRootNode()->AddChild(BaseNode);
-
-	if( SkelMeshActor )
-	{
-		// Set the default position of the actor on the transforms
-		// The Unreal transformation is different from FBX's Z-up: invert the Y-axis for translations and the Y/Z angle values in rotations.
-		BaseNode->LclTranslation.Set(Converter.ConvertToFbxPos(SkelMeshActor->GetActorLocation()));
-		BaseNode->LclRotation.Set(Converter.ConvertToFbxRot(SkelMeshActor->GetActorRotation().Euler()));
-		BaseNode->LclScaling.Set(Converter.ConvertToFbxScale(SkelMeshActor->GetRootComponent()->GetRelativeScale3D()));
-
-	}
-
-	// Create the Skeleton
-	TArray<FbxNode*> BoneNodes;
-	FbxNode* SkeletonRootNode = CreateSkeleton(SkelMesh, BoneNodes);
-	BaseNode->AddChild(SkeletonRootNode);
-
-	bool bAnyObjectMissingSourceData = false;
-	float ExportStartTime = 0.f;
-	for(int32 AnimSeqIndex = 0; AnimSeqIndex < AnimSeqList.Num(); ++AnimSeqIndex)
-	{
-		const UAnimSequence* AnimSeq = AnimSeqList[AnimSeqIndex];
-		const FAnimControlTrackKey& TrackKey = TrackKeys[AnimSeqIndex];
-
-		// Shift the anim sequences so the first one is at time zero in the FBX file
-		const float CurrentStartTime = TrackKey.StartTime - ExportStartTime;
-
-		ExportAnimSequenceToFbx(AnimSeq,
-			SkelMesh,
-			BoneNodes,
-			AnimLayer,
-			TrackKey.AnimStartOffset,
-			TrackKey.AnimEndOffset,
-			TrackKey.AnimPlayRate,
-			CurrentStartTime);
-	}
-
-	CorrectAnimTrackInterpolation(BoneNodes, AnimLayer);
-
-	if (bAnyObjectMissingSourceData)
-	{
-		FMessageDialog::Open( EAppMsgType::Ok, NSLOCTEXT("UnrealEd", "Exporter_Error_SourceDataUnavailable", "No source data available for some objects.  See the log for details.") );
-	}
-
-}
-
-
-
-/**
- * Exports all the animation sequences part of a single Group in a Matinee sequence
- * as a single animation in the FBX document.  The animation is created by sampling the
- * sequence at DEFAULT_SAMPLERATE updates/second and extracting the resulting bone transforms from the given
- * skeletal mesh
- */
-void FFbxExporter::ExportMatineeGroup(class AMatineeActor* MatineeActor, USkeletalMeshComponent* SkeletalMeshComponent)
-{
-	if (Scene == NULL || MatineeActor == NULL || SkeletalMeshComponent == NULL || MatineeActor->MatineeData->InterpLength == 0)
-	{
-		return;
-	}
-
-	FbxString NodeName("MatineeSequence");
-
-	FbxNode* BaseNode = FbxNode::Create(Scene, NodeName);
-	Scene->GetRootNode()->AddChild(BaseNode);
-
-	AActor* Owner = SkeletalMeshComponent->GetOwner();
-	if(Owner && Owner->GetRootComponent())
-	{
-		// Set the default position of the actor on the transforms
-		// The UE transformation is different from FBX's Z-up: invert the Y-axis for translations and the Y/Z angle values in rotations.
-		BaseNode->LclTranslation.Set(Converter.ConvertToFbxPos(Owner->GetActorLocation()));
-		BaseNode->LclRotation.Set(Converter.ConvertToFbxRot(Owner->GetActorRotation().Euler()));
-		BaseNode->LclScaling.Set(Converter.ConvertToFbxScale(Owner->GetRootComponent()->GetRelativeScale3D()));
-	}
-	// Create the Skeleton
-	TArray<FbxNode*> BoneNodes;
-	FbxNode* SkeletonRootNode = CreateSkeleton(SkeletalMeshComponent->GetSkeletalMesh(), BoneNodes);
-	FbxSkeletonRoots.Add(SkeletalMeshComponent, SkeletonRootNode);
-	BaseNode->AddChild(SkeletonRootNode);
-
-	static const float SamplingRate = 1.f / DEFAULT_SAMPLERATE;
-
-	FMatineeAnimTrackAdapter AnimTrackAdapter(MatineeActor);
-	ExportAnimTrack(AnimTrackAdapter, Owner, SkeletalMeshComponent, SamplingRate);
-}
-
 void FFbxExporter::ExportAnimTrack(IAnimTrackAdapter& AnimTrackAdapter, AActor* Actor, USkeletalMeshComponent* InSkeletalMeshComponent, float SamplingRate)
 {
 	// show a status update every 1 second worth of samples
@@ -693,7 +598,7 @@ void FFbxExporter::ExportAnimTrack(IAnimTrackAdapter& AnimTrackAdapter, AActor*
 		int32 LocalFrame = LocalStartFrame + FrameCount;
 		float SampleTime = (StartFrame + FrameCount) / FrameRate;
 
-		// This will call UpdateSkelPose on the skeletal mesh component to move bones based on animations in the matinee group
+		// This will call UpdateSkelPose on the skeletal mesh component to move bones based on animations in the sequence
 		AnimTrackAdapter.UpdateAnimation(LocalFrame);
 
 		if (FrameCount == 0)
diff --git a/Engine/Source/Editor/UnrealEd/Private/Fbx/FbxMainExport.cpp b/Engine/Source/Editor/UnrealEd/Private/Fbx/FbxMainExport.cpp
--- a/Engine/Source/Editor/UnrealEd/Private/Fbx/FbxMainExport.cpp
+++ b/Engine/Source/Editor/UnrealEd/Private/Fbx/FbxMainExport.cpp
@@ -15,6 +15,7 @@
 #include "GameFramework/Actor.h"
 #include "Engine/Blueprint.h"
 #include "RawIndexBuffer.h"
+#include "CinematicExporter.h"
 #include "Components/StaticMeshComponent.h"
 #include "Components/LightComponent.h"
 #include "Model.h"
@@ -52,15 +53,6 @@
 #include "Materials/MaterialExpressionTextureSample.h"
 #include "Materials/MaterialInstance.h"
 
-#include "Matinee/InterpData.h"
-#include "Matinee/InterpTrackMove.h"
-#include "Matinee/InterpTrackMoveAxis.h"
-#include "Matinee/InterpTrackFloatProp.h"
-#include "Matinee/InterpTrackInstFloatProp.h"
-#include "Matinee/InterpTrackInstMove.h"
-#include "Matinee/InterpTrackAnimControl.h"
-#include "Matinee/InterpTrackInstAnimControl.h"
-
 #include "LandscapeProxy.h"
 #include "LandscapeInfo.h"
 #include "LandscapeComponent.h"
@@ -68,9 +60,6 @@
 #include "Components/SplineMeshComponent.h"
 #include "StaticMeshResources.h"
 
-#include "Matinee/InterpGroup.h"
-#include "Matinee/InterpGroupInst.h"
-#include "Matinee/MatineeActor.h"
 #include "FbxExporter.h"
 
 #include "StaticMeshAttributes.h"
@@ -1626,121 +1615,6 @@ FbxSurfaceMaterial* FFbxExporter::ExportMaterial(UMaterialInterface* MaterialInt
 }
 
 
-FFbxExporter::FMatineeNodeNameAdapter::FMatineeNodeNameAdapter( AMatineeActor* InMatineeActor )
-{
-	MatineeActor = InMatineeActor;
-}
-
-FString FFbxExporter::FMatineeNodeNameAdapter::GetActorNodeName(const AActor* Actor)
-{
-	FString NodeName = Actor->GetName();
-	const UInterpGroupInst* FoundGroupInst = MatineeActor->FindGroupInst( Actor );
-	if( FoundGroupInst != NULL )
-	{
-		NodeName = FoundGroupInst->Group->GroupName.ToString();
-	}
-
-	// Maya does not support dashes.  Change all dashes to underscores
-	NodeName = NodeName.Replace(TEXT("-"), TEXT("_") );
-
-	return NodeName;
-}
-
-FMatineeAnimTrackAdapter::FMatineeAnimTrackAdapter( AMatineeActor* InMatineeActor )
-{
-	MatineeActor = InMatineeActor;
-}
-
-int32 FMatineeAnimTrackAdapter::GetLocalStartFrame() const
-{
-	return 0;
-}
-
-
-int32 FMatineeAnimTrackAdapter::GetLength() const
-{
-	return MatineeActor->MatineeData->InterpLength * (1.f / GetFrameRate());
-}
-
-
-void FMatineeAnimTrackAdapter::UpdateAnimation( int32 LocalFrame )
-{
-	float Time = LocalFrame * GetFrameRate();
-	MatineeActor->UpdateInterp( Time, true );
-}
-
-
-/**
- * Exports the given Matinee sequence information into a FBX document.
- */
-bool FFbxExporter::ExportMatinee(AMatineeActor* InMatineeActor)
-{
-	if (InMatineeActor == NULL || Scene == NULL) return false;
-
-	// If the Matinee editor is not open, we need to initialize the sequence.
-	//bool InitializeMatinee = InMatineeActor->MatineeData == NULL;
-	//if (InitializeMatinee)
-	//{
-	//	InMatineeActor->InitInterp();
-	//}
-
-	// Iterate over the Matinee data groups and export the known tracks
-	int32 GroupCount = InMatineeActor->GroupInst.Num();
-	for (int32 GroupIndex = 0; GroupIndex < GroupCount; ++GroupIndex)
-	{
-		UInterpGroupInst* Group = InMatineeActor->GroupInst[GroupIndex];
-		AActor* Actor = Group->GetGroupActor();
-		if (Group->Group == NULL || Actor == NULL) continue;
-
-		FbxNode* FbxActor = FindActor(Actor); 
-		// now it should export everybody
-		check (FbxActor);
-
-		// Look for the tracks that we currently support
-		int32 TrackCount = FMath::Min(Group->TrackInst.Num(), Group->Group->InterpTracks.Num());
-		for (int32 TrackIndex = 0; TrackIndex < TrackCount; ++TrackIndex)
-		{
-			UInterpTrackInst* TrackInst = Group->TrackInst[TrackIndex];
-			UInterpTrack* Track = Group->Group->InterpTracks[TrackIndex];
-			if ( Track->IsDisabled() == false )
-			{
-				if(TrackInst->IsA(UInterpTrackInstMove::StaticClass()) && Track->IsA(UInterpTrackMove::StaticClass()))
-				{
-					UInterpTrackInstMove* MoveTrackInst = (UInterpTrackInstMove*)TrackInst;
-					UInterpTrackMove* MoveTrack = (UInterpTrackMove*)Track;
-					ExportMatineeTrackMove(FbxActor, MoveTrackInst, MoveTrack, InMatineeActor->MatineeData->InterpLength);
-				}
-				else if(TrackInst->IsA(UInterpTrackInstFloatProp::StaticClass()) && Track->IsA(UInterpTrackFloatProp::StaticClass()))
-				{
-					UInterpTrackInstFloatProp* PropertyTrackInst = (UInterpTrackInstFloatProp*)TrackInst;
-					UInterpTrackFloatProp* PropertyTrack = (UInterpTrackFloatProp*)Track;
-					ExportMatineeTrackFloatProp(FbxActor, PropertyTrack);
-				}
-				else if(TrackInst->IsA(UInterpTrackInstAnimControl::StaticClass()) && Track->IsA(UInterpTrackAnimControl::StaticClass()))
-				{
-					USkeletalMeshComponent* SkeletalMeshComp = Cast<USkeletalMeshComponent>(Actor->GetComponentByClass(USkeletalMeshComponent::StaticClass()));
-					if(SkeletalMeshComp)
-					{
-						static const float SamplingRate = 1.f / DEFAULT_SAMPLERATE;
-
-						FMatineeAnimTrackAdapter AnimTrackAdapter(InMatineeActor);
-						ExportAnimTrack(AnimTrackAdapter, Actor, SkeletalMeshComp, SamplingRate);
-					}
-				}
-			}
-		}
-	}
-
-	//if (InitializeMatinee)
-	//{
-	//	InMatineeActor->TermInterp();
-	//}
-
-	DefaultCamera = NULL;
-	return true;
-}
-
-
 FFbxExporter::FLevelSequenceNodeNameAdapter::FLevelSequenceNodeNameAdapter( UMovieScene* InMovieScene, IMovieScenePlayer* InMovieScenePlayer, FMovieSceneSequenceIDRef InSequenceID )
 {
 	MovieScene = InMovieScene;
@@ -2343,121 +2217,6 @@ FbxNode* FFbxExporter::ExportActor(AActor* Actor, bool bExportComponents, INodeN
 	return ActorNode;
 }
 
-/**
- * Exports the Matinee movement track into the FBX animation library.
- */
-void FFbxExporter::ExportMatineeTrackMove(FbxNode* FbxActor, UInterpTrackInstMove* MoveTrackInst, UInterpTrackMove* MoveTrack, float InterpLength)
-{
-	if (FbxActor == NULL || MoveTrack == NULL) return;
-	
-	// For the Y and Z angular rotations, we need to invert the relative animation frames,
-	// While keeping the standard angles constant.
-
-	if (MoveTrack != NULL)
-	{
-		FbxAnimLayer* BaseLayer = AnimStack->GetMember<FbxAnimLayer>(0);
-		FbxAnimCurve* Curve;
-
-		bool bPosCurve = true;
-		if( MoveTrack->SubTracks.Num() == 0 )
-		{
-			// Translation;
-			FbxActor->LclTranslation.GetCurveNode(BaseLayer, true);
-			Curve = FbxActor->LclTranslation.GetCurve(BaseLayer, FBXSDK_CURVENODE_COMPONENT_X, true);
-			ExportAnimatedVector(Curve, "X", MoveTrack, MoveTrackInst, bPosCurve, 0, false, InterpLength);
-			Curve = FbxActor->LclTranslation.GetCurve(BaseLayer, FBXSDK_CURVENODE_COMPONENT_Y, true);
-			ExportAnimatedVector(Curve, "Y", MoveTrack, MoveTrackInst, bPosCurve, 1, true, InterpLength);
-			Curve = FbxActor->LclTranslation.GetCurve(BaseLayer, FBXSDK_CURVENODE_COMPONENT_Z, true);
-			ExportAnimatedVector(Curve, "Z", MoveTrack, MoveTrackInst, bPosCurve, 2, false, InterpLength);
-
-			// Rotation
-			FbxActor->LclRotation.GetCurveNode(BaseLayer, true);
-			bPosCurve = false;
-
-			Curve = FbxActor->LclRotation.GetCurve(BaseLayer, FBXSDK_CURVENODE_COMPONENT_X, true);
-			ExportAnimatedVector(Curve, "X", MoveTrack, MoveTrackInst, bPosCurve, 0, false, InterpLength);
-			Curve = FbxActor->LclRotation.GetCurve(BaseLayer, FBXSDK_CURVENODE_COMPONENT_Y, true);
-			ExportAnimatedVector(Curve, "Y", MoveTrack, MoveTrackInst, bPosCurve, 1, true, InterpLength);
-			Curve = FbxActor->LclRotation.GetCurve(BaseLayer, FBXSDK_CURVENODE_COMPONENT_Z, true);
-			ExportAnimatedVector(Curve, "Z", MoveTrack, MoveTrackInst, bPosCurve, 2, true, InterpLength);
-		}
-		else
-		{
-			// Translation;
-			FbxActor->LclTranslation.GetCurveNode(BaseLayer, true);
-			Curve = FbxActor->LclTranslation.GetCurve(BaseLayer, FBXSDK_CURVENODE_COMPONENT_X, true);
-			ExportMoveSubTrack(Curve, "X", CastChecked<UInterpTrackMoveAxis>(MoveTrack->SubTracks[0]), MoveTrackInst, bPosCurve, 0, false, InterpLength);
-			Curve = FbxActor->LclTranslation.GetCurve(BaseLayer, FBXSDK_CURVENODE_COMPONENT_Y, true);
-			ExportMoveSubTrack(Curve, "Y", CastChecked<UInterpTrackMoveAxis>(MoveTrack->SubTracks[1]), MoveTrackInst, bPosCurve, 1, true, InterpLength);
-			Curve = FbxActor->LclTranslation.GetCurve(BaseLayer, FBXSDK_CURVENODE_COMPONENT_Z, true);
-			ExportMoveSubTrack(Curve, "Z", CastChecked<UInterpTrackMoveAxis>(MoveTrack->SubTracks[2]), MoveTrackInst, bPosCurve, 2, false, InterpLength);
-
-			// Rotation
-			FbxActor->LclRotation.GetCurveNode(BaseLayer, true);
-			bPosCurve = false;
-
-			Curve = FbxActor->LclRotation.GetCurve(BaseLayer, FBXSDK_CURVENODE_COMPONENT_X, true);
-			ExportMoveSubTrack(Curve, "X", CastChecked<UInterpTrackMoveAxis>(MoveTrack->SubTracks[3]), MoveTrackInst, bPosCurve, 0, false, InterpLength);
-			Curve = FbxActor->LclRotation.GetCurve(BaseLayer, FBXSDK_CURVENODE_COMPONENT_Y, true);
-			ExportMoveSubTrack(Curve, "Y", CastChecked<UInterpTrackMoveAxis>(MoveTrack->SubTracks[4]), MoveTrackInst, bPosCurve, 1, true, InterpLength);
-			Curve = FbxActor->LclRotation.GetCurve(BaseLayer, FBXSDK_CURVENODE_COMPONENT_Z, true);
-			ExportMoveSubTrack(Curve, "Z", CastChecked<UInterpTrackMoveAxis>(MoveTrack->SubTracks[5]), MoveTrackInst, bPosCurve, 2, true, InterpLength);
-		}
-	}
-}
-
-/**
- * Exports the Matinee float property track into the FBX animation library.
- */
-void FFbxExporter::ExportMatineeTrackFloatProp(FbxNode* FbxActor, UInterpTrackFloatProp* PropTrack)
-{
-	if (FbxActor == NULL || PropTrack == NULL) return;
-	
-	FbxNodeAttribute* FbxNodeAttr = NULL;
-	// camera and light is appended on the fbx pivot node
-	if( FbxActor->GetChild(0) )
-	{
-		FbxNodeAttr = ((FbxNode*)FbxActor->GetChild(0))->GetNodeAttribute();
-
-		if (FbxNodeAttr == NULL) return;
-	}
-	
-	FbxProperty Property;
-	FString PropertyName = PropTrack->PropertyName.ToString();
-	bool IsFoV = false;
-	// most properties are created as user property, only FOV of camera in FBX supports animation
-	if (PropertyName == "Intensity")
-	{
-		Property = FbxActor->FindProperty("UE_Intensity", false);
-	}
-	else if (PropertyName == "FalloffExponent")
-	{
-		Property = FbxActor->FindProperty("UE_FalloffExponent", false);
-	}
-	else if (PropertyName == "AttenuationRadius")
-	{
-		Property = FbxActor->FindProperty("UE_Radius", false);
-	}
-	else if (PropertyName == "FOVAngle" && FbxNodeAttr )
-	{
-		Property = ((FbxCamera*)FbxNodeAttr)->FocalLength;
-		IsFoV = true;
-	}
-	else if (PropertyName == "AspectRatio")
-	{
-		Property = FbxActor->FindProperty("UE_AspectRatio", false);
-	}
-	else if (PropertyName == "MotionBlur_Amount")
-	{
-		Property = FbxActor->FindProperty("UE_MotionBlur_Amount", false);
-	}
-
-	if (Property != 0)
-	{
-		ExportAnimatedFloat(&Property, &PropTrack->FloatTrack, IsFoV);
-	}
-}
-
 void ConvertInterpToFBX(uint8 UnrealInterpMode, FbxAnimCurveDef::EInterpolationType& Interpolation, FbxAnimCurveDef::ETangentMode& Tangent)
 {
 	switch(UnrealInterpMode)
@@ -2501,349 +2260,6 @@ static bool IsEquivalent(float a, float b, float Tolerance = KINDA_SMALL_NUMBER)
 	return (a - b) > -Tolerance && (a - b) < Tolerance;
 }
 
-// Set the default FPS to 30 because the SetupMatinee MEL script sets up Maya this way.
-const float FFbxExporter::BakeTransformsFPS = DEFAULT_SAMPLERATE;
-
-/**
- * Exports a given interpolation curve into the FBX animation curve.
- */
-void FFbxExporter::ExportAnimatedVector(FbxAnimCurve* FbxCurve, const char* ChannelName, UInterpTrackMove* MoveTrack, UInterpTrackInstMove* MoveTrackInst, bool bPosCurve, int32 CurveIndex, bool bNegative, float InterpLength)
-{
-	if (Scene == NULL) return;
-	
-	FInterpCurveVector* Curve = bPosCurve ? &MoveTrack->PosTrack : &MoveTrack->EulerTrack;
-
-	if (Curve == NULL || CurveIndex >= 3) return;
-
-#define FLT_TOLERANCE 0.000001
-
-	// Determine how many key frames we are exporting. If the user wants to export a key every 
-	// frame, calculate this number. Otherwise, use the number of keys the user created. 
-	int32 KeyCount = bBakeKeys ? (InterpLength * BakeTransformsFPS) : Curve->Points.Num();
-
-	// Write out the key times from the curve to the FBX curve.
-	TArray<float> KeyTimes;
-	for (int32 KeyIndex = 0; KeyIndex < KeyCount; ++KeyIndex)
-	{
-		// The Unreal engine allows you to place more than one key at one time value:
-		// displace the extra keys. This assumes that Unreal's keys are always ordered.
-		float KeyTime = bBakeKeys ? (KeyIndex * InterpLength) / KeyCount : Curve->Points[KeyIndex].InVal;
-		if (KeyTimes.Num() && KeyTime < KeyTimes[KeyIndex-1] + FLT_TOLERANCE)
-		{
-			KeyTime = KeyTimes[KeyIndex-1] + 0.01f; // Add 1 millisecond to the timing of this key.
-		}
-		KeyTimes.Add(KeyTime);
-	}
-
-	// Write out the key values from the curve to the FBX curve.
-	FbxCurve->KeyModifyBegin();
-	for (int32 KeyIndex = 0; KeyIndex < KeyCount; ++KeyIndex)
-	{
-		// First, convert the output value to the correct coordinate system, if we need that.  For movement
-		// track keys that are in a local coordinate system (IMF_RelativeToInitial), we need to transform
-		// the keys to world space first
-		FVector FinalOutVec;
-		{
-			FVector KeyPosition;
-			FRotator KeyRotation;
-
-			// If we are baking trnasforms, ask the movement track what are transforms are at the given time.
-			if( bBakeKeys )
-			{
-				MoveTrack->GetKeyTransformAtTime(MoveTrackInst, KeyTimes[KeyIndex], KeyPosition, KeyRotation);
-			}
-			// Else, this information is already present in the position and rotation tracks stored on the movement track.
-			else
-			{
-				KeyPosition = MoveTrack->PosTrack.Points[KeyIndex].OutVal;
-				KeyRotation = FRotator( FQuat::MakeFromEuler(MoveTrack->EulerTrack.Points[KeyIndex].OutVal) );
-			}
-
-			if (bKeepHierarchy)
-			{
-				if(bPosCurve)
-				{
-					FinalOutVec = KeyPosition;
-				}
-				else
-				{
-					FinalOutVec = KeyRotation.Euler();
-				}
-			}
-			else
-			{
-				FVector WorldSpacePos;
-				FRotator WorldSpaceRotator;
-				MoveTrack->ComputeWorldSpaceKeyTransform(
-					MoveTrackInst,
-					KeyPosition,
-					KeyRotation,
-					WorldSpacePos,			// Out
-					WorldSpaceRotator);	// Out
-
-				if(bPosCurve)
-				{
-					FinalOutVec = WorldSpacePos;
-				}
-				else
-				{
-					FinalOutVec = WorldSpaceRotator.Euler();
-				}
-			}
-		}
-
-		float KeyTime = KeyTimes[KeyIndex];
-		float OutValue = (CurveIndex == 0) ? FinalOutVec.X : (CurveIndex == 1) ? FinalOutVec.Y : FinalOutVec.Z;
-		float FbxKeyValue = bNegative ? -OutValue : OutValue;
-		
-		// Add a new key to the FBX curve
-		FbxTime Time;
-		FbxAnimCurveKey FbxKey;
-		Time.SetSecondDouble((float)KeyTime);
-		int FbxKeyIndex = FbxCurve->KeyAdd(Time);
-		
-
-		FbxAnimCurveDef::EInterpolationType Interpolation = FbxAnimCurveDef::eInterpolationConstant;
-		FbxAnimCurveDef::ETangentMode Tangent = FbxAnimCurveDef::eTangentAuto;
-		
-		if( !bBakeKeys )
-		{
-			ConvertInterpToFBX(Curve->Points[KeyIndex].InterpMode, Interpolation, Tangent);
-		}
-
-		if (bBakeKeys || Interpolation != FbxAnimCurveDef::eInterpolationCubic)
-		{
-			FbxCurve->KeySet(FbxKeyIndex, Time, (float)FbxKeyValue, Interpolation, Tangent);
-		}
-		else
-		{
-			FInterpCurvePoint<FVector>& Key = Curve->Points[KeyIndex];
-
-			// Setup tangents for bezier curves. Avoid this for keys created from baking 
-			// transforms since there is no tangent info created for these types of keys. 
-			if( Interpolation == FbxAnimCurveDef::eInterpolationCubic ) //-V547
-			{
-				float OutTangentValue = (CurveIndex == 0) ? Key.LeaveTangent.X : (CurveIndex == 1) ? Key.LeaveTangent.Y : Key.LeaveTangent.Z;
-				float OutTangentX = (KeyIndex < KeyCount - 1) ? (KeyTimes[KeyIndex + 1] - KeyTime) / 3.0f : 0.333f;
-				if (IsEquivalent(OutTangentX, KeyTime))
-				{
-					OutTangentX = 0.00333f; // 1/3rd of a millisecond.
-				}
-				float OutTangentY = OutTangentValue / 3.0f;
-				float RightTangent =  OutTangentY / OutTangentX ;
-				
-				float NextLeftTangent = 0;
-				
-				if (KeyIndex < KeyCount - 1)
-				{
-					FInterpCurvePoint<FVector>& NextKey = Curve->Points[KeyIndex + 1];
-					float NextInTangentValue = (CurveIndex == 0) ? NextKey.ArriveTangent.X : (CurveIndex == 1) ? NextKey.ArriveTangent.Y : NextKey.ArriveTangent.Z;
-					float NextInTangentX;
-					NextInTangentX = (KeyTimes[KeyIndex + 1] - KeyTimes[KeyIndex]) / 3.0f;
-					float NextInTangentY = NextInTangentValue / 3.0f;
-					NextLeftTangent =  NextInTangentY / NextInTangentX ;
-				}
-
-				FbxCurve->KeySet(FbxKeyIndex, Time, (float)FbxKeyValue, Interpolation, Tangent, RightTangent, NextLeftTangent );
-			}
-		}
-	}
-	FbxCurve->KeyModifyEnd();
-}
-
-void FFbxExporter::ExportMoveSubTrack(FbxAnimCurve* FbxCurve, const ANSICHAR* ChannelName, UInterpTrackMoveAxis* SubTrack, UInterpTrackInstMove* MoveTrackInst, bool bPosCurve, int32 CurveIndex, bool bNegative, float InterpLength)
-{
-	if (Scene == NULL || FbxCurve == NULL) return;
-
-	FInterpCurveFloat* Curve = &SubTrack->FloatTrack;
-	UInterpTrackMove* ParentTrack = CastChecked<UInterpTrackMove>( SubTrack->GetOuter() );
-
-#define FLT_TOLERANCE 0.000001
-
-	// Determine how many key frames we are exporting. If the user wants to export a key every 
-	// frame, calculate this number. Otherwise, use the number of keys the user created. 
-	int32 KeyCount = bBakeKeys ? (InterpLength * BakeTransformsFPS) : Curve->Points.Num();
-
-	// Write out the key times from the curve to the FBX curve.
-	TArray<float> KeyTimes;
-	for(int32 KeyIndex = 0; KeyIndex < KeyCount; ++KeyIndex)
-	{
-		// The Unreal engine allows you to place more than one key at one time value:
-		// displace the extra keys. This assumes that Unreal's keys are always ordered.
-		float KeyTime = bBakeKeys ? (KeyIndex * InterpLength) / KeyCount : Curve->Points[KeyIndex].InVal;
-		if(KeyTimes.Num() && KeyTime < KeyTimes[KeyIndex-1] + FLT_TOLERANCE)
-		{
-			KeyTime = KeyTimes[KeyIndex-1] + 0.01f; // Add 1 millisecond to the timing of this key.
-		}
-		KeyTimes.Add(KeyTime);
-	}
-	// Write out the key values from the curve to the FBX curve.
-	FbxCurve->KeyModifyBegin();
-	for (int32 KeyIndex = 0; KeyIndex < KeyCount; ++KeyIndex)
-	{
-		// First, convert the output value to the correct coordinate system, if we need that.  For movement
-		// track keys that are in a local coordinate system (IMF_RelativeToInitial), we need to transform
-		// the keys to world space first
-		FVector FinalOutVec;
-		{
-			FVector KeyPosition;
-			FRotator KeyRotation;
-
-			ParentTrack->GetKeyTransformAtTime(MoveTrackInst, KeyTimes[KeyIndex], KeyPosition, KeyRotation);
-		
-			FVector WorldSpacePos;
-			FRotator WorldSpaceRotator;
-			ParentTrack->ComputeWorldSpaceKeyTransform(
-				MoveTrackInst,
-				KeyPosition,
-				KeyRotation,
-				WorldSpacePos,			// Out
-				WorldSpaceRotator );	// Out
-
-			if( bPosCurve )
-			{
-				FinalOutVec = WorldSpacePos;
-			}
-			else
-			{
-				FinalOutVec = WorldSpaceRotator.Euler();
-			}
-		}
-
-		float KeyTime = KeyTimes[KeyIndex];
-		float OutValue = (CurveIndex == 0) ? FinalOutVec.X : (CurveIndex == 1) ? FinalOutVec.Y : FinalOutVec.Z;
-		float FbxKeyValue = bNegative ? -OutValue : OutValue;
-
-		// Add a new key to the FBX curve
-		FbxTime Time;
-		FbxAnimCurveKey FbxKey;
-		Time.SetSecondDouble((float)KeyTime);
-		int FbxKeyIndex = FbxCurve->KeyAdd(Time);
-
-		FbxAnimCurveDef::EInterpolationType Interpolation = FbxAnimCurveDef::eInterpolationConstant;
-		FbxAnimCurveDef::ETangentMode Tangent = FbxAnimCurveDef::eTangentAuto;
-		
-		if (bBakeKeys || Interpolation != FbxAnimCurveDef::eInterpolationCubic)
-		{
-			FbxCurve->KeySet(FbxKeyIndex, Time, (float)FbxKeyValue, Interpolation, Tangent);
-		}
-		else
-		{
-			FInterpCurvePoint<float>& Key = Curve->Points[KeyIndex];
-			ConvertInterpToFBX(Key.InterpMode, Interpolation, Tangent);
-
-			// Setup tangents for bezier curves. Avoid this for keys created from baking 
-			// transforms since there is no tangent info created for these types of keys. 
-			if( Interpolation == FbxAnimCurveDef::eInterpolationCubic )
-			{
-				float OutTangentValue = Key.LeaveTangent;
-				float OutTangentX = (KeyIndex < KeyCount - 1) ? (KeyTimes[KeyIndex + 1] - KeyTime) / 3.0f : 0.333f;
-				if (IsEquivalent(OutTangentX, KeyTime))
-				{
-					OutTangentX = 0.00333f; // 1/3rd of a millisecond.
-				}
-				float OutTangentY = OutTangentValue / 3.0f;
-				float RightTangent =  OutTangentY / OutTangentX ;
-
-				float NextLeftTangent = 0;
-
-				if (KeyIndex < KeyCount - 1)
-				{
-					FInterpCurvePoint<float>& NextKey = Curve->Points[KeyIndex + 1];
-					float NextInTangentValue =  Key.LeaveTangent;
-					float NextInTangentX;
-					NextInTangentX = (KeyTimes[KeyIndex + 1] - KeyTimes[KeyIndex]) / 3.0f;
-					float NextInTangentY = NextInTangentValue / 3.0f;
-					NextLeftTangent =  NextInTangentY / NextInTangentX ;
-				}
-
-				FbxCurve->KeySet(FbxKeyIndex, Time, (float)FbxKeyValue, Interpolation, Tangent, RightTangent, NextLeftTangent );
-			}
-		}
-	}
-	FbxCurve->KeyModifyEnd();
-}
-
-void FFbxExporter::ExportAnimatedFloat(FbxProperty* FbxProperty, FInterpCurveFloat* Curve, bool IsCameraFoV)
-{
-	if (FbxProperty == NULL || Curve == NULL) return;
-
-	// do not export an empty anim curve
-	if (Curve->Points.Num() == 0) return;
-
-	FbxAnimCurve* AnimCurve = FbxAnimCurve::Create(Scene, "");
-	FbxAnimCurveNode* CurveNode = FbxProperty->GetCurveNode(true);
-	if (!CurveNode)
-	{
-		return;
-	}
-	CurveNode->SetChannelValue<double>(0U, Curve->Points[0].OutVal);
-	CurveNode->ConnectToChannel(AnimCurve, 0U);
-
-	// Write out the key times from the curve to the FBX curve.
-	int32 KeyCount = Curve->Points.Num();
-	TArray<float> KeyTimes;
-	for (int32 KeyIndex = 0; KeyIndex < KeyCount; ++KeyIndex)
-	{
-		FInterpCurvePoint<float>& Key = Curve->Points[KeyIndex];
-
-		// The Unreal engine allows you to place more than one key at one time value:
-		// displace the extra keys. This assumes that Unreal's keys are always ordered.
-		float KeyTime = Key.InVal;
-		if (KeyTimes.Num() && KeyTime < KeyTimes[KeyIndex-1] + FLT_TOLERANCE)
-		{
-			KeyTime = KeyTimes[KeyIndex-1] + 0.01f; // Add 1 millisecond to the timing of this key.
-		}
-		KeyTimes.Add(KeyTime);
-	}
-
-	// Write out the key values from the curve to the FBX curve.
-	AnimCurve->KeyModifyBegin();
-	for (int32 KeyIndex = 0; KeyIndex < KeyCount; ++KeyIndex)
-	{
-		FInterpCurvePoint<float>& Key = Curve->Points[KeyIndex];
-		float KeyTime = KeyTimes[KeyIndex];
-		
-		// Add a new key to the FBX curve
-		FbxTime Time;
-		FbxAnimCurveKey FbxKey;
-		Time.SetSecondDouble((float)KeyTime);
-		int FbxKeyIndex = AnimCurve->KeyAdd(Time);
-		float OutVal = (IsCameraFoV && DefaultCamera)? DefaultCamera->ComputeFocalLength(Key.OutVal): (float)Key.OutVal;
-
-		FbxAnimCurveDef::EInterpolationType Interpolation = FbxAnimCurveDef::eInterpolationConstant;
-		FbxAnimCurveDef::ETangentMode Tangent = FbxAnimCurveDef::eTangentAuto;
-		ConvertInterpToFBX(Key.InterpMode, Interpolation, Tangent);
-		
-		if (Interpolation != FbxAnimCurveDef::eInterpolationCubic)
-		{
-			AnimCurve->KeySet(FbxKeyIndex, Time, OutVal, Interpolation, Tangent);
-		}
-		else
-		{
-			// Setup tangents for bezier curves.
-			float OutTangentX = (KeyIndex < KeyCount - 1) ? (KeyTimes[KeyIndex + 1] - KeyTime) / 3.0f : 0.333f;
-			float OutTangentY = Key.LeaveTangent / 3.0f;
-			float RightTangent =  OutTangentY / OutTangentX ;
-
-			float NextLeftTangent = 0;
-
-			if (KeyIndex < KeyCount - 1)
-			{
-				FInterpCurvePoint<float>& NextKey = Curve->Points[KeyIndex + 1];
-				float NextInTangentX;
-				NextInTangentX = (KeyTimes[KeyIndex + 1] - KeyTimes[KeyIndex]) / 3.0f;
-				float NextInTangentY = NextKey.ArriveTangent / 3.0f;
-				NextLeftTangent =  NextInTangentY / NextInTangentX ;
-			}
-
-			AnimCurve->KeySet(FbxKeyIndex, Time, OutVal, Interpolation, Tangent, RightTangent, NextLeftTangent );
-
-		}
-	}
-	AnimCurve->KeyModifyEnd();
-}
-
 
 void RichCurveInterpolationToFbxInterpolation(ERichCurveInterpMode InInterpolation, ERichCurveTangentMode InTangentMode, ERichCurveTangentWeightMode InTangentWeightMode,
 	FbxAnimCurveDef::EInterpolationType &OutInterpolation, FbxAnimCurveDef::ETangentMode &OutTangentMode, FbxAnimCurveDef::EWeightedMode &OutTangentWeightMode)
diff --git a/Engine/Source/Editor/UnrealEd/Private/FbxExporter.h b/Engine/Source/Editor/UnrealEd/Private/FbxExporter.h
--- a/Engine/Source/Editor/UnrealEd/Private/FbxExporter.h
+++ b/Engine/Source/Editor/UnrealEd/Private/FbxExporter.h
@@ -6,7 +6,7 @@
 #include "Misc/Guid.h"
 #include "EngineDefines.h"
 #include "Engine/StaticMesh.h"
-#include "MatineeExporter.h"
+#include "CinematicExporter.h"
 #include "MovieSceneSequenceID.h"
 #include "Evaluation/MovieSceneSequenceTransform.h"
 #include "MovieSceneFwd.h"
@@ -18,16 +18,11 @@ class ABrush;
 class ACameraActor;
 class ALandscapeProxy;
 class ALight;
-class AMatineeActor;
 class ASkeletalMeshActor;
 class IMovieScenePlayer;
 class UAnimSequence;
 class UCameraComponent;
 class UInstancedStaticMeshComponent;
-class UInterpTrackFloatProp;
-class UInterpTrackInstMove;
-class UInterpTrackMove;
-class UInterpTrackMoveAxis;
 class ULightComponent;
 class UMaterialInterface;
 class UModel;
@@ -51,7 +46,7 @@ struct FMovieSceneSequenceTransform;
 
 namespace UnFbx
 {
-	/** Adapter interface which allows ExportAnimTrack to act on both sequencer and matinee data. */
+	/** Adapter interface which allows ExportAnimTrack to act on sequencer without a tight coupling. */
 	class UNREALED_API IAnimTrackAdapter
 	{
 	public:
@@ -68,18 +63,6 @@ namespace UnFbx
 		virtual float GetAnimTime(int32 LocalFrame) const { return 0.f; }
 	};
 
-	/** An anim track adapter for matinee. */
-	class UNREALED_API FMatineeAnimTrackAdapter : public IAnimTrackAdapter
-	{
-	public:
-		FMatineeAnimTrackAdapter(AMatineeActor* InMatineeActor);
-		virtual int32 GetLocalStartFrame() const override;
-		virtual int32 GetLength() const override;
-		virtual void UpdateAnimation(int32 LocalFrame) override;
-
-	private:
-		AMatineeActor* MatineeActor;
-	};
 	/** An anim track adapter for a level sequence. */
 	class UNREALED_API FLevelSequenceAnimTrackAdapter : public IAnimTrackAdapter
 	{
@@ -102,7 +85,7 @@ namespace UnFbx
 /**
  * Main FBX Exporter class.
  */
-class UNREALED_API FFbxExporter  : public MatineeExporter, public FGCObject
+class UNREALED_API FFbxExporter : public FCinematicExporter, public FGCObject
 {
 public:
 	/**
@@ -175,13 +158,6 @@ public:
 	virtual void ExportLevelMesh(ULevel* InLevel, bool bExportLevelGeometry, TArray<AActor*>& ActorToExport, INodeNameAdapter& NodeNameAdapter, bool bSaveAnimSeq = true);
 
 
-	/**
-	 * Exports the given Matinee sequence information into a FBX document.
-	 * 
-	 * @return	true, if sucessful
-	 */
-	virtual bool ExportMatinee(class AMatineeActor* InMatineeActor);
-
 	/**
 	 * Exports the given level sequence information into a FBX document.
 	 *
@@ -196,16 +172,6 @@ public:
 	 */
 	bool ExportLevelSequenceTracks(UMovieScene* MovieScene, IMovieScenePlayer* MovieScenePlayer, FMovieSceneSequenceIDRef InSequenceID, FbxNode* FbxActor, UObject* BoundObject, const TArray<UMovieSceneTrack*>& Tracks, const FMovieSceneSequenceTransform& RootToLocalTransform);
 
-	/**
-	 * Exports all the animation sequences part of a single Group in a Matinee sequence
-	 * as a single animation in the FBX document.  The animation is created by sampling the
-	 * sequence at 30 updates/second and extracting the resulting bone transforms from the given
-	 * skeletal mesh
-	 * @param MatineeSequence The Matinee Sequence containing the group to export
-	 * @param SkeletalMeshComponent The Skeletal mesh that the animations from the Matinee group are applied to
-	 */
-	virtual void ExportMatineeGroup(class AMatineeActor* MatineeActor, USkeletalMeshComponent* SkeletalMeshComponent);
-
 
 	/**
 	 * Exports the mesh and the actor information for a static mesh actor.
@@ -251,21 +217,6 @@ public:
 	 */
 	FbxNode* ExportAnimSequence( const UAnimSequence* AnimSeq, const USkeletalMesh* SkelMesh, bool bExportSkelMesh, const TCHAR* MeshNames=NULL, FbxNode* ActorRootNode=NULL, const TArray<UMaterialInterface*>* OverrideMaterials = nullptr);
 
-	/**
-	 * Exports the list of UAnimSequences as a single animation based on the settings in the TrackKeys
-	 */
-	void ExportAnimSequencesAsSingle( USkeletalMesh* SkelMesh, const ASkeletalMeshActor* SkelMeshActor, const FString& ExportName, const TArray<UAnimSequence*>& AnimSeqList, const TArray<FAnimControlTrackKey>& TrackKeys );
-
-	/** A node name adapter for matinee. */
-	class UNREALED_API FMatineeNodeNameAdapter : public INodeNameAdapter
-	{
-	public:
-		FMatineeNodeNameAdapter( AMatineeActor* InMatineeActor );
-		virtual FString GetActorNodeName(const AActor* InActor) override;
-	private:
-		AMatineeActor* MatineeActor;
-	};
-
 	/** A node name adapter for a level sequence. */
 	class UNREALED_API FLevelSequenceNodeNameAdapter : public INodeNameAdapter
 	{
@@ -442,28 +393,6 @@ private:
 	 */
 	void CorrectAnimTrackInterpolation( TArray<FbxNode*>& BoneNodes, FbxAnimLayer* AnimLayer );
 
-	/**
-	 * Exports the Matinee movement track into the FBX animation stack.
-	 */
-	void ExportMatineeTrackMove(FbxNode* FbxActor, UInterpTrackInstMove* MoveTrackInst, UInterpTrackMove* MoveTrack, float InterpLength);
-
-	/**
-	 * Exports the Matinee float property track into the FBX animation stack.
-	 */
-	void ExportMatineeTrackFloatProp(FbxNode* FbxActor, UInterpTrackFloatProp* PropTrack);
-
-	/**
-	 * Exports a given interpolation curve into the FBX animation curve.
-	 */
-	void ExportAnimatedVector(FbxAnimCurve* FbxCurve, const ANSICHAR* ChannelName, UInterpTrackMove* MoveTrack, UInterpTrackInstMove* MoveTrackInst, bool bPosCurve, int32 CurveIndex, bool bNegative, float InterpLength);
-	
-	/**
-	 * Exports a movement subtrack to an FBX curve
-	 */
-	void ExportMoveSubTrack(FbxAnimCurve* FbxCurve, const ANSICHAR* ChannelName, UInterpTrackMoveAxis* SubTrack, UInterpTrackInstMove* MoveTrackInst, bool bPosCurve, int32 CurveIndex, bool bNegative, float InterpLength);
-	
-	void ExportAnimatedFloat(FbxProperty* FbxProperty, FInterpCurveFloat* Curve, bool IsCameraFoV);
-
 	/**
 	 * Exports a level sequence 3D transform track into the FBX animation stack.
 	 */
diff --git a/Engine/Source/Editor/UnrealEd/Private/GroupActor.cpp b/Engine/Source/Editor/UnrealEd/Private/GroupActor.cpp
--- a/Engine/Source/Editor/UnrealEd/Private/GroupActor.cpp
+++ b/Engine/Source/Editor/UnrealEd/Private/GroupActor.cpp
@@ -683,17 +683,13 @@ void AGroupActor::UnlockSelectedGroups()
 
 void AGroupActor::ToggleGroupMode()
 {
-	// Group mode can only be toggled when not in InterpEdit mode
-	if( !GLevelEditorModeTools().IsModeActive(FBuiltinEditorModes::EM_InterpEdit) )
-	{
-		UActorGroupingUtils::SetGroupingActive(!UActorGroupingUtils::IsGroupingActive());
+	UActorGroupingUtils::SetGroupingActive(!UActorGroupingUtils::IsGroupingActive());
 
-		// Update group selection in the editor to reflect the toggle
-		SelectGroupsInSelection();
-		GEditor->RedrawAllViewports();
+	// Update group selection in the editor to reflect the toggle
+	SelectGroupsInSelection();
+	GEditor->RedrawAllViewports();
 
-		GEditor->SaveConfig();
-	}
+	GEditor->SaveConfig();
 }
 
 
diff --git a/Engine/Source/Editor/UnrealEd/Private/Kismet2/BlueprintEditorUtils.cpp b/Engine/Source/Editor/UnrealEd/Private/Kismet2/BlueprintEditorUtils.cpp
--- a/Engine/Source/Editor/UnrealEd/Private/Kismet2/BlueprintEditorUtils.cpp
+++ b/Engine/Source/Editor/UnrealEd/Private/Kismet2/BlueprintEditorUtils.cpp
@@ -44,7 +44,6 @@
 #include "EdMode.h"
 #include "Dialogs/Dialogs.h"
 #include "UnrealEdGlobals.h"
-#include "Matinee/MatineeActor.h"
 #include "Engine/LevelScriptBlueprint.h"
 #include "UObject/BlueprintsObjectVersion.h"
 #include "Kismet2/CompilerResultsLog.h"
@@ -74,7 +73,6 @@
 #include "K2Node_Literal.h"
 #include "K2Node_MacroInstance.h"
 #include "K2Node_MathExpression.h"
-#include "K2Node_MatineeController.h"
 #include "K2Node_SpawnActorFromClass.h"
 #include "K2Node_StructOperation.h"
 #include "K2Node_TemporaryVariable.h"
@@ -8363,20 +8361,6 @@ void FBlueprintEditorUtils::FixLevelScriptActorBindings(ALevelScriptActor* Level
 				}
 			}
 		}
-
-		// Find matinee controller nodes and update node name
-		TArray<UK2Node_MatineeController*> MatineeControllers;
-		(*GraphIt)->GetNodesOfClass(MatineeControllers);
-
-		for( TArray<UK2Node_MatineeController*>::TConstIterator NodeIt(MatineeControllers); NodeIt; ++NodeIt )
-		{
-			const UK2Node_MatineeController* MatController = *NodeIt;
-
-			if(MatController->MatineeActor != nullptr)
-			{
-				MatController->MatineeActor->MatineeControllerName = MatController->GetFName();
-			}
-		}
 	}
 }
 
diff --git a/Engine/Source/Editor/UnrealEd/Private/LevelEditorViewport.cpp b/Engine/Source/Editor/UnrealEd/Private/LevelEditorViewport.cpp
--- a/Engine/Source/Editor/UnrealEd/Private/LevelEditorViewport.cpp
+++ b/Engine/Source/Editor/UnrealEd/Private/LevelEditorViewport.cpp
@@ -1995,7 +1995,7 @@ bool FLevelEditorViewportClient::ShouldLockPitch() const
 		return false;
 	}
 	// Else use the standard rules
-	return FEditorViewportClient::ShouldLockPitch() || !ModeTools->GetActiveMode(FBuiltinEditorModes::EM_InterpEdit) ;
+	return FEditorViewportClient::ShouldLockPitch() ;
 }
 
 void FLevelEditorViewportClient::BeginCameraMovement(bool bHasMovement)
diff --git a/Engine/Source/Editor/UnrealEd/Private/PackageAutoSaver.cpp b/Engine/Source/Editor/UnrealEd/Private/PackageAutoSaver.cpp
--- a/Engine/Source/Editor/UnrealEd/Private/PackageAutoSaver.cpp
+++ b/Engine/Source/Editor/UnrealEd/Private/PackageAutoSaver.cpp
@@ -415,7 +415,6 @@ bool FPackageAutoSaver::CanAutoSave() const
 	const bool bDidInteractRecently = (FApp::GetCurrentTime() - LastInteractionTime) < InteractionDelay;
 	const bool bAutosaveEnabled	= LoadingSavingSettings->bAutoSaveEnable && bPackagesNeedAutoSave;
 	const bool bSlowTask = GIsSlowTask;
-	const bool bInterpEditMode = GLevelEditorModeTools().IsModeActive(FBuiltinEditorModes::EM_InterpEdit);
 	const bool bPlayWorldValid = GUnrealEd->PlayWorld != nullptr;
 	const bool bAnyMenusVisible	= FSlateApplication::Get().AnyMenusVisible();
 	const bool bAutomationTesting = GIsAutomationTesting;
@@ -439,7 +438,7 @@ bool FPackageAutoSaver::CanAutoSave() const
 	// query any active editor modes and allow them to prevent autosave
 	const bool bActiveModesAllowAutoSave = GLevelEditorModeTools().CanAutoSave();
 
-	return (bAutosaveEnabled && !bSlowTask && !bInterpEditMode && !bPlayWorldValid && !bAnyMenusVisible && !bAutomationTesting && !bIsInteracting && !GIsDemoMode && bHasGameOrProjectLoaded && !bAreShadersCompiling && !bAreAssetsCompiling && !bAreAnimationsCompressing && !bIsVREditorActive && !bIsSequencerPlaying && bActiveModesAllowAutoSave);
+	return (bAutosaveEnabled && !bSlowTask && !bPlayWorldValid && !bAnyMenusVisible && !bAutomationTesting && !bIsInteracting && !GIsDemoMode && bHasGameOrProjectLoaded && !bAreShadersCompiling && !bAreAssetsCompiling && !bAreAnimationsCompressing && !bIsVREditorActive && !bIsSequencerPlaying && bActiveModesAllowAutoSave);
 }
 
 bool FPackageAutoSaver::DoPackagesNeedAutoSave() const
diff --git a/Engine/Source/Editor/UnrealEd/Private/StaticLightingSystem/StaticLightingSystem.cpp b/Engine/Source/Editor/UnrealEd/Private/StaticLightingSystem/StaticLightingSystem.cpp
--- a/Engine/Source/Editor/UnrealEd/Private/StaticLightingSystem/StaticLightingSystem.cpp
+++ b/Engine/Source/Editor/UnrealEd/Private/StaticLightingSystem/StaticLightingSystem.cpp
@@ -2446,13 +2446,12 @@ bool FStaticLightingSystem::CanAutoApplyLighting() const
 {
 	const bool bAutoApplyEnabled = GetDefault<ULevelEditorMiscSettings>()->bAutoApplyLightingEnable;
 	const bool bSlowTask = GIsSlowTask;
-	const bool bInterpEditMode = GLevelEditorModeTools().IsModeActive( FBuiltinEditorModes::EM_InterpEdit );
 	const bool bPlayWorldValid = GEditor->PlayWorld != nullptr;
 	const bool bAnyMenusVisible = (FSlateApplication::IsInitialized() && FSlateApplication::Get().AnyMenusVisible());
 	//const bool bIsInteratcting = false;// FSlateApplication::Get().GetMouseCaptor().IsValid() || GEditor->IsUserInteracting();
 	const bool bHasGameOrProjectLoaded = FApp::HasProjectName();
 
-	return ( bAutoApplyEnabled && !bSlowTask && !bInterpEditMode && !bPlayWorldValid && !bAnyMenusVisible/* && !bIsInteratcting */&& !GIsDemoMode && bHasGameOrProjectLoaded );
+	return ( bAutoApplyEnabled && !bSlowTask && !bPlayWorldValid && !bAnyMenusVisible/* && !bIsInteratcting */&& !GIsDemoMode && bHasGameOrProjectLoaded );
 }
 
 /**
diff --git a/Engine/Source/Editor/UnrealEd/Private/Tests/AutomationEditorCommon.cpp b/Engine/Source/Editor/UnrealEd/Private/Tests/AutomationEditorCommon.cpp
--- a/Engine/Source/Editor/UnrealEd/Private/Tests/AutomationEditorCommon.cpp
+++ b/Engine/Source/Editor/UnrealEd/Private/Tests/AutomationEditorCommon.cpp
@@ -51,12 +51,6 @@ DEFINE_LOG_CATEGORY_STATIC(LogAutomationEditorCommon, Log, All);
 
 UWorld* FAutomationEditorCommonUtils::CreateNewMap()
 {
-	// Change out of Matinee when opening new map, so we avoid editing data in the old one.
-	if ( GLevelEditorModeTools().IsModeActive(FBuiltinEditorModes::EM_InterpEdit) )
-	{
-		GLevelEditorModeTools().DeactivateMode(FBuiltinEditorModes::EM_InterpEdit);
-	}
-
 	// Also change out of Landscape mode to ensure all references are cleared.
 	if ( GLevelEditorModeTools().IsModeActive(FBuiltinEditorModes::EM_Landscape) )
 	{
diff --git a/Engine/Source/Editor/UnrealEd/Public/Commandlets/AnalyzeReferencedContentStat.h b/Engine/Source/Editor/UnrealEd/Public/Commandlets/AnalyzeReferencedContentStat.h
--- a/Engine/Source/Editor/UnrealEd/Public/Commandlets/AnalyzeReferencedContentStat.h
+++ b/Engine/Source/Editor/UnrealEd/Public/Commandlets/AnalyzeReferencedContentStat.h
@@ -542,7 +542,6 @@ public:
 		enum EAnimReferenceType
 		{
 			ART_SkeletalMeshComponent, // Regular SkeletalMeshComponent - mostly from script
-			ART_Matinee, // From Matinee, cinematic animations
 			ART_Crowd, // From Crowd spawner, expected to be none or very small
 		};
 
diff --git a/Engine/Source/Editor/UnrealEd/Public/EditorViewportClient.h b/Engine/Source/Editor/UnrealEd/Public/EditorViewportClient.h
--- a/Engine/Source/Editor/UnrealEd/Public/EditorViewportClient.h
+++ b/Engine/Source/Editor/UnrealEd/Public/EditorViewportClient.h
@@ -1082,18 +1082,6 @@ public:
 	/** @return true if a mouse button is down and it's movement being tracked for operations inside the viewport */
 	bool IsTracking() const { return bIsTracking; }
 
-	/**
-	 * Allows custom disabling of camera recoil
-	 */
-	UE_DEPRECATED(5.0, "Matinee is no longer part of the editor.")
-	void SetMatineeRecordingWindow(class IMatineeBase* InInterpEd) {}
-
-	/**
-	 * Returns true if camera recoil is currently allowed
-	 */
-	UE_DEPRECATED(5.0, "Matinee is no longer part of the editor.")
-	bool IsMatineeRecordingWindow() const { return false; }
-	
 	EAxisList::Type GetCurrentWidgetAxis() const;
 
 	/** Overrides current cursor. */
diff --git a/Engine/Source/Editor/UnrealEd/Public/FbxImporter.h b/Engine/Source/Editor/UnrealEd/Public/FbxImporter.h
--- a/Engine/Source/Editor/UnrealEd/Public/FbxImporter.h
+++ b/Engine/Source/Editor/UnrealEd/Public/FbxImporter.h
@@ -18,15 +18,11 @@
 class AActor;
 class ACameraActor;
 class ALight;
-class AMatineeActor;
 class Error;
 class FSkeletalMeshImportData;
 class UActorComponent;
 class UAnimSequence;
 class UFbxSkeletalMeshImportData;
-class UInterpGroupInst;
-class UInterpTrackMove;
-class UInterpTrackMoveAxis;
 class ULightComponent;
 class UMaterial;
 class UMaterialInstanceConstant;
@@ -101,7 +97,6 @@ THIRD_PARTY_INCLUDES_END
 class FSkeletalMeshImportData;
 class FSkelMeshOptionalImportData;
 class ASkeletalMeshActor;
-class UInterpTrackMoveAxis;
 struct FbxSceneInfo;
 struct FExistingStaticMeshData;
 
@@ -1590,8 +1585,6 @@ protected:
 	* @param Scale scale factor for this skeleton node
 	*/
 	bool FillAnimSequenceByKey(FbxNode* Node, UAnimSequence* AnimSequence, const char* TakeName, FbxTime& Start, FbxTime& End, bool bIsRoot, FbxVector4 Scale);
-	/*bool CreateMatineeSkeletalAnimation(ASkeletalMeshActor* Actor, UAnimSet* AnimSet);
-	bool CreateMatineeAnimation(FbxNode* Node, AActor* Actor, bool bInvertOrient, bool bAddDirectorTrack);*/
 
 
 	// material
@@ -1760,88 +1753,16 @@ protected:
 	void LoadNodeKeyframeAnimation(FbxNode* NodeToQuery, FFbxCurvesAPI &CurvesAPI);
 	void SetupTransformForNode(FbxNode *Node);
 
-
-	//
-	// for matinee export
-	//
-public:
-	/**
-	 * Retrieves whether there are any unknown camera instances within the FBX document.
-	 */
-	UNREALED_API bool HasUnknownCameras( AMatineeActor* InMatineeActor ) const;
-	
-	/**
-	 * Sets the camera creation flag. Call this function before the import in order to enforce
-	 * the creation of FBX camera instances that do not exist in the current scene.
-	 */
-	inline void SetProcessUnknownCameras(bool bCreateMissingCameras)
-	{
-		bCreateUnknownCameras = bCreateMissingCameras;
-	}
-	
-	/**
-	 * Modifies the Matinee actor with the animations found in the FBX document.
-	 * 
-	 * @return	true, if sucessful
-	 */
-	UNREALED_API bool ImportMatineeSequence(AMatineeActor* InMatineeActor);
-
-
 	/** Create a new asset from the package and objectname and class */
 	static UObject* CreateAssetOfClass(UClass* AssetClass, FString ParentPackageName, FString ObjectName, bool bAllowReplace = false);
 
 	/* Templated function to create an asset with given package and name */
-	template< class T> 
-	static T * CreateAsset( FString ParentPackageName, FString ObjectName, bool bAllowReplace = false )
+	template< class T>
+	static T* CreateAsset(FString ParentPackageName, FString ObjectName, bool bAllowReplace = false)
 	{
 		return (T*)CreateAssetOfClass(T::StaticClass(), ParentPackageName, ObjectName, bAllowReplace);
 	}
 
-protected:
-	bool bCreateUnknownCameras;
-	
-	/**
-	 * Creates a Matinee group for a given actor within a given Matinee actor.
-	 */
-	UInterpGroupInst* CreateMatineeGroup(AMatineeActor* InMatineeActor, AActor* Actor, FString GroupName);
-	/**
-	 * Imports a FBX scene node into a Matinee actor group.
-	 */
-	float ImportMatineeActor(FbxNode* FbxNode, UInterpGroupInst* MatineeGroup);
-
-	/**
-	 * Imports an FBX transform curve into a movement subtrack
-	 */
-	void ImportMoveSubTrack( FbxAnimCurve* FbxCurve, int32 FbxDimension, UInterpTrackMoveAxis* SubTrack, int32 CurveIndex, bool bNegative, FbxAnimCurve* RealCurve, float DefaultVal );
-
-	/**
-	 * Imports a FBX animated element into a Matinee track.
-	 */
-	void ImportMatineeAnimated(FbxAnimCurve* FbxCurve, FInterpCurveVector& Curve, int32 CurveIndex, bool bNegative, FbxAnimCurve* RealCurve, float DefaultVal);
-	/**
-	 * Imports a FBX camera into properties tracks of a Matinee group for a camera actor.
-	 */
-	void ImportCamera(ACameraActor* Actor, UInterpGroupInst* MatineeGroup, FbxCamera* Camera);
-	/**
-	 * Imports a FBX animated value into a property track of a Matinee group.
-	 */
-	void ImportAnimatedProperty(float* Value, const TCHAR* ValueName, UInterpGroupInst* MatineeGroup, const float FbxValue, FbxProperty Property, bool bImportFOV = false, FbxCamera* Camera = NULL );
-	/**
-	 * Check if FBX node has transform animation (translation and rotation, not check scale animation)
-	 */
-	bool IsNodeAnimated(FbxNode* FbxNode, FbxAnimLayer* AnimLayer = NULL);
-
-	/** 
-	 * As movement tracks in Unreal cannot have differing interpolation modes for position & rotation,
-	 * we consolidate the two modes here.
-	 */
-	void ConsolidateMovementTrackInterpModes(UInterpTrackMove* MovementTrack);
-
-	/**
-	 * Get Unreal Interpolation mode from FBX interpolation mode
-	 */
-	EInterpCurveMode GetUnrealInterpMode(FbxAnimCurveKey FbxKey);
-
 	/**
 	 * Fill up and verify bone names for animation 
 	 */
diff --git a/Engine/Source/Editor/UnrealEd/Public/Kismet2/BlueprintEditorUtils.h b/Engine/Source/Editor/UnrealEd/Public/Kismet2/BlueprintEditorUtils.h
--- a/Engine/Source/Editor/UnrealEd/Public/Kismet2/BlueprintEditorUtils.h
+++ b/Engine/Source/Editor/UnrealEd/Public/Kismet2/BlueprintEditorUtils.h
@@ -1099,10 +1099,10 @@ public:
 	static void SetBlueprintPropertyReadOnlyFlag(UBlueprint* Blueprint, const FName& VarName, const bool bVariableReadOnly);
 
 	/**
-	 * Sets the Interp flag on the variable with the specified name to make available to matinee
+	 * Sets the Interp flag on the variable with the specified name to make available to sequencer
 	 *
 	 * @param	VarName				Name of the var to set the flag on
-	 * @param	bInterp	true to make variable available to Matinee, false otherwise
+	 * @param	bInterp	true to make variable available to sequencer, false otherwise
 	 */
 	static void SetInterpFlag(UBlueprint* Blueprint, const FName& VarName, const bool bInterp);
 
diff --git a/Engine/Source/Editor/UnrealEd/Public/LevelEditorViewport.h b/Engine/Source/Editor/UnrealEd/Public/LevelEditorViewport.h
--- a/Engine/Source/Editor/UnrealEd/Public/LevelEditorViewport.h
+++ b/Engine/Source/Editor/UnrealEd/Public/LevelEditorViewport.h
@@ -589,15 +589,6 @@ public:
 		return PreviousActorLocks.CinematicActorLock;
 	}
 
-	/** 
-	 * Set the actor locked to the viewport by Matinee.
-	 */
-	UE_DEPRECATED(4.27, "Matinee is being deprecated, use SetCinematicActorLock instead.")
-	void SetMatineeActorLock(AActor* Actor)
-	{
-		SetCinematicActorLock(Actor);
-	}
-
 	/** 
 	 * Check whether this viewport is locked to the specified actor
 	 */
@@ -606,15 +597,6 @@ public:
 		return ActorLocks.HasActorLocked(Actor);
 	}
 
-	/** 
-	 * Check whether this viewport is locked to display the matinee view
-	 */
-	UE_DEPRECATED(4.27, "Matinee is being deprecated, use IsLockedToCinematic instead.")
-	bool IsLockedToMatinee() const
-	{
-		return IsLockedToCinematic();
-	}
-
 	/**
 	 * Check whether this viewport is locked to display a cinematic camera, like a Sequencer camera.
 	 */
diff --git a/Engine/Source/Runtime/AutomationWorker/Private/AutomationAnalyticParams.h b/Engine/Source/Runtime/AutomationWorker/Private/AutomationAnalyticParams.h
--- a/Engine/Source/Runtime/AutomationWorker/Private/AutomationAnalyticParams.h
+++ b/Engine/Source/Runtime/AutomationWorker/Private/AutomationAnalyticParams.h
@@ -20,7 +20,7 @@ namespace EAutomationAnalyticParam
 	enum Type
 	{
 		MapName,
-		MatineeName,
+		SequenceName,
 		TimeStamp,
 		Platform,
 		Spec,
diff --git a/Engine/Source/Runtime/AutomationWorker/Private/AutomationAnalytics.cpp b/Engine/Source/Runtime/AutomationWorker/Private/AutomationAnalytics.cpp
--- a/Engine/Source/Runtime/AutomationWorker/Private/AutomationAnalytics.cpp
+++ b/Engine/Source/Runtime/AutomationWorker/Private/AutomationAnalytics.cpp
@@ -111,7 +111,7 @@ void FAutomationAnalytics::InititalizeAnalyticParameterNames()
 	AutomationParamNames.SetNum(EAutomationAnalyticParam::NUM_ENGINE_PARAMS);
 
 	AutomationParamNames[EAutomationAnalyticParam::MapName] = TEXT("MapName");
-	AutomationParamNames[EAutomationAnalyticParam::MatineeName] = TEXT("MatineeName");
+	AutomationParamNames[EAutomationAnalyticParam::SequenceName] = TEXT("SequenceName");
 	AutomationParamNames[EAutomationAnalyticParam::TimeStamp] = TEXT("TimeStamp");
 	AutomationParamNames[EAutomationAnalyticParam::Platform] = TEXT("Platform");
 	AutomationParamNames[EAutomationAnalyticParam::Spec] = TEXT("Spec");
@@ -188,7 +188,7 @@ void FAutomationAnalytics::FireEvent_FPSCapture(const FAutomationPerformanceSnap
 		ParamArray.Add(FAnalyticsEventAttribute(GetAutomationParamName(EAutomationAnalyticParam::CL), PerfSnapshot.Changelist));
 		ParamArray.Add(FAnalyticsEventAttribute(GetAutomationParamName(EAutomationAnalyticParam::Spec), MachineSpec));
 		ParamArray.Add(FAnalyticsEventAttribute(GetAutomationParamName(EAutomationAnalyticParam::MapName), PerfSnapshot.MapName));
-		ParamArray.Add(FAnalyticsEventAttribute(GetAutomationParamName(EAutomationAnalyticParam::MatineeName), PerfSnapshot.MatineeName));
+		ParamArray.Add(FAnalyticsEventAttribute(GetAutomationParamName(EAutomationAnalyticParam::SequenceName), PerfSnapshot.SequenceName));
 		ParamArray.Add(FAnalyticsEventAttribute(GetAutomationParamName(EAutomationAnalyticParam::FPS), PerfSnapshot.AverageFPS));
 
 		ParamArray.Add(FAnalyticsEventAttribute(GetAutomationParamName(EAutomationAnalyticParam::BuildConfiguration), PerfSnapshot.BuildConfiguration));
diff --git a/Engine/Source/Runtime/AutomationWorker/Public/IAutomationWorkerModule.h b/Engine/Source/Runtime/AutomationWorker/Public/IAutomationWorkerModule.h
--- a/Engine/Source/Runtime/AutomationWorker/Public/IAutomationWorkerModule.h
+++ b/Engine/Source/Runtime/AutomationWorker/Public/IAutomationWorkerModule.h
@@ -14,7 +14,7 @@ struct FAutomationPerformanceSnapshot
 	FString Changelist;
 	FString BuildConfiguration;
 	FString MapName;
-	FString MatineeName;
+	FString SequenceName;
 	FString AverageFPS;
 	FString AverageFrameTime;
 	FString AverageGameThreadTime;
@@ -31,7 +31,7 @@ struct FAutomationPerformanceSnapshot
 			Changelist + TEXT( "," ) +
 			BuildConfiguration + TEXT( "," ) +
 			MapName + TEXT( "," ) +
-			MatineeName + TEXT( "," ) +
+			SequenceName + TEXT( "," ) +
 			AverageFPS + TEXT( "," ) +
 			AverageFrameTime + TEXT( "," ) +
 			AverageGameThreadTime + TEXT( "," ) +
@@ -52,7 +52,7 @@ struct FAutomationPerformanceSnapshot
 		Changelist = Values[ NextValue++ ];
 		BuildConfiguration = Values[ NextValue++ ];
 		MapName = Values[ NextValue++ ];
-		MatineeName = Values[ NextValue++ ];
+		SequenceName = Values[ NextValue++ ];
 		AverageFPS = Values[ NextValue++ ];
 		AverageFrameTime = Values[ NextValue++ ];
 		AverageGameThreadTime = Values[ NextValue++ ];
diff --git a/Engine/Source/Runtime/CoreUObject/Private/UObject/LinkerLoad.cpp b/Engine/Source/Runtime/CoreUObject/Private/UObject/LinkerLoad.cpp
--- a/Engine/Source/Runtime/CoreUObject/Private/UObject/LinkerLoad.cpp
+++ b/Engine/Source/Runtime/CoreUObject/Private/UObject/LinkerLoad.cpp
@@ -639,22 +639,6 @@ void FLinkerLoad::InvalidateExport(UObject* OldObject)
 	}
 }
 
-FName FLinkerLoad::FindSubobjectRedirectName(const FName& Name, UClass* Class)
-{
-	const TMap<FString, FString>* ValueChanges = FCoreRedirects::GetValueRedirects(ECoreRedirectFlags::Type_Class, Class);
-
-	if (ValueChanges)
-	{
-		const FString* NewInstanceName = ValueChanges->Find(Name.ToString());
-		if (NewInstanceName)
-		{
-			return FName(**NewInstanceName);
-		}
-	}
-
-	return FName();
-}
-
 /**
  * Creates a FLinkerLoad object for async creation. Tick has to be called manually till it returns
  * true in which case the returned linker object has finished the async creation process.
diff --git a/Engine/Source/Runtime/CoreUObject/Public/UObject/LinkerLoad.h b/Engine/Source/Runtime/CoreUObject/Public/UObject/LinkerLoad.h
--- a/Engine/Source/Runtime/CoreUObject/Public/UObject/LinkerLoad.h
+++ b/Engine/Source/Runtime/CoreUObject/Public/UObject/LinkerLoad.h
@@ -714,9 +714,6 @@ public:
 	 */
 	COREUOBJECT_API static void InvalidateExport(UObject* OldObject);
 
-	/** Used by Matinee to fixup component renaming */
-	COREUOBJECT_API static FName FindSubobjectRedirectName(const FName& Name, UClass* Class);
-
 #if WITH_EDITOR
 	COREUOBJECT_API static bool GetPreloadingEnabled();
 	COREUOBJECT_API static void SetPreloadingEnabled(bool bEnabled);
diff --git a/Engine/Source/Runtime/Engine/Classes/Camera/CameraActor.h b/Engine/Source/Runtime/Engine/Classes/Camera/CameraActor.h
--- a/Engine/Source/Runtime/Engine/Classes/Camera/CameraActor.h
+++ b/Engine/Source/Runtime/Engine/Classes/Camera/CameraActor.h
@@ -9,8 +9,6 @@
 
 #include "CameraActor.generated.h"
 
-class UCameraAnim;
-
 /** 
  * A CameraActor is a camera viewpoint that can be placed in a level.
  */
@@ -35,9 +33,6 @@ private:
 	TObjectPtr<class USceneComponent> SceneComponent;
 public:
 
-	/** If this CameraActor is being used to preview a CameraAnim in the editor, this is the anim being previewed. */
-	TWeakObjectPtr<class UCameraAnim> PreviewedCameraAnim;
-
 	/** Returns index of the player for whom we auto-activate, or INDEX_NONE (-1) if disabled. */
 	UFUNCTION(BlueprintCallable, Category="AutoPlayerActivation")
 	int32 GetAutoActivatePlayerIndex() const;
@@ -65,8 +60,6 @@ public:
 
 #if WITH_EDITOR
 	virtual void PostLoadSubobjects(FObjectInstancingGraph* OuterInstanceGraph) override;
-
-	virtual void PostEditChangeProperty(FPropertyChangedEvent& PropertyChangedEvent) override;
 #endif
 
 	virtual class USceneComponent* GetDefaultAttachComponent() const override;
diff --git a/Engine/Source/Runtime/Engine/Classes/Camera/PlayerCameraManager.h b/Engine/Source/Runtime/Engine/Classes/Camera/PlayerCameraManager.h
--- a/Engine/Source/Runtime/Engine/Classes/Camera/PlayerCameraManager.h
+++ b/Engine/Source/Runtime/Engine/Classes/Camera/PlayerCameraManager.h
@@ -362,16 +362,6 @@ protected:
 	TObjectPtr<class UCameraModifier_CameraShake> CachedCameraShakeMod;
 
 
-	////////////////////////
-	// CameraAnim support
-	////////////////////////
-	
-	static const int32 MAX_ACTIVE_CAMERA_ANIMS = 8;
-
-	/** Internal pool of camera anim instance objects available for playing camera animations. Defines the max number of camera anims that can play simultaneously. */
-	UPROPERTY(transient)
-	TObjectPtr<class UCameraAnimInst> AnimInstPool[8];    /*MAX_ACTIVE_CAMERA_ANIMS @fixme constant */
-
 	/** Internal list of active post process effects. Parallel array to PostProcessBlendCacheWeights. */
 	UPROPERTY(transient)
 	TArray<struct FPostProcessSettings> PostProcessBlendCache;
@@ -387,18 +377,10 @@ protected:
 	void ClearCachedPPBlends();
 
 public:
-	/** Array of camera anim instances that are currently playing and in-use */
-	UPROPERTY(transient)
-	TArray<TObjectPtr<class UCameraAnimInst>> ActiveAnims;
-
 	/** Returns active post process info. */
 	void GetCachedPostProcessBlends(TArray<struct FPostProcessSettings> const*& OutPPSettings, TArray<float> const*& OutBlendWeigthts) const;
 	
 protected:
-	/** Array of camera anim instances that are not playing and available to be used. */
-	UPROPERTY(transient)
-	TArray<TObjectPtr<class UCameraAnimInst>> FreeAnims;
-
 	/** Internal. Receives the output of individual camera animations. */
 	UPROPERTY(transient)
 	TObjectPtr<class ACameraActor> AnimCameraActor;
@@ -844,7 +826,7 @@ public:
 	virtual void StopAllCameraShakesFromSource(UCameraShakeSourceComponent* SourceComponent, bool bImmediately = true);
 
 	//
-	//  CameraAnim fades.
+	//  Camera fades.
 	//
 
 	/** 
@@ -873,70 +855,10 @@ public:
 	virtual void SetManualCameraFade(float InFadeAmount, FLinearColor Color, bool bInFadeAudio);
 
 
-	//
-	//  CameraAnim support.
-	//
-	
-	/**
-	 * Play the indicated CameraAnim on this camera.
-	 * 
-	 * @param Anim The animation that should play on this instance.
-	 * @param Rate				How fast to play the animation. 1.0 is normal.
-	 * @param Scale				How "intense" to play the animation. 1.0 is normal.
-	 * @param BlendInTime		Time to linearly ramp in.
-	 * @param BlendOutTime		Time to linearly ramp out.
-	 * @param bLoop				True to loop the animation if it hits the end.
-	 * @param bRandomStartTime	Whether or not to choose a random time to start playing. Useful with bLoop=true and a duration to randomize things like shakes.
-	 * @param Duration			Optional total playtime for this animation, including blends. 0 means to use animations natural duration, or infinite if looping.
-	 * @param PlaySpace			Which space to play the animation in.
-	 * @param UserPlaySpaceRot  Custom play space, used when PlaySpace is UserDefined.
-	 * @return The CameraAnim instance, which can be stored to manipulate/stop the anim after the fact.
-	 */
-	UFUNCTION(BlueprintCallable, Category = "Camera Animation")
-	virtual class UCameraAnimInst* PlayCameraAnim(class UCameraAnim* Anim, float Rate = 1.f, float Scale = 1.f, float BlendInTime = 0.f, float BlendOutTime = 0.f, bool bLoop = false, bool bRandomStartTime = false, float Duration = 0.f, ECameraShakePlaySpace PlaySpace = ECameraShakePlaySpace::CameraLocal, FRotator UserPlaySpaceRot = FRotator::ZeroRotator);
-	
-	/**
-	 * Stop playing all instances of the indicated CameraAnim.
-	 * @param bImmediate	True to stop it right now and ignore blend out, false to let it blend out as indicated.
-	 */
-	UFUNCTION(BlueprintCallable, Category = "Camera Animation")
-	virtual void StopAllInstancesOfCameraAnim(class UCameraAnim* Anim, bool bImmediate = false);
-	
-	/** 
-	 * Stops the given CameraAnimInst from playing.  The given pointer should be considered invalid after this. 
-	 * @param bImmediate	True to stop it right now and ignore blend out, false to let it blend out as indicated.
-	 */
-	UFUNCTION(BlueprintCallable, Category = "Camera Animation")
-	virtual void StopCameraAnimInst(class UCameraAnimInst* AnimInst, bool bImmediate = false);
-	
-	/**
-	 * Stop playing all CameraAnims on this CameraManager.
-	 * @param bImmediate	True to stop it right now and ignore blend out, false to let it blend out as indicated.
-	 */
-	UFUNCTION(BlueprintCallable, Category = "Camera Animation")
-	virtual void StopAllCameraAnims(bool bImmediate = false);
-
-	/** Returns first existing instance of the specified camera anim, or NULL if none exists. */
-	class UCameraAnimInst* FindInstanceOfCameraAnim(class UCameraAnim const* Anim) const;
-
 	/** Sets the bGameCameraCutThisFrame flag to true (indicating we did a camera cut this frame; useful for game code to call, e.g., when performing a teleport that should be seamless) */
 	UFUNCTION(BlueprintCallable, Category = "Camera")
 	virtual void SetGameCameraCutThisFrame() { bGameCameraCutThisFrame = true; }
 
-protected:
-	/** Gets specified temporary CameraActor ready to update the specified Anim. */
-	void InitTempCameraActor(class ACameraActor* CamActor, class UCameraAnimInst const* AnimInstToInitFor) const;
-	void ApplyAnimToCamera(class ACameraActor const* AnimatedCamActor, class UCameraAnimInst const* AnimInst, FMinimalViewInfo& InOutPOV);
-
-	/** Returns an available CameraAnimInst, or NULL if no more are available. */
-	class UCameraAnimInst* AllocCameraAnimInst();
-
-	/** 
-	 * Frees an allocated CameraAnimInst for future use. 
-	 * @param Inst - Instance to free. 
-	 */
-	void ReleaseCameraAnimInst(class UCameraAnimInst* Inst);
-
 public:
 	/** 
 	 * Limit the player's view pitch. 
diff --git a/Engine/Source/Runtime/Engine/Classes/Engine/Engine.h b/Engine/Source/Runtime/Engine/Classes/Engine/Engine.h
--- a/Engine/Source/Runtime/Engine/Classes/Engine/Engine.h
+++ b/Engine/Source/Runtime/Engine/Classes/Engine/Engine.h
@@ -2543,15 +2543,15 @@ public:
 #if !UE_BUILD_SHIPPING
 	/** 
 	 * Capture screenshots and performance metrics
-	 * @param EventTime time of the Matinee event
+	 * @param EventTime time of the Sequencer event
 	 */
-	void PerformanceCapture(UWorld* World, const FString& MapName, const FString& MatineeName, float EventTime);
+	void PerformanceCapture(UWorld* World, const FString& MapName, const FString& SequenceName, float EventTime);
 
 	/**
 	 * Logs performance capture for use in automation analytics
-	 * @param EventTime time of the Matinee event
+	 * @param EventTime time of the Sequencer event
 	 */
-	void LogPerformanceCapture(UWorld* World, const FString& MapName, const FString& MatineeName, float EventTime);
+	void LogPerformanceCapture(UWorld* World, const FString& MapName, const FString& SequenceName, float EventTime);
 #endif	// UE_BUILD_SHIPPING
 
 	/**
diff --git a/Engine/Source/Runtime/Engine/Classes/Engine/EngineTypes.h b/Engine/Source/Runtime/Engine/Classes/Engine/EngineTypes.h
--- a/Engine/Source/Runtime/Engine/Classes/Engine/EngineTypes.h
+++ b/Engine/Source/Runtime/Engine/Classes/Engine/EngineTypes.h
@@ -2105,43 +2105,6 @@ enum EShadowMapFlags
 	SMF_Streamed		= 0x00000001
 };
 
-/** Reference to a specific material in a PrimitiveComponent, used by Matinee */
-USTRUCT()
-struct FPrimitiveMaterialRef
-{
-	GENERATED_BODY()
-
-	/** Material is on a primitive component */
-	UPROPERTY()
-	TObjectPtr<class UPrimitiveComponent> Primitive;
-
-	/** Material is on a decal component */
-	UPROPERTY()
-	TObjectPtr<class UDecalComponent> Decal;
-
-	/** Index into the material on the components data */
-	UPROPERTY()
-	int32 ElementIndex;
-
-	FPrimitiveMaterialRef()
-		: Primitive(nullptr)
-		, Decal(nullptr)
-		, ElementIndex(0)
-	{ }
-
-	FPrimitiveMaterialRef(UPrimitiveComponent* InPrimitive, int32 InElementIndex)
-		: Primitive(InPrimitive)
-		, Decal(nullptr)
-		, ElementIndex(InElementIndex)
-	{ 	}
-
-	FPrimitiveMaterialRef(UDecalComponent* InDecal, int32 InElementIndex)
-		: Primitive(nullptr)
-		, Decal(InDecal)
-		, ElementIndex(InElementIndex)
-	{ 	}
-};
-
 /** Whether to teleport physics body or not */
 UENUM()
 enum class ETeleportType : uint8
@@ -2179,21 +2142,6 @@ struct ENGINE_API FMTDResult
 	}
 };
 
-/** Struct used for passing information from Matinee to an Actor for blending animations during a sequence. */
-USTRUCT()
-struct FAnimSlotInfo
-{
-	GENERATED_BODY()
-
-	/** Name of slot that we want to play the animtion in. */
-	UPROPERTY()
-	FName SlotName;
-
-	/** Strength of each Channel within this Slot. Channel indexs are determined by track order in Matinee. */
-	UPROPERTY()
-	TArray<float> ChannelWeights;
-};
-
 /** Used to indicate each slot name and how many channels they have. */
 USTRUCT()
 struct FAnimSlotDesc
diff --git a/Engine/Source/Runtime/Engine/Classes/GameFramework/PlayerController.h b/Engine/Source/Runtime/Engine/Classes/GameFramework/PlayerController.h
--- a/Engine/Source/Runtime/Engine/Classes/GameFramework/PlayerController.h
+++ b/Engine/Source/Runtime/Engine/Classes/GameFramework/PlayerController.h
@@ -37,7 +37,6 @@ class FDebugDisplayInfo;
 class UActorChannel;
 class UCheatManager;
 class UGameViewportClient;
-class UInterpTrackInstDirector;
 class ULocalMessage;
 class UNetConnection;
 class UPlayer;
@@ -239,10 +238,6 @@ public:
 	UPROPERTY()
 	TObjectPtr<APawn> AcknowledgedPawn;
 
-	/** Director track that's currently possessing this player controller, or none if not possessed. */
-	UPROPERTY(Transient)
-	TObjectPtr<UInterpTrackInstDirector> ControllingDirTrackInst;
-
 	/** Heads up display associated with this PlayerController. */
 	UPROPERTY()
 	TObjectPtr<AHUD> MyHUD;
@@ -1021,20 +1016,6 @@ public:
 	UFUNCTION(Reliable, Client)
 	void ClientMessage(const FString& S, FName Type = NAME_None, float MsgLifeTime = 0.f);
 
-	/** Play the indicated CameraAnim on this camera.
-	 * @param AnimToPlay - Camera animation to play
-	 * @param Scale - "Intensity" scalar.  This is the scale at which the anim was first played.
-	 * @param Rate -  Multiplier for playback rate.  1.0 = normal.
-	 * @param BlendInTime - Time to interpolate in from zero, for smooth starts
-	 * @param BlendOutTime - Time to interpolate out to zero, for smooth finishes
-	 * @param bLoop - True if the animation should loop, false otherwise
-	 * @param bRandomStartTime - Whether or not to choose a random time to start playing.  Only really makes sense for bLoop = true
-	 * @param Space - Animation play area
-	 * @param CustomPlaySpace - Matrix used when Space = CAPS_UserDefined
-	 */
-	UFUNCTION(unreliable, client, BlueprintCallable, Category="Game|Feedback")
-	void ClientPlayCameraAnim(class UCameraAnim* AnimToPlay, float Scale=1.f, float Rate=1.f, float BlendInTime=0.f, float BlendOutTime=0.f, bool bLoop=false, bool bRandomStartTime=false, ECameraShakePlaySpace Space=ECameraShakePlaySpace::CameraLocal, FRotator CustomPlaySpace=FRotator::ZeroRotator);
-
 	/** 
 	 * Play Camera Shake 
 	 * @param Shake - Camera shake animation to play
@@ -1200,10 +1181,6 @@ public:
 	UFUNCTION(reliable, client, BlueprintCallable, Category="Game|Feedback")
 	virtual void ClientClearCameraLensEffects();
 
-	/** Stop camera animation on client. */
-	UFUNCTION(Reliable, Client)
-	void ClientStopCameraAnim(class UCameraAnim* AnimToStop);
-
 	/** Stop camera shake on client.  */
 	UFUNCTION(reliable, client, BlueprintCallable, Category="Game|Feedback")
 	void ClientStopCameraShake(TSubclassOf<class UCameraShakeBase> Shake, bool bImmediately = true);
@@ -2243,4 +2220,4 @@ private:
 	FInputCmdBuffer InputBuffer;
 	FClientFrameInfo ClientFrameInfo;
 	FServerFrameInfo ServerFrameInfo;
-};
\ No newline at end of file
+};
diff --git a/Engine/Source/Runtime/Engine/Classes/GameFramework/WorldSettings.h b/Engine/Source/Runtime/Engine/Classes/GameFramework/WorldSettings.h
--- a/Engine/Source/Runtime/Engine/Classes/GameFramework/WorldSettings.h
+++ b/Engine/Source/Runtime/Engine/Classes/GameFramework/WorldSettings.h
@@ -724,10 +724,10 @@ public:
 	UPROPERTY(transient, replicated)
 	float TimeDilation;
 
-	// Additional time dilation used by Matinee (or Sequencer) slomo track.  Transient because this is often 
+	// Additional time dilation used by Sequencer slomo track.  Transient because this is often 
 	// temporarily modified by the editor when previewing slow motion effects, yet we don't want it saved or loaded from level packages.
 	UPROPERTY(transient, replicated)
-	float MatineeTimeDilation;
+	float CinematicTimeDilation;
 
 	// Additional TimeDilation used to control demo playback speed
 	UPROPERTY(transient)
@@ -817,7 +817,7 @@ public:
 
 	virtual float GetEffectiveTimeDilation() const
 	{
-		return TimeDilation * MatineeTimeDilation * DemoPlayTimeDilation;
+		return TimeDilation * CinematicTimeDilation * DemoPlayTimeDilation;
 	}
 
 	/**
diff --git a/Engine/Source/Runtime/Engine/Private/Camera/CameraActor.cpp b/Engine/Source/Runtime/Engine/Private/Camera/CameraActor.cpp
--- a/Engine/Source/Runtime/Engine/Private/Camera/CameraActor.cpp
+++ b/Engine/Source/Runtime/Engine/Private/Camera/CameraActor.cpp
@@ -4,7 +4,6 @@
 #include "Engine/World.h"
 #include "Kismet/GameplayStatics.h"
 #include "Camera/CameraComponent.h"
-#include "Camera/CameraAnim.h"
 
 #define LOCTEXT_NAMESPACE "CameraActor"
 
@@ -77,18 +76,6 @@ void ACameraActor::PostLoadSubobjects(FObjectInstancingGraph* OuterInstanceGraph
 	}
 }
 
-void ACameraActor::PostEditChangeProperty(FPropertyChangedEvent& PropertyChangedEvent)
-{
-	Super::PostEditChangeProperty(PropertyChangedEvent);
-
-	if (PreviewedCameraAnim.IsValid() && CameraComponent)
-	{
-		PreviewedCameraAnim->BaseFOV = CameraComponent->FieldOfView;
-		PreviewedCameraAnim->BasePostProcessSettings = CameraComponent->PostProcessSettings;
-		PreviewedCameraAnim->BasePostProcessBlendWeight = CameraComponent->PostProcessBlendWeight;
-	}
-}
-
 #endif
 
 USceneComponent* ACameraActor::GetDefaultAttachComponent() const
diff --git a/Engine/Source/Runtime/Engine/Private/CharacterMovementComponentAsync.cpp b/Engine/Source/Runtime/Engine/Private/CharacterMovementComponentAsync.cpp
--- a/Engine/Source/Runtime/Engine/Private/CharacterMovementComponentAsync.cpp
+++ b/Engine/Source/Runtime/Engine/Private/CharacterMovementComponentAsync.cpp
@@ -343,7 +343,7 @@ void FCharacterMovementComponentAsyncInput::PerformMovement(float DeltaSeconds,
 		// Update character state based on change from movement
 		UpdateCharacterStateAfterMovement(DeltaSeconds, Output);
 
-		if ((bAllowPhysicsRotationDuringAnimRootMotion || !RootMotion.bHasAnimRootMotion)/* && !CharacterOwner->IsMatineeControlled()*/)
+		if ((bAllowPhysicsRotationDuringAnimRootMotion || !RootMotion.bHasAnimRootMotion))
 		{
 			PhysicsRotation(DeltaSeconds, Output);
 		}
diff --git a/Engine/Source/Runtime/Engine/Private/InterpCurveEdSetup.cpp b/Engine/Source/Runtime/Engine/Private/InterpCurveEdSetup.cpp
--- a/Engine/Source/Runtime/Engine/Private/InterpCurveEdSetup.cpp
+++ b/Engine/Source/Runtime/Engine/Private/InterpCurveEdSetup.cpp
@@ -8,7 +8,6 @@
 #include "Templates/Casts.h"
 #include "Distributions/DistributionFloat.h"
 #include "Distributions/DistributionVector.h"
-#include "Matinee/InterpTrack.h"
 
 
 void UInterpCurveEdSetup::PostLoad()
@@ -51,12 +50,6 @@ FCurveEdInterface* UInterpCurveEdSetup::GetCurveEdInterfacePointer(const FCurveE
 		return VectorDist;
 	}
 
-	UInterpTrack* InterpTrack = Cast<UInterpTrack>(Entry.CurveObject);
-	if(InterpTrack)
-	{
-		return InterpTrack;
-	}
-
 	return NULL;
 }
 
diff --git a/Engine/Source/Runtime/Engine/Private/PlayerCameraManager.cpp b/Engine/Source/Runtime/Engine/Private/PlayerCameraManager.cpp
--- a/Engine/Source/Runtime/Engine/Private/PlayerCameraManager.cpp
+++ b/Engine/Source/Runtime/Engine/Private/PlayerCameraManager.cpp
@@ -14,8 +14,6 @@
 #include "GameFramework/WorldSettings.h"
 #include "AudioDevice.h"
 #include "Particles/EmitterCameraLensEffectBase.h"
-#include "Camera/CameraAnim.h"
-#include "Camera/CameraAnimInst.h"
 #include "Camera/CameraComponent.h"
 #include "Camera/CameraModifier.h"
 #include "Camera/CameraModifier_CameraShake.h"
@@ -297,45 +295,6 @@ void APlayerCameraManager::ApplyCameraModifiers(float DeltaTime, FMinimalViewInf
 			}
 		}
 	}
-
-	// Now apply CameraAnims
-	// these essentially behave as the highest-pri modifier.
-	for (int32 Idx = 0; Idx < ActiveAnims.Num(); ++Idx)
-	{
-		UCameraAnimInst* const AnimInst = ActiveAnims[Idx];
-
-		if (AnimCameraActor && !AnimInst->bFinished)
-		{
-			// clear out animated camera actor
-			InitTempCameraActor(AnimCameraActor, AnimInst);
-
-			// evaluate the animation at the new time
-			AnimInst->AdvanceAnim(DeltaTime, false);
-
-			// Add weighted properties to the accumulator actor
-			if (AnimInst->CurrentBlendWeight > 0.f)
-			{
-				ApplyAnimToCamera(AnimCameraActor, AnimInst, InOutPOV);
-			}
-		}
-
-		// changes to this are good for a single update, so reset this to 1.f after processing
-		AnimInst->TransientScaleModifier = 1.f;
-
-		// handle animations that have finished
-		if (AnimInst->bFinished)
-		{
-			ReleaseCameraAnimInst(AnimInst);
-			Idx--;		// we removed this from the ActiveAnims array
-		}
-	}
-
-	// need to zero this when we are done with it.  playing another animation
-	// will calc a new InitialTM for the move track instance based on these values.
-	if (AnimCameraActor)
-	{
-		AnimCameraActor->TeleportTo(FVector::ZeroVector, FRotator::ZeroRotator);
-	}
 }
 
 void APlayerCameraManager::AddCachedPPBlend(struct FPostProcessSettings& PPSettings, float BlendWeight)
@@ -357,135 +316,6 @@ void APlayerCameraManager::GetCachedPostProcessBlends(TArray<FPostProcessSetting
 	OutBlendWeigthts = &PostProcessBlendCacheWeights;
 }
 
-void APlayerCameraManager::ApplyAnimToCamera(ACameraActor const* AnimatedCamActor, UCameraAnimInst const* AnimInst, FMinimalViewInfo& InOutPOV)
-{
-	AnimInst->ApplyToView(InOutPOV);
-
-	// postprocess
-	if (AnimatedCamActor->GetCameraComponent()->PostProcessBlendWeight > 0.f)
-	{
-		AddCachedPPBlend(AnimatedCamActor->GetCameraComponent()->PostProcessSettings, AnimatedCamActor->GetCameraComponent()->PostProcessBlendWeight * AnimInst->CurrentBlendWeight);
-	}
-}
-
-UCameraAnimInst* APlayerCameraManager::AllocCameraAnimInst()
-{
-	check(IsInGameThread());
-
-	UCameraAnimInst* FreeAnim = (FreeAnims.Num() > 0) ? ToRawPtr(FreeAnims.Pop()) : NULL;
-	if (FreeAnim)
-	{
-		UCameraAnimInst const* const DefaultInst = GetDefault<UCameraAnimInst>();
-
-		ActiveAnims.Push(FreeAnim);
-
-		// reset some defaults
-		if (DefaultInst)
-		{
-			FreeAnim->TransientScaleModifier = DefaultInst->TransientScaleModifier;
-			FreeAnim->PlaySpace = DefaultInst->PlaySpace;
-		}
-
-		// make sure any previous anim has been terminated correctly
-		check( (FreeAnim->MoveTrack == NULL) && (FreeAnim->MoveInst == NULL) );
-	}
-
-	return FreeAnim;
-}
-
-
-void APlayerCameraManager::ReleaseCameraAnimInst(UCameraAnimInst* Inst)
-{	
-	ActiveAnims.Remove(Inst);
-	FreeAnims.Push(Inst);
-}
-
-
-UCameraAnimInst* APlayerCameraManager::FindInstanceOfCameraAnim(UCameraAnim const* Anim) const
-{
-	int32 const NumActiveAnims = ActiveAnims.Num();
-	for (int32 Idx=0; Idx<NumActiveAnims; Idx++)
-	{
-		if (ActiveAnims[Idx]->CamAnim == Anim)
-		{
-			return ActiveAnims[Idx];
-		}
-	}
-
-	return NULL;
-}
-
-UCameraAnimInst* APlayerCameraManager::PlayCameraAnim(UCameraAnim* Anim, float Rate, float Scale, float BlendInTime, float BlendOutTime, bool bLoop, bool bRandomStartTime, float Duration, ECameraShakePlaySpace PlaySpace, FRotator UserPlaySpaceRot)
-{
-	// get a new instance and play it
-	if (AnimCameraActor != NULL)
-	{
-		UCameraAnimInst* const Inst = AllocCameraAnimInst();
-		if (Inst)
-		{
-			if (Anim != nullptr && !Anim->bRelativeToInitialFOV)
-			{
-				Inst->InitialFOV = ViewTarget.POV.FOV;
-			}
-			Inst->LastCameraLoc = FVector::ZeroVector;		// clear LastCameraLoc
-			Inst->Play(Anim, AnimCameraActor, Rate, Scale, BlendInTime, BlendOutTime, bLoop, bRandomStartTime, Duration);
-			Inst->SetPlaySpace(PlaySpace, UserPlaySpaceRot);
-			return Inst;
-		}
-	}
-
-	return NULL;
-}
-
-void APlayerCameraManager::StopAllInstancesOfCameraAnim(UCameraAnim* Anim, bool bImmediate)
-{
-	// find cameraaniminst for this.
-	for (int32 Idx=0; Idx<ActiveAnims.Num(); ++Idx)
-	{
-		if (ActiveAnims[Idx]->CamAnim == Anim)
-		{
-			ActiveAnims[Idx]->Stop(bImmediate);
-		}
-	}
-}
-
-void APlayerCameraManager::StopAllCameraAnims(bool bImmediate)
-{
-	for (int32 Idx=0; Idx<ActiveAnims.Num(); ++Idx)
-	{
-		ActiveAnims[Idx]->Stop(bImmediate);
-	}
-}
-
-void APlayerCameraManager::StopCameraAnimInst(class UCameraAnimInst* AnimInst, bool bImmediate)
-{
-	if (AnimInst != NULL)
-	{
-		AnimInst->Stop(bImmediate);
-	}
-}
-
-
-void APlayerCameraManager::InitTempCameraActor(ACameraActor* CamActor, UCameraAnimInst const* AnimInstToInitFor) const
-{
-	if (CamActor)
-	{
-		CamActor->TeleportTo(FVector::ZeroVector, FRotator::ZeroRotator);
-
-		if (AnimInstToInitFor)
-		{
-			ACameraActor const* const DefaultCamActor = GetDefault<ACameraActor>();
-			if (DefaultCamActor)
-			{
-				CamActor->GetCameraComponent()->AspectRatio = DefaultCamActor->GetCameraComponent()->AspectRatio;
-				CamActor->GetCameraComponent()->FieldOfView = AnimInstToInitFor->CamAnim->BaseFOV;
-				CamActor->GetCameraComponent()->PostProcessSettings = AnimInstToInitFor->CamAnim->BasePostProcessSettings;
-				CamActor->GetCameraComponent()->PostProcessBlendWeight = AnimInstToInitFor->CamAnim->BasePostProcessBlendWeight;
-			}
-		}
-	}
-}
-
 void APlayerCameraManager::UpdateViewTargetInternal(FTViewTarget& OutVT, float DeltaTime)
 {
 	if (OutVT.Target)
@@ -815,28 +645,10 @@ void APlayerCameraManager::PostInitializeComponents()
 			}
 		}
 	}
-
- 	// create CameraAnimInsts in pool
-	for (int32 Idx=0; Idx<MAX_ACTIVE_CAMERA_ANIMS; ++Idx)
-	{
-		AnimInstPool[Idx] = NewObject<UCameraAnimInst>(this);
-
-		// add everything to the free list initially
-		FreeAnims.Add(AnimInstPool[Idx]);
-	}
-
-	// spawn the temp CameraActor used for updating CameraAnims
-	FActorSpawnParameters SpawnInfo;
-	SpawnInfo.Owner = this;
-	SpawnInfo.Instigator = GetInstigator();
-	SpawnInfo.SpawnCollisionHandlingOverride = ESpawnActorCollisionHandlingMethod::AlwaysSpawn;
-	SpawnInfo.ObjectFlags |= RF_Transient;	// We never want to save these temp actors into a map
-	AnimCameraActor = GetWorld()->SpawnActor<ACameraActor>(SpawnInfo);
 }
 
 void APlayerCameraManager::EndPlay(const EEndPlayReason::Type EndPlayReason)
 {
-	ActiveAnims.Empty();
 	ModifierList.Empty();
 	CleanUpAnimCamera(EndPlayReason == EEndPlayReason::Destroyed);
 	Super::EndPlay(EndPlayReason);
diff --git a/Engine/Source/Runtime/Engine/Private/PlayerController.cpp b/Engine/Source/Runtime/Engine/Private/PlayerController.cpp
--- a/Engine/Source/Runtime/Engine/Private/PlayerController.cpp
+++ b/Engine/Source/Runtime/Engine/Private/PlayerController.cpp
@@ -4599,25 +4599,6 @@ void APlayerController::ClientStopCameraShakesFromSource(class UCameraShakeSourc
 	}
 }
 
-void APlayerController::ClientPlayCameraAnim_Implementation( UCameraAnim* AnimToPlay, float Scale, float Rate,
-						float BlendInTime, float BlendOutTime, bool bLoop,
-						bool bRandomStartTime, ECameraShakePlaySpace Space, FRotator CustomPlaySpace )
-{
-	if (PlayerCameraManager != NULL)
-	{
-		PlayerCameraManager->PlayCameraAnim(AnimToPlay, Rate, Scale, BlendInTime, BlendOutTime, bLoop, bRandomStartTime, 0.f, Space, CustomPlaySpace);
-	}
-}
-
-void APlayerController::ClientStopCameraAnim_Implementation(UCameraAnim* AnimToStop)
-{
-	if (PlayerCameraManager != NULL)
-	{
-		PlayerCameraManager->StopAllInstancesOfCameraAnim(AnimToStop);
-	}
-}
-
-
 void APlayerController::ClientSpawnGenericCameraLensEffect_Implementation(TSubclassOf<class AActor> LensEffectEmitterClass)
 {
 	if (PlayerCameraManager != NULL)
diff --git a/Engine/Source/Runtime/Engine/Private/Replication.cpp b/Engine/Source/Runtime/Engine/Private/Replication.cpp
--- a/Engine/Source/Runtime/Engine/Private/Replication.cpp
+++ b/Engine/Source/Runtime/Engine/Private/Replication.cpp
@@ -10,66 +10,12 @@
 #include "GameFramework/PlayerController.h"
 #include "Engine/LocalPlayer.h"
 #include "UnrealEngine.h"
-#include "Matinee/MatineeInterface.h"
-#include "Matinee/MatineeActor.h"
 #include "Net/UnrealNetwork.h"
 
-void AMatineeActor::GetLifetimeReplicatedProps( TArray< FLifetimeProperty > & OutLifetimeProps ) const
-{
-	Super::GetLifetimeReplicatedProps( OutLifetimeProps );
-
-	DISABLE_REPLICATED_PROPERTY( AMatineeActor, bPendingStop );
-
-	DOREPLIFETIME_CONDITION( AMatineeActor, MatineeData,		COND_InitialOnly );
-	DOREPLIFETIME_CONDITION( AMatineeActor, GroupActorInfos,	COND_InitialOnly );
-
-	DOREPLIFETIME( AMatineeActor, bIsPlaying );
-	DOREPLIFETIME( AMatineeActor, bLooping );
-	DOREPLIFETIME( AMatineeActor, bReversePlayback );
-	DOREPLIFETIME( AMatineeActor, bPaused );
-	DOREPLIFETIME( AMatineeActor, PlayRate );
-	DOREPLIFETIME( AMatineeActor, InterpPosition );
-	DOREPLIFETIME( AMatineeActor, ReplicationForceIsPlaying );
-}
-
 //
 // Static variables for networking.
 //
 
-static UInterpData*		SavedInterpData;
-static bool				SavedbIsPlaying;
-static bool				SavedbReversePlayback;
-static float			SavedPosition;
-static uint8			SavedReplicationForceIsPlaying;
-
-float AMatineeActor::GetNetPriority(const FVector& ViewPos, const FVector& ViewDir, AActor* Viewer, AActor* ViewTarget, UActorChannel* InChannel, float Time, bool bLowBandwidth)
-{
-	float Result = Super::GetNetPriority(ViewPos, ViewDir, Viewer, ViewTarget, InChannel, Time, bLowBandwidth);
-	// attempt to replicate MatineeActors approximately in the order that they were spawned to reduce ordering issues
-	// when LDs make multiple matinees affect the same target(s)
-	// not great, but without a full dependency setup this is the best we can do
-	if (InChannel == NULL)
-	{
-		Result += 1.0f - CreationTime / GetWorld()->TimeSeconds;
-	}
-	return Result;
-}
-
-void AMatineeActor::PreNetReceive()
-{
-	Super::PreNetReceive();
-
-	SavedInterpData = MatineeData;
-
-	if (MatineeData != NULL)
-	{
-		SavedbIsPlaying = bIsPlaying;
-		SavedPosition = InterpPosition;
-		SavedbReversePlayback = bReversePlayback;
-		SavedReplicationForceIsPlaying = ReplicationForceIsPlaying;
-	}
-}
-
 /** @hack: saves and restores fade state for a PC when it goes out of scope
  * used for fade track hack below */
 struct FSavedFadeState
@@ -91,148 +37,3 @@ private:
 	float FadeTimeRemaining;
 };
 
-void AMatineeActor::PostNetReceive()
-{
-	Super::PostNetReceive();
-
-	if (MatineeData != NULL)
-	{
-		TArray<AActor*> ControlledActors;
-		// Build a list of actors controlled by this matinee actor
-		for( int32 InfoIndex = 0; InfoIndex < GroupActorInfos.Num(); ++InfoIndex )
-		{
-			const FInterpGroupActorInfo& Info = GroupActorInfos[ InfoIndex ];
-
-			for (int32 ActorIndex = 0; ActorIndex < Info.Actors.Num(); ++ActorIndex )
-			{
-				AActor* Actor = Info.Actors[ ActorIndex ];
-				if (Actor != NULL)
-				{
-					ControlledActors.Add( Actor );
-				}
-			}
-		}
-
-		// if we just received the matinee action, set 'saved' values to default so we make sure to apply previously received values
-		if (SavedInterpData == NULL)
-		{
-			AMatineeActor* Default = GetClass()->GetDefaultObject<AMatineeActor>();
-			SavedbIsPlaying = Default->bIsPlaying;
-			SavedPosition = Default->InterpPosition;
-			SavedbReversePlayback = Default->bReversePlayback;
-		}
-
-		// Handle replication of flag saying that bIsPlaying really should have replicated as true.
-		if (SavedReplicationForceIsPlaying != ReplicationForceIsPlaying)
-		{
-			bIsPlaying = true;
-		}
-
-		// apply bReversePlayback
-		if (SavedbReversePlayback!= bReversePlayback)
-		{
-			if (SavedbIsPlaying && bIsPlaying)
-			{
-				// notify actors that something has changed
-				for (int32 ActorIndex = 0; ActorIndex < ControlledActors.Num(); ++ActorIndex )
-				{
-					IMatineeInterface * IMI = Cast<IMatineeInterface>(ControlledActors[ActorIndex]);
-					if (IMI)
-					{
-						IMI->InterpolationChanged(this);
-					}
-				}
-			}
-		}
-	
-		// start up interpolation, if necessary
-		if (!SavedbIsPlaying && (bIsPlaying || InterpPosition != SavedPosition))
-		{	
-			InitInterp();
-
-			// if we're playing forward, call Play() to process any special properties on InterpAction that may affect the meaning of 'Position' (bNoResetOnRewind, etc)
-			if (!bReversePlayback)
-			{
-				Play();
-			}
-
-			// find affected actors and set their ControllingMatineeActor
-			// @warning: this code requires the linked actors to be static object references (i.e., some other Kismet action can't be assigning them)
-			// this might not work for AI pawns
-		
-			for (int32 ActorIndex = 0; ActorIndex < ControlledActors.Num(); ++ActorIndex )
-			{
-				AActor* Actor = ControlledActors[ActorIndex];
-				UInterpGroupInst * GrInst = FindGroupInst(Actor);
-				if (IsValid(Actor) && GrInst != NULL) 
-				{
-					// fire an event if we're really playing (and not just starting it up to do a position update)
-					if (bIsPlaying)
-					{
-						IMatineeInterface * IMI = Cast<IMatineeInterface>(Actor);
-						if (IMI)
-						{
-							IMI->InterpolationStarted(this);
-						}
-					}
-				}
-			}
-
-		}
-
-		// if we received a different current position
-		if (InterpPosition != SavedPosition)
-		{
-			//@hack: negate fade tracks if we're updating a stopped matinee
-			// the right fix is probably to pass bJump=true to UpdateInterp() when (!bIsPlaying && !SavedbIsPlaying),
-			// but that may have lots of other side effects I don't have time to test right now
-			TArray<FSavedFadeState> SavedFadeStates;
-			if (!bIsPlaying && !SavedbIsPlaying && MatineeData != NULL)
-			{
-				for (FLocalPlayerIterator It(GEngine, GetWorld()); It; ++It)
-				{
-					if (It->PlayerController != NULL && It->PlayerController->PlayerCameraManager != NULL)
-					{
-						new(SavedFadeStates)FSavedFadeState(It->PlayerController->PlayerCameraManager);
-					}
-				}
-			}
-
-			if (bIsPlaying && SavedPosition != -1 && FMath::Abs(InterpPosition - SavedPosition) < ClientSidePositionErrorTolerance)
-			{
-				// The error value between us and the server is too small to change gameplay, but will cause visual pops
-				InterpPosition = SavedPosition;
-			}
-			else
-			{
-				// set to position replicated from server
-				UpdateInterp(InterpPosition, false, false);
-			}
-		}
-
-		// terminate interpolation, if necessary
-		if ((SavedbIsPlaying || InterpPosition != SavedPosition) && !bIsPlaying)
-		{
-			TermInterp();
-
-			// find affected actors and remove InterpAction from their LatentActions list
-			for (int32 ActorIndex = 0; ActorIndex < ControlledActors.Num(); ++ActorIndex )
-			{
-				AActor* Actor = ControlledActors[ ActorIndex ];
-				if (Actor != NULL)
-				{
-
-					// fire an event if we were really playing (and not just starting it up to do a position update)
-					if (SavedbIsPlaying)
-					{
-						IMatineeInterface * IMI = Cast<IMatineeInterface>(Actor);
-						if (IMI)
-						{
-							IMI->InterpolationFinished(this);
-						}
-					}
-				}
-			}
-		}
-	}
-}
diff --git a/Engine/Source/Runtime/Engine/Private/Tests/AutomationCommon.cpp b/Engine/Source/Runtime/Engine/Private/Tests/AutomationCommon.cpp
--- a/Engine/Source/Runtime/Engine/Private/Tests/AutomationCommon.cpp
+++ b/Engine/Source/Runtime/Engine/Private/Tests/AutomationCommon.cpp
@@ -18,7 +18,6 @@
 #include "ShaderCompiler.h"
 #include "GameFramework/GameStateBase.h"
 #include "Scalability.h"
-#include "Matinee/MatineeActor.h"
 #include "StereoRendering.h"
 #include "Misc/PackageName.h"
 #include "TextureCompiler.h"
@@ -525,26 +524,13 @@ bool FWaitForAverageFrameRate::Update()
 
 bool FPlayMatineeLatentCommand::Update()
 {
-	if (MatineeActor)
-	{
-		UE_LOG(LogEngineAutomationLatentCommand, Log, TEXT("Triggering the matinee named: '%s'"), *MatineeActor->GetName())
-
-		//force this matinee to not be looping so it doesn't infinitely loop
-		MatineeActor->bLooping = false;
-		MatineeActor->Play();
-	}
 	return true;
 }
 
 
 bool FWaitForMatineeToCompleteLatentCommand::Update()
 {
-	bool bTestComplete = true;
-	if (MatineeActor)
-	{
-		bTestComplete = !MatineeActor->bIsPlaying;
-	}
-	return bTestComplete;
+	return true;
 }
 
 
@@ -597,53 +583,12 @@ bool FStreamAllResourcesLatentCommand::Update()
 
 bool FEnqueuePerformanceCaptureCommands::Update()
 {
-	//for every matinee actor in the level
-	for (TObjectIterator<AMatineeActor> It; It; ++It)
-	{
-		AMatineeActor* MatineeActor = *It;
-		if (MatineeActor && MatineeActor->GetName().Contains(TEXT("Automation")))
-		{
-			//add latent action to execute this matinee
-			ADD_LATENT_AUTOMATION_COMMAND(FPlayMatineeLatentCommand(MatineeActor));
-
-			//add action to wait until matinee is complete
-			ADD_LATENT_AUTOMATION_COMMAND(FWaitForMatineeToCompleteLatentCommand(MatineeActor));
-		}
-	}
-
 	return true;
 }
 
 
 bool FMatineePerformanceCaptureCommand::Update()
 {
-	//for every matinee actor in the level
-	for (TObjectIterator<AMatineeActor> It; It; ++It)
-	{
-		AMatineeActor* MatineeActor = *It;
-		FString MatineeFOOName = MatineeActor->GetName();
-		if (MatineeActor->GetName().Equals(MatineeName, ESearchCase::IgnoreCase))
-		{
-
-
-			//add latent action to execute this matinee
-			ADD_LATENT_AUTOMATION_COMMAND(FPlayMatineeLatentCommand(MatineeActor));
-
-			//Run the Stat FPS Chart command
-			ADD_LATENT_AUTOMATION_COMMAND(FExecWorldStringLatentCommand(TEXT("StartFPSChart")));
-
-			//add action to wait until matinee is complete
-			ADD_LATENT_AUTOMATION_COMMAND(FWaitForMatineeToCompleteLatentCommand(MatineeActor));
-
-			//Stop the Stat FPS Chart command
-			ADD_LATENT_AUTOMATION_COMMAND(FExecWorldStringLatentCommand(TEXT("StopFPSChart")));
-		}
-		else
-		{
-			UE_LOG(LogEngineAutomationLatentCommand, Log, TEXT("'%s' is not the matinee name that is being searched for."), *MatineeActor->GetName())
-		}
-	}
-
 	return true;
 
 }
diff --git a/Engine/Source/Runtime/Engine/Private/Tests/EngineAutomationTests.cpp b/Engine/Source/Runtime/Engine/Private/Tests/EngineAutomationTests.cpp
--- a/Engine/Source/Runtime/Engine/Private/Tests/EngineAutomationTests.cpp
+++ b/Engine/Source/Runtime/Engine/Private/Tests/EngineAutomationTests.cpp
@@ -1078,18 +1078,11 @@ bool FAutomationAttachment::RunTest(const FString& Parameters)
 
 // ----------------------------------------------------------------------------------------------------------------
 
-#include "Matinee/MatineeActor.h"
-
-DEFINE_LATENT_AUTOMATION_COMMAND_ONE_PARAMETER(FWaitForMatineeToCompleteAndDoScreenshotsLatentCommand, AMatineeActor*, MatineeActor);
+DEFINE_LATENT_AUTOMATION_COMMAND_ONE_PARAMETER(FWaitForMatineeToCompleteAndDoScreenshotsLatentCommand, AActor*, MatineeActor);
 
 bool FWaitForMatineeToCompleteAndDoScreenshotsLatentCommand::Update()
 {
-	bool bTestComplete = true;
-	if (MatineeActor)
-	{
-		bTestComplete = !MatineeActor->bIsPlaying;
-	}
-	return bTestComplete;
+	return true;
 }
 
 ///**
diff --git a/Engine/Source/Runtime/Engine/Private/UnrealClient.cpp b/Engine/Source/Runtime/Engine/Private/UnrealClient.cpp
--- a/Engine/Source/Runtime/Engine/Private/UnrealClient.cpp
+++ b/Engine/Source/Runtime/Engine/Private/UnrealClient.cpp
@@ -15,7 +15,6 @@
 #include "Engine/LocalPlayer.h"
 #include "UnrealEngine.h"
 #include "Components/PostProcessComponent.h"
-#include "Matinee/MatineeActor.h"
 #include "HighResScreenshot.h"
 #include "GameFramework/GameUserSettings.h"
 #include "HModel.h"
@@ -2435,4 +2434,4 @@ void FDummyViewport::InitDynamicRHI()
 		.SetInitialState(ERHIAccess::SRVMask);
 
 	RenderTargetTextureRHI = RHICreateTexture(Desc);
-}
\ No newline at end of file
+}
diff --git a/Engine/Source/Runtime/Engine/Private/UnrealEngine.cpp b/Engine/Source/Runtime/Engine/Private/UnrealEngine.cpp
--- a/Engine/Source/Runtime/Engine/Private/UnrealEngine.cpp
+++ b/Engine/Source/Runtime/Engine/Private/UnrealEngine.cpp
@@ -10974,13 +10974,13 @@ void UEngine::ClearOnScreenDebugMessages()
 }
 
 #if !UE_BUILD_SHIPPING
-void UEngine::PerformanceCapture(UWorld* World, const FString& MapName, const FString& MatineeName, float EventTime)
+void UEngine::PerformanceCapture(UWorld* World, const FString& MapName, const FString& SequenceName, float EventTime)
 {
 	// todo
 	uint32 t = IStreamingManager::Get().StreamAllResources(5.0f);
 	ensure(!t);
 
-	LogPerformanceCapture(World, MapName, MatineeName, EventTime);
+	LogPerformanceCapture(World, MapName, SequenceName, EventTime);
 
 	// can be define by command line -BuildName="ByCustomBuildName" or "CL<changelist>"
 	FString BuildName = GetBuildNameForPerfTesting();
@@ -11017,7 +11017,7 @@ void UEngine::PerformanceCapture(UWorld* World, const FString& MapName, const FS
 	FScreenshotRequest::RequestScreenshot( ScreenshotName, bShowUI, bAddFilenameSuffix );
 }
 
-void UEngine::LogPerformanceCapture(UWorld* World, const FString& MapName, const FString& MatineeName, float EventTime)
+void UEngine::LogPerformanceCapture(UWorld* World, const FString& MapName, const FString& SequenceName, float EventTime)
 {
 	const FString EventType = TEXT("PERF");
 	const int32 ChangeList = FEngineVersion::Current().GetChangelist();
@@ -11032,7 +11032,7 @@ void UEngine::LogPerformanceCapture(UWorld* World, const FString& MapName, const
 		PerfSnapshot.Changelist = FString::FromInt(ChangeList);
 		PerfSnapshot.BuildConfiguration = LexToString(FApp::GetBuildConfiguration());
 		PerfSnapshot.MapName = MapName;
-		PerfSnapshot.MatineeName = MatineeName;
+		PerfSnapshot.SequenceName = SequenceName;
 		PerfSnapshot.AverageFPS = FString::Printf(TEXT("%0.2f"), GAverageFPS);
 		PerfSnapshot.AverageFrameTime = FString::Printf(TEXT("%0.2f"), StatUnitData->FrameTime);
 		PerfSnapshot.AverageGameThreadTime = FString::Printf(TEXT("%0.2f"), StatUnitData->GameThreadTime);
diff --git a/Engine/Source/Runtime/Engine/Private/WorldSettings.cpp b/Engine/Source/Runtime/Engine/Private/WorldSettings.cpp
--- a/Engine/Source/Runtime/Engine/Private/WorldSettings.cpp
+++ b/Engine/Source/Runtime/Engine/Private/WorldSettings.cpp
@@ -105,7 +105,7 @@ PRAGMA_ENABLE_DEPRECATION_WARNINGS
 	bReplicates = true;
 	bAlwaysRelevant = true;
 	TimeDilation = 1.0f;
-	MatineeTimeDilation = 1.0f;
+	CinematicTimeDilation = 1.0f;
 	DemoPlayTimeDilation = 1.0f;
 	PackedLightAndShadowMapTextureSize = 1024;
 	SetHidden(false);
@@ -324,7 +324,7 @@ void AWorldSettings::GetLifetimeReplicatedProps( TArray< FLifetimeProperty > & O
 
 	DOREPLIFETIME( AWorldSettings, PauserPlayerState );
 	DOREPLIFETIME( AWorldSettings, TimeDilation );
-	DOREPLIFETIME( AWorldSettings, MatineeTimeDilation );
+	DOREPLIFETIME( AWorldSettings, CinematicTimeDilation );
 	DOREPLIFETIME( AWorldSettings, WorldGravityZ );
 	DOREPLIFETIME( AWorldSettings, bHighPriorityLoading );
 }
@@ -1030,7 +1030,7 @@ void AWorldSettings::RewindForReplay()
 
 	PauserPlayerState = nullptr;
 	TimeDilation = 1.0;
-	MatineeTimeDilation = 1.0;
+	CinematicTimeDilation = 1.0;
 	bWorldGravitySet = false;
 	bHighPriorityLoading = false;
 }
diff --git a/Engine/Source/Runtime/Engine/Public/AnimationUtils.h b/Engine/Source/Runtime/Engine/Public/AnimationUtils.h
--- a/Engine/Source/Runtime/Engine/Public/AnimationUtils.h
+++ b/Engine/Source/Runtime/Engine/Public/AnimationUtils.h
@@ -9,7 +9,6 @@
 #include "CoreMinimal.h"
 #include "Animation/AnimSequence.h"
 
-ENGINE_API DECLARE_LOG_CATEGORY_EXTERN(LogMatinee, Warning, All);
 ENGINE_API DECLARE_LOG_CATEGORY_EXTERN(LogAnimNotify, Warning, All);
 ENGINE_API DECLARE_LOG_CATEGORY_EXTERN(LogAnimMontage, Warning, All);
 ENGINE_API DECLARE_LOG_CATEGORY_EXTERN(LogAnimMarkerSync, Warning, All);
diff --git a/Engine/Source/Runtime/Engine/Public/Engine.h b/Engine/Source/Runtime/Engine/Public/Engine.h
--- a/Engine/Source/Runtime/Engine/Public/Engine.h
+++ b/Engine/Source/Runtime/Engine/Public/Engine.h
@@ -205,8 +205,6 @@ MONOLITHIC_HEADER_BOILERPLATE()
 #include "Animation/AnimSequenceBase.h"
 #include "AlphaBlend.h"
 #include "Animation/AnimBlueprintGeneratedClass.h"
-#include "Camera/CameraAnim.h"
-#include "Camera/CameraAnimInst.h"
 #include "Camera/CameraModifier.h"
 #include "Camera/CameraShakeBase.h"
 #include "Camera/CameraModifier_CameraShake.h"
diff --git a/Engine/Source/Runtime/Engine/Public/EngineFwd.h b/Engine/Source/Runtime/Engine/Public/EngineFwd.h
--- a/Engine/Source/Runtime/Engine/Public/EngineFwd.h
+++ b/Engine/Source/Runtime/Engine/Public/EngineFwd.h
@@ -14,7 +14,6 @@ class ACameraActor;
 class AController;
 class AGameModeBase;
 class AGameStateBase;
-class AMatineeActor;
 class ANavigationData;
 class APawn;
 class APlayerCameraManager;
@@ -58,11 +57,6 @@ class UFont;
 class UGameInstance;
 class UGameViewportClient;
 class UInputComponent;
-class UInterpCurveEdSetup;
-class UInterpGroup;
-class UInterpGroupInst;
-class UInterpTrack;
-class UInterpTrackInst;
 class ULayer;
 class ULevel;
 class ULevelStreaming;
@@ -128,7 +122,6 @@ class FDebugDisplayInfo;
 class FExportObjectInnerContext;
 class FFinalPostProcessSettings;
 class FInBunch;
-class FInterpTrackDrawParams;
 class FLightSceneProxy;
 class FLightingBuildOptions;
 class FMaterial;
