Fix compile error for cooked editor due to asset state dumping not being available for the DumpAssetRegistryCommandlet.

diff --git a/Engine/Source/Editor/UnrealEd/Private/Commandlets/DumpAssetRegistryCommandlet.cpp b/Engine/Source/Editor/UnrealEd/Private/Commandlets/DumpAssetRegistryCommandlet.cpp
--- a/Engine/Source/Editor/UnrealEd/Private/Commandlets/DumpAssetRegistryCommandlet.cpp
+++ b/Engine/Source/Editor/UnrealEd/Private/Commandlets/DumpAssetRegistryCommandlet.cpp
@@ -106,6 +106,7 @@ bool UDumpAssetRegistryCommandlet::TryParseArgs()
 
 bool UDumpAssetRegistryCommandlet::TryDumpAssetRegistry()
 {
+#if ASSET_REGISTRY_STATE_DUMPING_ENABLED
 	IFileManager& FileManager = IFileManager::Get();
 	if (!FileManager.FileExists(*Path))
 	{
@@ -150,4 +151,8 @@ bool UDumpAssetRegistryCommandlet::TryDumpAssetRegistry()
 
 	UE_LOG(LogAssetRegistryDump, Display, TEXT("Wrote %d files to %s."), Pages.Num(), *OutDir);
 	return true;
+#else
+	UE_LOG(LogAssetRegistryDump, Error, TEXT("Asset registry state dumping disabled in this build."));
+	return false;
+#endif //ASSET_REGISTRY_STATE_DUMPING_ENABLED
 }
\ No newline at end of file
