WorldPartitionConvertCommandlet: Multiple fixes

- Fix IsRunningConvertWorldPartitionCommandletstoring early result in static const (always returning false afterwards) and make it public
- Call RebuildLevel so that level UModels get properly rebuilt with brushes changing levels (from sublevels)
- Fix SubLevel LevelInstances failing to set world asset on unsaved renamed SubLevels

diff --git a/Engine/Source/Editor/UnrealEd/Private/Commandlets/WorldPartitionConvertCommandlet.cpp b/Engine/Source/Editor/UnrealEd/Private/Commandlets/WorldPartitionConvertCommandlet.cpp
--- a/Engine/Source/Editor/UnrealEd/Private/Commandlets/WorldPartitionConvertCommandlet.cpp
+++ b/Engine/Source/Editor/UnrealEd/Private/Commandlets/WorldPartitionConvertCommandlet.cpp
@@ -1183,9 +1183,6 @@ int32 UWorldPartitionConvertCommandlet::Main(const FString& Params)
 		UWorld* SubWorld = SubLevel->GetTypedOuter<UWorld>();
 		UPackage* SubPackage = SubLevel->GetPackage();
 
-		SubWorld->PersistentLevel->Model->MarkAsGarbage();
-		SubWorld->PersistentLevel->Model = nullptr;
-
 		RemapSoftObjectPaths.Add(FSoftObjectPath(SubWorld).ToString(), FSoftObjectPath(MainWorld).ToString());
 		RemapSoftObjectPaths.Add(FSoftObjectPath(SubLevel).ToString(), FSoftObjectPath(MainLevel).ToString());
 		RemapSoftObjectPaths.Add(FSoftObjectPath(SubPackage).ToString(), FSoftObjectPath(MainPackage).ToString());
@@ -1292,6 +1289,7 @@ int32 UWorldPartitionConvertCommandlet::Main(const FString& Params)
 			LevelInstanceActor->DesiredRuntimeBehavior = ELevelInstanceRuntimeBehavior::LevelStreaming;
 			LevelInstanceActor->SetActorTransform(LevelTransform);
 			LevelInstanceActor->SetWorldAsset(SubLevelWorld);
+			LevelInstanceActor->SetActorLabel(SubLevelWorld->GetName());
 		}
 		else
 		{
@@ -1378,6 +1376,11 @@ int32 UWorldPartitionConvertCommandlet::Main(const FString& Params)
 				PackagesToDelete.Add(SubLevel->GetPackage());
 			}
 		}
+		else
+		{
+			// Rebuild Model to clear refs to actors that were converted
+			GEditor->RebuildLevel(*SubLevel);
+		}
 	}
 
 	// Clear streaming levels
@@ -1494,6 +1497,7 @@ int32 UWorldPartitionConvertCommandlet::Main(const FString& Params)
 
 		MainWorld->WorldComposition = nullptr;
 		MainLevel->bIsPartitioned = !bOnlyMergeSubLevels;
+		GEditor->RebuildLevel(*MainLevel);
 
 		if (bDeleteSourceLevels)
 		{
diff --git a/Engine/Source/Editor/UnrealEd/Private/EditorServer.cpp b/Engine/Source/Editor/UnrealEd/Private/EditorServer.cpp
--- a/Engine/Source/Editor/UnrealEd/Private/EditorServer.cpp
+++ b/Engine/Source/Editor/UnrealEd/Private/EditorServer.cpp
@@ -1748,7 +1748,10 @@ void UEditorEngine::RebuildLevel(ULevel& Level)
 		}
 	}
 	
-	GLevelEditorModeTools().MapChangeNotify();
+	if (!IsRunningCommandlet())
+	{
+		GLevelEditorModeTools().MapChangeNotify();
+	}
 }
 
 void UEditorEngine::RebuildModelFromBrushes(UModel* Model, bool bSelectedBrushesOnly, bool bTreatMovableBrushesAsStatic)
diff --git a/Engine/Source/Runtime/Engine/Private/LevelInstance/LevelInstanceSubsystem.cpp b/Engine/Source/Runtime/Engine/Private/LevelInstance/LevelInstanceSubsystem.cpp
--- a/Engine/Source/Runtime/Engine/Private/LevelInstance/LevelInstanceSubsystem.cpp
+++ b/Engine/Source/Runtime/Engine/Private/LevelInstance/LevelInstanceSubsystem.cpp
@@ -21,6 +21,7 @@
 #include "LevelInstance/LevelInstanceEditorInstanceActor.h"
 #include "LevelInstance/LevelInstanceEditorObject.h"
 #include "LevelInstance/LevelInstanceEditorPivotActor.h"
+#include "WorldPartition/WorldPartitionSubsystem.h"
 #include "WorldPartition/DataLayer/DataLayerSubsystem.h"
 #include "WorldPartition/DataLayer/DataLayerInstanceWithAsset.h"
 #include "WorldPartition/WorldPartitionConverter.h"
@@ -2323,6 +2324,12 @@ bool ULevelInstanceSubsystem::CanUsePackage(FName InPackageName)
 
 bool ULevelInstanceSubsystem::CanUseWorldAsset(const ILevelInstanceInterface* LevelInstance, TSoftObjectPtr<UWorld> WorldAsset, FString* OutReason)
 {
+	// Do not validate when running convert commandlet as package might not exist yet.
+	if (UWorldPartitionSubsystem::IsRunningConvertWorldPartitionCommandlet())
+	{
+		return true;
+	}
+
 	if (WorldAsset.IsNull())
 	{
 		return true;
diff --git a/Engine/Source/Runtime/Engine/Private/WorldPartition/WorldPartitionSubsystem.cpp b/Engine/Source/Runtime/Engine/Private/WorldPartition/WorldPartitionSubsystem.cpp
--- a/Engine/Source/Runtime/Engine/Private/WorldPartition/WorldPartitionSubsystem.cpp
+++ b/Engine/Source/Runtime/Engine/Private/WorldPartition/WorldPartitionSubsystem.cpp
@@ -116,12 +116,11 @@ const UWorldPartition* UWorldPartitionSubsystem::GetWorldPartition() const
 }
 
 #if WITH_EDITOR
-bool UWorldPartitionSubsystem::IsRunningConvertWorldPartitionCommandlet() const
+bool UWorldPartitionSubsystem::IsRunningConvertWorldPartitionCommandlet()
 {
 	static UClass* WorldPartitionConvertCommandletClass = FindObject<UClass>(nullptr, TEXT("/Script/UnrealEd.WorldPartitionConvertCommandlet"), true);
 	check(WorldPartitionConvertCommandletClass);
-	static const bool bIsRunningWorldPartitionConvertCommandlet = GetRunningCommandletClass() && GetRunningCommandletClass()->IsChildOf(WorldPartitionConvertCommandletClass);
-	return bIsRunningWorldPartitionConvertCommandlet;
+	return GetRunningCommandletClass() && GetRunningCommandletClass()->IsChildOf(WorldPartitionConvertCommandletClass);
 }
 #endif
 
@@ -130,7 +129,8 @@ void UWorldPartitionSubsystem::Initialize(FSubsystemCollectionBase& Collection)
 	Super::Initialize(Collection);
 
 #if WITH_EDITOR
-	if (IsRunningConvertWorldPartitionCommandlet())
+	bIsRunningConvertWorldPartitionCommandlet = IsRunningConvertWorldPartitionCommandlet();
+	if(bIsRunningConvertWorldPartitionCommandlet)
 	{
 		return;
 	}
@@ -143,7 +143,7 @@ void UWorldPartitionSubsystem::Initialize(FSubsystemCollectionBase& Collection)
 void UWorldPartitionSubsystem::Deinitialize()
 {
 #if WITH_EDITOR
-	if (IsRunningConvertWorldPartitionCommandlet())
+	if (bIsRunningConvertWorldPartitionCommandlet)
 	{
 		Super::Deinitialize();
 		return;
diff --git a/Engine/Source/Runtime/Engine/Public/WorldPartition/WorldPartitionSubsystem.h b/Engine/Source/Runtime/Engine/Public/WorldPartition/WorldPartitionSubsystem.h
--- a/Engine/Source/Runtime/Engine/Public/WorldPartition/WorldPartitionSubsystem.h
+++ b/Engine/Source/Runtime/Engine/Public/WorldPartition/WorldPartitionSubsystem.h
@@ -49,12 +49,11 @@ public:
 
 #if WITH_EDITOR
 	void ForEachWorldPartition(TFunctionRef<bool(UWorldPartition*)> Func);
+
+	static bool IsRunningConvertWorldPartitionCommandlet();
 #endif
 
 private:
-#if WITH_EDITOR
-	bool IsRunningConvertWorldPartitionCommandlet() const;
-#endif
 
 	void OnWorldPartitionInitialized(UWorldPartition* InWorldPartition);
 	void OnWorldPartitionUninitialized(UWorldPartition* InWorldPartition);
@@ -73,4 +72,8 @@ private:
 	// GC backup values
 	int32 LevelStreamingContinuouslyIncrementalGCWhileLevelsPendingPurge;
 	int32 LevelStreamingForceGCAfterLevelStreamedOut;
+
+#if WITH_EDITOR
+	bool bIsRunningConvertWorldPartitionCommandlet;
+#endif
 };
