Rename editor experimental setting bEnableAsyncSkeletalMeshCompilation to bEnableAsyncSkinnedAssetCompilation

diff --git a/Engine/Config/BaseEditorPerProjectUserSettings.ini b/Engine/Config/BaseEditorPerProjectUserSettings.ini
--- a/Engine/Config/BaseEditorPerProjectUserSettings.ini
+++ b/Engine/Config/BaseEditorPerProjectUserSettings.ini
@@ -155,7 +155,7 @@ bExampleLayersAndBlends = True
 ; Enabled by default now
 bEnableAsyncTextureCompilation=True
 bEnableAsyncStaticMeshCompilation=True
-bEnableAsyncSkeletalMeshCompilation=True
+bEnableAsyncSkinnedAssetCompilation=True
 bEnableAsyncSoundWaveCompilation=True
 
 ; Cooking in the editor is disabled by default now; WorldPartition does not yet support it and we want to deprecate the mode in the future to reduce the number of code flows that have to be supported
diff --git a/Engine/Source/Editor/UnrealEd/Classes/Settings/EditorExperimentalSettings.h b/Engine/Source/Editor/UnrealEd/Classes/Settings/EditorExperimentalSettings.h
--- a/Engine/Source/Editor/UnrealEd/Classes/Settings/EditorExperimentalSettings.h
+++ b/Engine/Source/Editor/UnrealEd/Classes/Settings/EditorExperimentalSettings.h
@@ -29,9 +29,14 @@ public:
 	bool bEnableAsyncStaticMeshCompilation;
 
 	/** Enable async skeletal mesh compilation to improve import and map load time performance when compilation is required */
-	UPROPERTY(EditAnywhere, config, Category = Performance, meta = (DisplayName = "Enable async skeletal mesh compilation and loading"))
+	UE_DEPRECATED(5.1, "Deprecated & replaced by bEnableAsyncSkinnedAssetCompilation.")
+	UPROPERTY(/*EditAnywhere - deprecated & replaced by bEnableAsyncSkinnedAssetCompilation, */config/*, Category = Performance, meta = (DisplayName = "Enable async skeletal mesh compilation and loading")*/)
 	bool bEnableAsyncSkeletalMeshCompilation;
 
+	/** Enable async skinned asset compilation to improve import and map load time performance when compilation is required */
+	UPROPERTY(EditAnywhere, config, Category = Performance, meta = (DisplayName = "Enable async skinned asset compilation and loading"))
+	bool bEnableAsyncSkinnedAssetCompilation;
+
 	/** Enable async sound compilation to improve import and map load time performance when compilation is required */
 	UPROPERTY(EditAnywhere, config, Category = Performance, meta = (DisplayName = "Enable async sound compilation and loading"))
 	bool bEnableAsyncSoundWaveCompilation;
diff --git a/Engine/Source/Editor/UnrealEd/Private/Settings/SettingsClasses.cpp b/Engine/Source/Editor/UnrealEd/Private/Settings/SettingsClasses.cpp
--- a/Engine/Source/Editor/UnrealEd/Private/Settings/SettingsClasses.cpp
+++ b/Engine/Source/Editor/UnrealEd/Private/Settings/SettingsClasses.cpp
@@ -191,7 +191,8 @@ UEditorExperimentalSettings::UEditorExperimentalSettings( const FObjectInitializ
 	: Super(ObjectInitializer)
 	, bEnableAsyncTextureCompilation(false)
 	, bEnableAsyncStaticMeshCompilation(false)
-	, bEnableAsyncSkeletalMeshCompilation(false)
+	, bEnableAsyncSkeletalMeshCompilation(true)	// This was false and set to True in /Engine/Config/BaseEditorPerProjectUserSettings.ini. The setting is removed from .ini so change it to default True.
+	, bEnableAsyncSkinnedAssetCompilation(false)
 	, bEnableAsyncSoundWaveCompilation(false)
 	, bHDREditor(false)
 	, HDREditorNITLevel(160.0f)
@@ -204,6 +205,16 @@ UEditorExperimentalSettings::UEditorExperimentalSettings( const FObjectInitializ
 
 void UEditorExperimentalSettings::PostInitProperties()
 {
+	PRAGMA_DISABLE_DEPRECATION_WARNINGS
+	// bEnableAsyncSkeletalMeshCompilation's default to True (see comment in constructor above).
+	// To be backwards compatible, if a user project overrides it to False, pass on the value to bEnableAsyncSkinnedAssetCompilation.
+	if (!bEnableAsyncSkeletalMeshCompilation)
+	{
+		UE_LOG(LogSettingsClasses, Warning, TEXT("bEnableAsyncSkeletalMeshCompilation is deprecated and replaced with bEnableAsyncSkinnedAssetCompilation. Please update the config. Setting bEnableAsyncSkinnedAssetCompilation to False."));
+		bEnableAsyncSkinnedAssetCompilation = false;
+	}
+	PRAGMA_ENABLE_DEPRECATION_WARNINGS
+
 	CVarEditorHDRSupport->Set(bHDREditor ? 1 : 0, ECVF_SetByProjectSetting);
 	CVarEditorHDRNITLevel->Set(HDREditorNITLevel, ECVF_SetByProjectSetting);
 	Super::PostInitProperties();
diff --git a/Engine/Source/Runtime/Engine/Private/SkinnedAssetCompiler.cpp b/Engine/Source/Runtime/Engine/Private/SkinnedAssetCompiler.cpp
--- a/Engine/Source/Runtime/Engine/Private/SkinnedAssetCompiler.cpp
+++ b/Engine/Source/Runtime/Engine/Private/SkinnedAssetCompiler.cpp
@@ -55,7 +55,7 @@ namespace SkinnedAssetCompilingManagerImpl
 				TEXT("skinnedasset"),
 				CVarAsyncSkinnedAssetStandard.AsyncCompilation,
 				CVarAsyncSkinnedAssetStandard.AsyncCompilationMaxConcurrency,
-				GET_MEMBER_NAME_CHECKED(UEditorExperimentalSettings, bEnableAsyncSkeletalMeshCompilation)); // TODO: rename to SkinnedAsset
+				GET_MEMBER_NAME_CHECKED(UEditorExperimentalSettings, bEnableAsyncSkinnedAssetCompilation));
 		}
 	}
 }
