Fix a bug causing editor bulkdata to attempt to both reference the payload in the workspace domain as well as writing out its payload to the editor domain.

#rnx

- During 21147257 one of the checks to disable payload writing, that was calling ::IsReferencingByPackagePath was changed to checking the EFlags::ReferencesLegacyFile flag directly, which meant we were no longer taking referencing payloads from the workspace domain into account.
-- This was due to the CL being an attempt to merge 20965910 (which only checked the flag) with UE5 code (which called ::IsReferencingByPackagePath due to better editor domain support)
- This would cause the payload to be both referenced and written out into the editor domain, so on load we would end up with a flag telling the bulkdata to load from the original workspace domain file BUT the offset would reference where in the editor domain the payload was written out too.
-- This would only affect older versions of the package, newer versions would continue to use the package trailer to find the correct offet to use when referencing.
- This change requires that we change the EditorDomainVersion to clear out anything people might have saved with the incorrect info.

diff --git a/Engine/Source/Editor/UnrealEd/Private/EditorDomain/EditorDomainUtils.cpp b/Engine/Source/Editor/UnrealEd/Private/EditorDomain/EditorDomainUtils.cpp
--- a/Engine/Source/Editor/UnrealEd/Private/EditorDomain/EditorDomainUtils.cpp
+++ b/Engine/Source/Editor/UnrealEd/Private/EditorDomain/EditorDomainUtils.cpp
@@ -167,7 +167,7 @@ bool bGGlobalConstructClassesInitialized = false;
 int64 GMaxBulkDataSize = -1;
 
 // Change to a new guid when EditorDomain needs to be invalidated
-const TCHAR* EditorDomainVersion = TEXT("D3E2B4E479C34305952A97AA63A8E19D");
+const TCHAR* EditorDomainVersion = TEXT("1A71E99C857343C7A6AEB0035E8F5F26");
 
 // Identifier of the CacheBuckets for EditorDomain tables
 const TCHAR* EditorDomainPackageBucketName = TEXT("EditorDomainPackage");
diff --git a/Engine/Source/Runtime/CoreUObject/Private/Serialization/EditorBulkData.cpp b/Engine/Source/Runtime/CoreUObject/Private/Serialization/EditorBulkData.cpp
--- a/Engine/Source/Runtime/CoreUObject/Private/Serialization/EditorBulkData.cpp
+++ b/Engine/Source/Runtime/CoreUObject/Private/Serialization/EditorBulkData.cpp
@@ -857,7 +857,7 @@ void FEditorBulkData::Serialize(FArchive& Ar, UObject* Owner, bool bAllowRegiste
 				// TODO: add iostore support for editor bulkdata
 				bWriteOutPayload = false;
 			}
-			else if (EnumHasAnyFlags(UpdatedFlags, EFlags::ReferencesLegacyFile))
+			else if (IsReferencingByPackagePath(UpdatedFlags))
 			{
 				// Write out required extra data if we're saving by reference
 				if (!IsStoredInPackageTrailer(UpdatedFlags))
