fix comment
#rnx

diff --git a/Engine/Source/Editor/UnrealEd/Classes/Exporters/FbxExportOption.h b/Engine/Source/Editor/UnrealEd/Classes/Exporters/FbxExportOption.h
--- a/Engine/Source/Editor/UnrealEd/Classes/Exporters/FbxExportOption.h
+++ b/Engine/Source/Editor/UnrealEd/Classes/Exporters/FbxExportOption.h
@@ -65,7 +65,7 @@ public:
 	UPROPERTY(EditAnywhere, BlueprintReadWrite, config, category = Animation)
 	uint32 MapSkeletalMotionToRoot : 1;
 
-	/** If enabled, export sequencer animation in its local time, relative to its master sequence. */
+	/** If enabled, export sequencer animation in its local time, relative to its sequence. */
 	UPROPERTY(EditAnywhere, BlueprintReadWrite, config, category = Animation)
 	uint32 bExportLocalTime : 1;
 
