-List cook comamnd to log everything that will be cooked

diff --git a/Engine/Source/Editor/UnrealEd/Private/CookOnTheFlyServer.cpp b/Engine/Source/Editor/UnrealEd/Private/CookOnTheFlyServer.cpp
--- a/Engine/Source/Editor/UnrealEd/Private/CookOnTheFlyServer.cpp
+++ b/Engine/Source/Editor/UnrealEd/Private/CookOnTheFlyServer.cpp
@@ -9458,6 +9458,11 @@ void UCookOnTheFlyServer::GenerateInitialRequests(FBeginCookContext& BeginContex
 			}
 		}
 	}
+
+	if (FParse::Param(FCommandLine::Get(), TEXT("List")))
+	{
+		WorkerRequests->LogAllRequestedFiles();
+	}
 }
 
 void UCookOnTheFlyServer::RecordDLCPackagesFromBaseGame(FBeginCookContext& BeginContext)
diff --git a/Engine/Source/Editor/UnrealEd/Private/Cooker/CookRequests.cpp b/Engine/Source/Editor/UnrealEd/Private/Cooker/CookRequests.cpp
--- a/Engine/Source/Editor/UnrealEd/Private/Cooker/CookRequests.cpp
+++ b/Engine/Source/Editor/UnrealEd/Private/Cooker/CookRequests.cpp
@@ -324,4 +324,13 @@ void FExternalRequests::RemapTargetPlatforms(const TMap<ITargetPlatform*, ITarge
 	}
 }
 
-}
\ No newline at end of file
+void FExternalRequests::LogAllRequestedFiles()
+{
+	UE_LOG(LogCook, Display, TEXT("Listing initially requested set of files to be cooked (does not include list of transitive dependencies):"));
+	for (const TPair<FName, FFilePlatformRequest>& KVPair : RequestMap)
+	{
+		UE_LOG(LogCook, Display, TEXT("%s"), *KVPair.Value.ToString());
+	}
+	UE_LOG(LogCook, Display, TEXT("----------------------------------------------------------\n\n"));
+}
+}
diff --git a/Engine/Source/Editor/UnrealEd/Private/Cooker/CookRequests.h b/Engine/Source/Editor/UnrealEd/Private/Cooker/CookRequests.h
--- a/Engine/Source/Editor/UnrealEd/Private/Cooker/CookRequests.h
+++ b/Engine/Source/Editor/UnrealEd/Private/Cooker/CookRequests.h
@@ -123,6 +123,9 @@ public:
 	/** Swap all ITargetPlatform* stored on this instance according to the mapping in @param Remap. */
 	void RemapTargetPlatforms(const TMap<ITargetPlatform*, ITargetPlatform*>& Remap);
 
+	/* Prints a list of all files in the RequestMap to the log */
+	void LogAllRequestedFiles();
+
 public:
 	/** An FEvent the scheduler can sleep on when waiting for new cookonthefly requests. */
 	TUniquePtr<FEvent, FFreeFEvent> CookRequestEvent;
diff --git a/Engine/Source/Editor/UnrealEd/Private/Cooker/IWorkerRequests.h b/Engine/Source/Editor/UnrealEd/Private/Cooker/IWorkerRequests.h
--- a/Engine/Source/Editor/UnrealEd/Private/Cooker/IWorkerRequests.h
+++ b/Engine/Source/Editor/UnrealEd/Private/Cooker/IWorkerRequests.h
@@ -80,6 +80,9 @@ public:
 	virtual void GetBeginCookConfigSettings(UCookOnTheFlyServer& COTFS, FBeginCookContext& BeginContext, UE::Cook::FBeginCookConfigSettings& Settings) = 0;
 	virtual void GetBeginCookIterativeFlags(UCookOnTheFlyServer& COTFS, FBeginCookContext& BeginContext) = 0;
 	virtual ECookMode::Type GetDirectorCookMode(UCookOnTheFlyServer& COTFS) = 0;
+
+	/* Prints a list of all files in the RequestMap to the log */
+	virtual void LogAllRequestedFiles() = 0;
 };
 
 }
\ No newline at end of file
diff --git a/Engine/Source/Editor/UnrealEd/Private/Cooker/WorkerRequestsLocal.cpp b/Engine/Source/Editor/UnrealEd/Private/Cooker/WorkerRequestsLocal.cpp
--- a/Engine/Source/Editor/UnrealEd/Private/Cooker/WorkerRequestsLocal.cpp
+++ b/Engine/Source/Editor/UnrealEd/Private/Cooker/WorkerRequestsLocal.cpp
@@ -127,4 +127,9 @@ ECookMode::Type FWorkerRequestsLocal::GetDirectorCookMode(UCookOnTheFlyServer& C
 	return COTFS.GetCookMode();
 }
 
+void FWorkerRequestsLocal::LogAllRequestedFiles()
+{
+	ExternalRequests.LogAllRequestedFiles();
+}
+
 }
\ No newline at end of file
diff --git a/Engine/Source/Editor/UnrealEd/Private/Cooker/WorkerRequestsLocal.h b/Engine/Source/Editor/UnrealEd/Private/Cooker/WorkerRequestsLocal.h
--- a/Engine/Source/Editor/UnrealEd/Private/Cooker/WorkerRequestsLocal.h
+++ b/Engine/Source/Editor/UnrealEd/Private/Cooker/WorkerRequestsLocal.h
@@ -36,6 +36,8 @@ public:
 	virtual void GetBeginCookIterativeFlags(UCookOnTheFlyServer& COTFS, FBeginCookContext& BeginContext) override;
 	virtual ECookMode::Type GetDirectorCookMode(UCookOnTheFlyServer& COTFS) override;
 
+	virtual void LogAllRequestedFiles() override;
+
 private:
 	FExternalRequests ExternalRequests;
 };
diff --git a/Engine/Source/Editor/UnrealEd/Private/Cooker/WorkerRequestsRemote.cpp b/Engine/Source/Editor/UnrealEd/Private/Cooker/WorkerRequestsRemote.cpp
--- a/Engine/Source/Editor/UnrealEd/Private/Cooker/WorkerRequestsRemote.cpp
+++ b/Engine/Source/Editor/UnrealEd/Private/Cooker/WorkerRequestsRemote.cpp
@@ -167,4 +167,9 @@ void FWorkerRequestsRemote::LogCalledEditorActionError(const TCHAR* FunctionName
 	UE_LOG(LogCook, Error, TEXT("Calling %s (an editor-mode-only function) is not allowed in a CookWorker."), FunctionName);
 }
 
+void FWorkerRequestsRemote::LogAllRequestedFiles()
+{
+	ExternalRequests.LogAllRequestedFiles();
+}
+
 }
\ No newline at end of file
diff --git a/Engine/Source/Editor/UnrealEd/Private/Cooker/WorkerRequestsRemote.h b/Engine/Source/Editor/UnrealEd/Private/Cooker/WorkerRequestsRemote.h
--- a/Engine/Source/Editor/UnrealEd/Private/Cooker/WorkerRequestsRemote.h
+++ b/Engine/Source/Editor/UnrealEd/Private/Cooker/WorkerRequestsRemote.h
@@ -42,6 +42,8 @@ public:
 	virtual void GetBeginCookIterativeFlags(UCookOnTheFlyServer& COTFS, FBeginCookContext& BeginContext) override;
 	virtual ECookMode::Type GetDirectorCookMode(UCookOnTheFlyServer& COTFS) override;
 
+	virtual void LogAllRequestedFiles() override;
+
 private:
 	void LogCalledCookByTheBookError(const TCHAR* FunctionName) const;
 	void LogCalledCookOnTheFlyError(const TCHAR* FunctionName) const;
