Fix Ctrl+Click not selecting/deselecting Level Instances

diff --git a/Engine/Source/Editor/LevelEditor/Private/Elements/Actor/ActorElementLevelEditorSelectionCustomization.cpp b/Engine/Source/Editor/LevelEditor/Private/Elements/Actor/ActorElementLevelEditorSelectionCustomization.cpp
--- a/Engine/Source/Editor/LevelEditor/Private/Elements/Actor/ActorElementLevelEditorSelectionCustomization.cpp
+++ b/Engine/Source/Editor/LevelEditor/Private/Elements/Actor/ActorElementLevelEditorSelectionCustomization.cpp
@@ -242,6 +242,12 @@ bool FActorElementLevelEditorSelectionCustomization::DeselectActorElement(const
 
 	bool bSelectionChanged = false;
 
+	// If trying to deselect an actor, use this actors root selection actor instead (if it has one)
+	if (AActor* RootSelection = Actor->GetRootSelectionParent())
+	{
+		Actor = RootSelection;
+	}
+
 	if (UActorGroupingUtils::IsGroupingActive() && InSelectionOptions.AllowGroups())
 	{
 		// If this actor is a group, do a group select
diff --git a/Engine/Source/Editor/UnrealEd/Private/LevelViewportClickHandlers.cpp b/Engine/Source/Editor/UnrealEd/Private/LevelViewportClickHandlers.cpp
--- a/Engine/Source/Editor/UnrealEd/Private/LevelViewportClickHandlers.cpp
+++ b/Engine/Source/Editor/UnrealEd/Private/LevelViewportClickHandlers.cpp
@@ -165,7 +165,7 @@ namespace LevelViewportClickHandlers
 					const bool bAllowSelectionModifiers = bIsLeftClickSelection && LevelEditorElementSelectionSet->AllowSelectionModifiers(ResolvedElement);
 					if (Click.IsControlDown() && bAllowSelectionModifiers)
 					{
-						if (LevelEditorElementSelectionSet->IsElementSelected(ResolvedElement, FTypedElementIsSelectedOptions()))
+						if (LevelEditorElementSelectionSet->IsElementSelected(ResolvedElement, FTypedElementIsSelectedOptions().SetAllowIndirect(true)))
 						{
 							LevelEditorElementSelectionSet->DeselectElement(ResolvedElement, SelectionOptions);
 						}
