#UE Fix crash when loading maps in the editor related to having already cleaned up a scene

diff --git a/Engine/Source/Editor/UnrealEd/Private/EditorViewportClient.cpp b/Engine/Source/Editor/UnrealEd/Private/EditorViewportClient.cpp
--- a/Engine/Source/Editor/UnrealEd/Private/EditorViewportClient.cpp
+++ b/Engine/Source/Editor/UnrealEd/Private/EditorViewportClient.cpp
@@ -3989,7 +3989,7 @@ void FEditorViewportClient::Draw(FViewport* InViewport, FCanvas* Canvas)
 		}
 
 		// Setup custom upscaler and screen percentage.
-		if (GCustomEditorStaticScreenPercentage && ViewFamily.ViewMode == EViewModeIndex::VMI_Lit && ViewFamily.Scene->GetShadingPath() == EShadingPath::Deferred && ViewFamily.bRealtimeUpdate)
+		if (GCustomEditorStaticScreenPercentage && ViewFamily.ViewMode == EViewModeIndex::VMI_Lit && ViewFamily.Scene && ViewFamily.Scene->GetShadingPath() == EShadingPath::Deferred && ViewFamily.bRealtimeUpdate)
 		{
 			GCustomEditorStaticScreenPercentage->SetupEditorViewFamily(ViewFamily, this);
 		}
