Partial rollback of 21069313 because of cook time regression.

diff --git a/Engine/Source/Editor/UnrealEd/Private/FileHelpers.cpp b/Engine/Source/Editor/UnrealEd/Private/FileHelpers.cpp
--- a/Engine/Source/Editor/UnrealEd/Private/FileHelpers.cpp
+++ b/Engine/Source/Editor/UnrealEd/Private/FileHelpers.cpp
@@ -804,8 +804,10 @@ static bool SaveWorld(UWorld* World,
 
 				if (!DuplicatedWorld)
 				{
+					// Explict Reset Loaders of Package here because we want to avoid resetting of all loaders which is the current behavior of UObject::Rename when passing in a UPackage
+					ResetLoaders(Package);
 					// Duplicate failed or not needed. Just do a rename.
-					Package->Rename(*NewPackageName, NULL, REN_NonTransactional | REN_DontCreateRedirectors);
+					Package->Rename(*NewPackageName, NULL, REN_NonTransactional | REN_DontCreateRedirectors | REN_ForceNoResetLoaders);
 					
 					if (bWorldNeedsRename)
 					{
@@ -822,7 +824,7 @@ static bool SaveWorld(UWorld* World,
 							}
 						}
 
-						World->Rename(*NewWorldAssetName, NULL, REN_NonTransactional | REN_DontCreateRedirectors);
+						World->Rename(*NewWorldAssetName, NULL, REN_NonTransactional | REN_DontCreateRedirectors | REN_ForceNoResetLoaders);
 					}
 
 					// We're changing the world path, add a path redirector so that soft object paths get fixed on save
diff --git a/Engine/Source/Runtime/Engine/Private/ActorEditor.cpp b/Engine/Source/Runtime/Engine/Private/ActorEditor.cpp
--- a/Engine/Source/Runtime/Engine/Private/ActorEditor.cpp
+++ b/Engine/Source/Runtime/Engine/Private/ActorEditor.cpp
@@ -838,8 +838,8 @@ void AActor::SetPackageExternal(bool bExternal, bool bShouldDirty)
 	else 
 	{
 		UPackage* ActorPackage = GetExternalPackage();
-		// Detach the loaders so it doesn't resolve to this actor anymore
-		ResetLoaders(ActorPackage);
+		// Detach the linker exports so it doesn't resolve to this actor anymore
+		ResetLinkerExports(ActorPackage);
 		SetExternalPackage(nullptr);
 	}
 
