Import high res mesh menu entry
- Added the ability to import, reimport and remove the high res source model.
- It'll import the FBX file as a new transient static mesh and then copy its LOD0 mesh description to the static mesh hires mesh description.

diff --git a/Engine/Source/Developer/AssetTools/Private/AssetTypeActions/AssetTypeActions_StaticMesh.cpp b/Engine/Source/Developer/AssetTools/Private/AssetTypeActions/AssetTypeActions_StaticMesh.cpp
--- a/Engine/Source/Developer/AssetTools/Private/AssetTypeActions/AssetTypeActions_StaticMesh.cpp
+++ b/Engine/Source/Developer/AssetTools/Private/AssetTypeActions/AssetTypeActions_StaticMesh.cpp
@@ -1,19 +1,22 @@
 // Copyright Epic Games, Inc. All Rights Reserved.
 
 #include "AssetTypeActions/AssetTypeActions_StaticMesh.h"
-#include "Framework/MultiBox/MultiBoxBuilder.h"
-#include "ToolMenus.h"
-#include "Styling/AppStyle.h"
-#include "EditorFramework/AssetImportData.h"
-#include "ThumbnailRendering/SceneThumbnailInfo.h"
+
 #include "AssetTools.h"
+#include "Editor.h"
+#include "EditorSupportDelegates.h"
 #include "Editor/StaticMeshEditor/Public/StaticMeshEditorModule.h"
+#include "EditorFramework/AssetImportData.h"
 #include "FbxMeshUtils.h"
-#include "StaticMeshCompiler.h"
+#include "Framework/MultiBox/MultiBoxBuilder.h"
 #include "Framework/Notifications/NotificationManager.h"
-#include "Widgets/Notifications/SNotificationList.h"
 #include "Misc/MessageDialog.h"
 #include "Misc/ScopedSlowTask.h"
+#include "StaticMeshCompiler.h"
+#include "Styling/AppStyle.h"
+#include "ThumbnailRendering/SceneThumbnailInfo.h"
+#include "ToolMenus.h"
+#include "Widgets/Notifications/SNotificationList.h"
 
 #define LOCTEXT_NAMESPACE "AssetTypeActions"
 
@@ -191,10 +194,10 @@ void FAssetTypeActions_StaticMesh::GetNaniteMenu(class FMenuBuilder& MenuBuilder
 void FAssetTypeActions_StaticMesh::GetImportLODMenu(class FMenuBuilder& MenuBuilder,TArray<TWeakObjectPtr<UStaticMesh>> Objects)
 {
 	check(Objects.Num() > 0);
-	auto First = Objects[0];
+	TWeakObjectPtr<UStaticMesh> First = Objects[0];
 	UStaticMesh* StaticMesh = First.Get();
 
-	for(int32 LOD = 1;LOD<=First->GetNumLODs();++LOD)
+	for(int32 LOD = 1; LOD <= First->GetNumLODs(); ++LOD)
 	{
 		FText LODText = FText::AsNumber( LOD );
 		FText Description = FText::Format( NSLOCTEXT("AssetTypeActions_StaticMesh", "Reimport LOD (number)", "Reimport LOD {0}"), LODText );
@@ -206,7 +209,7 @@ void FAssetTypeActions_StaticMesh::GetImportLODMenu(class FMenuBuilder& MenuBuil
 		}
 
 		MenuBuilder.AddMenuEntry( Description, ToolTip, FSlateIcon(),
-			FUIAction(FExecuteAction::CreateStatic( &FAssetTypeActions_StaticMesh::ExecuteImportMeshLOD, static_cast<UObject*>(StaticMesh), LOD) )) ;
+			FUIAction(FExecuteAction::CreateStatic( &FAssetTypeActions_StaticMesh::ExecuteImportMeshLOD, static_cast<UObject*>(StaticMesh), LOD) ));
 	}
 }
 
@@ -245,6 +248,14 @@ void FAssetTypeActions_StaticMesh::GetLODMenu(class FMenuBuilder& MenuBuilder, T
 		FCanExecuteAction::CreateSP(this, &FAssetTypeActions_StaticMesh::CanPasteLODSettings, Meshes)
 		)
 	);
+
+	MenuBuilder.AddMenuSeparator();
+
+	MenuBuilder.AddSubMenu(
+		NSLOCTEXT("AssetTypeActions_StaticMesh", "StaticMesh_HiRes", "High Res"),
+		NSLOCTEXT("AssetTypeActions_StaticMesh", "StaticMesh_HiRestooltip", "High resolution version of the mesh. Used instead of LOD0 by systems that can benefit from more detailed geometry (Nanite, modeling, etc.)"),
+		FNewMenuDelegate::CreateSP(this, &FAssetTypeActions_StaticMesh::GetImportHiResMenu, Meshes)
+	);
 }
 
 void FAssetTypeActions_StaticMesh::ExecuteImportMeshLOD(UObject* Mesh, int32 LOD)
@@ -252,6 +263,85 @@ void FAssetTypeActions_StaticMesh::ExecuteImportMeshLOD(UObject* Mesh, int32 LOD
 	FbxMeshUtils::ImportMeshLODDialog(Mesh, LOD);
 }
 
+void FAssetTypeActions_StaticMesh::ExecuteImportHiResMesh(UStaticMesh* StaticMesh)
+{
+	if (FbxMeshUtils::ImportStaticMeshHiResSourceModelDialog(StaticMesh))
+	{
+		FEditorDelegates::RefreshEditor.Broadcast();
+		FEditorSupportDelegates::RedrawAllViewports.Broadcast();
+	}
+}
+
+void FAssetTypeActions_StaticMesh::ExecuteRemoveHiResMesh(UStaticMesh* StaticMesh)
+{
+	if (StaticMesh && StaticMesh->IsHiResMeshDescriptionValid())
+	{
+		StaticMesh->Modify();
+
+		StaticMesh->ModifyHiResMeshDescription();
+		StaticMesh->ClearHiResMeshDescription();
+		StaticMesh->CommitHiResMeshDescription();
+
+		StaticMesh->GetHiResSourceModel().SourceImportFilename.Empty();
+
+		StaticMesh->PostEditChange();
+
+		FEditorDelegates::RefreshEditor.Broadcast();
+		FEditorSupportDelegates::RedrawAllViewports.Broadcast();
+	}
+}
+
+void FAssetTypeActions_StaticMesh::GetImportHiResMenu(class FMenuBuilder& MenuBuilder, TArray<TWeakObjectPtr<UStaticMesh>> Objects)
+{
+	check(Objects.Num() > 0);
+	TWeakObjectPtr<UStaticMesh> First = Objects[0];
+	UStaticMesh* StaticMesh = First.Get();
+
+	if (!StaticMesh)
+	{
+		return;
+	}
+
+	// Import / Reimport
+	{
+		FText Description = NSLOCTEXT("AssetTypeActions_StaticMesh", "Reimport_HighRes", "Reimport High Res");
+		FText ToolTip = NSLOCTEXT("AssetTypeActions_StaticMesh", "Reimport_HighRes_Tip", "Reimport over existing the existing high res mesh.");
+		if (!StaticMesh->IsHiResMeshDescriptionValid())
+		{
+			Description = NSLOCTEXT("AssetTypeActions_StaticMesh", "Import_HighRes", "Import High Res...");
+			ToolTip = NSLOCTEXT("AssetTypeActions_StaticMesh", "Import_HighResTip", "Adds a high resolution version of the mesh. Used instead of LOD0 by systems that can benefit from more detailed geometry (Nanite, modeling, etc.).");
+		}
+
+		MenuBuilder.AddMenuEntry(Description, ToolTip, FSlateIcon(),
+			FUIAction(FExecuteAction::CreateStatic(&FAssetTypeActions_StaticMesh::ExecuteImportHiResMesh, StaticMesh)));
+	}
+
+	// Remove
+	{
+		FText Description = NSLOCTEXT("AssetTypeActions_StaticMesh", "Remove_HighRes", "Remove High Res");
+		FText ToolTip = NSLOCTEXT("AssetTypeActions_StaticMesh", "Remove_HighRes_Tip", "Remove the high res version of the mesh");
+
+		MenuBuilder.AddMenuEntry(Description, ToolTip, FSlateIcon(),
+			FUIAction(
+				FExecuteAction::CreateStatic(&FAssetTypeActions_StaticMesh::ExecuteRemoveHiResMesh, StaticMesh),
+				FCanExecuteAction::CreateLambda(
+					[StaticMesh]()
+					{
+						if (StaticMesh)
+						{
+							return StaticMesh->IsHiResMeshDescriptionValid();
+						}
+						else
+						{
+							return false;
+						}
+					}
+				)
+			)
+		);
+	}
+}
+
 void FAssetTypeActions_StaticMesh::ExecuteCopyLODSettings(TArray<TWeakObjectPtr<UStaticMesh>> Objects)
 {
 	LODCopyMesh = Objects[0];
diff --git a/Engine/Source/Developer/AssetTools/Private/AssetTypeActions/AssetTypeActions_StaticMesh.h b/Engine/Source/Developer/AssetTools/Private/AssetTypeActions/AssetTypeActions_StaticMesh.h
--- a/Engine/Source/Developer/AssetTools/Private/AssetTypeActions/AssetTypeActions_StaticMesh.h
+++ b/Engine/Source/Developer/AssetTools/Private/AssetTypeActions/AssetTypeActions_StaticMesh.h
@@ -66,6 +66,16 @@ private:
 
 	/** Handler for calling import methods */
 	static void ExecuteImportMeshLOD(UObject* Mesh, int32 LOD);
+
+	/** Builds the High Res mesh sub-menu as part of the LOD menu */
+	void GetImportHiResMenu(class FMenuBuilder& MenuBuilder, TArray<TWeakObjectPtr<UStaticMesh>> Objects);
+
+	/** Handles import and reimport of the high res source model */
+	static void ExecuteImportHiResMesh(UStaticMesh* Mesh);
+
+	/** Handles removing the high res source model from the mesh */
+	static void ExecuteRemoveHiResMesh(UStaticMesh* Mesh);
+
 private:
 
 	TWeakObjectPtr<UStaticMesh> LODCopyMesh;	
diff --git a/Engine/Source/Editor/UnrealEd/Private/FbxMeshUtils.cpp b/Engine/Source/Editor/UnrealEd/Private/FbxMeshUtils.cpp
--- a/Engine/Source/Editor/UnrealEd/Private/FbxMeshUtils.cpp
+++ b/Engine/Source/Editor/UnrealEd/Private/FbxMeshUtils.cpp
@@ -12,6 +12,7 @@
 #include "Engine/StaticMesh.h"
 #include "EditorDirectories.h"
 #include "Framework/Application/SlateApplication.h"
+#include "MeshDescription.h"
 #include "Misc/MessageDialog.h"
 #include "ComponentReregisterContext.h"
 #include "Logging/TokenizedMessage.h"
@@ -47,8 +48,48 @@ DEFINE_LOG_CATEGORY_STATIC(LogExportMeshUtils, Log, All);
 
 namespace FbxMeshUtils
 {
+	namespace Private
+	{
+		void SetupFbxImportOptions(const UStaticMesh* BaseStaticMesh, UnFbx::FBXImportOptions* ImportOptions)
+		{
+			check(BaseStaticMesh);
+
+			UFbxStaticMeshImportData* ImportData = Cast<UFbxStaticMeshImportData>(BaseStaticMesh->AssetImportData);
+			if (ImportData != nullptr)
+			{
+
+				UFbxImportUI* ReimportUI = NewObject<UFbxImportUI>();
+				ReimportUI->MeshTypeToImport = FBXIT_StaticMesh;
+				UnFbx::FBXImportOptions::ResetOptions(ImportOptions);
+				// Import data already exists, apply it to the fbx import options
+				ReimportUI->StaticMeshImportData = ImportData;
+				ApplyImportUIToImportOptions(ReimportUI, *ImportOptions);
+			}
+			else if (BaseStaticMesh->IsSourceModelValid(0))
+			{
+				// Use the lod 0 to set the import settings
+				const FStaticMeshSourceModel& SourceModel = BaseStaticMesh->GetSourceModel(0);
+				ImportOptions->NormalGenerationMethod = SourceModel.BuildSettings.bUseMikkTSpace ? EFBXNormalGenerationMethod::MikkTSpace : EFBXNormalGenerationMethod::BuiltIn;
+				ImportOptions->bComputeWeightedNormals = SourceModel.BuildSettings.bComputeWeightedNormals;
+				ImportOptions->DistanceFieldResolutionScale = SourceModel.BuildSettings.DistanceFieldResolutionScale;
+				ImportOptions->bRemoveDegenerates = SourceModel.BuildSettings.bRemoveDegenerates;
+				ImportOptions->bBuildReversedIndexBuffer = SourceModel.BuildSettings.bBuildReversedIndexBuffer;
+				ImportOptions->bGenerateLightmapUVs = SourceModel.BuildSettings.bGenerateLightmapUVs;
+			}
+
+			// Set a couple of settings that shouldn't change while importing a lod
+			ImportOptions->bBuildNanite = BaseStaticMesh->NaniteSettings.bEnabled;
+			ImportOptions->StaticMeshLODGroup = BaseStaticMesh->LODGroup;
+			ImportOptions->bIsImportCancelable = false;
+			ImportOptions->bImportMaterials = false;
+			ImportOptions->bImportTextures = false;
+
+			ImportOptions->bAutoComputeLodDistances = true; //Setting auto compute distance to true will avoid changing the staticmesh flag
+		}
+	}
+
 	/** Helper function used for retrieving data required for importing static mesh LODs */
-	void PopulateFBXStaticMeshLODList(UnFbx::FFbxImporter* FFbxImporter, FbxNode* Node, TArray< TArray<FbxNode*>* >& LODNodeList, int32& MaxLODCount, bool bUseLODs)
+	void PopulateFBXStaticMeshLODList(UnFbx::FFbxImporter* FFbxImporter, FbxNode* Node, TArray< TUniquePtr<TArray<FbxNode*>> >& LODNodeList, int32& MaxLODCount, bool bUseLODs)
 	{
 		// Check for LOD nodes, if one is found, add it to the list
 		if (bUseLODs && Node->GetNodeAttribute() && Node->GetNodeAttribute()->GetAttributeType() == FbxNodeAttribute::eLODGroup)
@@ -57,8 +98,7 @@ namespace FbxMeshUtils
 			{
 				if ((LODNodeList.Num() - 1) < ChildIdx)
 				{
-					TArray<FbxNode*>* NodeList = new TArray<FbxNode*>;
-					LODNodeList.Add(NodeList);
+					LODNodeList.Add(MakeUnique<TArray<FbxNode*>>());
 				}
 				FFbxImporter->FindAllLODGroupNode(*(LODNodeList[ChildIdx]), Node, ChildIdx);
 			}
@@ -75,8 +115,7 @@ namespace FbxMeshUtils
 			{
 				if (LODNodeList.Num() == 0)
 				{
-					TArray<FbxNode*>* NodeList = new TArray<FbxNode*>;
-					LODNodeList.Add(NodeList);
+					LODNodeList.Add(MakeUnique<TArray<FbxNode*>>());
 				}
 
 				LODNodeList[0]->Add(Node);
@@ -90,7 +129,7 @@ namespace FbxMeshUtils
 		}
 	}
 
-	bool ImportStaticMeshLOD( UStaticMesh* BaseStaticMesh, const FString& Filename, int32 LODLevel)
+	bool ImportStaticMeshLOD( UStaticMesh* BaseStaticMesh, const FString& Filename, int32 LODLevel )
 	{
 		if (!BaseStaticMesh)
 		{
@@ -143,39 +182,12 @@ namespace FbxMeshUtils
 		UnFbx::FFbxLoggerSetter Logger(FFbxImporter);
 
 		UnFbx::FBXImportOptions* ImportOptions = FFbxImporter->GetImportOptions();
-		
-		bool IsReimport = BaseStaticMesh->GetRenderData()->LODResources.Num() > LODLevel;
+		Private::SetupFbxImportOptions(BaseStaticMesh, ImportOptions);
+
 		UFbxStaticMeshImportData* ImportData = Cast<UFbxStaticMeshImportData>(BaseStaticMesh->AssetImportData);
-		if (ImportData != nullptr)
-		{
-			
-			UFbxImportUI* ReimportUI = NewObject<UFbxImportUI>();
-			ReimportUI->MeshTypeToImport = FBXIT_StaticMesh;
-			UnFbx::FBXImportOptions::ResetOptions(ImportOptions);
-			// Import data already exists, apply it to the fbx import options
-			ReimportUI->StaticMeshImportData = ImportData;
-			ApplyImportUIToImportOptions(ReimportUI, *ImportOptions);
-		}
-		else if (BaseStaticMesh->IsSourceModelValid(0))
-		{
-			// Use the lod 0 to set the import settings
-			FStaticMeshSourceModel& SourceModel = BaseStaticMesh->GetSourceModel(0);
-			ImportOptions->NormalGenerationMethod = SourceModel.BuildSettings.bUseMikkTSpace ? EFBXNormalGenerationMethod::MikkTSpace : EFBXNormalGenerationMethod::BuiltIn;
-			ImportOptions->bComputeWeightedNormals = SourceModel.BuildSettings.bComputeWeightedNormals;
-			ImportOptions->DistanceFieldResolutionScale = SourceModel.BuildSettings.DistanceFieldResolutionScale;
-			ImportOptions->bRemoveDegenerates = SourceModel.BuildSettings.bRemoveDegenerates;
-			ImportOptions->bBuildReversedIndexBuffer = SourceModel.BuildSettings.bBuildReversedIndexBuffer;
-			ImportOptions->bGenerateLightmapUVs = SourceModel.BuildSettings.bGenerateLightmapUVs;
-		}
-
-		// Set a couple of settings that shouldn't change while importing a lod
-		ImportOptions->bBuildNanite = BaseStaticMesh->NaniteSettings.bEnabled;
-		ImportOptions->StaticMeshLODGroup = BaseStaticMesh->LODGroup;
-		ImportOptions->bIsImportCancelable = false;
-		ImportOptions->bImportMaterials = false;
-		ImportOptions->bImportTextures = false;
-
-		ImportOptions->bAutoComputeLodDistances = true; //Setting auto compute distance to true will avoid changing the staticmesh flag
+
+		const bool bIsReimport = BaseStaticMesh->GetRenderData()->LODResources.Num() > LODLevel;
+
 		if ( !FFbxImporter->ImportFromFile( *Filename, FPaths::GetExtension( Filename ), true ) )
 		{
 			// Log the error message and fail the import.
@@ -192,7 +204,7 @@ namespace FbxMeshUtils
 
 			bool bUseLODs = true;
 			int32 MaxLODLevel = 0;
-			TArray< TArray<FbxNode*>* > LODNodeList;
+			TArray< TUniquePtr<TArray<FbxNode*>> > LODNodeList;
 
 			// Create a list of LOD nodes
 			PopulateFBXStaticMeshLODList(FFbxImporter, FFbxImporter->Scene->GetRootNode(), LODNodeList, MaxLODLevel, bUseLODs);
@@ -217,7 +229,7 @@ namespace FbxMeshUtils
 			}
 
 			TSharedPtr<FExistingStaticMeshData> ExistMeshDataPtr;
-			if (IsReimport)
+			if (bIsReimport)
 			{
 				ExistMeshDataPtr = StaticMeshImportUtils::SaveExistingStaticMeshData(BaseStaticMesh, FFbxImporter->ImportOptions, LODLevel);
 			}
@@ -253,7 +265,7 @@ namespace FbxMeshUtils
 					TArray<int32> ReimportLodList;
 					ReimportLodList.Add(LODLevel);
 					StaticMeshImportUtils::UpdateSomeLodsImportMeshData(BaseStaticMesh, &ReimportLodList);
-					if(IsReimport)
+					if(bIsReimport)
 					{
 						StaticMeshImportUtils::RestoreExistingMeshData(ExistMeshDataPtr, BaseStaticMesh, LODLevel, false, ImportOptions->bResetToFbxOnMaterialConflict);
 					}
@@ -286,13 +298,118 @@ namespace FbxMeshUtils
 					bSuccess = false;
 				}
 			}
+		}
+		FFbxImporter->ReleaseScene();
 
-			// Cleanup
-			for (int32 i = 0; i < LODNodeList.Num(); ++i)
+		return bSuccess;
+	}
+
+	bool ImportStaticMeshHiResSourceModel(UStaticMesh* BaseStaticMesh, const FString& Filename)
+	{
+		if (!BaseStaticMesh)
+		{
+			UE_LOG(LogExportMeshUtils, Log, TEXT("Cannot import custom high res mesh because the staticmesh is NULL."));
+			return false;
+		}
+
+		bool bSuccess = false;
+
+		UE_LOG(LogExportMeshUtils, Log, TEXT("Fbx Mesh loading"));
+
+		UnFbx::FFbxImporter* FFbxImporter = UnFbx::FFbxImporter::GetInstance();
+		UnFbx::FFbxLoggerSetter Logger(FFbxImporter);
+
+		UnFbx::FBXImportOptions* ImportOptions = FFbxImporter->GetImportOptions();
+		Private::SetupFbxImportOptions(BaseStaticMesh, ImportOptions);
+		ImportOptions->StaticMeshLODGroup = NAME_None;
+		ImportOptions->bImportLOD = false;
+
+		UFbxStaticMeshImportData* ImportData = Cast<UFbxStaticMeshImportData>(BaseStaticMesh->AssetImportData);
+
+		const bool bPreventMaterialNameClash = true;
+		if (!FFbxImporter->ImportFromFile(*Filename, FPaths::GetExtension(Filename), bPreventMaterialNameClash))
+		{
+			FFbxImporter->FlushToTokenizedErrorMessage(EMessageSeverity::Error);
+		}
+		else
+		{
+			FFbxImporter->FlushToTokenizedErrorMessage(EMessageSeverity::Warning);
+			if (ImportData)
 			{
-				delete LODNodeList[i];
+				FFbxImporter->ApplyTransformSettingsToFbxNode(FFbxImporter->Scene->GetRootNode(), ImportData);
+			}
+
+			constexpr int32 TempLODLevel = 0; // We are importing as LOD0 in a temporary mesh then transferring the geometry to the high res source model
+			int32 MaxLODLevel = 0;
+			TArray< TUniquePtr<TArray<FbxNode*>> > MeshNodeList;
+
+			const bool bUseLODs = false;
+			PopulateFBXStaticMeshLODList(FFbxImporter, FFbxImporter->Scene->GetRootNode(), MeshNodeList, MaxLODLevel, bUseLODs);
+
+			// Nothing found, error out
+			if (MeshNodeList.Num() == 0)
+			{
+				FFbxImporter->AddTokenizedErrorMessage(FTokenizedMessage::Create(EMessageSeverity::Error, FText(LOCTEXT("HiResImport_NoMeshFound", "No meshes were found in file."))), FFbxErrors::Generic_Mesh_MeshNotFound);
+
+				FFbxImporter->ReleaseScene();
+				return bSuccess;
+			}
+
+			{
+				UStaticMesh* TempStaticMesh = NULL;
+
+				if (MeshNodeList.IsValidIndex(0))
+				{
+					TempStaticMesh = (UStaticMesh*)FFbxImporter->ImportStaticMeshAsSingle(GetTransientPackage(), *(MeshNodeList[0]), NAME_None, RF_Transient,
+						ImportData, BaseStaticMesh, TempLODLevel);
+				}
+
+				// Add the imported mesh to the existing model
+				if (TempStaticMesh && TempStaticMesh->IsSourceModelValid(0))
+				{
+					BaseStaticMesh->Modify();
+
+					FMeshDescription* HiResMeshDescription = BaseStaticMesh->GetHiResMeshDescription();
+					if (HiResMeshDescription == nullptr)
+					{
+						HiResMeshDescription = BaseStaticMesh->CreateHiResMeshDescription();
+					}
+					check(HiResMeshDescription);
+
+					BaseStaticMesh->ModifyHiResMeshDescription();
+
+					FMeshDescription* LOD0MeshDescription = TempStaticMesh->GetMeshDescription(0);
+					check(LOD0MeshDescription);
+
+					*HiResMeshDescription = MoveTemp(*LOD0MeshDescription);
+
+					BaseStaticMesh->CommitHiResMeshDescription();
+
+					BaseStaticMesh->PostEditChange();
+					BaseStaticMesh->MarkPackageDirty();
+
+					FNotificationInfo NotificationInfo(FText::GetEmpty());
+					NotificationInfo.Text = FText::Format(LOCTEXT("HiResMeshImportSuccessful", "High res mesh imported successfully!"), FText::AsNumber(0));
+					NotificationInfo.ExpireDuration = 5.0f;
+					FSlateNotificationManager::Get().AddNotification(NotificationInfo);
+
+					FStaticMeshSourceModel& SourceModel = BaseStaticMesh->GetHiResSourceModel();
+					SourceModel.SourceImportFilename = UAssetImportData::SanitizeImportFilename(Filename, nullptr);
+					SourceModel.bImportWithBaseMesh = false;
+
+					bSuccess = true;
+				}
+
+				if (!bSuccess) // Import failed
+				{
+					FNotificationInfo NotificationInfo(FText::GetEmpty());
+					NotificationInfo.Text = FText::Format(LOCTEXT("HiResMeshImportFail", "Failed to import high res mesh!"), FText::AsNumber(0));
+					NotificationInfo.ExpireDuration = 5.0f;
+					FSlateNotificationManager::Get().AddNotification(NotificationInfo);
+				}
 			}
 		}
+
 		FFbxImporter->ReleaseScene();
 
 		return bSuccess;
@@ -850,6 +967,72 @@ namespace FbxMeshUtils
 		return bImportSuccess;
 	}
 
+	bool ImportStaticMeshHiResSourceModelDialog( UStaticMesh* StaticMesh )
+	{
+		if (!StaticMesh)
+		{
+			return false;
+		}
+
+		// TODO: Interchange support
+
+		FString FilenameToImport("");
+
+		const FStaticMeshSourceModel& SourceModel = StaticMesh->GetHiResSourceModel();
+		FilenameToImport = SourceModel.SourceImportFilename.IsEmpty() ?
+			SourceModel.SourceImportFilename :
+			UAssetImportData::ResolveImportFilename(SourceModel.SourceImportFilename, nullptr);
+
+		// Check if the file exists first
+		const bool bSourceFileExists = FPaths::FileExists(FilenameToImport);
+
+		// We'll give the user a chance to choose a new file if a previously set file fails to import
+		const bool bPromptOnFail = bSourceFileExists;
+
+		if (!bSourceFileExists || FilenameToImport.IsEmpty())
+		{
+			FText PromptTitle;
+
+			if (FilenameToImport.IsEmpty())
+			{
+				PromptTitle = LOCTEXT("HiResImportPrompt_NoSource", "Choose a file to import for the High Resolution Mesh");
+			}
+			else if (!bSourceFileExists)
+			{
+				PromptTitle = LOCTEXT("HiResImportPrompt_SourceNotFound", "High Resolution Mesh Source file not found. Choose a new file.");
+			}
+
+			FilenameToImport = PromptForLODImportFile(PromptTitle);
+		}
+
+		bool bImportSuccess = false;
+
+		if (!FilenameToImport.IsEmpty())
+		{
+			bImportSuccess = ImportStaticMeshHiResSourceModel(StaticMesh, FilenameToImport);
+		}
+
+		if (!bImportSuccess && bPromptOnFail)
+		{
+			FMessageDialog::Open(EAppMsgType::Ok, LOCTEXT("HiResImport_SourceMissingDialog", "Failed to import the High Resolution Mesh as the source file failed to import, please select a new source file."));
+
+			FText PromptTitle = LOCTEXT("HiResImportPrompt_SourceFailed", "Failed to import source file for the High Resolution Mesh, choose a new file");
+			FilenameToImport = PromptForLODImportFile(PromptTitle);
+
+			if (FilenameToImport.Len() > 0 && FPaths::FileExists(FilenameToImport))
+			{
+				bImportSuccess = ImportStaticMeshHiResSourceModel(StaticMesh, FilenameToImport);
+			}
+		}
+
+		if (!bImportSuccess && !FilenameToImport.IsEmpty())
+		{
+			FMessageDialog::Open(EAppMsgType::Ok, LOCTEXT("HiResImport_Failure", "Failed to import the High Resolution Mesh"));
+		}
+
+		return bImportSuccess;
+	}
+
 	void SetImportOption(UFbxImportUI* ImportUI)
 	{
 		UnFbx::FFbxImporter* FFbxImporter = UnFbx::FFbxImporter::GetInstance();
diff --git a/Engine/Source/Editor/UnrealEd/Public/FbxMeshUtils.h b/Engine/Source/Editor/UnrealEd/Public/FbxMeshUtils.h
--- a/Engine/Source/Editor/UnrealEd/Public/FbxMeshUtils.h
+++ b/Engine/Source/Editor/UnrealEd/Public/FbxMeshUtils.h
@@ -21,6 +21,15 @@ namespace FbxMeshUtils
 	 */
 	UNREALED_API bool ImportStaticMeshLOD( UStaticMesh* BaseStaticMesh, const FString& Filename, int32 LODLevel );
 
+	/**
+	 * Imports a mesh as the high res source model of the given static mesh
+	 *
+	 * @param BaseStaticMesh	The static mesh to import the high res source model to
+	 * @param Filename			The filename of the FBX file containing the mesh
+	 * @return					Whether or not the import succeeded
+	 */
+	UNREALED_API bool ImportStaticMeshHiResSourceModel( UStaticMesh* BaseStaticMesh, const FString& Filename );
+
 	/**
 	 * Imports a skeletal mesh LOD to the given skeletal mesh
 	 *
@@ -32,14 +41,23 @@ namespace FbxMeshUtils
 	UNREALED_API bool ImportSkeletalMeshLOD( USkeletalMesh* Mesh, const FString& Filename, int32 LODLevel );
 
 	/**
-	 * Imports a skeletal mesh LOD to the given skeletal mesh
+	 * Imports a mesh LOD to the given mesh
 	 * Opens file dialog to do so!
 	 *
-	 * @param Mesh				The skeletal mesh to import the LOD to
+	 * @param Mesh				The mesh to import the LOD to
 	 * @param LODLevel			The level of the lod to import
 	 * @return					Whether or not the import succeeded
 	 */
-	UNREALED_API bool ImportMeshLODDialog( UObject* Mesh, int32 LODLevel, bool bNotifyCB = true);
+	UNREALED_API bool ImportMeshLODDialog( UObject* Mesh, int32 LODLevel, bool bNotifyCB = true );
+
+	/**
+	 * Imports a mesh as the high res source model of the given mesh
+	 * Opens file dialog to do so!
+	 *
+	 * @param Mesh				The mesh to import the high res mesh to
+	 * @return					Whether or not the import succeeded
+	 */
+	UNREALED_API bool ImportStaticMeshHiResSourceModelDialog( UStaticMesh* StaticMesh );
 
 	/**
 	 * Sets import option before importing
