Fix to enable root motion mode "Loop" to be selected in Animation Viewport while editing a Blend Space and Animation Blueprints

**Changed**
- Now `DoesCurrentAssetHaveRootMotion` method checks `UBlendSpaces` and `UAnimBlueprint` besides `UAnimSequence` for root motion .
- Now `CanUseProcessRootMotionMode`method disables "Loop and Reset" mode for blendspaces AND animation blueprints

diff --git a/Engine/Source/Editor/UnrealEd/Classes/Animation/DebugSkelMeshComponent.h b/Engine/Source/Editor/UnrealEd/Classes/Animation/DebugSkelMeshComponent.h
--- a/Engine/Source/Editor/UnrealEd/Classes/Animation/DebugSkelMeshComponent.h
+++ b/Engine/Source/Editor/UnrealEd/Classes/Animation/DebugSkelMeshComponent.h
@@ -449,7 +449,7 @@ class UNREALED_API UDebugSkelMeshComponent : public USkeletalMeshComponent
 	/** Whether the supplied root motion mode can be used for the current asset */
 	bool CanUseProcessRootMotionMode(EProcessRootMotionMode Mode) const;
 
-	/** Whether the current asset is using root motion */
+	/** Whether the current asset or animation blueprint is using root motion */
 	bool DoesCurrentAssetHaveRootMotion() const;
 
 	/** Whether the current LOD of the debug mesh is being synced with the attached (preview) mesh instance. */
diff --git a/Engine/Source/Editor/UnrealEd/Private/Animation/DebugSkelMeshComponent.cpp b/Engine/Source/Editor/UnrealEd/Private/Animation/DebugSkelMeshComponent.cpp
--- a/Engine/Source/Editor/UnrealEd/Private/Animation/DebugSkelMeshComponent.cpp
+++ b/Engine/Source/Editor/UnrealEd/Private/Animation/DebugSkelMeshComponent.cpp
@@ -217,7 +217,7 @@ void UDebugSkelMeshComponent::ConsumeRootMotion(const FVector& FloorMin, const F
 		return FTransform::Identity;
 	};
 
-	// Force ProcessRootMotionMode to Ignore if the current asset is not using root motion. 
+	// Force ProcessRootMotionMode to Ignore if the current asset/animation blueprint is not using root motion. 
 	if(ProcessRootMotionMode != EProcessRootMotionMode::Ignore && DoesCurrentAssetHaveRootMotion() == false)
 	{
 		SetProcessRootMotionMode(EProcessRootMotionMode::Ignore);
@@ -356,12 +356,43 @@ EProcessRootMotionMode UDebugSkelMeshComponent::GetProcessRootMotionMode() const
 
 bool UDebugSkelMeshComponent::DoesCurrentAssetHaveRootMotion() const
 {
-	if(PreviewInstance)
+	if (PreviewInstance)
 	{
-		if(UAnimSequenceBase* AnimSequenceBase = Cast<UAnimSequenceBase>(PreviewInstance->GetCurrentAsset()))
+		// Allow root motion if current asset is sequence
+		if(const UAnimSequenceBase* AnimSequenceBase = Cast<UAnimSequenceBase>(PreviewInstance->GetCurrentAsset()))
 		{
 			return AnimSequenceBase->HasRootMotion();
 		}
+		
+		// Allow root motion if current blend-space references any sequences with root motion
+		if (const UBlendSpace* BlendSpace = Cast<UBlendSpace>(PreviewInstance->GetCurrentAsset()))
+		{
+			bool bIsRootMotionUsedInBlendSpace = false;
+			
+			BlendSpace->ForEachImmutableSample([&bIsRootMotionUsedInBlendSpace](const FBlendSample & Sample)
+			{
+				const TObjectPtr<UAnimSequence> Sequence = Sample.Animation;
+				
+				if (IsValid(Sequence) && Sequence->HasRootMotion())
+				{
+					bIsRootMotionUsedInBlendSpace = true;
+				}
+			});
+			
+			if (bIsRootMotionUsedInBlendSpace)
+			{
+				return true;
+			}
+		}
+
+		// Allow previewing an animation blueprint with root motion
+		if (!PreviewInstance->GetCurrentAsset())
+		{
+			if (PreviewInstance->RootMotionMode == ERootMotionMode::RootMotionFromEverything || PreviewInstance->RootMotionMode == ERootMotionMode::RootMotionFromMontagesOnly)
+			{
+				return true;
+			}
+		}
 	}
 
 	return false;
@@ -374,7 +405,7 @@ bool UDebugSkelMeshComponent::CanUseProcessRootMotionMode(EProcessRootMotionMode
 		return false;
 	}
 	
-	// Disable Loop modes if the current asset doesn't have root motion
+	// Disable Loop modes if the current asset or animation blueprint doesn't have root motion
 	if(Mode != EProcessRootMotionMode::Ignore)
 	{
 		if(!DoesCurrentAssetHaveRootMotion())
@@ -383,10 +414,10 @@ bool UDebugSkelMeshComponent::CanUseProcessRootMotionMode(EProcessRootMotionMode
 		}
 	}
 
-	// Disable Loop and Reset mode for blend spaces
+	// Disable Loop and Reset mode for blend spaces and animation blueprints
 	if (Mode == EProcessRootMotionMode::LoopAndReset)
 	{
-		if (UBlendSpace* BlendSpace = Cast<UBlendSpace>(PreviewInstance->GetCurrentAsset()))
+		if (Cast<UBlendSpace>(PreviewInstance->GetCurrentAsset()) || !PreviewInstance->GetCurrentAsset())
 		{
 			return false;
 		}
diff --git a/Engine/Source/Runtime/Engine/Classes/Animation/BlendSpace.h b/Engine/Source/Runtime/Engine/Classes/Animation/BlendSpace.h
--- a/Engine/Source/Runtime/Engine/Classes/Animation/BlendSpace.h
+++ b/Engine/Source/Runtime/Engine/Classes/Animation/BlendSpace.h
@@ -541,7 +541,13 @@ public:
 	 * @param	InDeltaTime				The tick time for this update
 	 */
 	ENGINE_API bool UpdateBlendSamples(const FVector& InBlendSpacePosition, float InDeltaTime, TArray<FBlendSampleData>& InOutSampleDataCache, int32& InOutCachedTriangulationIndex) const;
-
+	
+	/**
+	 * Allows the user to iterate through all the data samples available in the blend space.
+	 * @param Func The function to run for each blend sample
+	 */
+	ENGINE_API void ForEachImmutableSample(const TFunctionRef<void(const FBlendSample&)> Func) const;
+	
 	/** Interpolate BlendInput based on Filter data **/
 	ENGINE_API FVector FilterInput(FBlendFilter* Filter, const FVector& BlendInput, float DeltaTime) const;
 
diff --git a/Engine/Source/Runtime/Engine/Private/Animation/BlendSpace.cpp b/Engine/Source/Runtime/Engine/Private/Animation/BlendSpace.cpp
--- a/Engine/Source/Runtime/Engine/Private/Animation/BlendSpace.cpp
+++ b/Engine/Source/Runtime/Engine/Private/Animation/BlendSpace.cpp
@@ -281,6 +281,14 @@ bool UBlendSpace::UpdateBlendSamples(const FVector& InBlendSpacePosition, float
 	return bResult;
 }
 
+void UBlendSpace::ForEachImmutableSample(const TFunctionRef<void(const FBlendSample&)> Func) const
+{
+	for (const FBlendSample & Sample : SampleData)
+	{
+		Func(Sample);
+	}
+}
+
 void UBlendSpace::TickAssetPlayer(FAnimTickRecord& Instance, struct FAnimNotifyQueue& NotifyQueue, FAnimAssetTickContext& Context) const
 {
 	check(Instance.BlendSpace.BlendSampleDataCache);
