Fixed unnecessary performance drain when doing interactive property changes.
Previously, the thumbnail of an object would be constantly dirtied while dragging, leading to poor editor performance

diff --git a/Engine/Source/Editor/UnrealEd/Private/AssetThumbnail.cpp b/Engine/Source/Editor/UnrealEd/Private/AssetThumbnail.cpp
--- a/Engine/Source/Editor/UnrealEd/Private/AssetThumbnail.cpp
+++ b/Engine/Source/Editor/UnrealEd/Private/AssetThumbnail.cpp
@@ -1546,7 +1546,10 @@ void FAssetThumbnailPool::OnActorPostEditMove( AActor* Actor )
 
 void FAssetThumbnailPool::OnObjectPropertyChanged( UObject* ObjectBeingModified, FPropertyChangedEvent& PropertyChangedEvent )
 {
-	DirtyThumbnailForObject(ObjectBeingModified);
+	if(PropertyChangedEvent.ChangeType != EPropertyChangeType::Interactive)
+	{
+		DirtyThumbnailForObject(ObjectBeingModified);
+	}
 }
 
 void FAssetThumbnailPool::DirtyThumbnailForObject(UObject* ObjectBeingModified)
