Fix non-unity

#rnx

diff --git a/Engine/Source/Editor/UnrealEd/Private/SWorldPartitionViewportWidget.cpp b/Engine/Source/Editor/UnrealEd/Private/SWorldPartitionViewportWidget.cpp
--- a/Engine/Source/Editor/UnrealEd/Private/SWorldPartitionViewportWidget.cpp
+++ b/Engine/Source/Editor/UnrealEd/Private/SWorldPartitionViewportWidget.cpp
@@ -11,6 +11,7 @@
 #include "Engine/World.h"
 #include "WorldPartition/WorldPartition.h"
 #include "Styling/AppStyle.h"
+#include "Modules/ModuleManager.h"
 
 #define LOCTEXT_NAMESPACE "WorldPartitionViewportWidget"
 
