Fix level editor dolly in perspective viewport by backing out changes made in CL 20626859 (persona orbit camera).

diff --git a/Engine/Source/Editor/UnrealEd/Private/EditorViewportClient.cpp b/Engine/Source/Editor/UnrealEd/Private/EditorViewportClient.cpp
--- a/Engine/Source/Editor/UnrealEd/Private/EditorViewportClient.cpp
+++ b/Engine/Source/Editor/UnrealEd/Private/EditorViewportClient.cpp
@@ -3545,7 +3545,12 @@ void FEditorViewportClient::OnDollyPerspectiveCamera( const FInputEventState& In
 	FKey Key = InputState.GetKey();
 
 	// Scrolling the mousewheel up/down moves the perspective viewport forwards/backwards.
-	FVector Drag = (GetLookAtLocation() - GetViewLocation()).GetSafeNormal();
+	FVector Drag(0,0,0);
+
+	const FRotator& ViewRotation = GetViewRotation();
+	Drag.X = FMath::Cos( ViewRotation.Yaw * PI / 180.f ) * FMath::Cos( ViewRotation.Pitch * PI / 180.f );
+	Drag.Y = FMath::Sin( ViewRotation.Yaw * PI / 180.f ) * FMath::Cos( ViewRotation.Pitch * PI / 180.f );
+	Drag.Z = FMath::Sin( ViewRotation.Pitch * PI / 180.f );
 
 	if( Key == EKeys::MouseScrollDown )
 	{
