Fixed SMInstance proxy objects lingering in locked details view after they had become invalid

#rnx

diff --git a/Engine/Source/Editor/PropertyEditor/Private/ObjectPropertyNode.cpp b/Engine/Source/Editor/PropertyEditor/Private/ObjectPropertyNode.cpp
--- a/Engine/Source/Editor/PropertyEditor/Private/ObjectPropertyNode.cpp
+++ b/Engine/Source/Editor/PropertyEditor/Private/ObjectPropertyNode.cpp
@@ -91,8 +91,10 @@ void FObjectPropertyNode::ClearObjectPackageOverrides()
 }
 
 // Purges any objects marked pending kill from the object list
-void FObjectPropertyNode::PurgeKilledObjects()
+bool FObjectPropertyNode::PurgeKilledObjects()
 {
+	bool bDidPurgeObjects = false;
+
 	// Purge any objects that are marked pending kill from the object list
 	for (int32 Index = 0; Index < Objects.Num(); )
 	{
@@ -101,12 +103,15 @@ void FObjectPropertyNode::PurgeKilledObjects()
 		if ( !Object.IsValid() )
 		{
 			Objects.RemoveAt(Index, 1);
+			bDidPurgeObjects = true;
 		}
 		else
 		{
 			++Index;
 		}
 	}
+
+	return bDidPurgeObjects;
 }
 
 // Called when the object list is finalized, Finalize() finishes the property window setup.
diff --git a/Engine/Source/Editor/PropertyEditor/Private/ObjectPropertyNode.h b/Engine/Source/Editor/PropertyEditor/Private/ObjectPropertyNode.h
--- a/Engine/Source/Editor/PropertyEditor/Private/ObjectPropertyNode.h
+++ b/Engine/Source/Editor/PropertyEditor/Private/ObjectPropertyNode.h
@@ -68,8 +68,9 @@ public:
 
 	/**
 	 * Purges any objects marked pending kill from the object list
+	 * @return True if any objects were purged, or false otherwise
 	 */
-	void PurgeKilledObjects();
+	bool PurgeKilledObjects();
 
 	// Called when the object list is finalized, Finalize() finishes the property window setup.
 	void Finalize();
diff --git a/Engine/Source/Editor/PropertyEditor/Private/SDetailsViewBase.cpp b/Engine/Source/Editor/PropertyEditor/Private/SDetailsViewBase.cpp
--- a/Engine/Source/Editor/PropertyEditor/Private/SDetailsViewBase.cpp
+++ b/Engine/Source/Editor/PropertyEditor/Private/SDetailsViewBase.cpp
@@ -928,14 +928,15 @@ void SDetailsViewBase::Tick( const FGeometry& AllottedGeometry, const double InC
 	FRootPropertyNodeList& RootPropertyNodes = GetRootNodes();
 
 	bool bHadDeferredActions = DeferredActions.Num() > 0;
-	auto PreProcessRootNode = [bHadDeferredActions, this](TSharedPtr<FComplexPropertyNode> RootPropertyNode) 
+	bool bDidPurgeObjects = false;
+	auto PreProcessRootNode = [bHadDeferredActions, &bDidPurgeObjects, this](TSharedPtr<FComplexPropertyNode> RootPropertyNode)
 	{
 		check(RootPropertyNode.IsValid());
 
 		// Purge any objects that are marked pending kill from the object list
 		if (FObjectPropertyNode* ObjectRoot = RootPropertyNode->AsObjectNode())
 		{
-			ObjectRoot->PurgeKilledObjects();
+			bDidPurgeObjects |= ObjectRoot->PurgeKilledObjects();
 		}
 
 		if (bHadDeferredActions)
@@ -990,7 +991,7 @@ void SDetailsViewBase::Tick( const FGeometry& AllottedGeometry, const double InC
 
 	int32 FoundIndex = RootPropertyNodes.Find(LastRootPendingKill);
 	bool bUpdateFilteredDetails = false;
-	if (FoundIndex != INDEX_NONE)
+	if (FoundIndex != INDEX_NONE || bDidPurgeObjects)
 	{ 
 		// Reacquire the root property nodes.  It may have been changed by the deferred actions if something like a blueprint editor forcefully resets a details panel during a posteditchange
 		ForceRefresh();
diff --git a/Engine/Source/Editor/UnrealEd/Private/Elements/SMInstance/SMInstanceElementDetailsInterface.cpp b/Engine/Source/Editor/UnrealEd/Private/Elements/SMInstance/SMInstanceElementDetailsInterface.cpp
--- a/Engine/Source/Editor/UnrealEd/Private/Elements/SMInstance/SMInstanceElementDetailsInterface.cpp
+++ b/Engine/Source/Editor/UnrealEd/Private/Elements/SMInstance/SMInstanceElementDetailsInterface.cpp
@@ -9,28 +9,37 @@ class FSMInstanceTypedElementDetailsObject : public ITypedElementDetailsObject
 {
 public:
 	explicit FSMInstanceTypedElementDetailsObject(const FSMInstanceElementId& InSMInstanceElementId)
-		: InstanceProxyObject(NewObject<USMInstanceElementDetailsProxyObject>())
 	{
-		InstanceProxyObject->Initialize(InSMInstanceElementId);
+		USMInstanceElementDetailsProxyObject* InstanceProxyObjectPtr = NewObject<USMInstanceElementDetailsProxyObject>();
+		InstanceProxyObjectPtr->Initialize(InSMInstanceElementId);
+
+		InstanceProxyObject = InstanceProxyObjectPtr;
 	}
 
 	~FSMInstanceTypedElementDetailsObject()
 	{
-		InstanceProxyObject->Shutdown();
+		if (USMInstanceElementDetailsProxyObject* InstanceProxyObjectPtr = InstanceProxyObject.Get())
+		{
+			InstanceProxyObjectPtr->Shutdown();
+		}
 	}
 
 	virtual UObject* GetObject() override
 	{
-		return InstanceProxyObject;
+		return InstanceProxyObject.Get();
 	}
 
 	virtual void AddReferencedObjects(FReferenceCollector& Collector) override
 	{
-		Collector.AddReferencedObject(InstanceProxyObject);
+		if (USMInstanceElementDetailsProxyObject* InstanceProxyObjectPtr = InstanceProxyObject.Get())
+		{
+			Collector.AddReferencedObject(InstanceProxyObjectPtr);
+			InstanceProxyObject = InstanceProxyObjectPtr;
+		}
 	}
 
 private:
-	USMInstanceElementDetailsProxyObject* InstanceProxyObject = nullptr;
+	TWeakObjectPtr<USMInstanceElementDetailsProxyObject> InstanceProxyObject;
 };
 
 TUniquePtr<ITypedElementDetailsObject> USMInstanceElementDetailsInterface::GetDetailsObject(const FTypedElementHandle& InElementHandle)
diff --git a/Engine/Source/Editor/UnrealEd/Private/Elements/SMInstance/SMInstanceElementDetailsProxyObject.cpp b/Engine/Source/Editor/UnrealEd/Private/Elements/SMInstance/SMInstanceElementDetailsProxyObject.cpp
--- a/Engine/Source/Editor/UnrealEd/Private/Elements/SMInstance/SMInstanceElementDetailsProxyObject.cpp
+++ b/Engine/Source/Editor/UnrealEd/Private/Elements/SMInstance/SMInstanceElementDetailsProxyObject.cpp
@@ -18,8 +18,13 @@ void USMInstanceElementDetailsProxyObject::Initialize(const FSMInstanceElementId
 
 	TickHandle = FTSTicker::GetCoreTicker().AddTicker(TEXT("USMInstanceElementDetailsProxyObject"), 0.1f, [this](float)
 	{
-		SyncProxyStateFromInstance();
-		return true;
+		const bool bDidSyncInstance = SyncProxyStateFromInstance();
+		if (!bDidSyncInstance)
+		{
+			// The referenced instance is invalid, so mark this proxy as garbage so that the details panel stops using it
+			MarkAsGarbage();
+		}
+		return bDidSyncInstance;
 	});
 	SyncProxyStateFromInstance();
 }
@@ -35,7 +40,14 @@ void USMInstanceElementDetailsProxyObject::Shutdown()
 	ISMComponent.Reset();
 	ISMInstanceId = 0;
 
-	SyncProxyStateFromInstance();
+	Transform = FTransform::Identity;
+}
+
+void USMInstanceElementDetailsProxyObject::BeginDestroy()
+{
+	Super::BeginDestroy();
+
+	Shutdown();
 }
 
 void USMInstanceElementDetailsProxyObject::PostEditChangeChainProperty(FPropertyChangedChainEvent& PropertyChangedEvent)
@@ -81,17 +93,17 @@ void USMInstanceElementDetailsProxyObject::PostEditChangeChainProperty(FProperty
 	Super::PostEditChangeChainProperty(PropertyChangedEvent);
 }
 
-void USMInstanceElementDetailsProxyObject::SyncProxyStateFromInstance()
+bool USMInstanceElementDetailsProxyObject::SyncProxyStateFromInstance()
 {
 	if (FSMInstanceManager SMInstance = GetSMInstance())
 	{
 		// TODO: Need flag for local/world space, like FComponentTransformDetails
 		SMInstance.GetSMInstanceTransform(Transform, /*bWorldSpace*/false);
+		return true;
 	}
-	else
-	{
-		Transform = FTransform::Identity;
-	}
+	
+	Transform = FTransform::Identity;
+	return false;
 }
 
 FSMInstanceManager USMInstanceElementDetailsProxyObject::GetSMInstance() const
diff --git a/Engine/Source/Editor/UnrealEd/Private/Elements/SMInstance/SMInstanceElementDetailsProxyObject.h b/Engine/Source/Editor/UnrealEd/Private/Elements/SMInstance/SMInstanceElementDetailsProxyObject.h
--- a/Engine/Source/Editor/UnrealEd/Private/Elements/SMInstance/SMInstanceElementDetailsProxyObject.h
+++ b/Engine/Source/Editor/UnrealEd/Private/Elements/SMInstance/SMInstanceElementDetailsProxyObject.h
@@ -21,13 +21,14 @@ public:
 	void Shutdown();
 
 	//~ UObject interface
+	virtual void BeginDestroy() override;
 	virtual void PostEditChangeChainProperty(FPropertyChangedChainEvent& PropertyChangedEvent) override;
 
 	UPROPERTY(EditAnywhere, Category="Transform", meta=(ShowOnlyInnerProperties))
 	FTransform Transform;
 
 private:
-	void SyncProxyStateFromInstance();
+	bool SyncProxyStateFromInstance();
 
 	FSMInstanceManager GetSMInstance() const;
 
