ReplaceActors optional param that outputs the new actors

diff --git a/Engine/Source/Editor/UnrealEd/Classes/Editor/EditorEngine.h b/Engine/Source/Editor/UnrealEd/Classes/Editor/EditorEngine.h
--- a/Engine/Source/Editor/UnrealEd/Classes/Editor/EditorEngine.h
+++ b/Engine/Source/Editor/UnrealEd/Classes/Editor/EditorEngine.h
@@ -1253,8 +1253,11 @@ public:
 	 * note that only Location, Rotation, Drawscale, Drawscale3D, Tag, and Group are copied from the old Actors
 	 * 
 	 * @param Factory - the Factory to use to create Actors
+	 * @param AssetData - the asset to feed the Factory
+	 * @param ActorsToReplace - Actors to replace
+	 * @param OutNewActors - Actors that were created
 	 */
-	void ReplaceActors(UActorFactory* Factory, const FAssetData& AssetData, const TArray<AActor*>& ActorsToReplace);
+	void ReplaceActors(UActorFactory* Factory, const FAssetData& AssetData, const TArray<AActor*>& ActorsToReplace, TArray<AActor*>* OutNewActors = nullptr);
 
 	/**
 	 * Converts passed in brushes into a single static mesh actor. 
diff --git a/Engine/Source/Editor/UnrealEd/Private/EditorEngine.cpp b/Engine/Source/Editor/UnrealEd/Private/EditorEngine.cpp
--- a/Engine/Source/Editor/UnrealEd/Private/EditorEngine.cpp
+++ b/Engine/Source/Editor/UnrealEd/Private/EditorEngine.cpp
@@ -5123,7 +5123,7 @@ void UEditorEngine::ReplaceSelectedActors(UActorFactory* Factory, const FAssetDa
 	ReplaceActors(Factory, AssetData, ActorsToReplace);
 }
 
-void UEditorEngine::ReplaceActors(UActorFactory* Factory, const FAssetData& AssetData, const TArray<AActor*>& ActorsToReplace)
+void UEditorEngine::ReplaceActors(UActorFactory* Factory, const FAssetData& AssetData, const TArray<AActor*>& ActorsToReplace, TArray<AActor*>* OutNewActors)
 {
 	// Cache for attachment info of all actors being converted.
 	TArray<ReattachActorsHelper::FActorAttachmentCache> AttachmentInfo;
@@ -5282,12 +5282,21 @@ void UEditorEngine::ReplaceActors(UActorFactory* Factory, const FAssetData& Asse
 	// Reattaches actors based on their previous parent child relationship.
 	ReattachActorsHelper::ReattachActors(ConvertedMap, AttachmentInfo);
 
+
+	// Output new actors and
 	// Perform reference replacement on all Actors referenced by World
 	TArray<UObject*> ReferencedLevels;
-
+	if (OutNewActors)
+	{
+		OutNewActors->Reserve(ConvertedMap.Num());
+	}
 	for (const TPair<AActor*, AActor*>& ReplacedObj : ConvertedMap)
 	{
 		ReferencedLevels.AddUnique(ReplacedObj.Value->GetLevel());
+		if (OutNewActors)
+		{
+			OutNewActors->Add(ReplacedObj.Value);
+		}
 	}
 
 	for (UObject* Referencer : ReferencedLevels)
