Reduce severity of EditorDomain message about a failure to load from EditorDomain from Warning to Display.  Can happen in natural circumstances when the EditorDomain representation of a package refers to a class that no longer exists.

diff --git a/Engine/Source/Editor/UnrealEd/Private/EditorDomain/EditorDomain.cpp b/Engine/Source/Editor/UnrealEd/Private/EditorDomain/EditorDomain.cpp
--- a/Engine/Source/Editor/UnrealEd/Private/EditorDomain/EditorDomain.cpp
+++ b/Engine/Source/Editor/UnrealEd/Private/EditorDomain/EditorDomain.cpp
@@ -221,7 +221,7 @@ bool FEditorDomain::TryFindOrAddPackageSource(FScopeLock& ScopeLock, bool& bOutR
 		}
 		return false;
 	default:
-		UE_LOG(LogEditorDomain, Warning,
+		UE_LOG(LogEditorDomain, Display,
 			TEXT("Could not load package %s from EditorDomain; it will be loaded from the WorkspaceDomain: %s"),
 			*WriteToString<256>(PackageName), *PackageDigest.GetStatusString());
 		PackageSource = new FPackageSource();
