Added null checks in the UMaterialExpressionComposite for the Input- and OutputExpressions properties in order to make the material editor handle the absence of these two properties (e.g. due to a corruption in the uasset file).

We are investigating what causes this corruption to happen and we should also improve the feedback given to the user when such corruption occurs.

diff --git a/Engine/Source/Editor/UnrealEd/Private/MaterialGraph.cpp b/Engine/Source/Editor/UnrealEd/Private/MaterialGraph.cpp
--- a/Engine/Source/Editor/UnrealEd/Private/MaterialGraph.cpp
+++ b/Engine/Source/Editor/UnrealEd/Private/MaterialGraph.cpp
@@ -518,20 +518,24 @@ void UMaterialGraph::LinkGraphNodesFromMaterial()
 					{
 						// This is an unseen composite reroute expression, find the actual expression output to connect to.
 						UMaterialExpressionComposite* OwningComposite = CastChecked<UMaterialExpressionComposite>(CompositeReroute->SubgraphExpression);
-
-						UMaterialGraphNode* OutputGraphNode;
-						int32 OutputPinIndex = OwningComposite->InputExpressions->ReroutePins.FindLastByPredicate(ExpressionMatchesPredicate(CompositeReroute));
-						if (OutputPinIndex != INDEX_NONE)
-						{
-							OutputGraphNode = CastChecked<UMaterialGraphNode>(OwningComposite->InputExpressions->GraphNode);
-						}
-						else
+						
+						// If the input- and output expressions are valid, look for the output pin in the reroute lists and make a link to the current Pin.
+						if (OwningComposite->InputExpressions && OwningComposite->OutputExpressions)
 						{
-							// Output pin base in the subgraph cannot have outputs, if this reroute isn't in the inputs connect to composite's outputs
-							OutputPinIndex = OwningComposite->OutputExpressions->ReroutePins.FindLastByPredicate(ExpressionMatchesPredicate(CompositeReroute));
-							OutputGraphNode = CastChecked<UMaterialGraphNode>(OwningComposite->GraphNode);
+							UMaterialGraphNode* OutputGraphNode;
+							int32 OutputPinIndex = OwningComposite->InputExpressions->ReroutePins.FindLastByPredicate(ExpressionMatchesPredicate(CompositeReroute));
+							if (OutputPinIndex != INDEX_NONE)
+							{
+								OutputGraphNode = CastChecked<UMaterialGraphNode>(OwningComposite->InputExpressions->GraphNode);
+							}
+							else
+							{
+								// Output pin base in the subgraph cannot have outputs, if this reroute isn't in the inputs connect to composite's outputs
+								OutputPinIndex = OwningComposite->OutputExpressions->ReroutePins.FindLastByPredicate(ExpressionMatchesPredicate(CompositeReroute));
+								OutputGraphNode = CastChecked<UMaterialGraphNode>(OwningComposite->GraphNode);
+							}
+							Pin->MakeLinkTo(OutputGraphNode->GetOutputPin(OutputPinIndex));
 						}
-						Pin->MakeLinkTo(OutputGraphNode->GetOutputPin(OutputPinIndex));
 					}
 				}
 			}
diff --git a/Engine/Source/Runtime/Engine/Private/Materials/MaterialExpressions.cpp b/Engine/Source/Runtime/Engine/Private/Materials/MaterialExpressions.cpp
--- a/Engine/Source/Runtime/Engine/Private/Materials/MaterialExpressions.cpp
+++ b/Engine/Source/Runtime/Engine/Private/Materials/MaterialExpressions.cpp
@@ -11580,13 +11580,19 @@ UMaterialExpressionComposite::UMaterialExpressionComposite(const FObjectInitiali
 TArray<UMaterialExpressionReroute*> UMaterialExpressionComposite::GetCurrentReroutes() const
 {
 	TArray<UMaterialExpressionReroute*> RerouteExpressions;
-	for (const FCompositeReroute& InputReroute : InputExpressions->ReroutePins)
+	if (InputExpressions)
 	{
-		RerouteExpressions.Add(InputReroute.Expression);
+		for (const FCompositeReroute& InputReroute : InputExpressions->ReroutePins)
+		{
+			RerouteExpressions.Add(InputReroute.Expression);
+		}
 	}
-	for (const FCompositeReroute& OutputReroute : OutputExpressions->ReroutePins)
+	if (OutputExpressions)
 	{
-		RerouteExpressions.Add(OutputReroute.Expression);
+		for (const FCompositeReroute& OutputReroute : OutputExpressions->ReroutePins)
+		{
+			RerouteExpressions.Add(OutputReroute.Expression);
+		}
 	}
 	return RerouteExpressions;
 }
