Fixes adding newfiles to Uncontrolled Changelist offline mode when Source Control is disabled.
Fixes Deleted new assets staying in Uncontrolled Changelist.

#changelist validated

diff --git a/Engine/Source/Developer/UncontrolledChangelists/Private/UncontrolledChangelistState.cpp b/Engine/Source/Developer/UncontrolledChangelists/Private/UncontrolledChangelistState.cpp
--- a/Engine/Source/Developer/UncontrolledChangelists/Private/UncontrolledChangelistState.cpp
+++ b/Engine/Source/Developer/UncontrolledChangelists/Private/UncontrolledChangelistState.cpp
@@ -96,6 +96,7 @@ bool FUncontrolledChangelistState::AddFiles(const TArray<FString>& InFilenames,
 {
 	TArray<FSourceControlStateRef> FileStates;
 	ISourceControlProvider& SourceControlProvider = ISourceControlModule::Get().GetProvider();
+	IFileManager& FileManager = IFileManager::Get();
 	bool bCheckStatus = (InCheckFlags & ECheckFlags::Modified) != ECheckFlags::None;
 	bool bCheckCheckout = (InCheckFlags & ECheckFlags::NotCheckedOut) != ECheckFlags::None;
 	bool bOutChanged = false;
@@ -109,7 +110,10 @@ bool FUncontrolledChangelistState::AddFiles(const TArray<FString>& InFilenames,
 	if (!SourceControlProvider.IsAvailable())
 	{
 		int32 OldSize = OfflineFiles.Num();
-		Algo::Copy(SourceControlHelpers::AbsoluteFilenames(InFilenames), OfflineFiles);
+		Algo::CopyIf(SourceControlHelpers::AbsoluteFilenames(InFilenames), OfflineFiles, [&FileManager](const FString& Filename)
+		{
+			return FileManager.FileExists(*Filename);
+		});
 		return OldSize != OfflineFiles.Num();
 	}
 
@@ -129,7 +133,7 @@ bool FUncontrolledChangelistState::AddFiles(const TArray<FString>& InFilenames,
 		for (FSourceControlStateRef FileState : FileStates)
 		{
 			const bool bIsSourceControlled = (!FileState->IsUnknown()) && FileState->IsSourceControlled();
-			const bool bFileExists = IFileManager::Get().FileExists(*FileState->GetFilename());
+			const bool bFileExists = FileManager.FileExists(*FileState->GetFilename());
 
 			const bool bIsUncontrolled = (!bIsSourceControlled) && bFileExists;
 			// File doesn't exist and is not marked for delete
diff --git a/Engine/Source/Developer/UncontrolledChangelists/Private/UncontrolledChangelistsModule.cpp b/Engine/Source/Developer/UncontrolledChangelists/Private/UncontrolledChangelistsModule.cpp
--- a/Engine/Source/Developer/UncontrolledChangelists/Private/UncontrolledChangelistsModule.cpp
+++ b/Engine/Source/Developer/UncontrolledChangelists/Private/UncontrolledChangelistsModule.cpp
@@ -80,7 +80,7 @@ TArray<FUncontrolledChangelistStateRef> FUncontrolledChangelistsModule::GetChang
 	return UncontrolledChangelistStates;
 }
 
-bool FUncontrolledChangelistsModule::OnMakeWritable(const FString& InFilename)
+bool FUncontrolledChangelistsModule::OnMakeWritable(const TArray<FString>& InFilenames)
 {
 	bool bHasStateChanged = false;
 
@@ -92,7 +92,7 @@ bool FUncontrolledChangelistsModule::OnMakeWritable(const FString& InFilename)
 	FUncontrolledChangelist DefaultUncontrolledChangelist(FUncontrolledChangelist::DEFAULT_UNCONTROLLED_CHANGELIST_GUID, FUncontrolledChangelist::DEFAULT_UNCONTROLLED_CHANGELIST_NAME.ToString());
 	FUncontrolledChangelistsStateCache::ValueType& UncontrolledChangelistState = UncontrolledChangelistsStateCache.FindOrAdd(MoveTemp(DefaultUncontrolledChangelist), MakeShareable(new FUncontrolledChangelistState(DefaultUncontrolledChangelist)));
 
-	bHasStateChanged = UncontrolledChangelistState->AddFiles({ InFilename }, FUncontrolledChangelistState::ECheckFlags::NotCheckedOut);
+	bHasStateChanged = UncontrolledChangelistState->AddFiles(InFilenames, FUncontrolledChangelistState::ECheckFlags::NotCheckedOut);
 
 	if (bHasStateChanged)
 	{
diff --git a/Engine/Source/Developer/UncontrolledChangelists/Public/UncontrolledChangelistsModule.h b/Engine/Source/Developer/UncontrolledChangelists/Public/UncontrolledChangelistsModule.h
--- a/Engine/Source/Developer/UncontrolledChangelists/Public/UncontrolledChangelistsModule.h
+++ b/Engine/Source/Developer/UncontrolledChangelists/Public/UncontrolledChangelistsModule.h
@@ -41,11 +41,11 @@ public:
 	 TArray<FUncontrolledChangelistStateRef> GetChangelistStates() const;
 
 	 /**
-	  * Called when a file has been made writable. Adds the file to the Default Uncontrolled Changelist
-	  * @param	InFilename			The file to be added.
-	  * @return True if the file has been added to the Default Uncontrolled Changelist.
+	  * Called when files have been made writable. Adds the files to the Default Uncontrolled Changelist
+	  * @param	InFilenames			The files to be added.
+	  * @return True if the files have been added to the Default Uncontrolled Changelist.
 	  */
-	 bool OnMakeWritable(const FString& InFilename);
+	 bool OnMakeWritable(const TArray<FString>& InFilenames);
 
 	/**
 	 * Updates the status of Uncontrolled Changelists and files.
diff --git a/Engine/Source/Editor/UnrealEd/Private/EditorEngine.cpp b/Engine/Source/Editor/UnrealEd/Private/EditorEngine.cpp
--- a/Engine/Source/Editor/UnrealEd/Private/EditorEngine.cpp
+++ b/Engine/Source/Editor/UnrealEd/Private/EditorEngine.cpp
@@ -101,6 +101,7 @@
 #include "Misc/AutomationTest.h"
 #include "ActorFolder.h"
 #include "Materials/MaterialInterface.h"
+#include "UncontrolledChangelistsModule.h"
 
 // needed for the RemotePropagator
 #include "AudioDevice.h"
@@ -4213,7 +4214,7 @@ bool UEditorEngine::IsPackageValidForAutoAdding(UPackage* InPackage, const FStri
 	// Ensure the package exists, the user is running the editor (and not a commandlet or cooking), and that source control
 	// is enabled and expecting new files to be auto-added before attempting to test the validity of the package
 	if (InPackage && GIsEditor && !IsRunningCommandlet() 
-		&& ISourceControlModule::Get().IsEnabled() 
+		&& (ISourceControlModule::Get().IsEnabled() || FUncontrolledChangelistsModule::Get().IsEnabled())
 		&& GetDefault<UEditorLoadingSavingSettings>()->bSCCAutoAddNewFiles)
 	{
 		const FString CleanFilename = FPaths::GetCleanFilename(InFilename);
@@ -4255,7 +4256,7 @@ void UEditorEngine::RunDeferredMarkForAddFiles(bool)
 		return;
 	}
 
-	if(ISourceControlModule::Get().IsEnabled())
+	if (ISourceControlModule::Get().IsEnabled())
 	{
 		ISourceControlProvider& SourceControlProvider = ISourceControlModule::Get().GetProvider();
 		if(SourceControlProvider.IsAvailable())
@@ -4263,6 +4264,12 @@ void UEditorEngine::RunDeferredMarkForAddFiles(bool)
 			SourceControlProvider.Execute(ISourceControlOperation::Create<FMarkForAdd>(), SourceControlHelpers::PackageFilenames(DeferredFilesToAddToSourceControl));
 		}
 	}
+	else if (FUncontrolledChangelistsModule::Get().IsEnabled())
+	{
+		FUncontrolledChangelistsModule& UncontrolledChangelistsModule = FUncontrolledChangelistsModule::Get();
+		UncontrolledChangelistsModule.OnMakeWritable(DeferredFilesToAddToSourceControl);
+	}
+
 	// Clear the list when this run whether source control is active or not, since we do not want to accumulate those if the user is running without source control
 	DeferredFilesToAddToSourceControl.Empty();
 }
@@ -4464,10 +4471,12 @@ FSavePackageResultStruct UEditorEngine::Save(UPackage* InOuter, UObject* InAsset
 	// to the default changelist
 	if (Result == ESavePackageResult::Success && bAutoAddPkgToSCC)
 	{
+		const bool bIsUncontrolledChangelistEnabled = FUncontrolledChangelistsModule::Get().IsEnabled();
+
 		// IsPackageValidForAutoAdding should not return true if SCC is disabled
-		check(ISourceControlModule::Get().IsEnabled());
+		check(ISourceControlModule::Get().IsEnabled() || bIsUncontrolledChangelistEnabled);
 
-		if(!ISourceControlModule::Get().GetProvider().IsAvailable())
+		if(!ISourceControlModule::Get().GetProvider().IsAvailable() && !bIsUncontrolledChangelistEnabled)
 		{
 			// Show the login window here & store the file we are trying to add.
 			// We defer the add operation until we have a valid source control connection.
diff --git a/Engine/Source/Editor/UnrealEd/Private/FileHelpers.cpp b/Engine/Source/Editor/UnrealEd/Private/FileHelpers.cpp
--- a/Engine/Source/Editor/UnrealEd/Private/FileHelpers.cpp
+++ b/Engine/Source/Editor/UnrealEd/Private/FileHelpers.cpp
@@ -1875,7 +1875,7 @@ bool FEditorFileUtils::PromptToCheckoutPackagesInternal(bool bCheckDirty, const
 						{
 							// Add to PackagesNotToPromptAnyMore only if not added to Uncontrolled Changelist.
 							// If added to Uncontrolled Changelist, we want the checkout prompt to be displayed again if the file is reverted
-							if (!UncontrolledChangelistModule.OnMakeWritable(Filename))
+							if (!UncontrolledChangelistModule.OnMakeWritable({Filename}))
 							{
 								PackagesNotToPromptAnyMore.Add(PackageToMakeWritable->GetName());
 							}
