Backout 17944249 causing regression

diff --git a/Engine/Source/Editor/UnrealEd/Private/EditorServer.cpp b/Engine/Source/Editor/UnrealEd/Private/EditorServer.cpp
--- a/Engine/Source/Editor/UnrealEd/Private/EditorServer.cpp
+++ b/Engine/Source/Editor/UnrealEd/Private/EditorServer.cpp
@@ -1924,9 +1924,11 @@ void UEditorEngine::RebuildAlteredBSP()
 			AActor* Actor = static_cast<AActor*>(*It);
 			checkSlow(Actor->IsA(AActor::StaticClass()));
 
-			if (Actor->GetClass() == ABrush::StaticClass() && !FActorEditorUtils::IsABuilderBrush(Actor))
+			ABrush* SelectedBrush = Cast< ABrush >(Actor);
+			if (SelectedBrush && !FActorEditorUtils::IsABuilderBrush(Actor))
 			{
-				if(ULevel* Level = Actor->GetLevel())
+				ULevel* Level = SelectedBrush->GetLevel();
+				if (Level)
 				{
 					LevelsToRebuild.AddUnique(Level);
 				}
