[WP navmesh builder] Avoid saving empty packages, delete is already saving the packages.

diff --git a/Engine/Source/Editor/UnrealEd/Private/WorldPartition/WorldPartitionNavigationDataBuilder.cpp b/Engine/Source/Editor/UnrealEd/Private/WorldPartition/WorldPartitionNavigationDataBuilder.cpp
--- a/Engine/Source/Editor/UnrealEd/Private/WorldPartition/WorldPartitionNavigationDataBuilder.cpp
+++ b/Engine/Source/Editor/UnrealEd/Private/WorldPartition/WorldPartitionNavigationDataBuilder.cpp
@@ -141,11 +141,6 @@ bool UWorldPartitionNavigationDataBuilder::RunInternal(UWorld* World, const FCel
 			UE_LOG(LogWorldPartitionNavigationDataBuilder, Error, TEXT("Error deleting packages."));
 		}
 
-		if (!SavePackages(PackagesToClean.Array()))
-		{
-			return true;
-		}
-
 		// If we had packages to delete we need to notify the delete after the save. Else WP might try to load the deleted descriptors on the next iteration.
 		// Note: this notification is expected to be done by the SavePackages()
 		for (UPackage* Package : PackagesToClean)
@@ -195,10 +190,8 @@ bool UWorldPartitionNavigationDataBuilder::RunInternal(UWorld* World, const FCel
 			else
 			{
 				PackagesToAdd.Add(ActorPackage);
+				PackagesToSave.Add(ActorPackage);
 			}
-
-			// Save all packages (we need to also save the ones we are deleting).
-			PackagesToSave.Add(ActorPackage);
 		}
 	}
 
