Expose ignore list for soft reference deletion warning

diff --git a/Engine/Source/Editor/UnrealEd/Classes/Editor/UnrealEdEngine.h b/Engine/Source/Editor/UnrealEd/Classes/Editor/UnrealEdEngine.h
--- a/Engine/Source/Editor/UnrealEd/Classes/Editor/UnrealEdEngine.h
+++ b/Engine/Source/Editor/UnrealEd/Classes/Editor/UnrealEdEngine.h
@@ -140,6 +140,10 @@ public:
 	UPROPERTY()
 	TObjectPtr<class UCookOnTheFlyServer> CookServer;
 
+	/** When deleting actors, these types should not generate warnings when references will be broken (this should only be types that don't affect gameplay) */
+	UPROPERTY()
+	TArray<TObjectPtr<UClass>> ClassesToIgnoreDeleteReferenceWarning;
+
 	/** A list of packages dirtied this tick */
 	TArray<TWeakObjectPtr<UPackage>> PackagesDirtiedThisTick;
 
diff --git a/Engine/Source/Editor/UnrealEd/Private/EditorActor.cpp b/Engine/Source/Editor/UnrealEd/Private/EditorActor.cpp
--- a/Engine/Source/Editor/UnrealEd/Private/EditorActor.cpp
+++ b/Engine/Source/Editor/UnrealEd/Private/EditorActor.cpp
@@ -905,15 +905,15 @@ bool UUnrealEdEngine::DeleteActors(const TArray<AActor*>& InActorsToDelete, UWor
 	TMap<FSoftObjectPath, TArray<UObject*>> SoftReferencingObjectsMap;
 	{
 		TArray<UClass*> ClassTypesToIgnore;
-		ClassTypesToIgnore.Add(ALevelScriptActor::StaticClass());
+		ClassesToIgnoreDeleteReferenceWarning.AddUnique(ALevelScriptActor::StaticClass());
 		// The delete warning is meant for actor references that affect gameplay.  Group actors do not affect gameplay and should not show up as a warning.
-		ClassTypesToIgnore.Add(AGroupActor::StaticClass());
+		ClassesToIgnoreDeleteReferenceWarning.AddUnique(AGroupActor::StaticClass());
 
 		// If we want to warn about references to the actors to be deleted, it is a lot more efficient to query
 		// the world first and build a map of actors referenced by other actors. We can then quickly look this up later on in the loop.
 		if (bWarnAboutReferences)
 		{
-			FBlueprintEditorUtils::GetActorReferenceMap(InWorld, ClassTypesToIgnore, ReferencingActorsMap);
+			FBlueprintEditorUtils::GetActorReferenceMap(InWorld, ClassesToIgnoreDeleteReferenceWarning, ReferencingActorsMap);
 
 			if (bWarnAboutSoftReferences)
 			{
@@ -973,7 +973,20 @@ bool UUnrealEdEngine::DeleteActors(const TArray<AActor*>& InActorsToDelete, UWor
 
 			if (SoftReferencingObjects)
 			{
-				bReferencedBySoftReference = true;
+				// Remove any references from object types marked to be ignored
+				for (int32 i = SoftReferencingObjects->Num() - 1; i >= 0; --i)
+				{
+					for (const TObjectPtr<UClass>& ClassToIgnore : ClassesToIgnoreDeleteReferenceWarning)
+					{
+						if ((*SoftReferencingObjects)[i]->IsA(ClassToIgnore))
+						{
+							SoftReferencingObjects->RemoveAt(i);
+							break;
+						}
+					}
+				}
+
+				bReferencedBySoftReference = SoftReferencingObjects->Num() > 0;
 			}
 		}
 
