- Themes are now saved to the right location in engine folder.
- Implemented theme name checks to avoid duplicating names.
- Editing existing theme name will update the .json file name as well in the engine theme folder.

diff --git a/Engine/Source/Editor/UnrealEd/Private/Settings/EditorStyleSettingsCustomization.cpp b/Engine/Source/Editor/UnrealEd/Private/Settings/EditorStyleSettingsCustomization.cpp
--- a/Engine/Source/Editor/UnrealEd/Private/Settings/EditorStyleSettingsCustomization.cpp
+++ b/Engine/Source/Editor/UnrealEd/Private/Settings/EditorStyleSettingsCustomization.cpp
@@ -20,6 +20,7 @@
 #define LOCTEXT_NAMESPACE "ThemeEditor"
 
 TWeakPtr<SWindow> ThemeEditorWindow;
+FString CurrentActiveThemeDisplayName;
 
 
 class SThemeEditor : public SCompoundWidget
@@ -81,9 +82,11 @@ public:
 					.VAlign(VAlign_Center)
 					.Padding(5.0f, 2.0f)
 					[
-						SNew(SEditableTextBox)
+						SAssignNew(EditableThemeName, SEditableTextBox)
 						.Text(this, &SThemeEditor::GetThemeName)
-						.OnTextCommitted(this, &SThemeEditor::OnThemeNameChanged)
+						.OnTextChanged(this, &SThemeEditor::OnThemeNameChanged)
+						.OnTextCommitted(this, &SThemeEditor::OnThemeNameCommitted)
+						.SelectAllTextWhenFocused(true)
 						//.IsReadOnly(true)
 					]
 				]
@@ -154,37 +157,90 @@ private:
 		return USlateThemeManager::Get().GetCurrentTheme().DisplayName;
 	}
 
-	void OnThemeNameChanged(const FText& NewName, ETextCommit::Type Type)
+	bool ValidateThemeName(const FText& ThemeName)
+	{
+		FText OutErrorMessage;
+		FString	Filename = USlateThemeManager::Get().GetEngineThemeDir() / ThemeName.ToString() + TEXT(".json");
+
+		// in cases the name user entered already exists when creating a new theme: 
+		if (FPaths::FileExists(Filename))
+		{
+			// show error message whenever there's duplicate 
+			OutErrorMessage = FText::Format(LOCTEXT("RenameThemeAlreadyExists", "A theme already exists with the name '{0}'."), ThemeName);
+			EditableThemeName->SetError(OutErrorMessage);
+			return false; 
+		}
+		EditableThemeName->SetError(FText::GetEmpty());
+		return true; 
+	}
+
+	void OnThemeNameChanged(const FText& NewName)
+	{
+		// verify duplicates before setting the display name. 
+		ValidateThemeName(NewName);
+	}
+
+	void OnThemeNameCommitted(const FText& NewName, ETextCommit::Type = ETextCommit::Default)
 	{
-		USlateThemeManager::Get().SetCurrentThemeDisplayName(NewName);
+		// if the new name is a valid name, write to display. Else, go back to the previous name. 
+		if (ValidateThemeName(NewName))
+		{
+			USlateThemeManager::Get().SetCurrentThemeDisplayName(NewName);
+		}
 	}
 
 	FReply OnSaveClicked()
 	{
 		FString Filename;
+		bool bSuccess = true; 
 
 		const FStyleTheme& Theme = USlateThemeManager::Get().GetCurrentTheme();
+
+		// New theme being created: 
 		if (Theme.Filename.IsEmpty())
 		{
 			TSharedPtr<SWindow> MyWindow = FSlateApplication::Get().FindWidgetWindow(SharedThis(this));
 
 			TArray<FString> Filenames;
 			IDesktopPlatform* DesktopPlatform = FDesktopPlatformModule::Get();
+			Filename = USlateThemeManager::Get().GetEngineThemeDir() / Theme.DisplayName.ToString() + TEXT(".json");
 
-			if (DesktopPlatform->SaveFileDialog(MyWindow->GetNativeWindow()->GetOSWindowHandle(), TEXT("Save Theme As"), USlateThemeManager::Get().GetEngineThemeDir(), Theme.DisplayName.ToString() + TEXT(".json"), TEXT("Theme Files (*.json)|*.json"), EFileDialogFlags::None, Filenames))
+			// if not a valid name, DO NOT SAVE. 
+			if (!ValidateThemeName(Theme.DisplayName))
 			{
-				Filename = Filenames[0];
+				bSuccess = false;
 			}
 		}
+		// Modifying an existing theme: 
 		else
 		{
-			Filename = Theme.Filename;
+			// updated name is not taken: SAVE. 
+			if (ValidateThemeName(Theme.DisplayName))
+			{
+				Filename = Theme.Filename;
+			}
+			// updated name either already exists, or is the same as before: 
+			else
+			{		
+				// updated name is the same as the current theme: SAVE anyway. 
+				if (Theme.DisplayName.ToString().Equals(CurrentActiveThemeDisplayName))
+				{
+					Filename = Theme.Filename;
+				}
+				// updated name already exists: DON'T SAVE. 
+				else
+				{
+					bSuccess = false;
+				}
+			}
 		}
 
-		if(!Filename.IsEmpty())
+		if(!Filename.IsEmpty() && bSuccess)
 		{
 			USlateThemeManager::Get().SaveCurrentThemeAs(Filename);
 
+			EditableThemeName->SetError(FText::GetEmpty()); 
+
 			ParentWindow.Pin()->SetOnWindowClosed(FOnWindowClosed());
 			ParentWindow.Pin()->RequestDestroyWindow();
 
@@ -209,6 +265,7 @@ private:
 
 private:
 	FOnThemeEditorClosed OnThemeEditorClosed;
+	TSharedPtr<SEditableTextBox> EditableThemeName; 
 	TWeakPtr<SWindow> ParentWindow;
 };
 
@@ -431,6 +488,8 @@ FReply FEditorStyleSettingsCustomization::OnDuplicateAndEditThemeClicked()
 {
 	FGuid PreviouslyActiveTheme = USlateThemeManager::Get().GetCurrentTheme().Id;
 
+	CurrentActiveThemeDisplayName = USlateThemeManager::Get().GetCurrentTheme().DisplayName.ToString(); 
+
 	FGuid NewThemeId = USlateThemeManager::Get().DuplicateActiveTheme();
 	USlateThemeManager::Get().ApplyTheme(NewThemeId);
 
@@ -443,6 +502,9 @@ FReply FEditorStyleSettingsCustomization::OnDuplicateAndEditThemeClicked()
 
 FReply FEditorStyleSettingsCustomization::OnEditThemeClicked()
 {
+
+	CurrentActiveThemeDisplayName = USlateThemeManager::Get().GetCurrentTheme().DisplayName.ToString();
+
 	OpenThemeEditorWindow(FOnThemeEditorClosed::CreateStatic(&OnThemeEditorClosed, TWeakPtr<FEditorStyleSettingsCustomization>(SharedThis(this)), FGuid(), FGuid()));
 
 	return FReply::Handled();
diff --git a/Engine/Source/Runtime/SlateCore/Private/Styling/StyleColors.cpp b/Engine/Source/Runtime/SlateCore/Private/Styling/StyleColors.cpp
--- a/Engine/Source/Runtime/SlateCore/Private/Styling/StyleColors.cpp
+++ b/Engine/Source/Runtime/SlateCore/Private/Styling/StyleColors.cpp
@@ -158,6 +158,7 @@ void USlateThemeManager::SaveCurrentThemeAs(const FString& Filename)
 {
 	FStyleTheme& CurrentTheme = GetMutableCurrentTheme();
 	CurrentTheme.Filename = Filename;
+	FString NewPath = CurrentTheme.Filename;
 	{
 		FString Output;
 		TSharedRef<TJsonWriter<>> WriterRef = TJsonWriterFactory<>::Create(&Output);
@@ -166,7 +167,7 @@ void USlateThemeManager::SaveCurrentThemeAs(const FString& Filename)
 		Writer.WriteValue(TEXT("Version"), 1);
 		Writer.WriteValue(TEXT("Id"), CurrentTheme.Id.ToString());
 		Writer.WriteValue(TEXT("DisplayName"), CurrentTheme.DisplayName.ToString());
-
+		
 		{
 			Writer.WriteObjectStart(TEXT("Colors"));
 			UEnum* Enum = StaticEnum<EStyleColor>();
@@ -184,8 +185,16 @@ void USlateThemeManager::SaveCurrentThemeAs(const FString& Filename)
 		if (FPlatformFileManager::Get().GetPlatformFile().FileExists(*Filename))
 		{
 			FPlatformFileManager::Get().GetPlatformFile().SetReadOnly(*Filename, false);
+			// create a new path if the filename has been changed. 
+			NewPath = USlateThemeManager::Get().GetEngineThemeDir() / CurrentTheme.DisplayName.ToString() + TEXT(".json");
+
+			if (!NewPath.Equals(CurrentTheme.Filename))
+			{
+				// rename the current .json file with the new name. 
+				IFileManager::Get().Move(*NewPath, *Filename);
+			}
 		}
-		FFileHelper::SaveStringToFile(Output, *Filename);
+		FFileHelper::SaveStringToFile(Output, *NewPath);
 	}
 	
 }
