World Folders Optimization: delay world folders rebuild to tick instead of doing it repeatedly.

#rnx

diff --git a/Engine/Source/Editor/UnrealEd/Private/WorldFolders.cpp b/Engine/Source/Editor/UnrealEd/Private/WorldFolders.cpp
--- a/Engine/Source/Editor/UnrealEd/Private/WorldFolders.cpp
+++ b/Engine/Source/Editor/UnrealEd/Private/WorldFolders.cpp
@@ -34,46 +34,7 @@ void UWorldFolders::Initialize(UWorld* InWorld)
 
 void UWorldFolders::RebuildList()
 {
-	Modify();
-	
-	// Clear folders with a Root Object.
-	TArray<FFolder> FoldersToRemove;
-	ForEachFolder([&FoldersToRemove](const FFolder& Folder)
-	{
-		if (!Folder.IsRootObjectPersistentLevel())
-		{
-			FoldersToRemove.Add(Folder);
-		}
-		return true;
-	});
-
-	for (const FFolder& Folder : FoldersToRemove)
-	{
-		RemoveFolder(Folder);
-	}
-
-	// Iterate over every actor in memory. WARNING: This is potentially very expensive!
-	for (FActorIterator ActorIt(GetWorld()); ActorIt; ++ActorIt)
-	{
-		AddFolder(ActorIt->GetFolder());
-	}
-
-	// Add levels ActorFolders as they are still needed for unloaded actors (only in editor)
-	if (!GetWorld()->IsGameWorld())
-	{
-		for (ULevel* Level : GetWorld()->GetLevels())
-		{
-			const bool bIsLevelVisibleOrAssociating = (Level->bIsVisible && !Level->bIsBeingRemoved) || Level->bIsAssociatingLevel || Level->bIsDisassociatingLevel;
-			if (bIsLevelVisibleOrAssociating)
-			{
-				Level->ForEachActorFolder([this](UActorFolder* ActorFolder)
-				{
-					AddFolder(ActorFolder->GetFolder());
-					return true;
-				}, /*bSkipDeleted*/ true);
-			}
-		}
-	}
+	bNeedRebuildList = true;
 }
 
 UWorld* UWorldFolders::GetWorld() const
@@ -277,6 +238,71 @@ void UWorldFolders::Serialize(FArchive& Ar)
 	Super::Serialize(Ar);
 }
 
+void UWorldFolders::Tick(float DeltaTime)
+{
+	TRACE_CPUPROFILER_EVENT_SCOPE(UWorldFolders::RebuildList)
+
+	check(bNeedRebuildList);
+	bNeedRebuildList = false;
+
+	// Clear folders with a Root Object.
+	TArray<FFolder> FoldersToRemove;
+	ForEachFolder([&FoldersToRemove](const FFolder& Folder)
+	{
+		if (!Folder.IsRootObjectPersistentLevel())
+		{
+			FoldersToRemove.Add(Folder);
+		}
+		return true;
+	});
+
+	for (const FFolder& Folder : FoldersToRemove)
+	{
+		RemoveFolder(Folder);
+	}
+
+	// Iterate over every actor in memory. WARNING: This is potentially very expensive!
+	if (UWorld* CurrentWorld = GetWorld())
+	{
+		for (FActorIterator ActorIt(CurrentWorld); ActorIt; ++ActorIt)
+		{
+			AddFolder(ActorIt->GetFolder());
+		}
+
+		// Add levels ActorFolders as they are still needed for unloaded actors (only in editor)
+		if (!CurrentWorld->IsGameWorld())
+		{
+			for (ULevel* Level : CurrentWorld->GetLevels())
+			{
+				const bool bIsLevelVisibleOrAssociating = (Level->bIsVisible && !Level->bIsBeingRemoved) || Level->bIsAssociatingLevel || Level->bIsDisassociatingLevel;
+				if (bIsLevelVisibleOrAssociating)
+				{
+					Level->ForEachActorFolder([this](UActorFolder* ActorFolder)
+					{
+						AddFolder(ActorFolder->GetFolder());
+						return true;
+					}, /*bSkipDeleted*/ true);
+				}
+			}
+		}
+	}
+}
+
+ETickableTickType UWorldFolders::GetTickableTickType() const
+{
+	return ETickableTickType::Conditional;
+}
+
+bool UWorldFolders::IsTickable() const
+{
+	return bNeedRebuildList;
+}
+
+TStatId UWorldFolders::GetStatId() const
+{
+	RETURN_QUICK_DECLARE_CYCLE_STAT(UWorldFolders, STATGROUP_Tickables);
+}
+
 FString UWorldFolders::GetWorldStateFilename() const
 {
 	UPackage* Package = World->GetOutermost();
diff --git a/Engine/Source/Editor/UnrealEd/Public/WorldFolders.h b/Engine/Source/Editor/UnrealEd/Public/WorldFolders.h
--- a/Engine/Source/Editor/UnrealEd/Public/WorldFolders.h
+++ b/Engine/Source/Editor/UnrealEd/Public/WorldFolders.h
@@ -7,6 +7,7 @@
 #include "UObject/Object.h"
 #include "UObject/WeakObjectPtr.h"
 #include "Misc/Guid.h"
+#include "TickableEditorObject.h"
 #include "WorldFoldersImplementation.h"
 #include "WorldPersistentFolders.h"
 #include "WorldTransientFolders.h"
@@ -67,7 +68,7 @@ struct FActorPlacementFolder
 
 /** Per-World Actor Folders UObject (used to support undo/redo reliably) */
 UCLASS()
-class UNREALED_API UWorldFolders : public UObject
+class UNREALED_API UWorldFolders : public UObject, public FTickableEditorObject
 {
 public:
 	GENERATED_BODY()
@@ -94,6 +95,13 @@ public:
 	virtual void Serialize(FArchive& Ar) override;
 	//~ End UObject
 
+	//~ Begin FTickableEditorObject interface
+	virtual void Tick(float DeltaTime) override;
+	virtual ETickableTickType GetTickableTickType() const override;
+	virtual bool IsTickable() const override;
+	virtual TStatId GetStatId() const override;
+	//~ End FTickableEditorObject interface
+
 	//~ Begin Deprecated
 	FActorFolderProps* GetFolderProperties(const FFolder& InFolder);
 	//~ End Deprecated
@@ -122,6 +130,8 @@ private:
 
 	TArray<FActorPlacementFolder> CurrentFolderStack;
 
+	bool bNeedRebuildList = false;
+
 	friend class FWorldFoldersImplementation;
 	friend class FWorldPersistentFolders;
 	friend class FWorldTransientFolders;
