Removed includes in favor of a forward decl

#rnx

diff --git a/Engine/Source/Editor/UnrealEd/Public/EditorUndoClient.h b/Engine/Source/Editor/UnrealEd/Public/EditorUndoClient.h
--- a/Engine/Source/Editor/UnrealEd/Public/EditorUndoClient.h
+++ b/Engine/Source/Editor/UnrealEd/Public/EditorUndoClient.h
@@ -7,8 +7,9 @@
 #pragma once
 
 #include "CoreMinimal.h"
-#include "Misc/ITransaction.h"
-#include "Misc/TransactionObjectEvent.h"
+
+struct FTransactionContext;
+class FTransactionObjectEvent;
 
 /**
  * Interface for tools wanting to handle undo/redo operations
