Adding additional includes in preparation for removing a global include

#rnx

diff --git a/Engine/Source/Editor/UnrealEd/Private/UserDefinedStructEditorData.cpp b/Engine/Source/Editor/UnrealEd/Private/UserDefinedStructEditorData.cpp
--- a/Engine/Source/Editor/UnrealEd/Private/UserDefinedStructEditorData.cpp
+++ b/Engine/Source/Editor/UnrealEd/Private/UserDefinedStructEditorData.cpp
@@ -1,7 +1,7 @@
 // Copyright Epic Games, Inc. All Rights Reserved.
 
 #include "UserDefinedStructure/UserDefinedStructEditorData.h"
-#include "Misc/ITransaction.h"
+#include "Misc/ITransactionObjectAnnotation.h"
 #include "UObject/UE5ReleaseStreamObjectVersion.h"
 #include "UObject/UnrealType.h"
 #include "UObject/ObjectSaveContext.h"
