Added a Delete button for Active Themes. When users delete any user-created themes, they're prompted with a "Are you sure" dialog box. If the user selects "OK", the active theme is set to the default dark theme. When the active theme is Dark, users do not have the option to delete it.

diff --git a/Engine/Source/Editor/UnrealEd/Private/Settings/EditorStyleSettingsCustomization.cpp b/Engine/Source/Editor/UnrealEd/Private/Settings/EditorStyleSettingsCustomization.cpp
--- a/Engine/Source/Editor/UnrealEd/Private/Settings/EditorStyleSettingsCustomization.cpp
+++ b/Engine/Source/Editor/UnrealEd/Private/Settings/EditorStyleSettingsCustomization.cpp
@@ -16,6 +16,8 @@
 #include "IDesktopPlatform.h"
 #include "DesktopPlatformModule.h"
 #include "SPrimaryButton.h"
+#include "HAL/FileManager.h"
+#include "Misc/MessageDialog.h"
 
 #define LOCTEXT_NAMESPACE "ThemeEditor"
 
@@ -451,6 +453,27 @@ void FEditorStyleSettingsCustomization::MakeThemePickerRow(IDetailPropertyRow& P
 				.Image(FAppStyle::Get().GetBrush("Icons.Duplicate"))
 			]
 		]
+		+ SHorizontalBox::Slot()
+			.VAlign(VAlign_Center)
+			.HAlign(HAlign_Center)
+			.AutoWidth()
+			[
+			SNew(SButton)
+			.ButtonStyle(FAppStyle::Get(), "SimpleButton")
+			.Visibility_Lambda(
+				[]()
+				{
+					// if current active theme is the default dark theme: hide the delete button. 
+					return USlateThemeManager::Get().IsDefaultThemeActive() ? EVisibility::Hidden : EVisibility::Visible;
+				})
+			.ToolTipText(LOCTEXT("DeleteThemeToolTip", "Delete this theme"))
+			.OnClicked(this, &FEditorStyleSettingsCustomization::OnDeleteThemeClicked)
+			[
+				SNew(SImage)
+				.ColorAndOpacity(FSlateColor::UseForeground())
+				.Image(FAppStyle::Get().GetBrush("Icons.Delete"))
+			]
+		]
 	];
 }
 
@@ -484,6 +507,31 @@ static void OnThemeEditorClosed(bool bSaved, TWeakPtr<FEditorStyleSettingsCustom
 	}
 }
 
+FReply FEditorStyleSettingsCustomization::OnDeleteThemeClicked()
+{
+	const FStyleTheme PreviouslyActiveTheme = USlateThemeManager::Get().GetCurrentTheme();
+
+	// Are you sure you want to do this?
+	const FText FileNameToRemove = FText::FromString(PreviouslyActiveTheme.DisplayName.ToString());
+	const FText TextBody = FText::Format(LOCTEXT("ActionRemoveMsg", "Are you sure you want to permanently delete the theme \"{0}\"? This action cannot be undone."), FileNameToRemove);
+	const FText TextTitle = FText::Format(LOCTEXT("RemoveTheme_Title", "Remove Theme \"{0}\"?"), FileNameToRemove);
+
+	// If user select "OK"...
+	if (EAppReturnType::Ok == FMessageDialog::Open(EAppMsgType::OkCancel, TextBody, &TextTitle))
+	{
+		// apply default theme
+		USlateThemeManager::Get().ApplyDefaultTheme();
+
+		// remove previously active theme
+		const FString Filename = USlateThemeManager::Get().GetEngineThemeDir() / PreviouslyActiveTheme.DisplayName.ToString() + TEXT(".json");
+		IFileManager::Get().Delete(*Filename);
+		USlateThemeManager::Get().RemoveTheme(PreviouslyActiveTheme.Id);
+		RefreshComboBox();
+	}
+	// Else, do nothing. 
+	return FReply::Handled();
+}
+
 FReply FEditorStyleSettingsCustomization::OnDuplicateAndEditThemeClicked()
 {
 	FGuid PreviouslyActiveTheme = USlateThemeManager::Get().GetCurrentTheme().Id;
diff --git a/Engine/Source/Editor/UnrealEd/Private/Settings/EditorStyleSettingsCustomization.h b/Engine/Source/Editor/UnrealEd/Private/Settings/EditorStyleSettingsCustomization.h
--- a/Engine/Source/Editor/UnrealEd/Private/Settings/EditorStyleSettingsCustomization.h
+++ b/Engine/Source/Editor/UnrealEd/Private/Settings/EditorStyleSettingsCustomization.h
@@ -40,6 +40,7 @@ private:
 	void GenerateThemeOptions(TSharedPtr<FString>& OutSelectedTheme);
 
 	void MakeThemePickerRow(IDetailPropertyRow& PropertyRow);
+	FReply OnDeleteThemeClicked();
 	FReply OnDuplicateAndEditThemeClicked();
 	FReply OnEditThemeClicked();
 	FString GetTextLabelForThemeEntry(TSharedPtr<FString> Entry);
diff --git a/Engine/Source/Runtime/SlateCore/Private/Styling/StyleColors.cpp b/Engine/Source/Runtime/SlateCore/Private/Styling/StyleColors.cpp
--- a/Engine/Source/Runtime/SlateCore/Private/Styling/StyleColors.cpp
+++ b/Engine/Source/Runtime/SlateCore/Private/Styling/StyleColors.cpp
@@ -232,6 +232,16 @@ void USlateThemeManager::ApplyTheme(FGuid ThemeId)
 	}
 }
 
+void USlateThemeManager::ApplyDefaultTheme()
+{
+	ApplyTheme(DefaultDarkTheme.Id); 
+}
+
+bool USlateThemeManager::IsDefaultThemeActive() const 
+{
+	return GetCurrentTheme() == DefaultDarkTheme; 
+}
+
 void USlateThemeManager::RemoveTheme(FGuid ThemeId)
 {
 	// Current Theme cannot currently be removed.  Apply a new theme first
@@ -369,7 +379,6 @@ bool USlateThemeManager::ReadTheme(const FString& ThemeData, FStyleTheme& Theme)
 
 void USlateThemeManager::EnsureValidCurrentTheme()
 {
-	FStyleTheme DefaultDarkTheme;
 	DefaultDarkTheme.DisplayName = NSLOCTEXT("StyleColors", "DefaultDarkTheme", "Dark");
 	// If you change this you invalidate the default dark theme forcing the default theme to be reset and the existing dark theme to become a user theme
 	// In general you should not do this. You should instead update the default "dark.json" file 
diff --git a/Engine/Source/Runtime/SlateCore/Public/Styling/StyleColors.h b/Engine/Source/Runtime/SlateCore/Public/Styling/StyleColors.h
--- a/Engine/Source/Runtime/SlateCore/Public/Styling/StyleColors.h
+++ b/Engine/Source/Runtime/SlateCore/Public/Styling/StyleColors.h
@@ -207,6 +207,16 @@ public:
 	 */
 	void ApplyTheme(FGuid ThemeId);
 
+	/**
+	 * Applies the default dark theme as the active theme
+	*/
+	void ApplyDefaultTheme();
+
+	/**
+	* Returns true if the default dark theme is the active theme
+	*/
+	bool IsDefaultThemeActive() const;
+
 	/**
 	 * Removes a theme. 
 	 * Note: The active theme cannot be removed and there must always be an active theme.  Apply a new theme first before removing the current theme.
@@ -269,6 +279,7 @@ private:
 
 
 private:
+	FStyleTheme DefaultDarkTheme;
 	TArray<FStyleTheme> Themes;
 	FLinearColor DefaultColors[(int32)EStyleColor::MAX];
 #endif // ALLOW_THEMES
