Preparing to remove global Transactor.h include

Removed a hard type dependency from EditorEngine.h in favor of a forward declaration

#rnx

diff --git a/Engine/Source/Editor/UnrealEd/Classes/Editor/EditorEngine.h b/Engine/Source/Editor/UnrealEd/Classes/Editor/EditorEngine.h
--- a/Engine/Source/Editor/UnrealEd/Classes/Editor/EditorEngine.h
+++ b/Engine/Source/Editor/UnrealEd/Classes/Editor/EditorEngine.h
@@ -64,6 +64,9 @@ class USkeleton;
 class USoundBase;
 class USoundNode;
 class UTextureRenderTarget2D;
+class UTransactor;
+class FTransactionObjectEvent;
+struct FTransactionContext;
 struct FTypedElementHandle;
 struct FAnalyticsEventAttribute;
 class UEditorWorldExtensionManager;
@@ -2784,23 +2787,7 @@ public:
 
 private:
 
-	/** Internal struct to hold undo/redo transaction object context */
-	struct FTransactionDeltaContext
-	{
-		FGuid	OuterOperationId;
-		int32	OperationDepth;
-		TArray<TPair<UObject*, FTransactionObjectEvent>> TransactionObjects;
-
-		void Reset()
-		{
-			OuterOperationId.Invalidate();
-			TransactionObjects.Empty();
-			OperationDepth = 0;
-		}
-
-		FTransactionDeltaContext() = default;
-	};
-	FTransactionDeltaContext CurrentUndoRedoContext;
+	TSharedPtr<struct FEditorTransactionDeltaContext> CurrentUndoRedoContext;
 
 	/** List of all viewport clients */
 	TArray<class FEditorViewportClient*> AllViewportClients;
diff --git a/Engine/Source/Editor/UnrealEd/Private/EditorServer.cpp b/Engine/Source/Editor/UnrealEd/Private/EditorServer.cpp
--- a/Engine/Source/Editor/UnrealEd/Private/EditorServer.cpp
+++ b/Engine/Source/Editor/UnrealEd/Private/EditorServer.cpp
@@ -179,6 +179,21 @@ DEFINE_LOG_CATEGORY_STATIC(LogEditorServer, Log, All);
 /** Used for the "tagsounds" and "checksounds" commands only			*/
 static FUObjectAnnotationSparseBool DebugSoundAnnotation;
 
+/** Internal struct to hold undo/redo transaction object context */
+struct FEditorTransactionDeltaContext
+{
+	FGuid	OuterOperationId;
+	int32	OperationDepth = 0;
+	TArray<TPair<UObject*, FTransactionObjectEvent>> TransactionObjects;
+
+	void Reset()
+	{
+		OuterOperationId.Invalidate();
+		TransactionObjects.Empty();
+		OperationDepth = 0;
+	}
+};
+
 namespace 
 {
 	/**
@@ -1172,10 +1187,10 @@ void UEditorEngine::HandleTransactorBeforeRedoUndo(const FTransactionContext& Tr
 	}
 
 	// Before an undo, store the current operation and hook on object transaction, if we do not have an outer operation already
-	if (CurrentUndoRedoContext.OperationDepth++ == 0)
+	if (CurrentUndoRedoContext->OperationDepth++ == 0)
 	{
-		check(!CurrentUndoRedoContext.OuterOperationId.IsValid());
-		CurrentUndoRedoContext.OuterOperationId = TransactionContext.OperationId;
+		check(!CurrentUndoRedoContext->OuterOperationId.IsValid());
+		CurrentUndoRedoContext->OuterOperationId = TransactionContext.OperationId;
 		FCoreUObjectDelegates::OnObjectTransacted.AddUObject(this, &UEditorEngine::HandleObjectTransacted);
 	}
 }
@@ -1191,17 +1206,17 @@ void UEditorEngine::HandleTransactorRedoUndo(const FTransactionContext& Transact
 	// Broadcast only if you have an actual transaction context
 	if (Succeeded)
 	{
-		check(CurrentUndoRedoContext.OuterOperationId.IsValid() && CurrentUndoRedoContext.OperationDepth > 0);
+		check(CurrentUndoRedoContext->OuterOperationId.IsValid() && CurrentUndoRedoContext->OperationDepth > 0);
 		if (!bSuspendBroadcastPostUndoRedo)
 		{
 			BroadcastPostUndoRedo(TransactionContext, WasUndo);
 		}
 
-		if (--CurrentUndoRedoContext.OperationDepth == 0)
+		if (--CurrentUndoRedoContext->OperationDepth == 0)
 		{
 			// Undo/Redo is done clear out operation
-			check(CurrentUndoRedoContext.OuterOperationId == TransactionContext.OperationId);
-			CurrentUndoRedoContext.Reset();
+			check(CurrentUndoRedoContext->OuterOperationId == TransactionContext.OperationId);
+			CurrentUndoRedoContext->Reset();
 			FCoreUObjectDelegates::OnObjectTransacted.RemoveAll(this);
 		}
 	}
@@ -1225,10 +1240,10 @@ void UEditorEngine::HandleTransactorUndo(const FTransactionContext& TransactionC
 
 void UEditorEngine::HandleObjectTransacted(UObject* InObject, const FTransactionObjectEvent& InTransactionObjectEvent)
 {
-	check(CurrentUndoRedoContext.OuterOperationId.IsValid() && CurrentUndoRedoContext.OperationDepth > 0);
+	check(CurrentUndoRedoContext->OuterOperationId.IsValid() && CurrentUndoRedoContext->OperationDepth > 0);
 	if (InTransactionObjectEvent.GetEventType() == ETransactionObjectEventType::UndoRedo)
 	{
-		CurrentUndoRedoContext.TransactionObjects.Add(TPair<UObject*, FTransactionObjectEvent>{ InObject, InTransactionObjectEvent });
+		CurrentUndoRedoContext->TransactionObjects.Add(TPair<UObject*, FTransactionObjectEvent>{ InObject, InTransactionObjectEvent });
 	}
 }
 
@@ -1254,6 +1269,8 @@ void UEditorEngine::CreateStartupAnalyticsAttributes( TArray<FAnalyticsEventAttr
 
 UTransactor* UEditorEngine::CreateTrans()
 {
+	CurrentUndoRedoContext = MakeShared<FEditorTransactionDeltaContext>();
+
 	int32 UndoBufferSize;
 
 	if (!GConfig->GetInt(TEXT("Undo"), TEXT("UndoBufferSize"), UndoBufferSize, GEditorPerProjectIni))
@@ -5375,7 +5392,7 @@ void UEditorEngine::BroadcastPostUndoRedo(const FTransactionContext& UndoContext
 	for (auto UndoIt = UndoClients.CreateIterator(); UndoIt; ++UndoIt)
 	{
 		FEditorUndoClient* Client = *UndoIt;
-		if (Client && Client->MatchesContext(UndoContext, CurrentUndoRedoContext.TransactionObjects))
+		if (Client && Client->MatchesContext(UndoContext, CurrentUndoRedoContext->TransactionObjects))
 		{
 			if (bWasUndo)
 			{
