Sphere Reflection Capture:  Fix for OOM prevention logic, so it properly handles building or loading cached captures, without generating invalid cache data or asserting.

Capturing is disabled if resolution has been reduced, requiring the user to solve memory issues before capturing (by deleting items or reducing resolution).  Logic was added so the reflection capture scene data is reset to allow it to shrink, in cases where an OOM had occurred due to resolution -- without this fix, the user would need to restart the editor.  And if the resolution has been reduced since captures were taken (some were successfully captured, but then more were added, triggering the OOM logic), excess mips are discarded from the cached data to avoid an assert.

#rnx

diff --git a/Engine/Source/Editor/UnrealEd/Private/EditorEngine.cpp b/Engine/Source/Editor/UnrealEd/Private/EditorEngine.cpp
--- a/Engine/Source/Editor/UnrealEd/Private/EditorEngine.cpp
+++ b/Engine/Source/Editor/UnrealEd/Private/EditorEngine.cpp
@@ -3845,6 +3845,10 @@ void UEditorEngine::BuildReflectionCaptures(UWorld* World)
 	// Calling code should not allow building reflection captures on lower feature levels
 	check(World->FeatureLevel >= ERHIFeatureLevel::SM5);
 
+	// Only reset reflection captures if we had hit an OOM condition, and there is a chance we could fit in memory if we rebuild from scratch
+	const bool bOnlyIfOOM = true;
+	World->Scene->ResetReflectionCaptures(bOnlyIfOOM);
+
 	// Update sky light first because it's considered direct lighting, sky diffuse will be visible in reflection capture indirect specular
 	World->UpdateAllSkyCaptures();
 
@@ -3964,16 +3968,22 @@ void UEditorEngine::BuildReflectionCaptures(UWorld* World)
 				// Capture can fail if there are more than GMaxNumReflectionCaptures captures
 				if (ReadbackCaptureData.CubemapSize > 0)
 				{
-					ULevel* StorageLevel = LightingScenarios[LevelIndex] ? LightingScenarios[LevelIndex] : CaptureComponent->GetOwner()->GetLevel();
-					UMapBuildDataRegistry* Registry = StorageLevel->GetOrCreateMapBuildData();
-					if (!CaptureComponent->bModifyMaxValueRGBM)
+					// Capture should also fail if memory limitations prevent full resolution captures from being generated.
+					// We report an error message for this case below in the non-mobile capture code path (no need for two errors).
+					int32 DesiredCaptureSize = UReflectionCaptureComponent::GetReflectionCaptureSize();
+					if (ReadbackCaptureData.CubemapSize == DesiredCaptureSize)
 					{
-						CaptureComponent->MaxValueRGBM = GetMaxValueRGBM(ReadbackCaptureData.FullHDRCapturedData, ReadbackCaptureData.CubemapSize);
+						ULevel* StorageLevel = LightingScenarios[LevelIndex] ? LightingScenarios[LevelIndex] : CaptureComponent->GetOwner()->GetLevel();
+						UMapBuildDataRegistry* Registry = StorageLevel->GetOrCreateMapBuildData();
+						if (!CaptureComponent->bModifyMaxValueRGBM)
+						{
+							CaptureComponent->MaxValueRGBM = GetMaxValueRGBM(ReadbackCaptureData.FullHDRCapturedData, ReadbackCaptureData.CubemapSize);
+						}
+						FString TextureName = CaptureComponent->GetName() + TEXT("Texture");
+						TextureName += LexToString(CaptureComponent->MapBuildDataId);
+						GenerateEncodedHDRTextureCube(Registry, ReadbackCaptureData, TextureName, CaptureComponent->MaxValueRGBM, CaptureComponent, bIsReflectionCaptureCompressionProjectSetting);
+						EncodedCaptures[CaptureIndex] = ReadbackCaptureData.EncodedCaptureData;
 					}
-					FString TextureName = CaptureComponent->GetName() + TEXT("Texture");
-					TextureName += LexToString(CaptureComponent->MapBuildDataId);
-					GenerateEncodedHDRTextureCube(Registry, ReadbackCaptureData, TextureName, CaptureComponent->MaxValueRGBM, CaptureComponent, bIsReflectionCaptureCompressionProjectSetting);
-					EncodedCaptures[CaptureIndex] = ReadbackCaptureData.EncodedCaptureData;
 				}
 			}
 		}
@@ -3994,14 +4004,24 @@ void UEditorEngine::BuildReflectionCaptures(UWorld* World)
 			// Capture can fail if there are more than GMaxNumReflectionCaptures captures
 			if (ReadbackCaptureData.CubemapSize > 0)
 			{
-				ULevel* StorageLevel = LightingScenarios[LevelIndex] ? LightingScenarios[LevelIndex] : CaptureComponent->GetOwner()->GetLevel();
-				UMapBuildDataRegistry* Registry = StorageLevel->GetOrCreateMapBuildData();
-				FReflectionCaptureMapBuildData& CaptureBuildData = Registry->AllocateReflectionCaptureBuildData(CaptureComponent->MapBuildDataId, true);
-				(FReflectionCaptureData&)CaptureBuildData = ReadbackCaptureData;
-				CaptureBuildData.EncodedCaptureData = EncodedCaptures[CaptureIndex];
-				CaptureBuildData.FinalizeLoad();
-				// Recreate capture render state now that we have valid BuildData
-				CaptureComponent->MarkRenderStateDirty();
+				// Capture should also fail if memory limitations prevent full resolution captures from being generated.
+				int32 DesiredCaptureSize = UReflectionCaptureComponent::GetReflectionCaptureSize();
+				if (ReadbackCaptureData.CubemapSize == DesiredCaptureSize)
+				{
+					ULevel* StorageLevel = LightingScenarios[LevelIndex] ? LightingScenarios[LevelIndex] : CaptureComponent->GetOwner()->GetLevel();
+					UMapBuildDataRegistry* Registry = StorageLevel->GetOrCreateMapBuildData();
+					FReflectionCaptureMapBuildData& CaptureBuildData = Registry->AllocateReflectionCaptureBuildData(CaptureComponent->MapBuildDataId, true);
+					(FReflectionCaptureData&)CaptureBuildData = ReadbackCaptureData;
+					CaptureBuildData.EncodedCaptureData = EncodedCaptures[CaptureIndex];
+					CaptureBuildData.FinalizeLoad();
+					// Recreate capture render state now that we have valid BuildData
+					CaptureComponent->MarkRenderStateDirty();
+				}
+				else
+				{
+					UE_LOG(LogEditor, Error, TEXT("Unable to build Reflection Capture %s, requested reflection capture cube size of %d didn't fit in memory on host machine (size clamped to %d)"),
+						*CaptureComponent->GetPathName(), DesiredCaptureSize, ReadbackCaptureData.CubemapSize);
+				}
 			}
 			else
 			{
diff --git a/Engine/Source/Runtime/Engine/Public/SceneInterface.h b/Engine/Source/Runtime/Engine/Public/SceneInterface.h
--- a/Engine/Source/Runtime/Engine/Public/SceneInterface.h
+++ b/Engine/Source/Runtime/Engine/Public/SceneInterface.h
@@ -174,6 +174,15 @@ public:
 	virtual void AllocateReflectionCaptures(const TArray<UReflectionCaptureComponent*>& NewCaptures, const TCHAR* CaptureReason, bool bVerifyOnlyCapturing, bool bCapturingForMobile, bool bInsideTick) {}
 	virtual void ReleaseReflectionCubemap(UReflectionCaptureComponent* CaptureComponent) {}
 
+	/**
+	  * Resets reflection capture data to allow it to be rebuilt from scratch.  Allows shrinking of the number of items in
+	  * the capture array, which otherwise only grows.  This can allow the user to solve cases where an out of memory (OOM)
+	  * condition prevented captures from being generated.  The user can remove some items or reduce the capture resolution
+	  * to fix the issue, and run "Build Reflection Captures" (which calls this function), instead of needing to restart
+	  * the editor.  The flag "bOnlyIfOOM" only does the reset if an OOM had occurred.
+	  */
+	virtual void ResetReflectionCaptures(bool bOnlyIfOOM) {}
+
 	UE_DEPRECATED(5.1, "This method now accepts a bInsideTick parameter, which specifies whether it's called during the frame Tick.")
 	inline void AllocateReflectionCaptures(const TArray<UReflectionCaptureComponent*>& NewCaptures, const TCHAR* CaptureReason, bool bVerifyOnlyCapturing, bool bCapturingForMobile)
 	{
diff --git a/Engine/Source/Runtime/Renderer/Private/ReflectionEnvironment.cpp b/Engine/Source/Runtime/Renderer/Private/ReflectionEnvironment.cpp
--- a/Engine/Source/Runtime/Renderer/Private/ReflectionEnvironment.cpp
+++ b/Engine/Source/Runtime/Renderer/Private/ReflectionEnvironment.cpp
@@ -383,10 +383,11 @@
 	RegisteredComponentMapBuildDataIds.Remove(Component);
 }
 
-void FReflectionEnvironmentSceneData::SetGameThreadTrackingData(int32 MaxAllocatedCubemaps, int32 CaptureSize)
+void FReflectionEnvironmentSceneData::SetGameThreadTrackingData(int32 MaxAllocatedCubemaps, int32 CaptureSize, int32 DesiredCaptureSize)
 {
 	MaxAllocatedReflectionCubemapsGameThread = MaxAllocatedCubemaps;
 	ReflectionCaptureSizeGameThread = CaptureSize;
+	DesiredReflectionCaptureSizeGameThread = DesiredCaptureSize;
 }
 
 bool FReflectionEnvironmentSceneData::DoesAllocatedDataNeedUpdate(int32 DesiredMaxCubemaps, int32 DesiredCaptureSize) const
@@ -459,6 +460,43 @@
 	CubemapArray.ResizeCubemapArrayGPU(InMaxCubemaps, InCubemapSize, IndexRemapping);
 }
 
+void FReflectionEnvironmentSceneData::Reset(FScene* Scene)
+{
+	for (TSparseArray<UReflectionCaptureComponent*>::TIterator It(AllocatedReflectionCapturesGameThread); It; ++It)
+	{
+		UReflectionCaptureComponent* CurrentCapture = *It;
+
+		// Unused slots will have a null component pointer
+		if (CurrentCapture)
+		{
+			Scene->RemoveReflectionCapture(CurrentCapture);
+		}
+	}
+
+	// Fields to reset in the game thread
+	AllocatedReflectionCapturesGameThread.Empty();
+	MaxAllocatedReflectionCubemapsGameThread = 0;
+	ReflectionCaptureSizeGameThread = 0;
+	DesiredReflectionCaptureSizeGameThread = 0;
+
+	ENQUEUE_RENDER_COMMAND(ReflectionEnvironmentSceneDataReset)(
+		[this](FRHICommandList& RHICmdList)
+		{
+			// Fields to reset in the render thread
+			bRegisteredReflectionCapturesHasChanged = true;
+			AllocatedReflectionCaptureStateHasChanged = false;
+
+			RegisteredReflectionCaptures.Empty();
+			RegisteredReflectionCapturePositionAndRadius.Empty();
+			CubemapArray.Reset();
+			AllocatedReflectionCaptureState.Empty();
+			CubemapArraySlotsUsed.Empty();
+			SortedCaptures.Empty();
+			NumBoxCaptures = 0;
+			NumSphereCaptures = 0;
+		});
+}
+
 void FReflectionEnvironmentCubemapArray::ResizeCubemapArrayGPU(uint32 InMaxCubemaps, int32 InCubemapSize, const TArray<int32>& IndexRemapping)
 {
 	check(InMaxCubemaps > 0);
@@ -534,6 +572,13 @@
 	}
 }
 
+void FReflectionEnvironmentCubemapArray::Reset()
+{
+	ReleaseDynamicRHI();
+	MaxCubemaps = 0;
+	CubemapSize = 0;
+}
+
 IMPLEMENT_STATIC_UNIFORM_BUFFER_SLOT(ReflectionCapture);
 IMPLEMENT_STATIC_AND_SHADER_UNIFORM_BUFFER_STRUCT(FReflectionCaptureShaderData, "ReflectionCaptureSM5", ReflectionCapture);
 IMPLEMENT_STATIC_AND_SHADER_UNIFORM_BUFFER_STRUCT_EX(FMobileReflectionCaptureShaderData, "ReflectionCaptureES31", ReflectionCapture, FShaderParametersMetadata::EUsageFlags::NoEmulatedUniformBuffer);
diff --git a/Engine/Source/Runtime/Renderer/Private/ReflectionEnvironmentCapture.cpp b/Engine/Source/Runtime/Renderer/Private/ReflectionEnvironmentCapture.cpp
--- a/Engine/Source/Runtime/Renderer/Private/ReflectionEnvironmentCapture.cpp
+++ b/Engine/Source/Runtime/Renderer/Private/ReflectionEnvironmentCapture.cpp
@@ -830,7 +830,7 @@ void FScene::AllocateReflectionCaptures(const TArray<UReflectionCaptureComponent
 				if (ReflectionCaptureSize == ReflectionSceneData.CubemapArray.GetCubemapSize())
 				{
 					// We can do a fast GPU copy to realloc the array, so we don't need to update all captures
-					ReflectionSceneData.SetGameThreadTrackingData(DesiredMaxCubemaps, ReflectionCaptureSize);
+					ReflectionSceneData.SetGameThreadTrackingData(DesiredMaxCubemaps, ReflectionCaptureSize, DesiredCaptureSize);
 
 					FScene* Scene = this;
 					uint32 MaxSize = ReflectionSceneData.MaxAllocatedReflectionCubemapsGameThread;
@@ -847,7 +847,7 @@ void FScene::AllocateReflectionCaptures(const TArray<UReflectionCaptureComponent
 
 			if (bNeedsUpdateAllCaptures)
 			{
-				ReflectionSceneData.SetGameThreadTrackingData(DesiredMaxCubemaps, ReflectionCaptureSize);
+				ReflectionSceneData.SetGameThreadTrackingData(DesiredMaxCubemaps, ReflectionCaptureSize, DesiredCaptureSize);
 
 				FScene* Scene = this;
 				uint32 MaxSize = ReflectionSceneData.MaxAllocatedReflectionCubemapsGameThread;
@@ -941,6 +941,14 @@ void FScene::UpdateAllReflectionCaptures(const TCHAR* CaptureReason, int32 Refle
 	}
 }
 
+void FScene::ResetReflectionCaptures(bool bOnlyIfOOM)
+{
+	if (bOnlyIfOOM == false || ReflectionSceneData.ReflectionCaptureSizeGameThread != ReflectionSceneData.DesiredReflectionCaptureSizeGameThread)
+	{
+		ReflectionSceneData.Reset(this);
+	}
+}
+
 void GetReflectionCaptureData_RenderingThread(FRHICommandListImmediate& RHICmdList, FScene* Scene, const UReflectionCaptureComponent* Component, FReflectionCaptureData* OutCaptureData)
 {
 	const FCaptureComponentSceneState* ComponentStatePtr = Scene->ReflectionSceneData.AllocatedReflectionCaptureState.Find(Component);
@@ -1015,20 +1023,31 @@ void FScene::GetReflectionCaptureData(UReflectionCaptureComponent* Component, FR
 
 void UploadReflectionCapture_RenderingThread(FScene* Scene, const FReflectionCaptureData* CaptureData, const UReflectionCaptureComponent* CaptureComponent)
 {
-	const int32 EffectiveTopMipSize = CaptureData->CubemapSize;
-	const int32 NumMips = FMath::CeilLogTwo(EffectiveTopMipSize) + 1;
+	// Due to memory limitations, it's possible for the in-memory size to be smaller than the originally captured size.
+	const int32 EffectiveTopMipSize = Scene->ReflectionSceneData.CubemapArray.GetCubemapSize();
+	const int32 NumMipsSource = FMath::CeilLogTwo(CaptureData->CubemapSize) + 1;
+	const int32 NumMipsDest = FMath::CeilLogTwo(EffectiveTopMipSize) + 1;
 
 	const int32 CaptureIndex = FindOrAllocateCubemapIndex(Scene, CaptureComponent);
-	check(CaptureData->CubemapSize == Scene->ReflectionSceneData.CubemapArray.GetCubemapSize());
+	check(CaptureData->CubemapSize >= Scene->ReflectionSceneData.CubemapArray.GetCubemapSize());
 	check(CaptureIndex < Scene->ReflectionSceneData.CubemapArray.GetMaxCubemaps());
 	FRHITexture* CubeMapArray = Scene->ReflectionSceneData.CubemapArray.GetRenderTarget()->GetRHI();
 	check(CubeMapArray->GetFormat() == PF_FloatRGBA);
 
 	int32 MipBaseIndex = 0;
 
-	for (int32 MipIndex = 0; MipIndex < NumMips; MipIndex++)
+	// Skip over mips in originally captured data, based on what we can fit in memory.
+	for (int32 MipIndex = 0; MipIndex < NumMipsSource - NumMipsDest; MipIndex++)
 	{
-		const int32 MipSize = 1 << (NumMips - MipIndex - 1);
+		const int32 MipSize = 1 << (NumMipsSource - MipIndex - 1);
+		const int32 CubeFaceBytes = MipSize * MipSize * sizeof(FFloat16Color);
+
+		MipBaseIndex += CubeFaceBytes * CubeFace_MAX;
+	}
+
+	for (int32 MipIndex = 0; MipIndex < NumMipsDest; MipIndex++)
+	{
+		const int32 MipSize = 1 << (NumMipsDest - MipIndex - 1);
 		const int32 CubeFaceBytes = MipSize * MipSize * sizeof(FFloat16Color);
 
 		for (int32 CubeFace = 0; CubeFace < CubeFace_MAX; CubeFace++)
diff --git a/Engine/Source/Runtime/Renderer/Private/ScenePrivate.h b/Engine/Source/Runtime/Renderer/Private/ScenePrivate.h
--- a/Engine/Source/Runtime/Renderer/Private/ScenePrivate.h
+++ b/Engine/Source/Runtime/Renderer/Private/ScenePrivate.h
@@ -1528,6 +1528,7 @@ public:
 	int32 GetCubemapSize() const { return CubemapSize; }
 	bool IsValid() const { return IsValidRef(ReflectionEnvs); }
 	const TRefCountPtr<IPooledRenderTarget>& GetRenderTarget() const { return ReflectionEnvs; }
+	void Reset();
 
 protected:
 	uint32 MaxCubemaps;
@@ -1671,22 +1672,27 @@ public:
 
 	/** Game thread tracking of what size cubemaps are in the cubemap array. */
 	int32 ReflectionCaptureSizeGameThread;
+	int32 DesiredReflectionCaptureSizeGameThread;
 
 	FReflectionEnvironmentSceneData(ERHIFeatureLevel::Type InFeatureLevel) :
 		bRegisteredReflectionCapturesHasChanged(true),
 		AllocatedReflectionCaptureStateHasChanged(false),
 		CubemapArray(InFeatureLevel),
 		MaxAllocatedReflectionCubemapsGameThread(0),
-		ReflectionCaptureSizeGameThread(0)
+		ReflectionCaptureSizeGameThread(0),
+		DesiredReflectionCaptureSizeGameThread(0)
 	{}
 
 	/** Set Data necessary to determine if GPU resources will need future updates */
-	void SetGameThreadTrackingData(int32 MaxAllocatedCubemaps, int32 CaptureSize);
+	void SetGameThreadTrackingData(int32 MaxAllocatedCubemaps, int32 CaptureSize, int32 DesiredCaptureSize);
 
 	/** Do the resources on the GPU match our desired state?  If not, reallocation will be necessary. */
 	bool DoesAllocatedDataNeedUpdate(int32 DesiredMaxCubemaps, int32 DesiredCaptureSize) const;
 
 	void ResizeCubemapArrayGPU(uint32 InMaxCubemaps, int32 InCubemapSize);
+
+	/** Resets the structure to empty, useful if you want to shrink the allocation. */
+	void Reset(FScene* Scene);
 };
 
 /** Scene state used to manage hair strands. */
@@ -3094,6 +3100,7 @@ public:
 	virtual void UpdateSceneCaptureContents(class USceneCaptureComponentCube* CaptureComponent) override;
 	virtual void UpdatePlanarReflectionContents(UPlanarReflectionComponent* CaptureComponent, FSceneRenderer& MainSceneRenderer) override;
 	virtual void AllocateReflectionCaptures(const TArray<UReflectionCaptureComponent*>& NewCaptures, const TCHAR* CaptureReason, bool bVerifyOnlyCapturing, bool bCapturingForMobile, bool bInsideTick) override;
+	virtual void ResetReflectionCaptures(bool bOnlyIfOOM) override;
 	virtual void UpdateSkyCaptureContents(const USkyLightComponent* CaptureComponent, bool bCaptureEmissiveOnly, UTextureCube* SourceCubemap, FTexture* OutProcessedTexture, float& OutAverageBrightness, FSHVectorRGB3& OutIrradianceEnvironmentMap, TArray<FFloat16Color>* OutRadianceMap) override; 
 	virtual void AllocateAndCaptureFrameSkyEnvMap(FRDGBuilder& GraphBuilder, FSceneRenderer& SceneRenderer, FViewInfo& MainView, bool bShouldRenderSkyAtmosphere, bool bShouldRenderVolumetricCloud, FInstanceCullingManager& InstanceCullingManager) override;
 	virtual void ValidateSkyLightRealTimeCapture(FRDGBuilder& GraphBuilder, const FViewInfo& View, FRDGTextureRef SceneColorTexture) override;
