Bulk reset loaders when saving actor files to avoid going through all imports of all packages once per saved actor.
Saves several seconds when saving lots of actor files.

#rnx

diff --git a/Engine/Source/Editor/UnrealEd/Private/FileHelpers.cpp b/Engine/Source/Editor/UnrealEd/Private/FileHelpers.cpp
--- a/Engine/Source/Editor/UnrealEd/Private/FileHelpers.cpp
+++ b/Engine/Source/Editor/UnrealEd/Private/FileHelpers.cpp
@@ -3161,6 +3161,11 @@ enum class InternalSavePackageResult : int8
 	Error,
 };
 
+static void PrepareSavePackages(const TArray<UPackage*>& PackagesToSave)
+{
+	ResetLoaders(MakeArrayView<UObject*>((UObject**)PackagesToSave.GetData(), PackagesToSave.Num()));
+}
+
 /**
  * Actually save a package. Prompting for Save as if necessary
  *
@@ -3595,6 +3600,8 @@ static bool InternalSavePackagesFast(const TArray<UPackage*>& PackagesToSave, bo
 		ObjectTools::CleanupAfterSuccessfulDelete(PackagesToClean, true);
 	}
 
+	PrepareSavePackages(FinalPackagesToSave);
+
 	for (UPackage* Package : FinalPackagesToSave)
 	{
 		bool bPackageLocallyWritable;
@@ -3989,6 +3996,8 @@ FEditorFileUtils::EPromptReturnCode InternalPromptForCheckoutAndSave(const TArra
 			ObjectTools::CleanupAfterSuccessfulDelete(PackagesToClean, true);
 		}
 
+		PrepareSavePackages(PackagesToSave);
+
 		for (UPackage* Package : PackagesToSave)
 		{
 			// Save the package
diff --git a/Engine/Source/Runtime/CoreUObject/Private/UObject/Linker.cpp b/Engine/Source/Runtime/CoreUObject/Private/UObject/Linker.cpp
--- a/Engine/Source/Runtime/CoreUObject/Private/UObject/Linker.cpp
+++ b/Engine/Source/Runtime/CoreUObject/Private/UObject/Linker.cpp
@@ -499,6 +499,18 @@ void ResetLoaders(UObject* InPkg)
 	FLinkerManager::Get().ResetLoaders(InPkg);
 }
 
+void ResetLoaders(TArrayView<UObject*> InOuters)
+{
+	if (IsAsyncLoading())
+	{
+		UE_LOG(LogLinker, Log, TEXT("ResetLoaders is flushing async loading"));
+	}
+
+	// Make sure we're not in the middle of loading something in the background.
+	FlushAsyncLoading();
+	FLinkerManager::Get().ResetLoaders(InOuters);
+}
+
 void DeleteLoaders()
 {
 	FLinkerManager::Get().DeleteLinkers();
diff --git a/Engine/Source/Runtime/CoreUObject/Private/UObject/LinkerManager.cpp b/Engine/Source/Runtime/CoreUObject/Private/UObject/LinkerManager.cpp
--- a/Engine/Source/Runtime/CoreUObject/Private/UObject/LinkerManager.cpp
+++ b/Engine/Source/Runtime/CoreUObject/Private/UObject/LinkerManager.cpp
@@ -167,6 +167,25 @@ void FLinkerManager::ResetLoaders(UObject* InPkg)
 	}
 }
 
+void FLinkerManager::ResetLoaders(TArrayView<UObject*> InPackages)
+{
+	TSet<FLinkerLoad*> LinkerLoads;
+	LinkerLoads.Reserve(InPackages.Num());
+
+	for (UObject* Object : InPackages)
+	{
+		if (UPackage* TopLevelPackage = Object->GetPackage())
+		{
+			if (FLinkerLoad* LinkerToReset = FLinkerLoad::FindExistingLinkerForPackage(TopLevelPackage))
+			{
+				LinkerLoads.Add(LinkerToReset);
+			}
+		}
+	}
+
+	ResetLoaders(LinkerLoads);
+}
+
 void FLinkerManager::ResetLoaders(const TSet<FLinkerLoad*>& InLinkerLoads)
 {
 	// Remove import references
diff --git a/Engine/Source/Runtime/CoreUObject/Private/UObject/LinkerManager.h b/Engine/Source/Runtime/CoreUObject/Private/UObject/LinkerManager.h
--- a/Engine/Source/Runtime/CoreUObject/Private/UObject/LinkerManager.h
+++ b/Engine/Source/Runtime/CoreUObject/Private/UObject/LinkerManager.h
@@ -116,6 +116,9 @@ public:
 	/** Empty the loaders */
 	void ResetLoaders(UObject* InPkg);
 
+	/** Empty the loaders */
+	void ResetLoaders(TArrayView<UObject*> InPackages);
+
 	/** Empty the loaders from the specified set */
 	void ResetLoaders(const TSet<FLinkerLoad*>& InLinkerLoadSet);
 
diff --git a/Engine/Source/Runtime/CoreUObject/Public/UObject/Linker.h b/Engine/Source/Runtime/CoreUObject/Public/UObject/Linker.h
--- a/Engine/Source/Runtime/CoreUObject/Public/UObject/Linker.h
+++ b/Engine/Source/Runtime/CoreUObject/Public/UObject/Linker.h
@@ -697,6 +697,7 @@ COREUOBJECT_API void ResetLinkerExports(UPackage* InPackage);
  * Flushes async loading.
  */
 COREUOBJECT_API void ResetLoaders(UObject* InOuter);
+COREUOBJECT_API void ResetLoaders(TArrayView<UObject*> InOuters);
 
 /** Deletes all linkers that have finished loading */
 COREUOBJECT_API void DeleteLoaders();
