Use a TPimplPtr rather than TSharedPtr for the FEditorTransactionDeltaContext instance

#rnx

diff --git a/Engine/Source/Editor/UnrealEd/Classes/Editor/EditorEngine.h b/Engine/Source/Editor/UnrealEd/Classes/Editor/EditorEngine.h
--- a/Engine/Source/Editor/UnrealEd/Classes/Editor/EditorEngine.h
+++ b/Engine/Source/Editor/UnrealEd/Classes/Editor/EditorEngine.h
@@ -30,6 +30,7 @@
 #include "Subsystems/SubsystemCollection.h"
 #include "RHI.h"
 #include "UnrealEngine.h"
+#include "Templates/PimplPtr.h"
 #include "Templates/UniqueObj.h"
 
 #include "EditorEngine.generated.h"
@@ -67,6 +68,7 @@ class UTextureRenderTarget2D;
 class UTransactor;
 class FTransactionObjectEvent;
 struct FTransactionContext;
+struct FEditorTransactionDeltaContext;
 struct FTypedElementHandle;
 struct FAnalyticsEventAttribute;
 class UEditorWorldExtensionManager;
@@ -2787,7 +2789,7 @@ public:
 
 private:
 
-	TSharedPtr<struct FEditorTransactionDeltaContext> CurrentUndoRedoContext;
+	TPimplPtr<FEditorTransactionDeltaContext> CurrentUndoRedoContext;
 
 	/** List of all viewport clients */
 	TArray<class FEditorViewportClient*> AllViewportClients;
diff --git a/Engine/Source/Editor/UnrealEd/Private/EditorServer.cpp b/Engine/Source/Editor/UnrealEd/Private/EditorServer.cpp
--- a/Engine/Source/Editor/UnrealEd/Private/EditorServer.cpp
+++ b/Engine/Source/Editor/UnrealEd/Private/EditorServer.cpp
@@ -1269,7 +1269,7 @@ void UEditorEngine::CreateStartupAnalyticsAttributes( TArray<FAnalyticsEventAttr
 
 UTransactor* UEditorEngine::CreateTrans()
 {
-	CurrentUndoRedoContext = MakeShared<FEditorTransactionDeltaContext>();
+	CurrentUndoRedoContext = MakePimpl<FEditorTransactionDeltaContext>();
 
 	int32 UndoBufferSize;
 
