Added missing call to Reset and Swap

#rnx

diff --git a/Engine/Source/Editor/UnrealEd/Classes/Editor/Transactor.h b/Engine/Source/Editor/UnrealEd/Classes/Editor/Transactor.h
--- a/Engine/Source/Editor/UnrealEd/Classes/Editor/Transactor.h
+++ b/Engine/Source/Editor/UnrealEd/Classes/Editor/Transactor.h
@@ -55,12 +55,14 @@ protected:
 			void Reset()
 			{
 				FTransactionSerializedObject::Reset();
+				ObjectInfo.Reset();
 				ObjectAnnotation.Reset();
 			}
 
 			void Swap(FSerializedObject& Other)
 			{
 				FTransactionSerializedObject::Swap(Other);
+				ObjectInfo.Swap(Other.ObjectInfo);
 				Exchange(ObjectAnnotation, Other.ObjectAnnotation);
 			}
 
