Multiprocess cook: Report packages discovered by CookWorkers to the Director.
#rnx

diff --git a/Engine/Source/Editor/UnrealEd/Private/CookOnTheFlyServer.cpp b/Engine/Source/Editor/UnrealEd/Private/CookOnTheFlyServer.cpp
--- a/Engine/Source/Editor/UnrealEd/Private/CookOnTheFlyServer.cpp
+++ b/Engine/Source/Editor/UnrealEd/Private/CookOnTheFlyServer.cpp
@@ -2436,7 +2436,11 @@ void UCookOnTheFlyServer::AssignRequests(TArrayView<UE::Cook::FPackageData*> Req
 		{
 			FPackageData* PackageData = Requests[Index];
 			FWorkerId Assignment = Assignments[Index];
-			if (Assignment.IsLocal())
+			if (Assignment.IsInvalid())
+			{
+				DemoteToIdle(*PackageData, ESendFlags::QueueAdd, ESuppressCookReason::MultiprocessAssignmentError);
+			}
+			else if (Assignment.IsLocal())
 			{
 				RequestQueue.AddReadyRequest(PackageData);
 			}
diff --git a/Engine/Source/Editor/UnrealEd/Private/Cooker/CookDirector.cpp b/Engine/Source/Editor/UnrealEd/Private/Cooker/CookDirector.cpp
--- a/Engine/Source/Editor/UnrealEd/Private/Cooker/CookDirector.cpp
+++ b/Engine/Source/Editor/UnrealEd/Private/Cooker/CookDirector.cpp
@@ -130,15 +130,39 @@ void FCookDirector::AssignRequests(TArrayView<UE::Cook::FPackageData*> Requests,
 	LoadBalanceStriped(SortedWorkers, Requests, OutAssignments);
 
 	// Split the output array of WorkerId assignments into a batch for each of the RemoteWorkers 
-	TArray<TArray<FPackageData*>> RemoteBatches;
+	TArray<TArray<FPackageData*>> RemoteBatches; // Indexed by WorkerId.GetRemoteIndex()
+	TArray<bool> RemoteIndexIsValid; // Indexed by WorkerId.GetRemoteIndex()
 	RemoteBatches.SetNum(MaxRemoteIndex+1);
+	RemoteIndexIsValid.Init(false, MaxRemoteIndex+1);
+	for (FCookWorkerServer* Worker : SortedWorkers)
+	{
+		RemoteIndexIsValid[Worker->GetWorkerId().GetRemoteIndex()] = true;
+	}
+
 	for (int32 RequestIndex = 0; RequestIndex < Requests.Num(); ++RequestIndex)
 	{
 		FWorkerId WorkerId = OutAssignments[RequestIndex];
+		// Override the loadbalancer's assignment if the Package has a WorkerAssignmentConstraint
+		// This allows us to guarantee that generated packages will be cooked on the worker that cooked
+		// their generator package
+		FWorkerId WorkerIdConstraint = Requests[RequestIndex]->GetWorkerAssignmentConstraint();
+		if (WorkerIdConstraint.IsValid())
+		{
+			OutAssignments[RequestIndex] = WorkerIdConstraint;
+			WorkerId = WorkerIdConstraint;
+		}
+
 		if (!WorkerId.IsLocal())
 		{
 			uint8 RemoteIndex = WorkerId.GetRemoteIndex();
 			check(RemoteIndex < RemoteBatches.Num());
+			if (!RemoteIndexIsValid[RemoteIndex])
+			{
+				UE_LOG(LogCook, Error, TEXT("Package %s can only be cooked by CookWorkerServer %d, but this worker has disconnected. The package can not be cooked."),
+					*Requests[RequestIndex]->GetPackageName().ToString(), RemoteIndex);
+				OutAssignments[RequestIndex] = FWorkerId::Invalid();
+				continue;
+			}
 			TArray<FPackageData*>& RemoteBatch = RemoteBatches[RemoteIndex];
 			if (RemoteBatch.Num() == 0)
 			{
diff --git a/Engine/Source/Editor/UnrealEd/Private/Cooker/CookPackageData.cpp b/Engine/Source/Editor/UnrealEd/Private/Cooker/CookPackageData.cpp
--- a/Engine/Source/Editor/UnrealEd/Private/Cooker/CookPackageData.cpp
+++ b/Engine/Source/Editor/UnrealEd/Private/Cooker/CookPackageData.cpp
@@ -1403,6 +1403,7 @@ bool FGeneratorPackage::TryGenerateList(UObject* OwnerObject, FPackageDatas& Pac
 		GeneratedInfo.Dependencies = MoveTemp(SplitterData.Dependencies);
 		GeneratedInfo.SetIsCreateAsMap(bCreateAsMap);
 		PackageData->SetGeneratedOwner(this);
+		PackageData->SetWorkerAssignmentConstraint(FWorkerId::Local());
 	}
 	RemainingToPopulate = GeneratorDatas.Num() + 1; // GeneratedPackaged plus one for the Generator
 	return true;
diff --git a/Engine/Source/Editor/UnrealEd/Private/Cooker/CookPackageData.h b/Engine/Source/Editor/UnrealEd/Private/Cooker/CookPackageData.h
--- a/Engine/Source/Editor/UnrealEd/Private/Cooker/CookPackageData.h
+++ b/Engine/Source/Editor/UnrealEd/Private/Cooker/CookPackageData.h
@@ -481,10 +481,14 @@ public:
 	/** Set whether COTFS is keeping this package referenced referenced during GC. */
 	void SetKeepReferencedDuringGC(bool Value) { bKeepReferencedDuringGC = Value != 0; }
 
-	/** For MultiProcessCooks, Get the id of the worker this Package is assigned to. Returns Invalid if not assigned to a worker. */
+	/** For MultiProcessCooks, Get the id of the worker this Package is assigned to; InvalidId means owned by local. */
 	FWorkerId GetWorkerAssignment() const { return WorkerAssignment; }
 	/** Set the id of the worker this Package is assigned to. */
 	void SetWorkerAssignment(FWorkerId InWorkerAssignment) { WorkerAssignment = InWorkerAssignment; }
+	/** Get the workerid that is the only worker allowed to cook this package; InvalidId means no constraint. */
+	FWorkerId GetWorkerAssignmentConstraint() const { return WorkerAssignmentConstraint; }
+	/** Set the workerid that is the only worker allowed to cook this package; default is InvalidId; */
+	void SetWorkerAssignmentConstraint(FWorkerId InWorkerAssignment) { WorkerAssignmentConstraint = InWorkerAssignment; }
 
 	/** Marshall this PackageData to a ConstructData that can be used later or on a remote machine to reconstruct it. */
 	FConstructPackageData CreateConstructData();
@@ -592,6 +596,7 @@ private:
 	FInstigator Instigator;
 
 	FWorkerId WorkerAssignment = FWorkerId::Invalid();
+	FWorkerId WorkerAssignmentConstraint = FWorkerId::Invalid();
 	uint32 State : int32(EPackageState::BitCount);
 	uint32 bIsUrgent : 1;
 	uint32 bIsVisited : 1;
diff --git a/Engine/Source/Editor/UnrealEd/Private/Cooker/CookTypes.h b/Engine/Source/Editor/UnrealEd/Private/Cooker/CookTypes.h
--- a/Engine/Source/Editor/UnrealEd/Private/Cooker/CookTypes.h
+++ b/Engine/Source/Editor/UnrealEd/Private/Cooker/CookTypes.h
@@ -143,6 +143,7 @@ namespace UE::Cook
 		SaveError,
 		OnlyEditorOnly,
 		CookCanceled,
+		MultiprocessAssignmentError,
 	};
 
 	/** The type of callback for External Requests that needs to be executed within the Scheduler's lock. */
diff --git a/Engine/Source/Editor/UnrealEd/Private/Cooker/CookWorkerClient.cpp b/Engine/Source/Editor/UnrealEd/Private/Cooker/CookWorkerClient.cpp
--- a/Engine/Source/Editor/UnrealEd/Private/Cooker/CookWorkerClient.cpp
+++ b/Engine/Source/Editor/UnrealEd/Private/Cooker/CookWorkerClient.cpp
@@ -118,7 +118,7 @@ void FCookWorkerClient::DoneWithInitialSettings()
 	InitialConfigMessage.Reset();
 }
 
-void FCookWorkerClient::ReportDemoteToIdle(FPackageData& PackageData, ESuppressCookReason Reason)
+void FCookWorkerClient::ReportDemoteToIdle(const FPackageData& PackageData, ESuppressCookReason Reason)
 {
 	FPackageRemoteResult& Result = PendingResults.Emplace_GetRef();
 	Result.PackageName = PackageData.GetPackageName();
@@ -183,6 +183,14 @@ void FCookWorkerClient::ReportPromoteToSaveComplete(FPackageData& PackageData)
 	}
 }
 
+void FCookWorkerClient::ReportDiscoveredPackage(const FPackageData& PackageData, const FInstigator& Instigator)
+{
+	FDiscoveredPackage& Discovered = PendingDiscoveredPackages.Emplace_GetRef();
+	Discovered.PackageName = PackageData.GetPackageName();
+	Discovered.NormalizedFileName = PackageData.GetFileName();
+	Discovered.Instigator = Instigator;
+}
+
 EPollStatus FCookWorkerClient::PollTryConnect(const FDirectorConnectionInfo& ConnectInfo)
 {
 	for (;;)
@@ -370,14 +378,20 @@ void FCookWorkerClient::PumpSendMessages()
 
 void FCookWorkerClient::SendPendingResults()
 {
-	if (PendingResults.IsEmpty())
+	if (!PendingResults.IsEmpty())
 	{
-		return;
+		FPackageResultsMessage Message;
+		Message.Results = MoveTemp(PendingResults);
+		SendMessage(Message);
+		PendingResults.Reset();
+	}
+	if (!PendingDiscoveredPackages.IsEmpty())
+	{
+		FDiscoveredPackagesMessage Message;
+		Message.Packages = MoveTemp(PendingDiscoveredPackages);
+		SendMessage(Message);
+		PendingDiscoveredPackages.Reset();
 	}
-	FPackageResultsMessage Message;
-	Message.Results = MoveTemp(PendingResults);
-	SendMessage(Message);
-	PendingResults.Reset();
 }
 
 void FCookWorkerClient::PumpReceiveMessages()
@@ -405,7 +419,6 @@ void FCookWorkerClient::HandleReceiveMessages(TArray<UE::CompactBinaryTCP::FMars
 			if (AbortMessage.Type == FAbortWorkerMessage::EType::CookComplete)
 			{
 				UE_LOG(LogCook, Display, TEXT("CookWorkerClient received CookComplete message from Director. Flushing messages and shutting down."));
-				// MPCOOKTODO: Add synchronous flush of messages here
 			}
 			else
 			{
diff --git a/Engine/Source/Editor/UnrealEd/Private/Cooker/CookWorkerClient.h b/Engine/Source/Editor/UnrealEd/Private/Cooker/CookWorkerClient.h
--- a/Engine/Source/Editor/UnrealEd/Private/Cooker/CookWorkerClient.h
+++ b/Engine/Source/Editor/UnrealEd/Private/Cooker/CookWorkerClient.h
@@ -11,6 +11,7 @@
 namespace UE::Cook { class IMPCollector; }
 namespace UE::Cook { struct FAssignPackagesMessage; }
 namespace UE::Cook { struct FDirectorConnectionInfo; }
+namespace UE::Cook { struct FDiscoveredPackage; }
 namespace UE::Cook { struct FInitialConfigMessage; }
 namespace UE::Cook { struct FPackageRemoteResult; }
 
@@ -47,9 +48,11 @@ public:
 	void DoneWithInitialSettings();
 
 	/** Queue a message to the server that the Package was cook-suppressed. Will be sent during Tick. */
-	void ReportDemoteToIdle(FPackageData& PackageData, ESuppressCookReason Reason);
+	void ReportDemoteToIdle(const FPackageData& PackageData, ESuppressCookReason Reason);
 	/** Queue a message to the server that the Package was saved. Will be sent during Tick. */
 	void ReportPromoteToSaveComplete(FPackageData& PackageData);
+	/** Queue a message to the server that a package was discovered as needed in the cook. Will be sent during Tick. */
+	void ReportDiscoveredPackage(const FPackageData& PackageData, const FInstigator& Instigator);
 
 	/** Register a Collector for periodic ticking that sends messages to the Director. */
 	void Register(IMPCollector* Collector);
@@ -106,6 +109,7 @@ private:
 	TUniquePtr<FInitialConfigMessage> InitialConfigMessage;
 	TArray<ITargetPlatform*> OrderedSessionPlatforms;
 	TArray<FPackageRemoteResult> PendingResults;
+	TArray<FDiscoveredPackage> PendingDiscoveredPackages;
 	TArray<TRefCountPtr<IMPCollector>> CollectorsToTick;
 	UE::CompactBinaryTCP::FSendBuffer SendBuffer;
 	UE::CompactBinaryTCP::FReceiveBuffer ReceiveBuffer;
diff --git a/Engine/Source/Editor/UnrealEd/Private/Cooker/CookWorkerServer.cpp b/Engine/Source/Editor/UnrealEd/Private/Cooker/CookWorkerServer.cpp
--- a/Engine/Source/Editor/UnrealEd/Private/Cooker/CookWorkerServer.cpp
+++ b/Engine/Source/Editor/UnrealEd/Private/Cooker/CookWorkerServer.cpp
@@ -428,6 +428,21 @@ void FCookWorkerServer::HandleReceiveMessages(TArray<UE::CompactBinaryTCP::FMars
 				RecordResults(ResultsMessage);
 			}
 		}
+		else if (Message.MessageType == FDiscoveredPackagesMessage::MessageType)
+		{
+			FDiscoveredPackagesMessage DiscoveredMessage;
+			if (!DiscoveredMessage.TryRead(MoveTemp(Message.Object)))
+			{
+				LogInvalidMessage(TEXT("FDiscoveredPackagesMessage"));
+			}
+			else
+			{
+				for (FDiscoveredPackage& DiscoveredPackage : DiscoveredMessage.Packages)
+				{
+					AddDiscoveredPackage(MoveTemp(DiscoveredPackage));
+				}
+			}
+		}
 		else
 		{
 			TRefCountPtr<IMPCollector>* Collector = Director.MessageHandlers.Find(Message.MessageType);
@@ -530,6 +545,24 @@ void FCookWorkerServer::LogInvalidMessage(const TCHAR* MessageTypeName)
 		MessageTypeName);
 }
 
+void FCookWorkerServer::AddDiscoveredPackage(FDiscoveredPackage&& DiscoveredPackage)
+{
+	FPackageData& PackageData = COTFS.PackageDatas->FindOrAddPackageData(DiscoveredPackage.PackageName,
+		DiscoveredPackage.NormalizedFileName);
+	if (PackageData.IsInProgress() || PackageData.HasAnyCookedPlatform())
+	{
+		// The CookWorker thought this was a new package, but the Director already knows about it; ignore the report
+		return;
+	}
+
+	if (DiscoveredPackage.Instigator.Category == EInstigator::GeneratedPackage)
+	{
+		PackageData.SetGenerated(true);
+		PackageData.SetWorkerAssignmentConstraint(GetWorkerId());
+	}
+	COTFS.QueueDiscoveredPackageData(PackageData, MoveTemp(DiscoveredPackage.Instigator));
+}
+
 FAssignPackagesMessage::FAssignPackagesMessage(TArray<FConstructPackageData>&& InPackageDatas)
 	: PackageDatas(MoveTemp(InPackageDatas))
 {
@@ -665,4 +698,50 @@ bool FInitialConfigMessage::TryRead(FCbObject&& Object)
 
 FGuid FInitialConfigMessage::MessageType(TEXT("340CDCB927304CEB9C0A66B5F707FC2B"));
 
+FCbWriter& operator<<(FCbWriter& Writer, const FDiscoveredPackage& Package)
+{
+	Writer.BeginObject();
+	Writer << "PackageName" << Package.PackageName;
+	Writer << "NormalizedFileName" << Package.NormalizedFileName;
+	Writer << "Instigator.Category" << static_cast<uint8>(Package.Instigator.Category);
+	Writer << "Instigator.Referencer" << Package.Instigator.Referencer;
+	Writer.EndObject();
+	return Writer;
+}
+
+bool LoadFromCompactBinary(FCbFieldView Field, FDiscoveredPackage& OutPackage)
+{
+	bool bOk = LoadFromCompactBinary(Field["PackageName"], OutPackage.PackageName);
+	bOk = LoadFromCompactBinary(Field["NormalizedFileName"], OutPackage.NormalizedFileName) & bOk;
+	uint8 CategoryInt;
+	if (LoadFromCompactBinary(Field["Instigator.Category"], CategoryInt) &&
+		CategoryInt < static_cast<uint8>(EInstigator::Count))
+
+	{
+		OutPackage.Instigator.Category = static_cast<EInstigator>(CategoryInt);
+	}
+	else
+	{
+		bOk = false;
+	}
+	bOk = LoadFromCompactBinary(Field["Instigator.Referencer"], OutPackage.Instigator.Referencer) & bOk;
+	if (!bOk)
+	{
+		OutPackage = FDiscoveredPackage();
+	}
+	return bOk;
+}
+
+void FDiscoveredPackagesMessage::Write(FCbWriter& Writer) const
+{
+	Writer << "Packages" << Packages;
+}
+
+bool FDiscoveredPackagesMessage::TryRead(FCbObject&& Object)
+{
+	return LoadFromCompactBinary(Object["Packages"], Packages);
+}
+
+FGuid FDiscoveredPackagesMessage::MessageType(TEXT("C9F5BC5C11484B06B346B411F1ED3090"));
+
 }
\ No newline at end of file
diff --git a/Engine/Source/Editor/UnrealEd/Private/Cooker/CookWorkerServer.h b/Engine/Source/Editor/UnrealEd/Private/Cooker/CookWorkerServer.h
--- a/Engine/Source/Editor/UnrealEd/Private/Cooker/CookWorkerServer.h
+++ b/Engine/Source/Editor/UnrealEd/Private/Cooker/CookWorkerServer.h
@@ -17,6 +17,7 @@ struct FProcHandle;
 namespace UE::CompactBinaryTCP { struct FMarshalledMessage; }
 namespace UE::Cook { class FCookDirector; }
 namespace UE::Cook { struct FConstructPackageData; }
+namespace UE::Cook { struct FDiscoveredPackage; }
 namespace UE::Cook { struct FPackageData; }
 namespace UE::Cook { struct FPackageResultsMessage; }
 namespace UE::Cook { struct FWorkerConnectMessage; }
@@ -100,6 +101,7 @@ private:
 	/** Add results from the client to the local CookOnTheFlyServer. */
 	void RecordResults(FPackageResultsMessage& Message);
 	void LogInvalidMessage(const TCHAR* MessageTypeName);
+	void AddDiscoveredPackage(FDiscoveredPackage&& DiscoveredPackage);
 
 	TArray<FPackageData*> PackagesToAssign;
 	TSet<FPackageData*> PendingPackages;
@@ -206,4 +208,30 @@ private:
 	bool bZenStore = false;
 };
 
+/** Information about a discovered package sent from a CookWorker to the Director. */
+struct FDiscoveredPackage
+{
+	FName PackageName;
+	FName NormalizedFileName;
+	FInstigator Instigator;
+};
+FCbWriter& operator<<(FCbWriter& Writer, const FDiscoveredPackage& Package);
+bool LoadFromCompactBinary(FCbFieldView Field, FDiscoveredPackage& OutPackage);
+
+/**
+ * Message from CookWorker to Director that reports dependency packages discovered during load/save of
+ * a package that were not found in the earlier traversal of the packages dependencies.
+ */
+struct FDiscoveredPackagesMessage : public UE::CompactBinaryTCP::IMessage
+{
+public:
+	virtual void Write(FCbWriter& Writer) const override;
+	virtual bool TryRead(FCbObject&& Object) override;
+	virtual FGuid GetMessageType() const override { return MessageType; }
+
+public:
+	TArray<FDiscoveredPackage> Packages;
+	static FGuid MessageType;
+};
+
 }
\ No newline at end of file
diff --git a/Engine/Source/Editor/UnrealEd/Private/Cooker/IWorkerRequests.h b/Engine/Source/Editor/UnrealEd/Private/Cooker/IWorkerRequests.h
--- a/Engine/Source/Editor/UnrealEd/Private/Cooker/IWorkerRequests.h
+++ b/Engine/Source/Editor/UnrealEd/Private/Cooker/IWorkerRequests.h
@@ -49,7 +49,8 @@ public:
 	virtual void DequeueAllExternal(TArray<FSchedulerCallback>& OutCallbacks, TArray<FFilePlatformRequest>& OutCookRequests) = 0;
 
 	// Writing Packages during Cook
-	virtual void AddDiscoveredPackage(FPackageData& PackageData, FInstigator& Instigator, bool bLoadReady, bool& bOutShouldAddToQueue) = 0;
+	virtual void AddDiscoveredPackage(const FPackageData& PackageData, const FInstigator& Instigator,
+		bool bLoadReady, bool& bOutShouldAddToQueue) = 0;
 
 	// Writing Packages from CookByTheBook
 	virtual void AddStartCookByTheBookRequest(FFilePlatformRequest&& Request) = 0;
diff --git a/Engine/Source/Editor/UnrealEd/Private/Cooker/WorkerRequestsLocal.cpp b/Engine/Source/Editor/UnrealEd/Private/Cooker/WorkerRequestsLocal.cpp
--- a/Engine/Source/Editor/UnrealEd/Private/Cooker/WorkerRequestsLocal.cpp
+++ b/Engine/Source/Editor/UnrealEd/Private/Cooker/WorkerRequestsLocal.cpp
@@ -32,7 +32,8 @@ void FWorkerRequestsLocal::DequeueAllExternal(TArray<FSchedulerCallback>& OutCal
 	ExternalRequests.DequeueAll(OutCallbacks, OutCookRequests);
 }
 
-void FWorkerRequestsLocal::AddDiscoveredPackage(FPackageData& PackageData, FInstigator& Instigator, bool bLoadReady, bool& bOutShouldAddToQueue)
+void FWorkerRequestsLocal::AddDiscoveredPackage(const FPackageData& PackageData, const FInstigator& Instigator,
+	bool bLoadReady, bool& bOutShouldAddToQueue)
 {
 	bOutShouldAddToQueue = true;
 }
diff --git a/Engine/Source/Editor/UnrealEd/Private/Cooker/WorkerRequestsLocal.h b/Engine/Source/Editor/UnrealEd/Private/Cooker/WorkerRequestsLocal.h
--- a/Engine/Source/Editor/UnrealEd/Private/Cooker/WorkerRequestsLocal.h
+++ b/Engine/Source/Editor/UnrealEd/Private/Cooker/WorkerRequestsLocal.h
@@ -14,12 +14,15 @@ class FWorkerRequestsLocal : public IWorkerRequests
 public:
 	virtual bool HasExternalRequests() const override;
 	virtual int32 GetNumExternalRequests() const override;
-	virtual EExternalRequestType DequeueNextCluster(TArray<FSchedulerCallback>& OutCallbacks, TArray<FFilePlatformRequest>& OutBuildRequests) override;
+	virtual EExternalRequestType DequeueNextCluster(TArray<FSchedulerCallback>& OutCallbacks,
+		TArray<FFilePlatformRequest>& OutBuildRequests) override;
 	virtual bool DequeueSchedulerCallbacks(TArray<FSchedulerCallback>& OutCallbacks) override;
-	virtual void DequeueAllExternal(TArray<FSchedulerCallback>& OutCallbacks, TArray<FFilePlatformRequest>& OutCookRequests) override;
+	virtual void DequeueAllExternal(TArray<FSchedulerCallback>& OutCallbacks,
+		TArray<FFilePlatformRequest>& OutCookRequests) override;
 	virtual void RemapTargetPlatforms(const TMap<ITargetPlatform*, ITargetPlatform*>& Remap) override;
 	virtual void OnRemoveSessionPlatform(const ITargetPlatform* TargetPlatform) override;
-	virtual void AddDiscoveredPackage(FPackageData& PackageData, FInstigator& Instigator, bool bLoadReady, bool& bOutShouldAddToQueue) override;
+	virtual void AddDiscoveredPackage(const FPackageData& PackageData, const FInstigator& Instigator,
+		bool bLoadReady, bool& bOutShouldAddToQueue) override;
 	virtual void AddStartCookByTheBookRequest(FFilePlatformRequest&& Request) override;
 	virtual void InitializeCookOnTheFly() override;
 	virtual void AddCookOnTheFlyRequest(FFilePlatformRequest&& Request) override;
@@ -31,8 +34,10 @@ public:
 	virtual void ReportAccessedIniSettings(UCookOnTheFlyServer& COTFS, const FConfigFile& Config) override;
 	virtual void ReportDemoteToIdle(UE::Cook::FPackageData& PackageData, ESuppressCookReason Reason) override;
 	virtual void ReportPromoteToSaveComplete(UE::Cook::FPackageData& PackageData) override;
-	virtual void GetInitializeConfigSettings(UCookOnTheFlyServer& COTFS, const FString& OutputOverrideDirectory, UE::Cook::FInitializeConfigSettings& Settings) override;
-	virtual void GetBeginCookConfigSettings(UCookOnTheFlyServer& COTFS, FBeginCookContext& BeginContext, UE::Cook::FBeginCookConfigSettings& Settings) override;
+	virtual void GetInitializeConfigSettings(UCookOnTheFlyServer& COTFS, const FString& OutputOverrideDirectory,
+		UE::Cook::FInitializeConfigSettings& Settings) override;
+	virtual void GetBeginCookConfigSettings(UCookOnTheFlyServer& COTFS, FBeginCookContext& BeginContext,
+		UE::Cook::FBeginCookConfigSettings& Settings) override;
 	virtual void GetBeginCookIterativeFlags(UCookOnTheFlyServer& COTFS, FBeginCookContext& BeginContext) override;
 	virtual ECookMode::Type GetDirectorCookMode(UCookOnTheFlyServer& COTFS) override;
 
diff --git a/Engine/Source/Editor/UnrealEd/Private/Cooker/WorkerRequestsRemote.cpp b/Engine/Source/Editor/UnrealEd/Private/Cooker/WorkerRequestsRemote.cpp
--- a/Engine/Source/Editor/UnrealEd/Private/Cooker/WorkerRequestsRemote.cpp
+++ b/Engine/Source/Editor/UnrealEd/Private/Cooker/WorkerRequestsRemote.cpp
@@ -25,7 +25,8 @@ int32 FWorkerRequestsRemote::GetNumExternalRequests() const
 	return ExternalRequests.GetNumRequests();
 }
 
-EExternalRequestType FWorkerRequestsRemote::DequeueNextCluster(TArray<FSchedulerCallback>& OutCallbacks, TArray<FFilePlatformRequest>& OutBuildRequests)
+EExternalRequestType FWorkerRequestsRemote::DequeueNextCluster(TArray<FSchedulerCallback>& OutCallbacks,
+	TArray<FFilePlatformRequest>& OutBuildRequests)
 {
 	return ExternalRequests.DequeueNextCluster(OutCallbacks, OutBuildRequests);
 }
@@ -40,12 +41,11 @@ void FWorkerRequestsRemote::DequeueAllExternal(TArray<FSchedulerCallback>& OutCa
 	ExternalRequests.DequeueAll(OutCallbacks, OutCookRequests);
 }
 
-void FWorkerRequestsRemote::AddDiscoveredPackage(FPackageData& PackageData, FInstigator& Instigator, bool bLoadReady, bool& bOutShouldAddToQueue)
+void FWorkerRequestsRemote::AddDiscoveredPackage(const FPackageData& PackageData, const FInstigator& Instigator,
+	bool bLoadReady, bool& bOutShouldAddToQueue)
 {
 	bOutShouldAddToQueue = false;
-	// MPCOOKTODO: Not yet implemented
-	// Send package information to the Director, and make sure to send the instigator.
-	// The director needs to handle GeneratedPackages by forcing them to go to this CookWorker, or if they need to go to another worker, add handling that calls pumpsave on the generator first
+	CookWorkerClient.ReportDiscoveredPackage(PackageData, Instigator);
 }
 
 void FWorkerRequestsRemote::AddStartCookByTheBookRequest(FFilePlatformRequest&& Request)
diff --git a/Engine/Source/Editor/UnrealEd/Private/Cooker/WorkerRequestsRemote.h b/Engine/Source/Editor/UnrealEd/Private/Cooker/WorkerRequestsRemote.h
--- a/Engine/Source/Editor/UnrealEd/Private/Cooker/WorkerRequestsRemote.h
+++ b/Engine/Source/Editor/UnrealEd/Private/Cooker/WorkerRequestsRemote.h
@@ -20,12 +20,15 @@ public:
 
 	virtual bool HasExternalRequests() const override;
 	virtual int32 GetNumExternalRequests() const override;
-	virtual EExternalRequestType DequeueNextCluster(TArray<FSchedulerCallback>& OutCallbacks, TArray<FFilePlatformRequest>& OutBuildRequests) override;
+	virtual EExternalRequestType DequeueNextCluster(TArray<FSchedulerCallback>& OutCallbacks,
+		TArray<FFilePlatformRequest>& OutBuildRequests) override;
 	virtual bool DequeueSchedulerCallbacks(TArray<FSchedulerCallback>& OutCallbacks) override;
-	virtual void DequeueAllExternal(TArray<FSchedulerCallback>& OutCallbacks, TArray<FFilePlatformRequest>& OutCookRequests) override;
+	virtual void DequeueAllExternal(TArray<FSchedulerCallback>& OutCallbacks,
+		TArray<FFilePlatformRequest>& OutCookRequests) override;
 	virtual void RemapTargetPlatforms(const TMap<ITargetPlatform*, ITargetPlatform*>& Remap) override;
 	virtual void OnRemoveSessionPlatform(const ITargetPlatform* TargetPlatform) override;
-	virtual void AddDiscoveredPackage(FPackageData& PackageData, FInstigator& Instigator, bool bLoadReady, bool& bOutShouldAddToQueue) override;
+	virtual void AddDiscoveredPackage(const FPackageData& PackageData, const FInstigator& Instigator,
+		bool bLoadReady, bool& bOutShouldAddToQueue) override;
 	virtual void AddStartCookByTheBookRequest(FFilePlatformRequest&& Request) override;
 	virtual void InitializeCookOnTheFly() override;
 	virtual void AddCookOnTheFlyRequest(FFilePlatformRequest&& Request) override;
@@ -37,8 +40,10 @@ public:
 	virtual void ReportAccessedIniSettings(UCookOnTheFlyServer& COTFS, const FConfigFile& Config) override;
 	virtual void ReportDemoteToIdle(UE::Cook::FPackageData& PackageData, ESuppressCookReason Reason) override;
 	virtual void ReportPromoteToSaveComplete(UE::Cook::FPackageData& PackageData) override;
-	virtual void GetInitializeConfigSettings(UCookOnTheFlyServer& COTFS, const FString& OutputOverrideDirectory, UE::Cook::FInitializeConfigSettings& Settings) override;
-	virtual void GetBeginCookConfigSettings(UCookOnTheFlyServer& COTFS, FBeginCookContext& BeginContext, UE::Cook::FBeginCookConfigSettings& Settings) override;
+	virtual void GetInitializeConfigSettings(UCookOnTheFlyServer& COTFS, const FString& OutputOverrideDirectory,
+		UE::Cook::FInitializeConfigSettings& Settings) override;
+	virtual void GetBeginCookConfigSettings(UCookOnTheFlyServer& COTFS, FBeginCookContext& BeginContext,
+		UE::Cook::FBeginCookConfigSettings& Settings) override;
 	virtual void GetBeginCookIterativeFlags(UCookOnTheFlyServer& COTFS, FBeginCookContext& BeginContext) override;
 	virtual ECookMode::Type GetDirectorCookMode(UCookOnTheFlyServer& COTFS) override;
 
diff --git a/Engine/Source/Runtime/Engine/Public/Engine/ICookInfo.h b/Engine/Source/Runtime/Engine/Public/Engine/ICookInfo.h
--- a/Engine/Source/Runtime/Engine/Public/Engine/ICookInfo.h
+++ b/Engine/Source/Runtime/Engine/Public/Engine/ICookInfo.h
@@ -54,11 +54,12 @@ namespace UE::Cook
 
 
 /** The different ways a package can be discovered by the cooker. */
-enum class EInstigator
+enum class EInstigator : uint8
 {
 #define EINSTIGATOR_VALUE_CALLBACK(Name, bAllowUnparameterized) Name,
 	EINSTIGATOR_VALUES(EINSTIGATOR_VALUE_CALLBACK)
 #undef EINSTIGATOR_VALUE_CALLBACK
+	Count,
 };
 ENGINE_API const TCHAR* LexToString(EInstigator Value);
 
