DDC: Added ticking of FTickableCookObject to the cache commandlet

#rnx

diff --git a/Engine/Source/Editor/UnrealEd/Private/Commandlets/DerivedDataCacheCommandlet.cpp b/Engine/Source/Editor/UnrealEd/Private/Commandlets/DerivedDataCacheCommandlet.cpp
--- a/Engine/Source/Editor/UnrealEd/Private/Commandlets/DerivedDataCacheCommandlet.cpp
+++ b/Engine/Source/Editor/UnrealEd/Private/Commandlets/DerivedDataCacheCommandlet.cpp
@@ -135,6 +135,22 @@ void UDerivedDataCacheCommandlet::MaybeMarkPackageAsAlreadyLoaded(UPackage* Pack
 	}
 }
 
+namespace UE::Private::DerivedDataCacheCommandlet
+{
+
+static void TickCookObjects(bool bCookComplete = false)
+{
+	constexpr double TickPeriod = 0.1;
+	static double LastTickTime = FPlatformTime::Seconds();
+	const double CurrentTime = FPlatformTime::Seconds();
+	if (bCookComplete || LastTickTime + TickPeriod <= CurrentTime)
+	{
+		TRACE_CPUPROFILER_EVENT_SCOPE(TickCookObjects);
+		FTickableCookObject::TickObjects(CurrentTime - LastTickTime, bCookComplete);
+		LastTickTime = CurrentTime;
+	}
+}
+
 static void WaitForCompilationToFinish(bool& bInOutHadActivity)
 {
 	auto LogStatus =
@@ -153,6 +169,8 @@ static void WaitForCompilationToFinish(bool& bInOutHadActivity)
 
 	while (FAssetCompilingManager::Get().GetNumRemainingAssets() > 0)
 	{
+		TickCookObjects();
+
 		for (IAssetCompilingManager* CompilingManager : FAssetCompilingManager::Get().GetRegisteredManagers())
 		{
 			int32 CachedAssetCount = CompilingManager->GetNumRemainingAssets();
@@ -183,15 +201,7 @@ static void WaitForCompilationToFinish(bool& bInOutHadActivity)
 	}
 }
 
-static void PumpAsync(bool* bInOutHadActivity = nullptr)
-{
-	bool bHadActivity = false;
-	WaitForCompilationToFinish(bHadActivity);
-	if (bInOutHadActivity)
-	{
-		*bInOutHadActivity = *bInOutHadActivity || bHadActivity;
-	}
-}
+} // UE::Private::DerivedDataCacheCommandlet
 
 void UDerivedDataCacheCommandlet::CacheLoadedPackages(UPackage* CurrentPackage, uint8 PackageFilter, const TArray<ITargetPlatform*>& Platforms, TSet<FName>& OutNewProcessedPackages)
 {
@@ -312,12 +322,14 @@ void UDerivedDataCacheCommandlet::FinishCachingObjects(const TArray<ITargetPlatf
 
 	while (CachingObjects.Num() > 0)
 	{
+		UE::Private::DerivedDataCacheCommandlet::TickCookObjects();
+
 		bool bHadActivity = ProcessCachingObjects(Platforms);
 
 		double CurrentTime = FPlatformTime::Seconds();
 		if (!bHadActivity)
 		{
-			PumpAsync(&bHadActivity);
+			UE::Private::DerivedDataCacheCommandlet::WaitForCompilationToFinish(bHadActivity);
 		}
 		if (!bHadActivity)
 		{
@@ -791,6 +803,8 @@ int32 UDerivedDataCacheCommandlet::Main( const FString& Params )
 				}
 			}
 
+			UE::Private::DerivedDataCacheCommandlet::TickCookObjects();
+
 			// Perform a GC if conditions are met
 			if (NumProcessedSinceLastGC >= GCInterval || PackagePaths.IsEmpty() || bLastPackageWasMap)
 			{
@@ -815,6 +829,8 @@ int32 UDerivedDataCacheCommandlet::Main( const FString& Params )
 
 	FinishCachingObjects(GetTargetPlatformManager()->GetActiveTargetPlatforms());
 
+	UE::Private::DerivedDataCacheCommandlet::TickCookObjects(/*bCookComplete*/ true);
+
 	GetDerivedDataCacheRef().WaitForQuiescence(true);
 
 	UE_LOG(LogDerivedDataCacheCommandlet, Display, TEXT("BeginCacheTime=%.2lfs, FinishCacheTime=%.2lfs, GCTime=%.2lfs."), BeginCacheTime, FinishCacheTime, GCTime);
