Removed global Transactor.h include

#rnx

diff --git a/Engine/Source/Editor/UnrealEd/Classes/Editor/EditorEngine.h b/Engine/Source/Editor/UnrealEd/Classes/Editor/EditorEngine.h
--- a/Engine/Source/Editor/UnrealEd/Classes/Editor/EditorEngine.h
+++ b/Engine/Source/Editor/UnrealEd/Classes/Editor/EditorEngine.h
@@ -18,7 +18,6 @@
 #include "UObject/UObjectAnnotation.h"
 #include "Engine/Brush.h"
 #include "Model.h"
-#include "Editor/Transactor.h"
 #include "Engine/Engine.h"
 #include "Settings/LevelEditorPlaySettings.h"
 #include "Settings/LevelEditorViewportSettings.h"
diff --git a/Engine/Source/Editor/UnrealEd/Private/UnrealEdPrivatePCH.h b/Engine/Source/Editor/UnrealEd/Private/UnrealEdPrivatePCH.h
--- a/Engine/Source/Editor/UnrealEd/Private/UnrealEdPrivatePCH.h
+++ b/Engine/Source/Editor/UnrealEd/Private/UnrealEdPrivatePCH.h
@@ -469,7 +469,6 @@
 #include "Engine/MaterialMerging.h"
 #include "Engine/Level.h"
 #include "GameFramework/Actor.h"
-#include "ComponentInstanceDataCache.h"
 #include "Components/ActorComponent.h"
 #include "Components/SceneComponent.h"
 #include "EngineDefines.h"
@@ -614,7 +613,6 @@
 
 // From UnrealEd:
 #include "Editor/UnrealEdTypes.h"
-#include "Editor/Transactor.h"
 #include "Settings/LevelEditorViewportSettings.h"
 #include "Viewports.h"
 #include "Editor/EditorEngine.h"
diff --git a/Engine/Source/Editor/UnrealEd/Public/UnrealEdSharedPCH.h b/Engine/Source/Editor/UnrealEd/Public/UnrealEdSharedPCH.h
--- a/Engine/Source/Editor/UnrealEd/Public/UnrealEdSharedPCH.h
+++ b/Engine/Source/Editor/UnrealEd/Public/UnrealEdSharedPCH.h
@@ -458,7 +458,6 @@
 #include "EdGraph/EdGraphPin.h"
 #include "Interfaces/Interface_AssetUserData.h"
 #include "Components/ActorComponent.h"
-#include "ComponentInstanceDataCache.h"
 #include "Components/SceneComponent.h"
 #include "Engine/MaterialMerging.h"
 #include "Engine/Level.h"
@@ -601,7 +600,6 @@
 #include "TickableEditorObject.h"
 #include "Editor/UnrealEdTypes.h"
 #include "Viewports.h"
-#include "Editor/Transactor.h"
 #include "Settings/LevelEditorViewportSettings.h"
 #include "Settings/LevelEditorPlaySettings.h"
 #include "Editor/EditorEngine.h"
