[NO NON-PII COMMIT MESSAGE]

diff --git a/Engine/Source/Editor/UnrealEd/Private/Settings/EditorStyleSettingsCustomization.cpp b/Engine/Source/Editor/UnrealEd/Private/Settings/EditorStyleSettingsCustomization.cpp
--- a/Engine/Source/Editor/UnrealEd/Private/Settings/EditorStyleSettingsCustomization.cpp
+++ b/Engine/Source/Editor/UnrealEd/Private/Settings/EditorStyleSettingsCustomization.cpp
@@ -430,6 +430,12 @@ void FEditorStyleSettingsCustomization::MakeThemePickerRow(IDetailPropertyRow& P
 		[
 			SNew(SButton)
 			.ButtonStyle(FAppStyle::Get(), "SimpleButton")
+			.Visibility_Lambda(
+				[]()
+				{
+					// if current active theme is the default dark theme: collapse the edit button. 
+					return USlateThemeManager::Get().IsDefaultThemeActive() ? EVisibility::Collapsed : EVisibility::Visible;
+				})
 			.ToolTipText(LOCTEXT("EditThemeToolTip", "Edit this theme"))
 			.OnClicked(this, &FEditorStyleSettingsCustomization::OnEditThemeClicked)
 			[
@@ -463,8 +469,8 @@ void FEditorStyleSettingsCustomization::MakeThemePickerRow(IDetailPropertyRow& P
 			.Visibility_Lambda(
 				[]()
 				{
-					// if current active theme is the default dark theme: hide the delete button. 
-					return USlateThemeManager::Get().IsDefaultThemeActive() ? EVisibility::Hidden : EVisibility::Visible;
+					// if current active theme is the default dark theme: collapse the delete button. 
+					return USlateThemeManager::Get().IsDefaultThemeActive() ? EVisibility::Collapsed : EVisibility::Visible;
 				})
 			.ToolTipText(LOCTEXT("DeleteThemeToolTip", "Delete this theme"))
 			.OnClicked(this, &FEditorStyleSettingsCustomization::OnDeleteThemeClicked)
