Reword the low drive space warning dialog.

diff --git a/Engine/Source/Editor/UnrealEd/Private/UnrealEdEngine.cpp b/Engine/Source/Editor/UnrealEd/Private/UnrealEdEngine.cpp
--- a/Engine/Source/Editor/UnrealEd/Private/UnrealEdEngine.cpp
+++ b/Engine/Source/Editor/UnrealEd/Private/UnrealEdEngine.cpp
@@ -1646,9 +1646,9 @@ bool IsBelowFreeDiskSpaceLimit(const TCHAR* TestDir, FText& OutAppendMessage, co
 		const uint64 HardDriveFreeMB = FreeDiskSpace / (1024 * 1024);
 		if (HardDriveFreeMB < MinMB)
 		{
-			static const FText AppendWarning = NSLOCTEXT("DriveSpaceDialog", "LowHardDriveSpaceFormatMsg", "{0}\n  {1} MB Free \t\t {2}\n \t\t\t\t {3} \n");
+			static const FText AppendWarning = NSLOCTEXT("DriveSpaceDialog", "LowHardDriveSpaceFormatMsg", "{0}\n\n{1}\n\t\t{2}\n\t\tRecommended: {4} MB\n\t\tFree: {3} MB");
 
-			OutAppendMessage = FText::Format(AppendWarning, OutAppendMessage, HardDriveFreeMB, FText::FromString(FPaths::ConvertRelativePathToFull(TestDir)), LocationDescriptor);
+			OutAppendMessage = FText::Format(AppendWarning, OutAppendMessage, LocationDescriptor, FText::FromString(FPaths::ConvertRelativePathToFull(TestDir)), HardDriveFreeMB, MinMB);
 
 			return true;
 		}
@@ -1658,18 +1658,18 @@ bool IsBelowFreeDiskSpaceLimit(const TCHAR* TestDir, FText& OutAppendMessage, co
 
 void UUnrealEdEngine::ValidateFreeDiskSpace() const
 {
-	FText Message = NSLOCTEXT("DriveSpaceDialog", "LowHardDriveSpaceMsgHeader", "The following drive locations have limited free space.\nIt is recommended that you free some space to avoid issues such as crashed and data loss due to files not being able to be written and saved.\n");
+	FText Message = NSLOCTEXT("DriveSpaceDialog", "LowHardDriveSpaceMsgHeader", "The following locations have limited free space. To avoid potential problems, please consider freeing up at least the amounts recommended below.");
 	
 	bool bShowWarning = false;
-	bShowWarning |= IsBelowFreeDiskSpaceLimit(FPlatformProcess::BaseDir(), Message, NSLOCTEXT("DriveSpaceDialog", "BaseDirDescriptor", "The base engine directory."));
-	bShowWarning |= IsBelowFreeDiskSpaceLimit(FPlatformMisc::ProjectDir(), Message, NSLOCTEXT("DriveSpaceDialog", "ProjectDirDescriptor", "The project directory."));
-	bShowWarning |= IsBelowFreeDiskSpaceLimit(FPlatformProcess::UserDir(), Message, NSLOCTEXT("DriveSpaceDialog", "UserDirDescriptor", "User directory where user specific settings are stored."), 1024);
+	bShowWarning |= IsBelowFreeDiskSpaceLimit(FPlatformProcess::BaseDir(), Message, NSLOCTEXT("DriveSpaceDialog", "BaseDirDescriptor", "The base engine directory:"));
+	bShowWarning |= IsBelowFreeDiskSpaceLimit(FPlatformMisc::ProjectDir(), Message, NSLOCTEXT("DriveSpaceDialog", "ProjectDirDescriptor", "The project directory:"));
+	bShowWarning |= IsBelowFreeDiskSpaceLimit(FPlatformProcess::UserDir(), Message, NSLOCTEXT("DriveSpaceDialog", "UserDirDescriptor", "The current user directory, for saving user-specific settings:"), 1024);
 
 	if (bShowWarning)
 	{
 		FEngineAnalytics::LowDriveSpaceDetected();
 
-		const FText Title = NSLOCTEXT("DriveSpaceDialog", "LowHardDriveSpaceMsgTitle", "Low drive space warning");
+		const FText Title = NSLOCTEXT("DriveSpaceDialog", "LowHardDriveSpaceMsgTitle", "Warning: Low Drive Space");
 
 		FMessageDialog::Open(EAppMsgType::Ok, Message, &Title);
 	}
