Fixed WorldPartitionConvertCommandlet re-creating multiple times LandscapeSplineActors which ends-up with potentially thousands of invalid/empty spline actors.

diff --git a/Engine/Source/Editor/UnrealEd/Private/Commandlets/WorldPartitionConvertCommandlet.cpp b/Engine/Source/Editor/UnrealEd/Private/Commandlets/WorldPartitionConvertCommandlet.cpp
--- a/Engine/Source/Editor/UnrealEd/Private/Commandlets/WorldPartitionConvertCommandlet.cpp
+++ b/Engine/Source/Editor/UnrealEd/Private/Commandlets/WorldPartitionConvertCommandlet.cpp
@@ -944,7 +944,9 @@ int32 UWorldPartitionConvertCommandlet::Main(const FString& Params)
 		return true;
 	};
 
-	auto PartitionLandscape = [this, MainWorld](ULandscapeInfo* LandscapeInfo)
+	TSet<AActor*> NewSplineActors;
+
+	auto PartitionLandscape = [this, MainWorld, &NewSplineActors](ULandscapeInfo* LandscapeInfo)
 	{
 		TRACE_CPUPROFILER_EVENT_SCOPE(PartitionLandscape);
 
@@ -963,8 +965,6 @@ int32 UWorldPartitionConvertCommandlet::Main(const FString& Params)
 			LandscapeInfo->RegisterActor(NewLandscape);
 		}
 
-		TSet<AActor*> NewSplineActors;
-						
 		auto MoveControlPointToNewSplineActor = [&NewSplineActors, LandscapeInfo](ULandscapeSplineControlPoint* ControlPoint)
 		{
 			AActor* CurrentOwner = ControlPoint->GetTypedOuter<AActor>();
