Fixed transactions being incorrectly discarded due to the diffable object having no changes

We can no longer use the DeltaChange for this as Ar.IsTransacting() may have added extra data in the transacted object that is important for undo/redo, but not for delta change detection. This can result in an empty delta for a change that still needs to be kept for undo/redo (eg, changes to transient object data).

#rnx

diff --git a/Engine/Source/Editor/UnrealEd/Classes/Editor/Transactor.h b/Engine/Source/Editor/UnrealEd/Classes/Editor/Transactor.h
--- a/Engine/Source/Editor/UnrealEd/Classes/Editor/Transactor.h
+++ b/Engine/Source/Editor/UnrealEd/Classes/Editor/Transactor.h
@@ -103,6 +103,8 @@ protected:
 		bool				bFinalized = false;
 		/** True if record should serialize data as binary blob (more compact). False to use tagged serialization (more robust) */
 		bool				bWantsBinarySerialization = true;
+		/** True if the serialized object state changed between when the transaction started and when the transaction ended (set once finalized) */
+		bool				bHasSerializedObjectChanges = false;
 		/** The serialized object data */
 		FSerializedObject	SerializedObject;
 		/** The serialized object data that will be used when the transaction is flipped */
@@ -111,7 +113,7 @@ protected:
 		TUniquePtr<FTransactionDiffableObject> DiffableObject;
 		/** The diffable object data when it was last snapshot (always null once finalized) */
 		TUniquePtr<FTransactionDiffableObject> DiffableObjectSnapshot;
-		/** Delta change information between the state of the object when the transaction started, and the state of the object when the transaction ended */
+		/** Delta change information between the diffable state of the object when the transaction started, and the diffable state of the object when the transaction ended */
 		FTransactionObjectDeltaChange DeltaChange;
 
 		// Constructors.
diff --git a/Engine/Source/Editor/UnrealEd/Private/EditorTransaction.cpp b/Engine/Source/Editor/UnrealEd/Private/EditorTransaction.cpp
--- a/Engine/Source/Editor/UnrealEd/Private/EditorTransaction.cpp
+++ b/Engine/Source/Editor/UnrealEd/Private/EditorTransaction.cpp
@@ -282,6 +282,14 @@ void FTransaction::FObjectRecord::Finalize( FTransaction* Owner, TSharedPtr<ITra
 				SerializeObject(Writer);
 			}
 
+			// Determine if the actual transacted object has changes, as we use this to detect if the transaction was a no-op and can be discarded
+			// Note: We don't use the DeltaChange for this as Ar.IsTransacting() may have added extra data in the transacted object that is important for undo/redo, but not for 
+			//       delta change detection. This can result in an empty delta for a change that still needs to be kept for undo/redo (eg, changes to transient object data).
+			bHasSerializedObjectChanges 
+				 = SerializedObject.SerializedData != SerializedObjectFlip.SerializedData
+				|| SerializedObject.ReferencedNames != SerializedObjectFlip.ReferencedNames
+				|| SerializedObject.ReferencedObjects != SerializedObjectFlip.ReferencedObjects;
+
 			// Serialize the object so we can diff it
 			FTransactionDiffableObject CurrentDiffableObject;
 			{
@@ -548,7 +556,7 @@ bool FTransaction::FObjectRecord::ContainsPieObject() const
 bool FTransaction::FObjectRecord::HasChanges() const
 {
 	// A record contains change if it has a detected delta or a custom change or an object annotation
-	return DeltaChange.HasChanged() || CustomChange || SerializedObject.ObjectAnnotation || SerializedObjectFlip.ObjectAnnotation;
+	return bHasSerializedObjectChanges || CustomChange || SerializedObject.ObjectAnnotation || SerializedObjectFlip.ObjectAnnotation;
 }
 
 bool FTransaction::FObjectRecord::HasExpired() const
@@ -766,7 +774,7 @@ void FTransaction::Apply()
 
 		const FObjectRecord& ChangedObjectRecord = Records[ChangedObjectIt.Value.RecordIndex];
 		const FTransactionObjectDeltaChange& DeltaChange = ChangedObjectRecord.DeltaChange;
-		if (DeltaChange.HasChanged() || ChangedObjectTransactionAnnotation.IsValid())
+		if (DeltaChange.HasChanged() || ChangedObjectTransactionAnnotation.IsValid() || ChangedObjectRecord.bHasSerializedObjectChanges)
 		{
 			const FObjectRecord::FSerializedObject& InitialSerializedObject = ChangedObjectRecord.SerializedObject;
 			ChangedObject->PostTransacted(FTransactionObjectEvent(Id, OperationId, ETransactionObjectEventType::UndoRedo, DeltaChange, ChangedObjectTransactionAnnotation
@@ -831,7 +839,7 @@ void FTransaction::Finalize()
 
 		const FObjectRecord& ChangedObjectRecord = Records[ChangedObjectIt.Value.RecordIndex];
 		const FTransactionObjectDeltaChange& DeltaChange = ChangedObjectRecord.DeltaChange;
-		if (DeltaChange.HasChanged() || ChangedObjectTransactionAnnotation.IsValid())
+		if (DeltaChange.HasChanged() || ChangedObjectTransactionAnnotation.IsValid() || ChangedObjectRecord.bHasSerializedObjectChanges)
 		{
 			UObject* ChangedObject = ChangedObjectIt.Key;
 			
