AssetRegistry: Add IAssetRegistry::AssetTagsFinalized so that classes that have incomplete tags after load until compilation is finished can notify the assetregistry to update their tags when compilation is finished.
#rnx

diff --git a/Engine/Source/Editor/UnrealEd/Classes/Editor/EditorEngine.h b/Engine/Source/Editor/UnrealEd/Classes/Editor/EditorEngine.h
--- a/Engine/Source/Editor/UnrealEd/Classes/Editor/EditorEngine.h
+++ b/Engine/Source/Editor/UnrealEd/Classes/Editor/EditorEngine.h
@@ -70,6 +70,7 @@ struct FTransactionContext;
 struct FEditorTransactionDeltaContext;
 struct FTypedElementHandle;
 struct FAnalyticsEventAttribute;
+struct FAssetCompileData;
 class UEditorWorldExtensionManager;
 class ITargetDevice;
 class ULevelEditorDragDropHandler;
@@ -3149,6 +3150,9 @@ private:
 	/** Handler for when an asset is created (used to detect world duplication after PostLoad) */
 	void OnAssetCreated( UObject* Asset );
 
+	/** Handler for when an asset finishes compiling (used to notify AssetRegistry to update tags after Load) */
+	void OnAssetPostCompile(const TArray<FAssetCompileData>& CompiledAssets);
+
 	/** Handler for when a world is duplicated in the editor */
 	void InitializeNewlyCreatedInactiveWorld(UWorld* World);
 
diff --git a/Engine/Source/Editor/UnrealEd/Private/EditorEngine.cpp b/Engine/Source/Editor/UnrealEd/Private/EditorEngine.cpp
--- a/Engine/Source/Editor/UnrealEd/Private/EditorEngine.cpp
+++ b/Engine/Source/Editor/UnrealEd/Private/EditorEngine.cpp
@@ -404,6 +404,9 @@ UEditorEngine::UEditorEngine(const FObjectInitializer& ObjectInitializer)
 	ActorGroupingUtilsClassName = UActorGroupingUtils::StaticClass();
 
 	bUATSuccessfullyCompiledOnce = FApp::IsEngineInstalled() || FApp::GetEngineIsPromotedBuild();
+
+	// The AssetRegistry module is needed early in initialization functions so load it here rather than in Init
+	FModuleManager::LoadModuleChecked<FAssetRegistryModule>("AssetRegistry");
 }
 
 
@@ -868,6 +871,8 @@ void UEditorEngine::InitEditor(IEngineLoop* InEngineLoop)
 			GEditor->Exec(nullptr, *InString);
 		}));
 	}
+
+	FAssetCompilingManager::Get().OnAssetPostCompileEvent().AddUObject(this, &UEditorEngine::OnAssetPostCompile);
 }
 
 bool UEditorEngine::HandleOpenAsset(UObject* Asset)
@@ -946,8 +951,7 @@ void UEditorEngine::Init(IEngineLoop* InEngineLoop)
 	FWorldDelegates::LevelAddedToWorld.AddUObject(this, &UEditorEngine::OnLevelAddedToWorld);
 	FWorldDelegates::LevelRemovedFromWorld.AddUObject(this, &UEditorEngine::OnLevelRemovedFromWorld);
 
-	FAssetRegistryModule& AssetRegistryModule = FModuleManager::LoadModuleChecked<FAssetRegistryModule>("AssetRegistry");
-	AssetRegistryModule.Get().OnInMemoryAssetCreated().AddUObject(this, &UEditorEngine::OnAssetCreated);
+	IAssetRegistry::GetChecked().OnInMemoryAssetCreated().AddUObject(this, &UEditorEngine::OnAssetCreated);
 
 	FEditorDelegates::BeginPIE.AddLambda([](bool)
 	{
@@ -1035,7 +1039,7 @@ void UEditorEngine::Init(IEngineLoop* InEngineLoop)
 			TEXT("MainFrame"),
 			TEXT("PropertyEditor"),
 			TEXT("PackagesDialog"),
-			TEXT("AssetRegistry"),
+			// TEXT("AssetRegistry"), // Loaded in constructor
 			TEXT("DetailCustomizations"),
 			TEXT("ComponentVisualizers"),
 			TEXT("Layers"),
@@ -1301,6 +1305,8 @@ void UEditorEngine::FinishDestroy()
 				AssetRegistry->OnInMemoryAssetCreated().RemoveAll(this);
 			}
 		}
+		FAssetCompilingManager::Get().OnAssetPostCompileEvent().RemoveAll(this);
+
 
 		// Shut down transaction tracking system.
 		if( Trans )
@@ -2993,8 +2999,7 @@ void UEditorEngine::GetAssetsToSyncToContentBrowser(TArray<FAssetData>& Assets,
 
 					if (SoftObjects.Num())
 					{
-						FAssetRegistryModule& AssetRegistryModule = FModuleManager::LoadModuleChecked<FAssetRegistryModule>("AssetRegistry");
-						IAssetRegistry& AssetRegistry = AssetRegistryModule.Get();
+						IAssetRegistry& AssetRegistry = IAssetRegistry::GetChecked();
 
 						for (const FSoftObjectPath& SoftObject : SoftObjects)
 						{
@@ -5825,7 +5830,7 @@ AActor* UEditorEngine::ConvertBrushesToStaticMesh(const FString& InStaticMeshPac
 		}
 
 		// Notify the asset registry
-		FAssetRegistryModule::AssetCreated(NewMesh);
+		IAssetRegistry::GetChecked().AssetCreated(NewMesh);
 	}
 
 	ConversionTempModel->EmptyModel(1, 1);
@@ -6326,9 +6331,10 @@ bool UEditorEngine::ShouldThrottleCPUUsage() const
 			if (bShouldThrottle)
 			{
 				static const FName AssetRegistryName(TEXT("AssetRegistry"));
-				FAssetRegistryModule* AssetRegistryModule = FModuleManager::GetModulePtr<FAssetRegistryModule>(AssetRegistryName);
+				IAssetRegistry* AssetRegistry = IAssetRegistry::Get();
 				// Don't throttle during amortized export, greatly increases export time
-				if (IsLightingBuildCurrentlyExporting() || FStaticLightingSystemInterface::IsStaticLightingSystemRunning() || GShaderCompilingManager->IsCompiling() || (AssetRegistryModule && AssetRegistryModule->Get().IsLoadingAssets()))
+				if (IsLightingBuildCurrentlyExporting() || FStaticLightingSystemInterface::IsStaticLightingSystemRunning() ||
+					GShaderCompilingManager->IsCompiling() || (AssetRegistry && AssetRegistry->IsLoadingAssets()))
 				{
 					bShouldThrottle = false;
 				}
@@ -7331,6 +7337,21 @@ void UEditorEngine::OnAssetCreated(UObject* Asset)
 	}
 }
 
+void UEditorEngine::OnAssetPostCompile(const TArray<FAssetCompileData>& CompiledAssets)
+{
+	IAssetRegistry* AssetRegistry = IAssetRegistry::Get();
+	if (AssetRegistry)
+	{
+		for (const FAssetCompileData& CompileData : CompiledAssets)
+		{
+			if (CompileData.Asset.IsValid())
+			{
+				AssetRegistry->AssetTagsFinalized(*CompileData.Asset);
+			}
+		}
+	}
+}
+
 void UEditorEngine::InitializeNewlyCreatedInactiveWorld(UWorld* World)
 {
 	check(World);
diff --git a/Engine/Source/Runtime/AssetRegistry/Private/AssetRegistry.cpp b/Engine/Source/Runtime/AssetRegistry/Private/AssetRegistry.cpp
--- a/Engine/Source/Runtime/AssetRegistry/Private/AssetRegistry.cpp
+++ b/Engine/Source/Runtime/AssetRegistry/Private/AssetRegistry.cpp
@@ -3441,7 +3441,20 @@ void UAssetRegistryImpl::AssetSaved(const UObject& SavedAsset)
 	}
 
 	FWriteScopeLock InterfaceScopeLock(InterfaceLock);
-	GuardedData.OnAssetSaved(SavedAsset);
+	GuardedData.AddLoadedAssetToProcess(SavedAsset);
+#endif
+}
+
+void UAssetRegistryImpl::AssetTagsFinalized(const UObject& FinalizedAsset)
+{
+#if WITH_EDITOR
+	if (!FinalizedAsset.IsAsset())
+	{
+		return;
+	}
+
+	FWriteScopeLock InterfaceScopeLock(InterfaceLock);
+	GuardedData.AddLoadedAssetToProcess(FinalizedAsset);
 #endif
 }
 
@@ -5031,7 +5044,7 @@ void FAssetRegistryImpl::OnDirectoryChanged(Impl::FEventContext& EventContext, T
 void UAssetRegistryImpl::OnAssetLoaded(UObject *AssetLoaded)
 {
 	FWriteScopeLock InterfaceScopeLock(InterfaceLock);
-	GuardedData.OnAssetLoaded(AssetLoaded);
+	GuardedData.AddLoadedAssetToProcess(*AssetLoaded);
 }
 
 void UAssetRegistryImpl::ProcessLoadedAssetsToUpdateCache(UE::AssetRegistry::Impl::FEventContext& EventContext,
@@ -5109,14 +5122,9 @@ void UAssetRegistryImpl::ProcessLoadedAssetsToUpdateCache(UE::AssetRegistry::Imp
 namespace UE::AssetRegistry
 {
 
-void FAssetRegistryImpl::OnAssetLoaded(UObject* AssetLoaded)
-{
-	LoadedAssetsToProcess.Add(AssetLoaded);
-}
-
-void FAssetRegistryImpl::OnAssetSaved(const UObject& AssetSaved)
+void FAssetRegistryImpl::AddLoadedAssetToProcess(const UObject& AssetLoaded)
 {
-	LoadedAssetsToProcess.Add(&AssetSaved);
+	LoadedAssetsToProcess.Add(&AssetLoaded);
 }
 
 void FAssetRegistryImpl::GetProcessLoadedAssetsBatch(TArray<const UObject*>& OutLoadedAssets, uint32 BatchSize)
@@ -5148,8 +5156,14 @@ void FAssetRegistryImpl::GetProcessLoadedAssetsBatch(TArray<const UObject*>& Out
 		UPackage* InMemoryPackage = LoadedAsset->GetOutermost();
 		if (InMemoryPackage->IsDirty())
 		{
-			// Package is dirty, which means it has temporary changes other than just a PostLoad, ignore
-			continue;
+			// Package is dirty, which means it has changes other than just a PostLoad
+			// In editor, ignore the update of the asset; it will be updated when saved
+			// In the cook commandlet, in which editoruser-created changes are impossible, do the update anyway.
+			// Occurrences of IsDirty in the cook commandlet are spurious and a code bug.
+			if (!IsRunningCookCommandlet())
+			{
+				continue;
+			}
 		}
 
 		OutLoadedAssets.Add(LoadedAsset);
diff --git a/Engine/Source/Runtime/AssetRegistry/Private/AssetRegistry.h b/Engine/Source/Runtime/AssetRegistry/Private/AssetRegistry.h
--- a/Engine/Source/Runtime/AssetRegistry/Private/AssetRegistry.h
+++ b/Engine/Source/Runtime/AssetRegistry/Private/AssetRegistry.h
@@ -130,6 +130,7 @@ public:
 	virtual void AssetDeleted(UObject* DeletedAsset) override;
 	virtual void AssetRenamed(const UObject* RenamedAsset, const FString& OldObjectPath) override;
 	virtual void AssetSaved(const UObject& SavedAsset) override;
+	virtual void AssetTagsFinalized(const UObject& FinalizedAsset) override;
 
 	virtual void PackageDeleted(UPackage* DeletedPackage) override;
 
diff --git a/Engine/Source/Runtime/AssetRegistry/Private/AssetRegistryImpl.h b/Engine/Source/Runtime/AssetRegistry/Private/AssetRegistryImpl.h
--- a/Engine/Source/Runtime/AssetRegistry/Private/AssetRegistryImpl.h
+++ b/Engine/Source/Runtime/AssetRegistry/Private/AssetRegistryImpl.h
@@ -144,8 +144,7 @@ public:
 	bool ResolveRedirect(const FString& InPackageName, FString& OutPackageName) const;
 #if WITH_EDITOR
 	void OnDirectoryChanged(Impl::FEventContext& EventContext, TArray<FFileChangeData>& FileChangesProcessed);
-	void OnAssetLoaded(UObject* AssetLoaded);
-	void OnAssetSaved(const UObject& AssetSaved);
+	void AddLoadedAssetToProcess(const UObject& AssetLoaded);
 #endif
 	void OnContentPathMounted(Impl::FEventContext& EventContext, const FString& InAssetPath,
 		const FString& AssetPathWithTrailingSlash, const FString& FileSystemPath);
diff --git a/Engine/Source/Runtime/AssetRegistry/Private/AssetRegistryState.cpp b/Engine/Source/Runtime/AssetRegistry/Private/AssetRegistryState.cpp
--- a/Engine/Source/Runtime/AssetRegistry/Private/AssetRegistryState.cpp
+++ b/Engine/Source/Runtime/AssetRegistry/Private/AssetRegistryState.cpp
@@ -2439,7 +2439,16 @@ void FAssetRegistryState::Dump(const TArray<FString>& Arguments, TArray<FString>
 						{
 							PageBuffer.Appendf(TEXT("		%s"), CategoryName);
 							AddLine();
-							Dependencies.Sort([](const FAssetDependency& A, const FAssetDependency& B) { return A.AssetId.ToString() < B.AssetId.ToString(); });
+							Dependencies.Sort([](const FAssetDependency& A, const FAssetDependency& B)
+								{
+									FString AString = A.AssetId.ToString();
+									FString BString = B.AssetId.ToString();
+									if (AString != BString)
+									{
+										return AString < BString;
+									}
+									return A.Properties < B.Properties;
+								});
 							for (const FAssetDependency& AssetDependency : Dependencies)
 							{
 								PageBuffer.Append(TEXT("			"));
diff --git a/Engine/Source/Runtime/AssetRegistry/Public/AssetRegistry/IAssetRegistry.h b/Engine/Source/Runtime/AssetRegistry/Public/AssetRegistry/IAssetRegistry.h
--- a/Engine/Source/Runtime/AssetRegistry/Public/AssetRegistry/IAssetRegistry.h
+++ b/Engine/Source/Runtime/AssetRegistry/Public/AssetRegistry/IAssetRegistry.h
@@ -505,6 +505,9 @@ public:
 	/** Informs the asset registry that an in-memory package has been deleted, and all associated assets should be removed */
 	virtual void PackageDeleted (UPackage* DeletedPackage) = 0;
 
+	/** Informs the asset registry that an Asset has finalized its tags after loading. Ignored if the Asset's package has been modified. */
+	virtual void AssetTagsFinalized(const UObject& FinalizedAsset) = 0;
+
 	/** Event for when assets are added to the registry */
 	DECLARE_EVENT_OneParam( IAssetRegistry, FAssetAddedEvent, const FAssetData& );
 	virtual FAssetAddedEvent& OnAssetAdded() = 0;
diff --git a/Engine/Source/Runtime/Engine/Classes/Animation/AnimSequence.h b/Engine/Source/Runtime/Engine/Classes/Animation/AnimSequence.h
--- a/Engine/Source/Runtime/Engine/Classes/Animation/AnimSequence.h
+++ b/Engine/Source/Runtime/Engine/Classes/Animation/AnimSequence.h
@@ -592,6 +592,10 @@ protected:
     void ClearCompressedCurveData();
     // Write the compressed data to the supplied FArchive
     void SerializeCompressedData(FArchive& Ar, bool bDDCData);
+#if WITH_EDITOR
+	void SetCompressionComplete();
+#endif
+
 public:
 	bool IsCompressedDataValid() const;
 	bool IsCurveCompressedDataValid() const;
diff --git a/Engine/Source/Runtime/Engine/Private/Animation/AnimSequence.cpp b/Engine/Source/Runtime/Engine/Private/Animation/AnimSequence.cpp
--- a/Engine/Source/Runtime/Engine/Private/Animation/AnimSequence.cpp
+++ b/Engine/Source/Runtime/Engine/Private/Animation/AnimSequence.cpp
@@ -63,6 +63,7 @@ LLM_DEFINE_TAG(SequenceData);
 #if WITH_EDITOR
 #include "Animation/AnimData/AnimDataModel.h"
 #include "Animation/BuiltInAttributeTypes.h"
+#include "AssetRegistry/AssetRegistryModule.h"
 #endif // WITH_EDITOR
 
 #include "Animation/AnimSequenceHelpers.h"
@@ -1869,11 +1870,21 @@ void UAnimSequence::WaitOnExistingCompression(const bool bWantResults)
 	{
 		COOK_STAT(auto Timer = AnimSequenceCookStats::UsageStats.TimeAsyncWait());
 		FAsyncCompressedAnimationsManagement::Get().WaitOnExistingCompression(this, bWantResults);
-		bCompressionInProgress = false;
+		SetCompressionComplete();
 		COOK_STAT(Timer.TrackCyclesOnly()); // Need to get hit/miss and size from WaitOnExistingCompression!
 	}
 }
 
+void UAnimSequence::SetCompressionComplete()
+{
+	if (bCompressionInProgress)
+	{
+		bCompressionInProgress = false;
+		FAssetRegistryModule& AssetRegistryModule = FModuleManager::LoadModuleChecked<FAssetRegistryModule>(TEXT("AssetRegistry"));
+		AssetRegistryModule.Get().AssetTagsFinalized(*this);
+	}
+}
+
 void UAnimSequence::RequestAnimCompression(FRequestAnimCompressionParams Params)
 {
 	check(IsInGameThread());
@@ -1995,7 +2006,7 @@ void UAnimSequence::ApplyCompressedData(const FString& DataCacheKeySuffix, const
 	}
 	else
 	{
-		bCompressionInProgress = false;
+		SetCompressionComplete();
 	}
 }
 #endif
@@ -2003,7 +2014,7 @@ void UAnimSequence::ApplyCompressedData(const FString& DataCacheKeySuffix, const
 void UAnimSequence::ApplyCompressedData(const TArray<uint8>& Data)
 {
 #if WITH_EDITOR
-	bCompressionInProgress = false;
+	SetCompressionComplete();
 	
 	SynchronousAnimatedBoneAttributesCompression();
 #endif
diff --git a/Engine/Source/Runtime/Engine/Private/World.cpp b/Engine/Source/Runtime/Engine/Private/World.cpp
--- a/Engine/Source/Runtime/Engine/Private/World.cpp
+++ b/Engine/Source/Runtime/Engine/Private/World.cpp
@@ -140,6 +140,7 @@
 #include "ObjectTrace.h"
 #include "ReplaySubsystem.h"
 #include "Net/NetPing.h"
+#include "AssetRegistry/AssetRegistryModule.h"
 
 #if UE_WITH_IRIS
 #include "Iris/ReplicationSystem/ReplicationSystem.h"
@@ -1958,6 +1959,9 @@ void UWorld::InitWorld(const InitializationValues IVS)
 	PostInitializeSubsystems();
 
 	BroadcastLevelsChanged();
+
+	FAssetRegistryModule& AssetRegistryModule = FModuleManager::LoadModuleChecked<FAssetRegistryModule>(TEXT("AssetRegistry"));
+	AssetRegistryModule.Get().AssetTagsFinalized(*this);
 }
 
 #if WITH_EDITOR
