EditorDomain: Fix crash when CalculatePackageDigest is called while saving.
#rnx

diff --git a/Engine/Source/Editor/UnrealEd/Private/EditorDomain/EditorDomainUtils.cpp b/Engine/Source/Editor/UnrealEd/Private/EditorDomain/EditorDomainUtils.cpp
--- a/Engine/Source/Editor/UnrealEd/Private/EditorDomain/EditorDomainUtils.cpp
+++ b/Engine/Source/Editor/UnrealEd/Private/EditorDomain/EditorDomainUtils.cpp
@@ -40,6 +40,7 @@
 #include "UObject/Package.h"
 #include "UObject/PackageFileSummary.h"
 #include "UObject/SavePackage.h"
+#include "UObject/UObjectGlobals.h"
 #include "UObject/UObjectHash.h"
 #include "UObject/UE5MainStreamObjectVersion.h"
 
@@ -361,8 +362,8 @@ FPackageDigest CalculatePackageDigest(IAssetRegistry& AssetRegistry, FName Packa
 
 		if (NextClass < ImportedClasses.Num())
 		{
-			// EDITORDOMAIN_TODO: Remove this !IsInGameThread check once FindObject no longer asserts if GIsSavingPackage
-			if (bHasTriedPrecacheClassDigests || !IsInGameThread())
+			// EDITORDOMAIN_TODO: Remove the clauses !IsInGameThread || GIsSavingPackage once FindObject no longer asserts if GIsSavingPackage
+			if (bHasTriedPrecacheClassDigests || !IsInGameThread() || GIsSavingPackage)
 			{
 				return FPackageDigest(FPackageDigest::EStatus::MissingClass, ImportedClasses[NextClass]);
 			}
