Fixed WorldPartitionConvertCommandlet : When converted actor doesn't support external packaging, use Rename instead of SetExternalPackage when moving actor's references into proper destination package.

diff --git a/Engine/Source/Editor/UnrealEd/Private/Commandlets/WorldPartitionConvertCommandlet.cpp b/Engine/Source/Editor/UnrealEd/Private/Commandlets/WorldPartitionConvertCommandlet.cpp
--- a/Engine/Source/Editor/UnrealEd/Private/Commandlets/WorldPartitionConvertCommandlet.cpp
+++ b/Engine/Source/Editor/UnrealEd/Private/Commandlets/WorldPartitionConvertCommandlet.cpp
@@ -66,7 +66,7 @@ class FArchiveGatherPrivateImports : public FArchiveUObject
 	AActor* Root;
 	UPackage* RootPackage;
 	UObject* CurrentObject;
-	TMap<UObject*, UObject*>& PrivateRefsMap;
+	TMap<UObject*, AActor*>& PrivateRefsMap;
 	TSet<FString>& ActorsReferencesToActors;
 
 	void HandleObjectReference(UObject* Obj)
@@ -82,7 +82,7 @@ class FArchiveGatherPrivateImports : public FArchiveUObject
 	}
 
 public:
-	FArchiveGatherPrivateImports(AActor* InRoot, TMap<UObject*, UObject*>& InPrivateRefsMap, TSet<FString>& InActorsReferencesToActors)
+	FArchiveGatherPrivateImports(AActor* InRoot, TMap<UObject*, AActor*>& InPrivateRefsMap, TSet<FString>& InActorsReferencesToActors)
 		: Root(InRoot)
 		, RootPackage(InRoot->GetPackage())
 		, CurrentObject(nullptr)
@@ -113,7 +113,7 @@ public:
 			{
 				if(!Obj->GetTypedOuter<AActor>())
 				{
-					UObject** OriginalRoot = PrivateRefsMap.Find(Obj);
+					AActor** OriginalRoot = PrivateRefsMap.Find(Obj);
 					if(OriginalRoot && (*OriginalRoot != Root))
 					{
 						SET_WARN_COLOR(COLOR_RED);
@@ -143,7 +143,7 @@ public:
 								PrivateRefsMap.Add(Obj, Root);
 
 								SET_WARN_COLOR(COLOR_WHITE);
-								UE_LOG(LogWorldPartitionConvertCommandlet, Warning, TEXT("Encountered reference %s.%s(%s)"), *Root->GetName(), *Obj->GetName(), *Obj->GetClass()->GetName());
+								UE_LOG(LogWorldPartitionConvertCommandlet, Warning, TEXT("Encountered actor %s referencing %s (%s)"), *Root->GetName(), *Obj->GetPathName(), *Obj->GetClass()->GetName());
 								CLEAR_WARN_COLOR();
 							}
 
@@ -1175,7 +1175,7 @@ int32 UWorldPartitionConvertCommandlet::Main(const FString& Params)
 		RemapSoftObjectPaths.Add(OldPackagePath, FSoftObjectPath(MainPackage).ToString());
 	}
 
-	TMap<UObject*, UObject*> PrivateRefsMap;
+	TMap<UObject*, AActor*> PrivateRefsMap;
 	for(ULevel* SubLevel : SubLevelsToConvert)
 	{
 		TRACE_CPUPROFILER_EVENT_SCOPE(ConvertSubLevel);
@@ -1530,14 +1530,47 @@ int32 UWorldPartitionConvertCommandlet::Main(const FString& Params)
 			}
 		}
 
-		for (TMap<UObject*, UObject*>::TConstIterator It(PrivateRefsMap); It; ++It)
+		SET_WARN_COLOR(COLOR_YELLOW);
+		bool bRemapError = false;
+		for (TMap<UObject*, AActor*>::TConstIterator It(PrivateRefsMap); It; ++It)
 		{
-			SET_WARN_COLOR(COLOR_YELLOW);
-			UE_LOG(LogWorldPartitionConvertCommandlet, Warning, TEXT("Renaming %s from %s to %s"), *It->Key->GetName(), *It->Key->GetPackage()->GetName(), *It->Value->GetPackage()->GetName());
-			CLEAR_WARN_COLOR();
-
-			It->Key->SetExternalPackage(It->Value->GetPackage());
+			check(It->Value->IsPackageExternal() == It->Value->SupportsExternalPackaging());
+			if (It->Value->SupportsExternalPackaging())
+			{
+				UE_LOG(LogWorldPartitionConvertCommandlet, Warning, TEXT("Changing object %s package from %s to actor external package %s"), *It->Key->GetName(), *It->Key->GetPackage()->GetName(), *It->Value->GetPackage()->GetName());
+				check(It->Value->GetExternalPackage());
+				It->Key->SetExternalPackage(It->Value->GetExternalPackage());
+			}
+			else
+			{
+				// Remap obj's outer
+				//
+				// Before remapping, validate that object is still in a different package than the actor's package :
+				// Calling Rename on an object can also affect other objects of PrivateRefsMap (UModel::Rename is one example).
+				UPackage* ActorPackage = It->Value->GetPackage();
+				if (!It->Key->IsInPackage(ActorPackage))
+				{
+					UObject* ObjectOuter = It->Key->GetOuter();
+					FString* RemappedOuterPath = RemapSoftObjectPaths.Find(FSoftObjectPath(ObjectOuter).ToString());
+					if (UObject* RemappedOuterObject = RemappedOuterPath ? FSoftObjectPath(*RemappedOuterPath).ResolveObject() : nullptr)
+					{
+						FString OldPathName = It->Key->GetPathName();
+						It->Key->Rename(nullptr, RemappedOuterObject, REN_DontCreateRedirectors);
+						UE_LOG(LogWorldPartitionConvertCommandlet, Warning, TEXT("Renamed object from %s to %s"), *OldPathName, *It->Key->GetPathName());
+					}
+					else
+					{
+						UE_LOG(LogWorldPartitionConvertCommandlet, Error, TEXT("Failed to find a corresponding outer for object %s."), *It->Key->GetPathName());
+						bRemapError = true;
+					}
+				}
+			}
+		}
+		if (bRemapError)
+		{
+			return 1;
 		}
+		CLEAR_WARN_COLOR();
 	
 		// Save packages
 		{
