Fix bug where Step Over will sometimes break inside a function you're trying to step over

diff --git a/Engine/Source/Editor/UnrealEd/Private/Kismet2/KismetDebugUtilities.cpp b/Engine/Source/Editor/UnrealEd/Private/Kismet2/KismetDebugUtilities.cpp
--- a/Engine/Source/Editor/UnrealEd/Private/Kismet2/KismetDebugUtilities.cpp
+++ b/Engine/Source/Editor/UnrealEd/Private/Kismet2/KismetDebugUtilities.cpp
@@ -30,6 +30,7 @@
 #include "K2Node.h"
 #include "K2Node_Tunnel.h"
 #include "K2Node_FunctionEntry.h"
+#include "K2Node_CallFunction.h"
 #include "K2Node_Knot.h"
 #include "K2Node_MacroInstance.h"
 #include "K2Node_Composite.h"
@@ -151,7 +152,7 @@ void FKismetDebugUtilities::RequestStepOver()
 
 		if (const UEdGraphNode* StoppedNode = Data.MostRecentStoppedNode.Get())
 		{
-			if (StoppedNode->IsA<UK2Node_MacroInstance>() || StoppedNode->IsA<UK2Node_Composite>())
+			if (StoppedNode->IsA<UK2Node_MacroInstance>() || StoppedNode->IsA<UK2Node_Composite>() || StoppedNode->IsA<UK2Node_CallFunction>())
 			{
 				for (const UEdGraphPin* Pin : StoppedNode->Pins)
 				{
