Skeletal Mesh - Replace the newly created GetSkeletalMesh function by GetSkeletalMeshAsset so it matches the setter SetSkeletalMeshAsset function (which itself couldn't be named SetSkeletalMesh due to

diff --git a/Engine/Source/Editor/Persona/Private/SPoseEditor.cpp b/Engine/Source/Editor/Persona/Private/SPoseEditor.cpp
--- a/Engine/Source/Editor/Persona/Private/SPoseEditor.cpp
+++ b/Engine/Source/Editor/Persona/Private/SPoseEditor.cpp
@@ -463,7 +463,7 @@ void RestartAnimations(const USkeleton* CurrentSkeleton)
 	for (FThreadSafeObjectIterator Iter(USkeletalMeshComponent::StaticClass()); Iter; ++Iter)
 	{
 		USkeletalMeshComponent* SkeletalMeshComponent = Cast<USkeletalMeshComponent>(*Iter);
-		if (SkeletalMeshComponent->GetSkeletalMesh() && SkeletalMeshComponent->GetSkeletalMesh()->GetSkeleton() == CurrentSkeleton)
+		if (SkeletalMeshComponent->GetSkeletalMeshAsset() && SkeletalMeshComponent->GetSkeletalMeshAsset()->GetSkeleton() == CurrentSkeleton)
 		{
 			SkeletalMeshComponent->InitAnim(true);
 		}
diff --git a/Engine/Source/Editor/Persona/Private/SRetargetSources.cpp b/Engine/Source/Editor/Persona/Private/SRetargetSources.cpp
--- a/Engine/Source/Editor/Persona/Private/SRetargetSources.cpp
+++ b/Engine/Source/Editor/Persona/Private/SRetargetSources.cpp
@@ -284,9 +284,9 @@ bool SRetargetSources::ShouldFilterAsset(const FAssetData& InAssetData)
 void SRetargetSources::ResetRetargetBasePose()
 {
 	UDebugSkelMeshComponent * PreviewMeshComp = PreviewScenePtr.Pin()->GetPreviewMeshComponent();
-	if(PreviewMeshComp && PreviewMeshComp->GetSkeletalMesh())
+	if(PreviewMeshComp && PreviewMeshComp->GetSkeletalMeshAsset())
 	{
-		USkeletalMesh * PreviewMesh = PreviewMeshComp->GetSkeletalMesh();
+		USkeletalMesh * PreviewMesh = PreviewMeshComp->GetSkeletalMeshAsset();
 
 		check(PreviewMesh && &EditableSkeletonPtr.Pin()->GetSkeleton() == PreviewMesh->GetSkeleton());
 
@@ -306,9 +306,9 @@ void SRetargetSources::ResetRetargetBasePose()
 void SRetargetSources::UseCurrentPose()
 {
 	UDebugSkelMeshComponent * PreviewMeshComp = PreviewScenePtr.Pin()->GetPreviewMeshComponent();
-	if (PreviewMeshComp && PreviewMeshComp->GetSkeletalMesh())
+	if (PreviewMeshComp && PreviewMeshComp->GetSkeletalMeshAsset())
 	{
-		USkeletalMesh * PreviewMesh = PreviewMeshComp->GetSkeletalMesh();
+		USkeletalMesh * PreviewMesh = PreviewMeshComp->GetSkeletalMeshAsset();
 
 		check(PreviewMesh && &EditableSkeletonPtr.Pin()->GetSkeleton() == PreviewMesh->GetSkeleton());
 
@@ -386,9 +386,9 @@ void SRetargetSources::ImportPose(const UPoseAsset* PoseAsset, const FName& Pose
 
 			// now I have pose, I have to copy to the retarget base pose
 			UDebugSkelMeshComponent * PreviewMeshComp = PreviewScenePtr.Pin()->GetPreviewMeshComponent();
-			if (PreviewMeshComp && PreviewMeshComp->GetSkeletalMesh())
+			if (PreviewMeshComp && PreviewMeshComp->GetSkeletalMeshAsset())
 			{
-				USkeletalMesh * PreviewMesh = PreviewMeshComp->GetSkeletalMesh();
+				USkeletalMesh * PreviewMesh = PreviewMeshComp->GetSkeletalMeshAsset();
 
 				check(PreviewMesh && &EditableSkeletonPtr.Pin()->GetSkeleton() == PreviewMesh->GetSkeleton());
 
diff --git a/Engine/Source/Editor/PhysicsAssetEditor/Private/PhysicsAssetEditor.cpp b/Engine/Source/Editor/PhysicsAssetEditor/Private/PhysicsAssetEditor.cpp
--- a/Engine/Source/Editor/PhysicsAssetEditor/Private/PhysicsAssetEditor.cpp
+++ b/Engine/Source/Editor/PhysicsAssetEditor/Private/PhysicsAssetEditor.cpp
@@ -466,12 +466,12 @@ void FPhysicsAssetEditor::OnAssetReimport(UObject* Object)
 	RefreshHierachyTree();
 	RefreshPreviewViewport();
 
-	if (SharedData->EditorSkelComp && SharedData->EditorSkelComp->GetSkeletalMesh())
+	if (SharedData->EditorSkelComp && SharedData->EditorSkelComp->GetSkeletalMeshAsset())
 	{
 		IMeshUtilities& MeshUtilities = FModuleManager::Get().LoadModuleChecked<IMeshUtilities>("MeshUtilities");
 		// Update various infos based on the mesh
-		MeshUtilities.CalcBoneVertInfos(SharedData->EditorSkelComp->GetSkeletalMesh(), SharedData->DominantWeightBoneInfos, true);
-		MeshUtilities.CalcBoneVertInfos(SharedData->EditorSkelComp->GetSkeletalMesh(), SharedData->AnyWeightBoneInfos, false);
+		MeshUtilities.CalcBoneVertInfos(SharedData->EditorSkelComp->GetSkeletalMeshAsset(), SharedData->DominantWeightBoneInfos, true);
+		MeshUtilities.CalcBoneVertInfos(SharedData->EditorSkelComp->GetSkeletalMeshAsset(), SharedData->AnyWeightBoneInfos, false);
 	}
 }
 
diff --git a/Engine/Source/Editor/PhysicsAssetEditor/Private/PhysicsAssetEditorSharedData.cpp b/Engine/Source/Editor/PhysicsAssetEditor/Private/PhysicsAssetEditorSharedData.cpp
--- a/Engine/Source/Editor/PhysicsAssetEditor/Private/PhysicsAssetEditorSharedData.cpp
+++ b/Engine/Source/Editor/PhysicsAssetEditor/Private/PhysicsAssetEditorSharedData.cpp
@@ -2786,7 +2786,7 @@ void FPhysicsAssetEditorSharedData::EnableSimulation(bool bEnableSimulation)
 			EditorSkelComp->SetAllBodiesSimulatePhysics(false);
 
 			// make sure we enable the preview animation is any compatible with the skeleton
-			if (PreviewAnimationAsset && EditorSkelComp->GetSkeletalMesh() && PreviewAnimationAsset->GetSkeleton() == EditorSkelComp->GetSkeletalMesh()->GetSkeleton())
+			if (PreviewAnimationAsset && EditorSkelComp->GetSkeletalMeshAsset() && PreviewAnimationAsset->GetSkeleton() == EditorSkelComp->GetSkeletalMeshAsset()->GetSkeleton())
 			{
 				EditorSkelComp->EnablePreview(true, PreviewAnimationAsset);
 				EditorSkelComp->Play(true);
diff --git a/Engine/Source/Editor/PhysicsAssetEditor/Private/PhysicsAssetEditorSkeletonTreeBuilder.cpp b/Engine/Source/Editor/PhysicsAssetEditor/Private/PhysicsAssetEditorSkeletonTreeBuilder.cpp
--- a/Engine/Source/Editor/PhysicsAssetEditor/Private/PhysicsAssetEditorSkeletonTreeBuilder.cpp
+++ b/Engine/Source/Editor/PhysicsAssetEditor/Private/PhysicsAssetEditorSkeletonTreeBuilder.cpp
@@ -116,9 +116,9 @@ void FPhysicsAssetEditorSkeletonTreeBuilder::AddBodies(FSkeletonTreeBuilderOutpu
 	if (PreviewScenePtr.IsValid())
 	{
 		UDebugSkelMeshComponent* PreviewMeshComponent = PreviewScenePtr.Pin()->GetPreviewMeshComponent();
-		if (PreviewMeshComponent->GetSkeletalMesh())
+		if (PreviewMeshComponent->GetSkeletalMeshAsset())
 		{
-			FReferenceSkeleton& RefSkeleton = PreviewMeshComponent->GetSkeletalMesh()->GetRefSkeleton();
+			FReferenceSkeleton& RefSkeleton = PreviewMeshComponent->GetSkeletalMeshAsset()->GetRefSkeleton();
 
 			for (int32 BoneIndex = 0; BoneIndex < RefSkeleton.GetRawBoneNum(); ++BoneIndex)
 			{
diff --git a/Engine/Source/Editor/SequenceRecorder/Private/AnimationRecorder.cpp b/Engine/Source/Editor/SequenceRecorder/Private/AnimationRecorder.cpp
--- a/Engine/Source/Editor/SequenceRecorder/Private/AnimationRecorder.cpp
+++ b/Engine/Source/Editor/SequenceRecorder/Private/AnimationRecorder.cpp
@@ -119,7 +119,7 @@ bool FAnimationRecorder::TriggerRecordAnimation(USkeletalMeshComponent* Componen
 	FFrameRate SampleRate;
 	float MaximumLength;
 
-	if (!Component || !Component->GetSkeletalMesh() || !Component->GetSkeletalMesh()->GetSkeleton())
+	if (!Component || !Component->GetSkeletalMeshAsset() || !Component->GetSkeletalMeshAsset()->GetSkeleton())
 	{
 		return false;
 	}
@@ -136,7 +136,7 @@ bool FAnimationRecorder::TriggerRecordAnimation(USkeletalMeshComponent* Componen
 
 bool FAnimationRecorder::TriggerRecordAnimation(USkeletalMeshComponent* Component, const FString& InAssetPath, const FString& InAssetName)
 {
-	if (!Component || !Component->GetSkeletalMesh() || !Component->GetSkeletalMesh()->GetSkeleton())
+	if (!Component || !Component->GetSkeletalMeshAsset() || !Component->GetSkeletalMeshAsset()->GetSkeleton())
 	{
 		return false;
 	}
@@ -182,7 +182,7 @@ bool FAnimationRecorder::TriggerRecordAnimation(USkeletalMeshComponent* Componen
 	if (NewSeq)
 	{
 		// set skeleton
-		NewSeq->SetSkeleton(Component->GetSkeletalMesh()->GetSkeleton());
+		NewSeq->SetSkeleton(Component->GetSkeletalMeshAsset()->GetSkeleton());
 		// Notify the asset registry
 		FAssetRegistryModule::AssetCreated(NewSeq);
 		StartRecord(Component, NewSeq);
@@ -272,8 +272,8 @@ void FAnimationRecorder::StartRecord(USkeletalMeshComponent* Component, UAnimSeq
 		// verify if this bone exists in skeleton
 		const int32 BoneTreeIndex = AnimSkeleton->GetSkeletonBoneIndexFromMeshBoneIndex(
 			Cast<USkeletalMeshComponent>(Component->LeaderPoseComponent) ?
-			Cast<USkeletalMeshComponent>(Component->LeaderPoseComponent)->GetSkeletalMesh() :
-			Component->GetSkeletalMesh(), BoneIndex);
+			Cast<USkeletalMeshComponent>(Component->LeaderPoseComponent)->GetSkeletalMeshAsset() :
+			Component->GetSkeletalMeshAsset(), BoneIndex);
 		if (BoneTreeIndex != INDEX_NONE)
 		{
 			// add tracks for the bone existing
@@ -283,10 +283,10 @@ void FAnimationRecorder::StartRecord(USkeletalMeshComponent* Component, UAnimSeq
 		}
 	}
 
-	AnimationObject->RetargetSource = Component->GetSkeletalMesh() ? AnimSkeleton->GetRetargetSourceForMesh(Component->GetSkeletalMesh()) : NAME_None;
+	AnimationObject->RetargetSource = Component->GetSkeletalMeshAsset() ? AnimSkeleton->GetRetargetSourceForMesh(Component->GetSkeletalMeshAsset()) : NAME_None;
 	if (AnimationObject->RetargetSource == NAME_None)
 	{
-		AnimationObject->RetargetSourceAsset = Component->GetSkeletalMesh();
+		AnimationObject->RetargetSourceAsset = Component->GetSkeletalMeshAsset();
 	}
 
 	// record the first frame
@@ -657,8 +657,8 @@ void FAnimationRecorder::ProcessRecordedTimes(UAnimSequence* AnimSequence, USkel
 		{
 			const int32 BoneTreeIndex = AnimSkeleton->GetSkeletonBoneIndexFromMeshBoneIndex(
 				Cast<USkeletalMeshComponent>(SkeletalMeshComponent->LeaderPoseComponent) ?
-				Cast<USkeletalMeshComponent>(SkeletalMeshComponent->LeaderPoseComponent)->GetSkeletalMesh() :
-				SkeletalMeshComponent->GetSkeletalMesh(), BoneIndex);
+				Cast<USkeletalMeshComponent>(SkeletalMeshComponent->LeaderPoseComponent)->GetSkeletalMeshAsset() :
+				SkeletalMeshComponent->GetSkeletalMeshAsset(), BoneIndex);
 			if (BoneTreeIndex != INDEX_NONE)
 			{
 				FName BoneTreeName = AnimSkeleton->GetReferenceSkeleton().GetBoneName(BoneTreeIndex);
@@ -681,8 +681,8 @@ void FAnimationRecorder::ProcessRecordedTimes(UAnimSequence* AnimSequence, USkel
 		// verify if this bone exists in skeleton
 		const int32 BoneTreeIndex = AnimSkeleton->GetSkeletonBoneIndexFromMeshBoneIndex(
 			Cast<USkeletalMeshComponent>(SkeletalMeshComponent->LeaderPoseComponent) ?
-			Cast<USkeletalMeshComponent>(SkeletalMeshComponent->LeaderPoseComponent)->GetSkeletalMesh() :
-			SkeletalMeshComponent->GetSkeletalMesh(), BoneIndex);
+			Cast<USkeletalMeshComponent>(SkeletalMeshComponent->LeaderPoseComponent)->GetSkeletalMeshAsset() :
+			SkeletalMeshComponent->GetSkeletalMeshAsset(), BoneIndex);
 		if (BoneTreeIndex != INDEX_NONE)
 		{
 			// add tracks for the bone existing
@@ -837,8 +837,8 @@ bool FAnimationRecorder::Record(USkeletalMeshComponent* Component, FTransform co
 		IAnimationDataController& Controller = AnimationObject->GetController();
 		USkeletalMesh* SkeletalMesh =
 			Cast<USkeletalMeshComponent>(Component->LeaderPoseComponent) ?
-			Cast<USkeletalMeshComponent>(Component->LeaderPoseComponent)->GetSkeletalMesh() :
-			Component->GetSkeletalMesh();
+			Cast<USkeletalMeshComponent>(Component->LeaderPoseComponent)->GetSkeletalMeshAsset() :
+			Component->GetSkeletalMeshAsset();
 
 		const TArray<FBoneAnimationTrack>& BoneAnimationTracks = AnimationObject->GetDataModel()->GetBoneAnimationTracks();
 
@@ -932,7 +932,7 @@ bool FAnimationRecorder::Record(USkeletalMeshComponent* Component, FTransform co
 				{
 					if (FrameToAdd == 0)
 					{
-						const FTransform RefPose = Component->GetSkeletalMesh()->GetRefSkeleton().GetRefBonePose()[BoneIndex];
+						const FTransform RefPose = Component->GetSkeletalMeshAsset()->GetRefSkeleton().GetRefBonePose()[BoneIndex];
 						RawTrack.PosKeys.Add((FVector3f)RefPose.GetTranslation());
 						RawTrack.RotKeys.Add(FQuat4f(RefPose.GetRotation()));
 						RawTrack.ScaleKeys.Add((FVector3f)RefPose.GetScale3D());
diff --git a/Engine/Source/Editor/SequenceRecorder/Private/Sections/MovieScene3DTransformSectionRecorder.cpp b/Engine/Source/Editor/SequenceRecorder/Private/Sections/MovieScene3DTransformSectionRecorder.cpp
--- a/Engine/Source/Editor/SequenceRecorder/Private/Sections/MovieScene3DTransformSectionRecorder.cpp
+++ b/Engine/Source/Editor/SequenceRecorder/Private/Sections/MovieScene3DTransformSectionRecorder.cpp
@@ -134,8 +134,8 @@ void FMovieScene3DTransformSectionRecorder::FinalizeSection(float CurrentTime)
 		{
 			USkeletalMesh* SkeletalMesh =
 				Cast<USkeletalMeshComponent>(SkeletalMeshComponent->LeaderPoseComponent) ?
-				Cast<USkeletalMeshComponent>(SkeletalMeshComponent->LeaderPoseComponent)->GetSkeletalMesh() :
-				SkeletalMeshComponent->GetSkeletalMesh();
+				Cast<USkeletalMeshComponent>(SkeletalMeshComponent->LeaderPoseComponent)->GetSkeletalMeshAsset() :
+				SkeletalMeshComponent->GetSkeletalMeshAsset();
 			if (AnimSequence && SkeletalMesh)
 			{
 				// find the root bone
diff --git a/Engine/Source/Editor/SequenceRecorder/Private/Sections/MovieSceneAnimationSectionRecorder.cpp b/Engine/Source/Editor/SequenceRecorder/Private/Sections/MovieSceneAnimationSectionRecorder.cpp
--- a/Engine/Source/Editor/SequenceRecorder/Private/Sections/MovieSceneAnimationSectionRecorder.cpp
+++ b/Engine/Source/Editor/SequenceRecorder/Private/Sections/MovieSceneAnimationSectionRecorder.cpp
@@ -26,7 +26,7 @@ bool FMovieSceneAnimationSectionRecorderFactory::CanRecordObject(UObject* InObje
 	if (InObjectToRecord->IsA<USkeletalMeshComponent>())
 	{
 		USkeletalMeshComponent* SkeletalMeshComponent = Cast<USkeletalMeshComponent>(InObjectToRecord);
-		if (SkeletalMeshComponent && SkeletalMeshComponent->GetSkeletalMesh())
+		if (SkeletalMeshComponent && SkeletalMeshComponent->GetSkeletalMeshAsset())
 		{
 			return true;
 		}
@@ -64,7 +64,7 @@ void FMovieSceneAnimationSectionRecorder::CreateSection(UObject* InObjectToRecor
 
 	if(SkeletalMeshComponent.IsValid())
 	{
-		SkeletalMesh = SkeletalMeshComponent->GetSkeletalMesh();
+		SkeletalMesh = SkeletalMeshComponent->GetSkeletalMeshAsset();
 		if (SkeletalMesh != nullptr)
 		{
 			ComponentTransform = SkeletalMeshComponent->GetComponentToWorld().GetRelativeTransform(SkeletalMeshComponent->GetOwner()->GetTransform());
@@ -99,7 +99,7 @@ void FMovieSceneAnimationSectionRecorder::CreateSection(UObject* InObjectToRecor
 					FAssetRegistryModule::AssetCreated(AnimSequence.Get());
 
 					// set skeleton
-					AnimSequence->SetSkeleton(SkeletalMeshComponent->GetSkeletalMesh()->GetSkeleton());
+					AnimSequence->SetSkeleton(SkeletalMeshComponent->GetSkeletalMeshAsset()->GetSkeleton());
 				}
 			}
 
diff --git a/Engine/Source/Editor/SkeletonEditor/Private/BoneProxy.cpp b/Engine/Source/Editor/SkeletonEditor/Private/BoneProxy.cpp
--- a/Engine/Source/Editor/SkeletonEditor/Private/BoneProxy.cpp
+++ b/Engine/Source/Editor/SkeletonEditor/Private/BoneProxy.cpp
@@ -62,7 +62,7 @@ void UBoneProxy::Tick(float DeltaTime)
 					Scale = BoneTransform.GetScale3D();
 				}
 
-				FTransform ReferenceTransform = Component->GetSkeletalMesh()->GetRefSkeleton().GetRefBonePose()[BoneIndex];
+				FTransform ReferenceTransform = Component->GetSkeletalMeshAsset()->GetRefSkeleton().GetRefBonePose()[BoneIndex];
 				ReferenceLocation = ReferenceTransform.GetLocation();
 				ReferenceRotation = ReferenceTransform.GetRotation().Rotator();
 				ReferenceScale = ReferenceTransform.GetScale3D();
diff --git a/Engine/Source/Editor/SkeletonEditor/Private/EditableSkeleton.cpp b/Engine/Source/Editor/SkeletonEditor/Private/EditableSkeleton.cpp
--- a/Engine/Source/Editor/SkeletonEditor/Private/EditableSkeleton.cpp
+++ b/Engine/Source/Editor/SkeletonEditor/Private/EditableSkeleton.cpp
@@ -796,7 +796,7 @@ void FEditableSkeleton::HandleRemoveAllAssets(TSharedPtr<IPersonaPreviewScene> I
 
 	DeleteAttachedObjects(Skeleton->PreviewAttachedAssetContainer, InPreviewScene);
 
-	USkeletalMesh* SkeletalMesh = InPreviewScene->GetPreviewMeshComponent()->GetSkeletalMesh();
+	USkeletalMesh* SkeletalMesh = InPreviewScene->GetPreviewMeshComponent()->GetSkeletalMeshAsset();
 	if (SkeletalMesh)
 	{
 		SkeletalMesh->Modify();
@@ -910,7 +910,7 @@ void FEditableSkeleton::HandleAttachAssets(const TArray<UObject*>& InObjects, co
 
 		if (bAttachToMesh)
 		{
-			USkeletalMesh* SkeletalMesh = InPreviewScene->GetPreviewMeshComponent()->GetSkeletalMesh();
+			USkeletalMesh* SkeletalMesh = InPreviewScene->GetPreviewMeshComponent()->GetSkeletalMeshAsset();
 			if (SkeletalMesh != nullptr)
 			{
 				FScopedTransaction Transaction(LOCTEXT("DragDropAttachMeshUndo", "Attach Assets to Mesh"));
@@ -957,7 +957,7 @@ void FEditableSkeleton::HandleDeleteAttachedAssets(const TArray<FPreviewAttached
 	if (InAttachedObjects.Num() > 0)
 	{
 		Skeleton->Modify();
-		USkeletalMesh* SkeletalMesh = InPreviewScene->GetPreviewMeshComponent()->GetSkeletalMesh();
+		USkeletalMesh* SkeletalMesh = InPreviewScene->GetPreviewMeshComponent()->GetSkeletalMeshAsset();
 		if (SkeletalMesh != nullptr)
 		{
 			SkeletalMesh->Modify();
@@ -993,7 +993,7 @@ void FEditableSkeleton::HandleDeleteSockets(const TArray<FSelectedSocketInfo>& I
 		}
 		else
 		{
-			USkeletalMesh* SkeletalMesh = InPreviewScene->GetPreviewMeshComponent()->GetSkeletalMesh();
+			USkeletalMesh* SkeletalMesh = InPreviewScene->GetPreviewMeshComponent()->GetSkeletalMeshAsset();
 			if (SkeletalMesh != nullptr)
 			{
 				UObject* Object = SkeletalMesh->GetPreviewAttachedAssetContainer().GetAttachedObjectByAttachName(SocketName);
diff --git a/Engine/Source/Editor/SkeletonEditor/Private/SSkeletonTree.cpp b/Engine/Source/Editor/SkeletonEditor/Private/SSkeletonTree.cpp
--- a/Engine/Source/Editor/SkeletonEditor/Private/SSkeletonTree.cpp
+++ b/Engine/Source/Editor/SkeletonEditor/Private/SSkeletonTree.cpp
@@ -1078,7 +1078,7 @@ void SSkeletonTree::RemoveFromLOD(int32 LODIndex, bool bIncludeSelected, bool bI
 	}
 
 	UDebugSkelMeshComponent* PreviewMeshComponent = GetPreviewScene()->GetPreviewMeshComponent();
-	if (!PreviewMeshComponent->GetSkeletalMesh())
+	if (!PreviewMeshComponent->GetSkeletalMeshAsset())
 	{
 		return;
 	}
@@ -1095,7 +1095,7 @@ void SSkeletonTree::RemoveFromLOD(int32 LODIndex, bool bIncludeSelected, bool bI
 
 		//Scoped post edit change
 		{
-			FScopedSkeletalMeshPostEditChange ScopedPostEditChange(PreviewMeshComponent->GetSkeletalMesh());
+			FScopedSkeletalMeshPostEditChange ScopedPostEditChange(PreviewMeshComponent->GetSkeletalMeshAsset());
 
 			for (const TSharedPtr<FSkeletonTreeBoneItem>& Item : TreeSelection.GetSelectedItems<FSkeletonTreeBoneItem>())
 			{
@@ -1105,7 +1105,7 @@ void SSkeletonTree::RemoveFromLOD(int32 LODIndex, bool bIncludeSelected, bool bI
 				{
 					if (bIncludeSelected)
 					{
-						PreviewMeshComponent->GetSkeletalMesh()->AddBoneToReductionSetting(LODIndex, BoneName);
+						PreviewMeshComponent->GetSkeletalMeshAsset()->AddBoneToReductionSetting(LODIndex, BoneName);
 						BonesToRemove.AddUnique(BoneName);
 					}
 					else
@@ -1115,7 +1115,7 @@ void SSkeletonTree::RemoveFromLOD(int32 LODIndex, bool bIncludeSelected, bool bI
 							if (RefSkeleton.GetParentIndex(ChildIndex) == BoneIndex)
 							{
 								FName ChildBoneName = RefSkeleton.GetBoneName(ChildIndex);
-								PreviewMeshComponent->GetSkeletalMesh()->AddBoneToReductionSetting(LODIndex, ChildBoneName);
+								PreviewMeshComponent->GetSkeletalMeshAsset()->AddBoneToReductionSetting(LODIndex, ChildBoneName);
 								BonesToRemove.AddUnique(ChildBoneName);
 							}
 						}
@@ -1123,20 +1123,20 @@ void SSkeletonTree::RemoveFromLOD(int32 LODIndex, bool bIncludeSelected, bool bI
 				}
 			}
 
-			int32 TotalLOD = PreviewMeshComponent->GetSkeletalMesh()->GetLODNum();
+			int32 TotalLOD = PreviewMeshComponent->GetSkeletalMeshAsset()->GetLODNum();
 			IMeshUtilities& MeshUtilities = FModuleManager::Get().LoadModuleChecked<IMeshUtilities>("MeshUtilities");
 
 			if (bIncludeBelowLODs)
 			{
 				for (int32 Index = LODIndex + 1; Index < TotalLOD; ++Index)
 				{
-					MeshUtilities.RemoveBonesFromMesh(PreviewMeshComponent->GetSkeletalMesh(), Index, &BonesToRemove);
-					PreviewMeshComponent->GetSkeletalMesh()->AddBoneToReductionSetting(Index, BonesToRemove);
+					MeshUtilities.RemoveBonesFromMesh(PreviewMeshComponent->GetSkeletalMeshAsset(), Index, &BonesToRemove);
+					PreviewMeshComponent->GetSkeletalMeshAsset()->AddBoneToReductionSetting(Index, BonesToRemove);
 				}
 			}
 
 			// remove from current LOD
-			MeshUtilities.RemoveBonesFromMesh(PreviewMeshComponent->GetSkeletalMesh(), LODIndex, &BonesToRemove);
+			MeshUtilities.RemoveBonesFromMesh(PreviewMeshComponent->GetSkeletalMeshAsset(), LODIndex, &BonesToRemove);
 		}
 		// update UI to reflect the change
 		OnLODSwitched();
@@ -1257,7 +1257,7 @@ void SSkeletonTree::OnPasteSockets(bool bPasteToSelectedBone)
 	if ( TreeSelection.IsSingleOfTypeSelected<FSkeletonTreeBoneItem>() )
 	{
 		FName DestBoneName = bPasteToSelectedBone ? TreeSelection.GetSingleSelectedItem()->GetRowItemName() : NAME_None;
-		USkeletalMesh* SkeletalMesh = GetPreviewScene().IsValid() ? ToRawPtr(GetPreviewScene()->GetPreviewMeshComponent()->GetSkeletalMesh()) : nullptr;
+		USkeletalMesh* SkeletalMesh = GetPreviewScene().IsValid() ? ToRawPtr(GetPreviewScene()->GetPreviewMeshComponent()->GetSkeletalMeshAsset()) : nullptr;
 		GetEditableSkeletonInternal()->HandlePasteSockets(DestBoneName, SkeletalMesh);
 
 		CreateFromSkeleton();
@@ -1318,7 +1318,7 @@ void SSkeletonTree::OnCustomizeSocket()
 	if(TreeSelection.IsSingleOfTypeSelected<FSkeletonTreeSocketItem>())
 	{
 		USkeletalMeshSocket* SocketToCustomize = StaticCastSharedPtr<FSkeletonTreeSocketItem>(TreeSelection.GetSingleSelectedItem())->GetSocket();
-		USkeletalMesh* SkeletalMesh = GetPreviewScene().IsValid() ? ToRawPtr(GetPreviewScene()->GetPreviewMeshComponent()->GetSkeletalMesh()) : nullptr;
+		USkeletalMesh* SkeletalMesh = GetPreviewScene().IsValid() ? ToRawPtr(GetPreviewScene()->GetPreviewMeshComponent()->GetSkeletalMeshAsset()) : nullptr;
 		GetEditableSkeletonInternal()->HandleCustomizeSocket(SocketToCustomize, SkeletalMesh);
 		CreateFromSkeleton();
 	}
@@ -1388,7 +1388,7 @@ void  SSkeletonTree::OnRemoveAllAssets()
 
 bool SSkeletonTree::CanRemoveAllAssets() const
 {
-	USkeletalMesh* SkeletalMesh = GetPreviewScene().IsValid() ? ToRawPtr(GetPreviewScene()->GetPreviewMeshComponent()->GetSkeletalMesh()) : nullptr;
+	USkeletalMesh* SkeletalMesh = GetPreviewScene().IsValid() ? ToRawPtr(GetPreviewScene()->GetPreviewMeshComponent()->GetSkeletalMeshAsset()) : nullptr;
 
 	const bool bHasPreviewAttachedObjects = GetEditableSkeletonInternal()->GetSkeleton().PreviewAttachedAssetContainer.Num() > 0;
 	const bool bHasMeshPreviewAttachedObjects = ( SkeletalMesh && SkeletalMesh->GetPreviewAttachedAssetContainer().Num() );
@@ -2308,7 +2308,7 @@ void SSkeletonTree::SelectItemsBy(TFunctionRef<bool(const TSharedRef<ISkeletonTr
 
 void SSkeletonTree::DuplicateAndSelectSocket(const FSelectedSocketInfo& SocketInfoToDuplicate, const FName& NewParentBoneName /*= FName()*/)
 {
-	USkeletalMesh* SkeletalMesh = GetPreviewScene().IsValid() ? ToRawPtr(GetPreviewScene()->GetPreviewMeshComponent()->GetSkeletalMesh()) : nullptr;
+	USkeletalMesh* SkeletalMesh = GetPreviewScene().IsValid() ? ToRawPtr(GetPreviewScene()->GetPreviewMeshComponent()->GetSkeletalMeshAsset()) : nullptr;
 	USkeletalMeshSocket* NewSocket = GetEditableSkeleton()->DuplicateSocket(SocketInfoToDuplicate, NewParentBoneName, SkeletalMesh);
 
 	if (GetPreviewScene().IsValid())
diff --git a/Engine/Source/Editor/SkeletonEditor/Private/SkeletonEditor.cpp b/Engine/Source/Editor/SkeletonEditor/Private/SkeletonEditor.cpp
--- a/Engine/Source/Editor/SkeletonEditor/Private/SkeletonEditor.cpp
+++ b/Engine/Source/Editor/SkeletonEditor/Private/SkeletonEditor.cpp
@@ -348,7 +348,7 @@ TStatId FSkeletonEditor::GetStatId() const
 bool FSkeletonEditor::CanRemoveBones() const
 {
 	UDebugSkelMeshComponent* PreviewMeshComponent = PersonaToolkit->GetPreviewMeshComponent();
-	return PreviewMeshComponent && PreviewMeshComponent->GetSkeletalMesh();
+	return PreviewMeshComponent && PreviewMeshComponent->GetSkeletalMeshAsset();
 }
 
 void FSkeletonEditor::RemoveUnusedBones()
@@ -364,9 +364,9 @@ void FSkeletonEditor::TestSkeletonCurveNamesForUse() const
 
 void FSkeletonEditor::UpdateSkeletonRefPose()
 {
-	if (PersonaToolkit->GetPreviewMeshComponent()->GetSkeletalMesh())
+	if (PersonaToolkit->GetPreviewMeshComponent()->GetSkeletalMeshAsset())
 	{
-		GetSkeletonTree()->GetEditableSkeleton()->UpdateSkeletonReferencePose(PersonaToolkit->GetPreviewMeshComponent()->GetSkeletalMesh());
+		GetSkeletonTree()->GetEditableSkeleton()->UpdateSkeletonReferencePose(PersonaToolkit->GetPreviewMeshComponent()->GetSkeletalMeshAsset());
 	}
 }
 
diff --git a/Engine/Source/Editor/SkeletonEditor/Private/SkeletonTreeBoneItem.cpp b/Engine/Source/Editor/SkeletonEditor/Private/SkeletonTreeBoneItem.cpp
--- a/Engine/Source/Editor/SkeletonEditor/Private/SkeletonTreeBoneItem.cpp
+++ b/Engine/Source/Editor/SkeletonEditor/Private/SkeletonTreeBoneItem.cpp
@@ -427,7 +427,7 @@ FReply FSkeletonTreeBoneItem::HandleDrop(const FDragDropEvent& DragDropEvent)
 		else if (BoneName != SocketInfo.Socket->BoneName)
 		{
 			// The socket can be dropped here if we're a bone and NOT the socket's existing parent
-			USkeletalMesh* SkeletalMesh = GetSkeletonTree()->GetPreviewScene().IsValid() ? ToRawPtr(GetSkeletonTree()->GetPreviewScene()->GetPreviewMeshComponent()->GetSkeletalMesh()) : nullptr;
+			USkeletalMesh* SkeletalMesh = GetSkeletonTree()->GetPreviewScene().IsValid() ? ToRawPtr(GetSkeletonTree()->GetPreviewScene()->GetPreviewMeshComponent()->GetSkeletalMeshAsset()) : nullptr;
 			GetEditableSkeleton()->SetSocketParent(SocketInfo.Socket->SocketName, BoneName, SkeletalMesh);
 
 			return FReply::Handled();
@@ -454,15 +454,15 @@ bool FSkeletonTreeBoneItem::IsBoneWeighted(int32 MeshBoneIndex, UDebugSkelMeshCo
 {
 	// MeshBoneIndex must be an index into the mesh's skeleton, *not* the source skeleton!!!
 
-	if (!PreviewComponent || !PreviewComponent->GetSkeletalMesh() || !PreviewComponent->GetSkeletalMesh()->GetResourceForRendering() || !PreviewComponent->GetSkeletalMesh()->GetResourceForRendering()->LODRenderData.Num())
+	if (!PreviewComponent || !PreviewComponent->GetSkeletalMeshAsset() || !PreviewComponent->GetSkeletalMeshAsset()->GetResourceForRendering() || !PreviewComponent->GetSkeletalMeshAsset()->GetResourceForRendering()->LODRenderData.Num())
 	{
 		// If there's no mesh, then this bone can't possibly be weighted!
 		return false;
 	}
 
 	//Get current LOD
-	const int32 LODIndex = FMath::Clamp(PreviewComponent->GetPredictedLODLevel(), 0, PreviewComponent->GetSkeletalMesh()->GetResourceForRendering()->LODRenderData.Num() - 1);
-	FSkeletalMeshLODRenderData& LODData = PreviewComponent->GetSkeletalMesh()->GetResourceForRendering()->LODRenderData[LODIndex];
+	const int32 LODIndex = FMath::Clamp(PreviewComponent->GetPredictedLODLevel(), 0, PreviewComponent->GetSkeletalMeshAsset()->GetResourceForRendering()->LODRenderData.Num() - 1);
+	FSkeletalMeshLODRenderData& LODData = PreviewComponent->GetSkeletalMeshAsset()->GetResourceForRendering()->LODRenderData[LODIndex];
 
 	//Check whether the bone is vertex weighted
 	int32 Index = LODData.ActiveBoneIndices.Find(MeshBoneIndex);
@@ -474,15 +474,15 @@ bool FSkeletonTreeBoneItem::IsBoneRequired(int32 MeshBoneIndex, UDebugSkelMeshCo
 {
 	// MeshBoneIndex must be an index into the mesh's skeleton, *not* the source skeleton!!!
 
-	if (!PreviewComponent || !PreviewComponent->GetSkeletalMesh() || !PreviewComponent->GetSkeletalMesh()->GetResourceForRendering() || !PreviewComponent->GetSkeletalMesh()->GetResourceForRendering()->LODRenderData.Num())
+	if (!PreviewComponent || !PreviewComponent->GetSkeletalMeshAsset() || !PreviewComponent->GetSkeletalMeshAsset()->GetResourceForRendering() || !PreviewComponent->GetSkeletalMeshAsset()->GetResourceForRendering()->LODRenderData.Num())
 	{
 		// If there's no mesh, then this bone can't possibly be weighted!
 		return false;
 	}
 
 	//Get current LOD
-	const int32 LODIndex = FMath::Clamp(PreviewComponent->GetPredictedLODLevel(), 0, PreviewComponent->GetSkeletalMesh()->GetResourceForRendering()->LODRenderData.Num() - 1);
-	FSkeletalMeshLODRenderData& LODData = PreviewComponent->GetSkeletalMesh()->GetResourceForRendering()->LODRenderData[LODIndex];
+	const int32 LODIndex = FMath::Clamp(PreviewComponent->GetPredictedLODLevel(), 0, PreviewComponent->GetSkeletalMeshAsset()->GetResourceForRendering()->LODRenderData.Num() - 1);
+	FSkeletalMeshLODRenderData& LODData = PreviewComponent->GetSkeletalMeshAsset()->GetResourceForRendering()->LODRenderData[LODIndex];
 
 	//Check whether the bone is vertex weighted
 	int32 Index = LODData.RequiredBones.Find(MeshBoneIndex);
diff --git a/Engine/Source/Editor/SkeletonEditor/Private/SkeletonTreeBuilder.cpp b/Engine/Source/Editor/SkeletonEditor/Private/SkeletonTreeBuilder.cpp
--- a/Engine/Source/Editor/SkeletonEditor/Private/SkeletonTreeBuilder.cpp
+++ b/Engine/Source/Editor/SkeletonEditor/Private/SkeletonTreeBuilder.cpp
@@ -304,7 +304,7 @@ void FSkeletonTreeBuilder::AddSockets(FSkeletonTreeBuilderOutput& Output)
 	if (PreviewScenePtr.IsValid())
 	{
 		UDebugSkelMeshComponent* PreviewMeshComponent = PreviewScenePtr.Pin()->GetPreviewMeshComponent();
-		if (USkeletalMesh* const SkeletalMesh = PreviewMeshComponent->GetSkeletalMesh())
+		if (USkeletalMesh* const SkeletalMesh = PreviewMeshComponent->GetSkeletalMeshAsset())
 		{
 			AddSocketsFromData(SkeletalMesh->GetMeshOnlySocketList(), ESocketParentType::Mesh, Output);
 		}
@@ -319,7 +319,7 @@ void FSkeletonTreeBuilder::AddAttachedAssets(FSkeletonTreeBuilderOutput& Output)
 	if (PreviewScenePtr.IsValid())
 	{
 		UDebugSkelMeshComponent* PreviewMeshComponent = PreviewScenePtr.Pin()->GetPreviewMeshComponent();
-		if (USkeletalMesh* const SkeletalMesh = PreviewMeshComponent->GetSkeletalMesh())
+		if (USkeletalMesh* const SkeletalMesh = PreviewMeshComponent->GetSkeletalMeshAsset())
 		{
 			AddAttachedAssetContainer(SkeletalMesh->GetPreviewAttachedAssetContainer(), Output);
 		}
@@ -348,7 +348,7 @@ void FSkeletonTreeBuilder::AddSocketsFromData(const TArray< USkeletalMeshSocket*
 			if (PreviewScenePtr.IsValid())
 			{
 				UDebugSkelMeshComponent* PreviewMeshComponent = PreviewScenePtr.Pin()->GetPreviewMeshComponent();
-				if (USkeletalMesh* const SkeletalMesh = PreviewMeshComponent->GetSkeletalMesh())
+				if (USkeletalMesh* const SkeletalMesh = PreviewMeshComponent->GetSkeletalMeshAsset())
 				{
 					bIsCustomized = EditableSkeletonPtr.Pin()->DoesSocketAlreadyExist(nullptr, FText::FromName(Socket->SocketName), ESocketParentType::Mesh, SkeletalMesh);
 				}
diff --git a/Engine/Source/Editor/SkeletonEditor/Private/SkeletonTreeSocketItem.cpp b/Engine/Source/Editor/SkeletonEditor/Private/SkeletonTreeSocketItem.cpp
--- a/Engine/Source/Editor/SkeletonEditor/Private/SkeletonTreeSocketItem.cpp
+++ b/Engine/Source/Editor/SkeletonEditor/Private/SkeletonTreeSocketItem.cpp
@@ -147,7 +147,7 @@ bool FSkeletonTreeSocketItem::CanCustomizeSocket() const
 	if (GetSkeletonTree()->GetPreviewScene().IsValid())
 	{
 		UDebugSkelMeshComponent* PreviewComponent = GetSkeletonTree()->GetPreviewScene()->GetPreviewMeshComponent();
-		return PreviewComponent && PreviewComponent->GetSkeletalMesh() && !IsSocketCustomized();
+		return PreviewComponent && PreviewComponent->GetSkeletalMeshAsset() && !IsSocketCustomized();
 	}
 	return false;
 }
@@ -172,7 +172,7 @@ bool FSkeletonTreeSocketItem::OnVerifySocketNameChanged( const FText& InText, FT
 	}
 	else
 	{
-		USkeletalMesh* SkeletalMesh = GetSkeletonTree()->GetPreviewScene().IsValid() ? ToRawPtr(GetSkeletonTree()->GetPreviewScene()->GetPreviewMeshComponent()->GetSkeletalMesh()) : nullptr;
+		USkeletalMesh* SkeletalMesh = GetSkeletonTree()->GetPreviewScene().IsValid() ? ToRawPtr(GetSkeletonTree()->GetPreviewScene()->GetPreviewMeshComponent()->GetSkeletalMeshAsset()) : nullptr;
 		bVerifyName = !GetEditableSkeleton()->DoesSocketAlreadyExist( Socket, NewText, ParentType, SkeletalMesh );
 
 		// Needs to be checked on verify.
@@ -192,7 +192,7 @@ void FSkeletonTreeSocketItem::OnCommitSocketName( const FText& InText, ETextComm
 	FText NewText = FText::TrimPrecedingAndTrailing(InText);
 
 	// Notify skeleton tree of socket rename
-	USkeletalMesh* SkeletalMesh = GetSkeletonTree()->GetPreviewScene().IsValid() ? ToRawPtr(GetSkeletonTree()->GetPreviewScene()->GetPreviewMeshComponent()->GetSkeletalMesh()) : nullptr;
+	USkeletalMesh* SkeletalMesh = GetSkeletonTree()->GetPreviewScene().IsValid() ? ToRawPtr(GetSkeletonTree()->GetPreviewScene()->GetPreviewMeshComponent()->GetSkeletalMeshAsset()) : nullptr;
 	GetEditableSkeleton()->RenameSocket(Socket->SocketName, FName(*NewText.ToString()), SkeletalMesh);
 }
 
diff --git a/Engine/Source/Editor/StatsViewer/Private/StatsPages/PrimitiveStatsPage.cpp b/Engine/Source/Editor/StatsViewer/Private/StatsPages/PrimitiveStatsPage.cpp
--- a/Engine/Source/Editor/StatsViewer/Private/StatsPages/PrimitiveStatsPage.cpp
+++ b/Engine/Source/Editor/StatsViewer/Private/StatsPages/PrimitiveStatsPage.cpp
@@ -156,7 +156,7 @@ struct PrimitiveStatsGenerator
 		// The skeletal mesh of a skeletal mesh component is its resource.
 		else if( SkeletalMeshComponent )
 		{
-			USkeletalMesh* Mesh = SkeletalMeshComponent->GetSkeletalMesh();
+			USkeletalMesh* Mesh = SkeletalMeshComponent->GetSkeletalMeshAsset();
 			Resource = Mesh;
 			// Calculate vertex color usage for skeletal meshes
 			if( Mesh )
@@ -324,7 +324,7 @@ struct PrimitiveStatsGenerator
 				// ... in the case of skeletal mesh component.
 				else if( SkeletalMeshComponent )
 				{
-					USkeletalMesh* SkeletalMesh = SkeletalMeshComponent->GetSkeletalMesh();
+					USkeletalMesh* SkeletalMesh = SkeletalMeshComponent->GetSkeletalMeshAsset();
 					if( SkeletalMesh )
 					{
 						FSkeletalMeshRenderData* SkelMeshRenderData = SkeletalMesh->GetResourceForRendering();
diff --git a/Engine/Source/Editor/UnrealEd/Classes/Animation/DebugSkelMeshComponent.h b/Engine/Source/Editor/UnrealEd/Classes/Animation/DebugSkelMeshComponent.h
--- a/Engine/Source/Editor/UnrealEd/Classes/Animation/DebugSkelMeshComponent.h
+++ b/Engine/Source/Editor/UnrealEd/Classes/Animation/DebugSkelMeshComponent.h
@@ -597,9 +597,9 @@ public:
 	 */
 	virtual const FReferenceSkeleton& GetReferenceSkeleton() const
 	{
-		if (GetSkeletalMesh())
+		if (GetSkeletalMeshAsset())
 		{
-			return GetSkeletalMesh()->GetRefSkeleton();
+			return GetSkeletalMeshAsset()->GetRefSkeleton();
 		}
 
 		static FReferenceSkeleton EmptySkeleton;
diff --git a/Engine/Source/Editor/UnrealEd/Private/Animation/DebugSkelMeshComponent.cpp b/Engine/Source/Editor/UnrealEd/Private/Animation/DebugSkelMeshComponent.cpp
--- a/Engine/Source/Editor/UnrealEd/Private/Animation/DebugSkelMeshComponent.cpp
+++ b/Engine/Source/Editor/UnrealEd/Private/Animation/DebugSkelMeshComponent.cpp
@@ -97,9 +97,9 @@ FBoxSphereBounds UDebugSkelMeshComponent::CalcBounds(const FTransform& LocalToWo
 			}			
 		}
 
-		if (GetSkeletalMesh())
+		if (GetSkeletalMeshAsset())
 		{
-			Result = Result + GetSkeletalMesh()->GetBounds();
+			Result = Result + GetSkeletalMeshAsset()->GetBounds();
 		}
 	}
 
@@ -455,7 +455,7 @@ FPrimitiveSceneProxy* UDebugSkelMeshComponent::CreateSceneProxy()
 {
 	FDebugSkelMeshSceneProxy* Result = NULL;
 	ERHIFeatureLevel::Type SceneFeatureLevel = GetWorld()->FeatureLevel;
-	FSkeletalMeshRenderData* SkelMeshRenderData = GetSkeletalMesh() ? GetSkeletalMesh()->GetResourceForRendering() : NULL;
+	FSkeletalMeshRenderData* SkelMeshRenderData = GetSkeletalMeshAsset() ? GetSkeletalMeshAsset()->GetResourceForRendering() : NULL;
 
 	// only create a scene proxy for rendering if
 	// properly initialized
@@ -519,9 +519,9 @@ void UDebugSkelMeshComponent::InitAnim(bool bForceReinit)
 {
 	// If we already have PreviewInstance and its asset's Skeleton isn't compatible with the mesh's Skeleton
 	// then we need to clear it up to avoid an issue
-	if ( PreviewInstance && PreviewInstance->GetCurrentAsset() && GetSkeletalMesh())
+	if ( PreviewInstance && PreviewInstance->GetCurrentAsset() && GetSkeletalMeshAsset())
 	{
-		if (!GetSkeletalMesh()->GetSkeleton()->IsCompatible(PreviewInstance->GetCurrentAsset()->GetSkeleton()))
+		if (!GetSkeletalMeshAsset()->GetSkeleton()->IsCompatible(PreviewInstance->GetCurrentAsset()->GetSkeleton()))
 		{
 			// if it doesn't match, just clear it
 			PreviewInstance->SetAnimationAsset(NULL);
@@ -748,7 +748,7 @@ void UDebugSkelMeshComponent::GenSpaceBases(TArray<FTransform>& OutSpaceBases)
 	FBlendedHeapCurve TempCurve;
 	UE::Anim::FMeshAttributeContainer TempAtttributes;
 	DoInstancePreEvaluation();
-	PerformAnimationEvaluation(GetSkeletalMesh(), AnimScriptInstance, OutSpaceBases, TempBoneSpaceTransforms, TempRootBoneTranslation, TempCurve, TempAtttributes);
+	PerformAnimationEvaluation(GetSkeletalMeshAsset(), AnimScriptInstance, OutSpaceBases, TempBoneSpaceTransforms, TempRootBoneTranslation, TempCurve, TempAtttributes);
 	DoInstancePostEvaluation();
 }
 
@@ -954,7 +954,7 @@ void UDebugSkelMeshComponent::ToggleClothSectionsVisibility(bool bShowOnlyClothS
 
 void UDebugSkelMeshComponent::RestoreClothSectionsVisibility()
 {
-	if (!GetSkeletalMesh())
+	if (!GetSkeletalMeshAsset())
 	{
 		return;
 	}
@@ -1007,11 +1007,11 @@ void UDebugSkelMeshComponent::RebuildClothingSectionsFixedVerts(bool bInvalidate
 	// TODO: There is no need to rebuild all section/LODs at once.
 	//       It should only do the section associated to the current cloth asset being
 	//        painted instead, and only when the MaxDistance mask changes.
-	FScopedSkeletalMeshPostEditChange ScopedSkeletalMeshPostEditChange(GetSkeletalMesh());
+	FScopedSkeletalMeshPostEditChange ScopedSkeletalMeshPostEditChange(GetSkeletalMeshAsset());
 
-	GetSkeletalMesh()->PreEditChange(nullptr);
+	GetSkeletalMeshAsset()->PreEditChange(nullptr);
 
-	TIndirectArray<FSkeletalMeshLODModel>& LODModels = GetSkeletalMesh()->GetImportedModel()->LODModels;
+	TIndirectArray<FSkeletalMeshLODModel>& LODModels = GetSkeletalMeshAsset()->GetImportedModel()->LODModels;
 
 	for (int32 LODIndex = 0; LODIndex < LODModels.Num(); ++LODIndex)
 	{
@@ -1025,7 +1025,7 @@ void UDebugSkelMeshComponent::RebuildClothingSectionsFixedVerts(bool bInvalidate
 
 	if (bInvalidateDerivedDataCache)
 	{
-		GetSkeletalMesh()->InvalidateDeriveDataCacheGUID();  // Dirty the DDC key unless previewing
+		GetSkeletalMeshAsset()->InvalidateDeriveDataCacheGUID();  // Dirty the DDC key unless previewing
 	}
 
 	ReregisterComponent();
@@ -1033,7 +1033,7 @@ void UDebugSkelMeshComponent::RebuildClothingSectionsFixedVerts(bool bInvalidate
 
 void UDebugSkelMeshComponent::RebuildClothingSectionFixedVerts(int32 LODIndex, int32 SectionIndex)
 {
-	FSkeletalMeshModel* const SkeletalMeshModel = GetSkeletalMesh()->GetImportedModel();
+	FSkeletalMeshModel* const SkeletalMeshModel = GetSkeletalMeshAsset()->GetImportedModel();
 	if (!ensure(SkeletalMeshModel && LODIndex < SkeletalMeshModel->LODModels.Num()))
 	{
 		return;
@@ -1051,7 +1051,7 @@ void UDebugSkelMeshComponent::RebuildClothingSectionFixedVerts(int32 LODIndex, i
 		return;
 	}
 
-	const UClothingAssetCommon* const ClothingAsset = Cast<UClothingAssetCommon>(GetSkeletalMesh()->GetClothingAsset(UpdatedSection.ClothingData.AssetGuid));
+	const UClothingAssetCommon* const ClothingAsset = Cast<UClothingAssetCommon>(GetSkeletalMeshAsset()->GetClothingAsset(UpdatedSection.ClothingData.AssetGuid));
 	check(ClothingAsset);  // Must have a valid clothing asset at this point, or something has gone terribly wrong
 
 	const FClothLODDataCommon& ClothLODData = ClothingAsset->LodData[UpdatedSection.ClothingData.AssetLodIndex];
@@ -1125,9 +1125,9 @@ void UDebugSkelMeshComponent::TickComponent(float DeltaTime, enum ELevelTick Tic
 
 void UDebugSkelMeshComponent::RefreshSelectedClothingSkinnedPositions()
 {
-	if(GetSkeletalMesh() && SelectedClothingGuidForPainting.IsValid())
+	if(GetSkeletalMeshAsset() && SelectedClothingGuidForPainting.IsValid())
 	{
-		UClothingAssetBase** Asset = GetSkeletalMesh()->GetMeshClothingAssets().FindByPredicate([&](UClothingAssetBase* Item)
+		UClothingAssetBase** Asset = GetSkeletalMeshAsset()->GetMeshClothingAssets().FindByPredicate([&](UClothingAssetBase* Item)
 		{
 			return Item && SelectedClothingGuidForPainting == Item->GetAssetGuid();
 		});
@@ -1336,7 +1336,7 @@ FDebugSkelMeshDynamicData::FDebugSkelMeshDynamicData(UDebugSkelMeshComponent* In
 		SkinnedPositions = InComponent->SkinnedSelectedClothingPositions;
 		SkinnedNormals = InComponent->SkinnedSelectedClothingNormals;
 
-		if(USkeletalMesh* Mesh = InComponent->GetSkeletalMesh())
+		if(USkeletalMesh* Mesh = InComponent->GetSkeletalMeshAsset())
 		{
 			const int32 NumClothingAssets = Mesh->GetMeshClothingAssets().Num();
 			for(int32 ClothingAssetIndex = 0; ClothingAssetIndex < NumClothingAssets; ++ClothingAssetIndex)
@@ -1392,7 +1392,7 @@ FScopedSuspendAlternateSkinWeightPreview::FScopedSuspendAlternateSkinWeightPrevi
 		for (TObjectIterator<UDebugSkelMeshComponent> It; It; ++It)
 		{
 			UDebugSkelMeshComponent* DebugSKComp = *It;
-			if (DebugSKComp->GetSkeletalMesh() == SkeletalMesh)
+			if (DebugSKComp->GetSkeletalMeshAsset() == SkeletalMesh)
 			{
 				const FName ProfileName = DebugSKComp->GetCurrentSkinWeightProfileName();
 				if (ProfileName != NAME_None)
diff --git a/Engine/Source/Editor/UnrealEd/Private/ApexClothingUtils.cpp b/Engine/Source/Editor/UnrealEd/Private/ApexClothingUtils.cpp
--- a/Engine/Source/Editor/UnrealEd/Private/ApexClothingUtils.cpp
+++ b/Engine/Source/Editor/UnrealEd/Private/ApexClothingUtils.cpp
@@ -33,7 +33,7 @@ void ReregisterSkelMeshComponents(USkeletalMesh* SkelMesh)
 		USkeletalMeshComponent* MeshComponent = *It;
 		if( MeshComponent && 
 			!MeshComponent->IsTemplate() &&
-			MeshComponent->GetSkeletalMesh() == SkelMesh )
+			MeshComponent->GetSkeletalMeshAsset() == SkelMesh )
 		{
 			MeshComponent->ReregisterComponent();
 		}
@@ -47,7 +47,7 @@ void RefreshSkelMeshComponents(USkeletalMesh* SkelMesh)
 		USkeletalMeshComponent* MeshComponent = *It;
 		if( MeshComponent && 
 			!MeshComponent->IsTemplate() &&
-			MeshComponent->GetSkeletalMesh() == SkelMesh )
+			MeshComponent->GetSkeletalMeshAsset() == SkelMesh )
 		{
 			MeshComponent->RecreateRenderState_Concurrent();
 		}
diff --git a/Engine/Source/Editor/UnrealEd/Private/ComponentAssetBroker.cpp b/Engine/Source/Editor/UnrealEd/Private/ComponentAssetBroker.cpp
--- a/Engine/Source/Editor/UnrealEd/Private/ComponentAssetBroker.cpp
+++ b/Engine/Source/Editor/UnrealEd/Private/ComponentAssetBroker.cpp
@@ -81,7 +81,7 @@ public:
 	{
 		if (USkeletalMeshComponent* SkelMeshComp = Cast<USkeletalMeshComponent>(InComponent))
 		{
-			return SkelMeshComp->GetSkeletalMesh();
+			return SkelMeshComp->GetSkeletalMeshAsset();
 		}
 		return NULL;
 	}
diff --git a/Engine/Source/Editor/UnrealEd/Private/EditorActor.cpp b/Engine/Source/Editor/UnrealEd/Private/EditorActor.cpp
--- a/Engine/Source/Editor/UnrealEd/Private/EditorActor.cpp
+++ b/Engine/Source/Editor/UnrealEd/Private/EditorActor.cpp
@@ -2370,7 +2370,7 @@ void UUnrealEdEngine::edactSelectMatchingSkeletalMesh(bool bAllClasses)
 		if(SkelMeshActor && SkelMeshActor->GetSkeletalMeshComponent())
 		{
 			bSelectSkelMeshActors = true;
-			SelectedMeshes.AddUnique(SkelMeshActor->GetSkeletalMeshComponent()->GetSkeletalMesh());
+			SelectedMeshes.AddUnique(SkelMeshActor->GetSkeletalMeshComponent()->GetSkeletalMeshAsset());
 			SelectedWorlds.AddUnique(Actor->GetWorld());			
 		}
 
@@ -2382,7 +2382,7 @@ void UUnrealEdEngine::edactSelectMatchingSkeletalMesh(bool bAllClasses)
 			if (PawnSkeletalMesh)
 			{
 				bSelectPawns = true;
-				SelectedMeshes.AddUnique(PawnSkeletalMesh->GetSkeletalMesh());
+				SelectedMeshes.AddUnique(PawnSkeletalMesh->GetSkeletalMeshAsset());
 				SelectedWorlds.AddUnique(Actor->GetWorld());
 			}
 		}
@@ -2418,7 +2418,7 @@ void UUnrealEdEngine::edactSelectMatchingSkeletalMesh(bool bAllClasses)
 				ASkeletalMeshActor* SkelMeshActor = Cast<ASkeletalMeshActor>(Actor);
 				if( SkelMeshActor && 
 					SkelMeshActor->GetSkeletalMeshComponent() && 
-					SelectedMeshes.Contains(SkelMeshActor->GetSkeletalMeshComponent()->GetSkeletalMesh()) )
+					SelectedMeshes.Contains(SkelMeshActor->GetSkeletalMeshComponent()->GetSkeletalMeshAsset()) )
 				{
 					bSelectActor = true;
 				}
@@ -2430,7 +2430,7 @@ void UUnrealEdEngine::edactSelectMatchingSkeletalMesh(bool bAllClasses)
 				if (Pawn)
 				{
 					USkeletalMeshComponent* PawnSkeletalMesh = Pawn->FindComponentByClass<USkeletalMeshComponent>();
-					if (PawnSkeletalMesh && SelectedMeshes.Contains(PawnSkeletalMesh->GetSkeletalMesh()) )
+					if (PawnSkeletalMesh && SelectedMeshes.Contains(PawnSkeletalMesh->GetSkeletalMeshAsset()) )
 					{
 						bSelectActor = true;
 					}
diff --git a/Engine/Source/Editor/UnrealEd/Private/EditorEngine.cpp b/Engine/Source/Editor/UnrealEd/Private/EditorEngine.cpp
--- a/Engine/Source/Editor/UnrealEd/Private/EditorEngine.cpp
+++ b/Engine/Source/Editor/UnrealEd/Private/EditorEngine.cpp
@@ -3573,7 +3573,7 @@ struct FConvertStaticMeshActorInfo
 		InternalGetFromActor(Actor);
 
 		// Copy over component properties.
-		SkeletalMesh			= MeshComp->GetSkeletalMesh();
+		SkeletalMesh			= MeshComp->GetSkeletalMeshAsset();
 		OverrideMaterials		= MeshComp->OverrideMaterials;
 		CachedMaxDrawDistance	= MeshComp->CachedMaxDrawDistance;
 		CastShadow				= MeshComp->CastShadow;
diff --git a/Engine/Source/Editor/UnrealEd/Private/EditorExporters.cpp b/Engine/Source/Editor/UnrealEd/Private/EditorExporters.cpp
--- a/Engine/Source/Editor/UnrealEd/Private/EditorExporters.cpp
+++ b/Engine/Source/Editor/UnrealEd/Private/EditorExporters.cpp
@@ -1569,7 +1569,7 @@ bool IsSomethingToExport(AActor* Actor)
 		UChildActorComponent* ChildActorComp = Cast<UChildActorComponent>(Component);
 
 		if ((StaticMeshComp && StaticMeshComp->GetStaticMesh()) ||
-			(SkelMeshComp && SkelMeshComp->GetSkeletalMesh()) ||
+			(SkelMeshComp && SkelMeshComp->GetSkeletalMeshAsset()) ||
 			Component->IsA(UCameraComponent::StaticClass()) ||
 			Component->IsA(ULightComponent::StaticClass()) ||
 			(ChildActorComp && ChildActorComp->GetChildActor() && IsSomethingToExport(ChildActorComp->GetChildActor()))
diff --git a/Engine/Source/Editor/UnrealEd/Private/Fbx/FbxAnimationExport.cpp b/Engine/Source/Editor/UnrealEd/Private/Fbx/FbxAnimationExport.cpp
--- a/Engine/Source/Editor/UnrealEd/Private/Fbx/FbxAnimationExport.cpp
+++ b/Engine/Source/Editor/UnrealEd/Private/Fbx/FbxAnimationExport.cpp
@@ -659,7 +659,7 @@ void FFbxExporter::ExportAnimTrack(IAnimTrackAdapter& AnimTrackAdapter, AActor*
 		// Add the animation data to the bone nodes
 		for(int32 BoneIndex = 0; BoneIndex < BoneNodes.Num(); ++BoneIndex)
 		{
-			FName BoneName = InSkeletalMeshComponent->GetSkeletalMesh()->GetRefSkeleton().GetBoneName(BoneIndex);
+			FName BoneName = InSkeletalMeshComponent->GetSkeletalMeshAsset()->GetRefSkeleton().GetBoneName(BoneIndex);
 			FbxNode* CurrentBoneNode = BoneNodes[BoneIndex];
 
 			// Create the AnimCurves
diff --git a/Engine/Source/Editor/UnrealEd/Private/Fbx/FbxMainExport.cpp b/Engine/Source/Editor/UnrealEd/Private/Fbx/FbxMainExport.cpp
--- a/Engine/Source/Editor/UnrealEd/Private/Fbx/FbxMainExport.cpp
+++ b/Engine/Source/Editor/UnrealEd/Private/Fbx/FbxMainExport.cpp
@@ -1213,7 +1213,7 @@ void FFbxExporter::ExportSkeletalMesh( AActor* Actor, USkeletalMeshComponent* Sk
 	if (Scene == NULL || Actor == NULL || SkeletalMeshComponent == NULL) return;
 
 	// Retrieve the skeletal mesh rendering information.
-	USkeletalMesh* SkeletalMesh = SkeletalMeshComponent->GetSkeletalMesh();
+	USkeletalMesh* SkeletalMesh = SkeletalMeshComponent->GetSkeletalMeshAsset();
 
 	FString FbxNodeName = NodeNameAdapter.GetActorNodeName(Actor);
 
@@ -2091,7 +2091,7 @@ FbxNode* FFbxExporter::ExportActor(AActor* Actor, bool bExportComponents, INodeN
 				{
 					ComponentsToExport.Add( StaticMeshComp );
 				}
-				else if( SkelMeshComp && SkelMeshComp->GetSkeletalMesh())
+				else if( SkelMeshComp && SkelMeshComp->GetSkeletalMeshAsset())
 				{
 					ComponentsToExport.Add( SkelMeshComp );
 				}
@@ -2186,7 +2186,7 @@ FbxNode* FFbxExporter::ExportActor(AActor* Actor, bool bExportComponents, INodeN
 						ExportStaticMeshToFbx(StaticMeshComp->GetStaticMesh(), LODIndex, *StaticMeshComp->GetName(), ExportNode, LightmapUVChannel, ColorBuffer, MaterialOrderOverride, &ToRawPtrTArrayUnsafe(StaticMeshComp->OverrideMaterials));
 					}
 				}
-				else if (SkelMeshComp && SkelMeshComp->GetSkeletalMesh())
+				else if (SkelMeshComp && SkelMeshComp->GetSkeletalMeshAsset())
 				{
 					ExportSkeletalMeshComponent(SkelMeshComp, *SkelMeshComp->GetName(), ExportNode, NodeNameAdapter, bSaveAnimSeq);
 				}
diff --git a/Engine/Source/Editor/UnrealEd/Private/Fbx/FbxSkeletalMeshExport.cpp b/Engine/Source/Editor/UnrealEd/Private/Fbx/FbxSkeletalMeshExport.cpp
--- a/Engine/Source/Editor/UnrealEd/Private/Fbx/FbxSkeletalMeshExport.cpp
+++ b/Engine/Source/Editor/UnrealEd/Private/Fbx/FbxSkeletalMeshExport.cpp
@@ -587,11 +587,11 @@ void FFbxExporter::CreateBindPose(FbxNode* MeshRootNode)
 
 void FFbxExporter::ExportSkeletalMeshComponent(USkeletalMeshComponent* SkelMeshComp, const TCHAR* MeshName, FbxNode* ActorRootNode, INodeNameAdapter& NodeNameAdapter, bool bSaveAnimSeq)
 {
-	if (SkelMeshComp && SkelMeshComp->GetSkeletalMesh())
+	if (SkelMeshComp && SkelMeshComp->GetSkeletalMeshAsset())
 	{
 		UAnimSequence* AnimSeq = (bSaveAnimSeq && SkelMeshComp->GetAnimationMode() == EAnimationMode::AnimationSingleNode) ? 
 			Cast<UAnimSequence>(SkelMeshComp->AnimationData.AnimToPlay) : NULL;
-		FbxNode* SkeletonRootNode = ExportSkeletalMeshToFbx(SkelMeshComp->GetSkeletalMesh(), AnimSeq, MeshName, ActorRootNode, &ToRawPtrTArrayUnsafe(SkelMeshComp->OverrideMaterials));
+		FbxNode* SkeletonRootNode = ExportSkeletalMeshToFbx(SkelMeshComp->GetSkeletalMeshAsset(), AnimSeq, MeshName, ActorRootNode, &ToRawPtrTArrayUnsafe(SkelMeshComp->OverrideMaterials));
 		if(SkeletonRootNode)
 		{
 			FbxSkeletonRoots.Add(SkelMeshComp, SkeletonRootNode);
diff --git a/Engine/Source/Editor/UnrealEd/Private/Fbx/ReimportFbxSceneFactory.cpp b/Engine/Source/Editor/UnrealEd/Private/Fbx/ReimportFbxSceneFactory.cpp
--- a/Engine/Source/Editor/UnrealEd/Private/Fbx/ReimportFbxSceneFactory.cpp
+++ b/Engine/Source/Editor/UnrealEd/Private/Fbx/ReimportFbxSceneFactory.cpp
@@ -894,7 +894,7 @@ void UReimportFbxSceneFactory::RecursivelySetComponentProperties(USCS_Node* Curr
 		{
 			USkeletalMeshComponent *CurrentNodeMeshComponent = Cast<USkeletalMeshComponent>(CurrentNodeSceneComponent);
 			USkeletalMeshComponent *MeshComponent = Cast<USkeletalMeshComponent>(SceneComponent);
-			if (CurrentNodeMeshComponent->GetSkeletalMesh() != MeshComponent->GetSkeletalMesh())
+			if (CurrentNodeMeshComponent->GetSkeletalMeshAsset() != MeshComponent->GetSkeletalMeshAsset())
 				bShouldSerializeProperty = false;
 		}
 
diff --git a/Engine/Source/Editor/UnrealEd/Private/LevelEditorViewport.cpp b/Engine/Source/Editor/UnrealEd/Private/LevelEditorViewport.cpp
--- a/Engine/Source/Editor/UnrealEd/Private/LevelEditorViewport.cpp
+++ b/Engine/Source/Editor/UnrealEd/Private/LevelEditorViewport.cpp
@@ -704,7 +704,7 @@ static bool AttemptApplyObjToComponent(UObject* ObjToUse, USceneComponent* Compo
 						SkeletalMeshComponent->Modify();
 
 						// If the component doesn't have a mesh or the anim blueprint's skeleton isn't compatible with the existing mesh's skeleton, the mesh should change
-						const bool bShouldChangeMesh = !SkeletalMeshComponent->GetSkeletalMesh() || !SkeletalMeshComponent->GetSkeletalMesh()->GetSkeleton()->IsCompatible(AnimBPSkeleton);
+						const bool bShouldChangeMesh = !SkeletalMeshComponent->GetSkeletalMeshAsset() || !SkeletalMeshComponent->GetSkeletalMeshAsset()->GetSkeleton()->IsCompatible(AnimBPSkeleton);
 
 						if (bShouldChangeMesh)
 						{
@@ -712,7 +712,7 @@ static bool AttemptApplyObjToComponent(UObject* ObjToUse, USceneComponent* Compo
 						}
 
 						// Verify that the skeletons are compatible before changing the anim BP
-						if (SkeletalMeshComponent->GetSkeletalMesh() && SkeletalMeshComponent->GetSkeletalMesh()->GetSkeleton()->IsCompatible(AnimBPSkeleton))
+						if (SkeletalMeshComponent->GetSkeletalMeshAsset() && SkeletalMeshComponent->GetSkeletalMeshAsset()->GetSkeleton()->IsCompatible(AnimBPSkeleton))
 						{
 							SkeletalMeshComponent->SetAnimInstanceClass(DroppedObjAsAnimBlueprint->GeneratedClass);
 							bResult = true;
@@ -739,7 +739,7 @@ static bool AttemptApplyObjToComponent(UObject* ObjToUse, USceneComponent* Compo
 						SkeletalMeshComponent->Modify();
 
 						// If the component doesn't have a mesh or the anim blueprint's skeleton isn't compatible with the existing mesh's skeleton, the mesh should change
-						const bool bShouldChangeMesh = !SkeletalMeshComponent->GetSkeletalMesh() || !SkeletalMeshComponent->GetSkeletalMesh()->GetSkeleton()->IsCompatible(AnimSkeleton);
+						const bool bShouldChangeMesh = !SkeletalMeshComponent->GetSkeletalMeshAsset() || !SkeletalMeshComponent->GetSkeletalMeshAsset()->GetSkeleton()->IsCompatible(AnimSkeleton);
 
 						if (bShouldChangeMesh)
 						{
@@ -752,7 +752,7 @@ static bool AttemptApplyObjToComponent(UObject* ObjToUse, USceneComponent* Compo
 						// set runtime data
 						SkeletalMeshComponent->SetAnimation(DroppedObjAsAnimSequence);
 
-						if (SkeletalMeshComponent->GetSkeletalMesh())
+						if (SkeletalMeshComponent->GetSkeletalMeshAsset())
 						{
 							bResult = true;
 							SkeletalMeshComponent->InitAnim(true);
diff --git a/Engine/Source/Editor/UnrealEd/Private/ThumbnailHelpers.cpp b/Engine/Source/Editor/UnrealEd/Private/ThumbnailHelpers.cpp
--- a/Engine/Source/Editor/UnrealEd/Private/ThumbnailHelpers.cpp
+++ b/Engine/Source/Editor/UnrealEd/Private/ThumbnailHelpers.cpp
@@ -476,9 +476,9 @@ void FSkeletalMeshThumbnailScene::GetViewMatrixParameters(const float InFOVDegre
 	const float TargetDistance = HalfMeshSize / FMath::Tan(HalfFOVRadians);
 
 	USceneThumbnailInfo* ThumbnailInfo = nullptr;
-	if(PreviewActor->GetSkeletalMeshComponent()->GetSkeletalMesh())
+	if(PreviewActor->GetSkeletalMeshComponent()->GetSkeletalMeshAsset())
 	{
-		ThumbnailInfo = Cast<USceneThumbnailInfo>(PreviewActor->GetSkeletalMeshComponent()->GetSkeletalMesh()->GetThumbnailInfo());
+		ThumbnailInfo = Cast<USceneThumbnailInfo>(PreviewActor->GetSkeletalMeshComponent()->GetSkeletalMeshAsset()->GetThumbnailInfo());
 	}
 	
 	if ( ThumbnailInfo )
@@ -690,7 +690,7 @@ void FAnimationSequenceThumbnailScene::GetViewMatrixParameters(const float InFOV
 {
 	check(PreviewAnimation);
 	check(PreviewActor->GetSkeletalMeshComponent());
-	check(PreviewActor->GetSkeletalMeshComponent()->GetSkeletalMesh());
+	check(PreviewActor->GetSkeletalMeshComponent()->GetSkeletalMeshAsset());
 
 	const float HalfFOVRadians = FMath::DegreesToRadians<float>(InFOVDegrees) * 0.5f;
 	// No need to add extra size to view slightly outside of the sphere to compensate for perspective since skeletal meshes already buffer bounds.
@@ -824,7 +824,7 @@ void FBlendSpaceThumbnailScene::GetViewMatrixParameters(const float InFOVDegrees
 {
 	check(PreviewAnimation);
 	check(PreviewActor->GetSkeletalMeshComponent());
-	check(PreviewActor->GetSkeletalMeshComponent()->GetSkeletalMesh());
+	check(PreviewActor->GetSkeletalMeshComponent()->GetSkeletalMeshAsset());
 
 	const float HalfFOVRadians = FMath::DegreesToRadians<float>(InFOVDegrees) * 0.5f;
 	// No need to add extra size to view slightly outside of the sphere to compensate for perspective since skeletal meshes already buffer bounds.
@@ -945,7 +945,7 @@ void FAnimBlueprintThumbnailScene::GetViewMatrixParameters(const float InFOVDegr
 {
 	check(PreviewBlueprint);
 	check(PreviewActor->GetSkeletalMeshComponent());
-	check(PreviewActor->GetSkeletalMeshComponent()->GetSkeletalMesh());
+	check(PreviewActor->GetSkeletalMeshComponent()->GetSkeletalMeshAsset());
 
 	const float HalfFOVRadians = FMath::DegreesToRadians<float>(InFOVDegrees) * 0.5f;
 	// No need to add extra size to view slightly outside of the sphere to compensate for perspective since skeletal meshes already buffer bounds.
@@ -1029,9 +1029,9 @@ void FPhysicsAssetThumbnailScene::GetViewMatrixParameters(const float InFOVDegre
 	const float TargetDistance = HalfMeshSize / FMath::Tan(HalfFOVRadians);
 
 	USceneThumbnailInfo* ThumbnailInfo = USceneThumbnailInfo::StaticClass()->GetDefaultObject<USceneThumbnailInfo>();
-	if(PreviewActor->GetSkeletalMeshComponent()->GetSkeletalMesh() && PreviewActor->GetSkeletalMeshComponent()->GetSkeletalMesh()->GetPhysicsAsset())
+	if(PreviewActor->GetSkeletalMeshComponent()->GetSkeletalMeshAsset() && PreviewActor->GetSkeletalMeshComponent()->GetSkeletalMeshAsset()->GetPhysicsAsset())
 	{
-		if ( USceneThumbnailInfo* InteralThumbnailInfo = Cast<USceneThumbnailInfo>(PreviewActor->GetSkeletalMeshComponent()->GetSkeletalMesh()->GetPhysicsAsset()->ThumbnailInfo) )
+		if ( USceneThumbnailInfo* InteralThumbnailInfo = Cast<USceneThumbnailInfo>(PreviewActor->GetSkeletalMeshComponent()->GetSkeletalMeshAsset()->GetPhysicsAsset()->ThumbnailInfo) )
 		{
 			ThumbnailInfo = InteralThumbnailInfo;
 			if ( TargetDistance + InteralThumbnailInfo->OrbitZoom < 0 )
@@ -1124,7 +1124,7 @@ bool FClassActorThumbnailScene::IsValidComponentForVisualization(UActorComponent
 		}
 
 		USkeletalMeshComponent* SkelMeshComp = Cast<USkeletalMeshComponent>(Component);
-		if ( SkelMeshComp && SkelMeshComp->GetSkeletalMesh())
+		if ( SkelMeshComp && SkelMeshComp->GetSkeletalMeshAsset())
 		{
 			return true;
 		}
diff --git a/Engine/Source/Runtime/AnimGraphRuntime/Private/AnimNodes/AnimNode_AimOffsetLookAt.cpp b/Engine/Source/Runtime/AnimGraphRuntime/Private/AnimNodes/AnimNode_AimOffsetLookAt.cpp
--- a/Engine/Source/Runtime/AnimGraphRuntime/Private/AnimNodes/AnimNode_AimOffsetLookAt.cpp
+++ b/Engine/Source/Runtime/AnimGraphRuntime/Private/AnimNodes/AnimNode_AimOffsetLookAt.cpp
@@ -31,7 +31,7 @@ void FAnimNode_AimOffsetLookAt::OnInitializeAnimInstance(const FAnimInstanceProx
 	SocketBoneReference.BoneName = NAME_None;
 	if (USkeletalMeshComponent* SkelMeshComp = InAnimInstance->GetSkelMeshComponent())
 	{
-		if (USkeletalMesh* SkelMesh = SkelMeshComp->GetSkeletalMesh())
+		if (USkeletalMesh* SkelMesh = SkelMeshComp->GetSkeletalMeshAsset())
 		{
 			if (const USkeletalMeshSocket* Socket = SkelMesh->FindSocket(SourceSocketName))
 			{
diff --git a/Engine/Source/Runtime/AnimGraphRuntime/Private/AnimNodes/AnimNode_CopyPoseFromMesh.cpp b/Engine/Source/Runtime/AnimGraphRuntime/Private/AnimNodes/AnimNode_CopyPoseFromMesh.cpp
--- a/Engine/Source/Runtime/AnimGraphRuntime/Private/AnimNodes/AnimNode_CopyPoseFromMesh.cpp
+++ b/Engine/Source/Runtime/AnimGraphRuntime/Private/AnimNodes/AnimNode_CopyPoseFromMesh.cpp
@@ -45,14 +45,14 @@ void FAnimNode_CopyPoseFromMesh::RefreshMeshComponent(USkeletalMeshComponent* Ta
 				ReinitializeMeshComponent(InMeshComponent, InTargetMeshComponent);
 			}
 			// if component is still same but mesh has been changed, we have to reinitialize
-			else if (CurrentMeshComponent->GetSkeletalMesh() != CurrentlyUsedSourceMesh.Get())
+			else if (CurrentMeshComponent->GetSkeletalMeshAsset() != CurrentlyUsedSourceMesh.Get())
 			{
 				ReinitializeMeshComponent(InMeshComponent, InTargetMeshComponent);
 			}
 			else if (InTargetMeshComponent)
 			{
 				// see if target mesh has changed
-				if (InTargetMeshComponent->GetSkeletalMesh() != CurrentlyUsedTargetMesh.Get())
+				if (InTargetMeshComponent->GetSkeletalMeshAsset() != CurrentlyUsedTargetMesh.Get())
 				{
 					ReinitializeMeshComponent(InMeshComponent, InTargetMeshComponent);
 				}
@@ -112,7 +112,7 @@ void FAnimNode_CopyPoseFromMesh::PreUpdate(const UAnimInstance* InAnimInstance)
 
 	USkeletalMeshComponent* CurrentMeshComponent = CurrentlyUsedSourceMeshComponent.IsValid() ? CurrentlyUsedSourceMeshComponent.Get() : nullptr;
 
-	if (CurrentMeshComponent && CurrentMeshComponent->GetSkeletalMesh() && CurrentMeshComponent->IsRegistered())
+	if (CurrentMeshComponent && CurrentMeshComponent->GetSkeletalMeshAsset() && CurrentMeshComponent->IsRegistered())
 	{
 		// If our source is running under leader-pose, then get bone data from there
 		if(USkeletalMeshComponent* LeaderPoseComponent = Cast<USkeletalMeshComponent>(CurrentMeshComponent->LeaderPoseComponent.Get()))
@@ -121,7 +121,7 @@ void FAnimNode_CopyPoseFromMesh::PreUpdate(const UAnimInstance* InAnimInstance)
 		}
 
 		// re-check mesh component validity as it may have changed to leader
-		if(CurrentMeshComponent->GetSkeletalMesh() && CurrentMeshComponent->IsRegistered())
+		if(CurrentMeshComponent->GetSkeletalMeshAsset() && CurrentMeshComponent->IsRegistered())
 		{
 			const bool bUROInSync = CurrentMeshComponent->ShouldUseUpdateRateOptimizations() && CurrentMeshComponent->AnimUpdateRateParams != nullptr && CurrentMeshComponent->AnimUpdateRateParams == InAnimInstance->GetSkelMeshComponent()->AnimUpdateRateParams;
 			const bool bUsingExternalInterpolation = CurrentMeshComponent->IsUsingExternalInterpolation();
@@ -133,7 +133,7 @@ void FAnimNode_CopyPoseFromMesh::PreUpdate(const UAnimInstance* InAnimInstance)
 			SourceMeshTransformArray.Append((bUROInSync || bUsingExternalInterpolation) && bArraySizesMatch ? CachedComponentSpaceTransforms : CurrentMeshComponent->GetComponentSpaceTransforms());
 
 			// Ref skeleton is need for parent index lookups later, so store it now
-			CurrentlyUsedMesh = CurrentMeshComponent->GetSkeletalMesh();
+			CurrentlyUsedMesh = CurrentMeshComponent->GetSkeletalMeshAsset();
 
 			if(bCopyCurves)
 			{
@@ -170,7 +170,7 @@ void FAnimNode_CopyPoseFromMesh::Update_AnyThread(const FAnimationUpdateContext&
 	GetEvaluateGraphExposedInputs().Execute(Context);
 
 	TRACE_ANIM_NODE_VALUE(Context, TEXT("Component"), *GetNameSafe(CurrentlyUsedSourceMeshComponent.IsValid() ? CurrentlyUsedSourceMeshComponent.Get() : nullptr));
-	TRACE_ANIM_NODE_VALUE(Context, TEXT("Mesh"), *GetNameSafe(CurrentlyUsedSourceMeshComponent.IsValid() ? CurrentlyUsedSourceMeshComponent.Get()->GetSkeletalMesh() : nullptr));
+	TRACE_ANIM_NODE_VALUE(Context, TEXT("Mesh"), *GetNameSafe(CurrentlyUsedSourceMeshComponent.IsValid() ? CurrentlyUsedSourceMeshComponent.Get()->GetSkeletalMeshAsset() : nullptr));
 }
 
 void FAnimNode_CopyPoseFromMesh::Evaluate_AnyThread(FPoseContext& Output)
@@ -253,7 +253,7 @@ void FAnimNode_CopyPoseFromMesh::GatherDebugData(FNodeDebugData& DebugData)
 	DECLARE_SCOPE_HIERARCHICAL_COUNTER_ANIMNODE(GatherDebugData)
 	FString DebugLine = DebugData.GetNodeName(this);
 
-	DebugLine += FString::Printf(TEXT("('%s')"), *GetNameSafe(CurrentlyUsedSourceMeshComponent.IsValid() ? CurrentlyUsedSourceMeshComponent.Get()->GetSkeletalMesh() : nullptr));
+	DebugLine += FString::Printf(TEXT("('%s')"), *GetNameSafe(CurrentlyUsedSourceMeshComponent.IsValid() ? CurrentlyUsedSourceMeshComponent.Get()->GetSkeletalMeshAsset() : nullptr));
 	DebugData.AddDebugItem(DebugLine, true);
 }
 
@@ -266,10 +266,10 @@ void FAnimNode_CopyPoseFromMesh::ReinitializeMeshComponent(USkeletalMeshComponen
 	BoneMapToSource.Reset();
 	CurveNameToUIDMap.Reset();
 
-	if (TargetMeshComponent && IsValid(NewSourceMeshComponent) && NewSourceMeshComponent->GetSkeletalMesh())
+	if (TargetMeshComponent && IsValid(NewSourceMeshComponent) && NewSourceMeshComponent->GetSkeletalMeshAsset())
 	{
-		USkeletalMesh* SourceSkelMesh = NewSourceMeshComponent->GetSkeletalMesh();
-		USkeletalMesh* TargetSkelMesh = TargetMeshComponent->GetSkeletalMesh();
+		USkeletalMesh* SourceSkelMesh = NewSourceMeshComponent->GetSkeletalMeshAsset();
+		USkeletalMesh* TargetSkelMesh = TargetMeshComponent->GetSkeletalMeshAsset();
 		
 		if (IsValid(SourceSkelMesh) && !SourceSkelMesh->HasAnyFlags(RF_NeedPostLoad) &&
 			IsValid(TargetSkelMesh) && !TargetSkelMesh->HasAnyFlags(RF_NeedPostLoad))
diff --git a/Engine/Source/Runtime/AnimGraphRuntime/Private/AnimNodes/AnimNode_PoseSnapshot.cpp b/Engine/Source/Runtime/AnimGraphRuntime/Private/AnimNodes/AnimNode_PoseSnapshot.cpp
--- a/Engine/Source/Runtime/AnimGraphRuntime/Private/AnimNodes/AnimNode_PoseSnapshot.cpp
+++ b/Engine/Source/Runtime/AnimGraphRuntime/Private/AnimNodes/AnimNode_PoseSnapshot.cpp
@@ -22,7 +22,7 @@ void FAnimNode_PoseSnapshot::PreUpdate(const UAnimInstance* InAnimInstance)
 	USkeletalMesh* CurrentSkeletalMesh = nullptr;
 	if (InAnimInstance->GetSkelMeshComponent() && InAnimInstance->GetSkelMeshComponent()->IsRegistered())
 	{
-		CurrentSkeletalMesh = InAnimInstance->GetSkelMeshComponent()->GetSkeletalMesh();
+		CurrentSkeletalMesh = InAnimInstance->GetSkelMeshComponent()->GetSkeletalMeshAsset();
 	}
 
 	if (CurrentSkeletalMesh)
diff --git a/Engine/Source/Runtime/AnimGraphRuntime/Private/AnimSequencerInstance.cpp b/Engine/Source/Runtime/AnimGraphRuntime/Private/AnimSequencerInstance.cpp
--- a/Engine/Source/Runtime/AnimGraphRuntime/Private/AnimSequencerInstance.cpp
+++ b/Engine/Source/Runtime/AnimGraphRuntime/Private/AnimSequencerInstance.cpp
@@ -65,7 +65,7 @@ void UAnimSequencerInstance::SavePose()
 {
 	if (USkeletalMeshComponent* SkeletalMeshComponent = GetSkelMeshComponent())
 	{
-		if (SkeletalMeshComponent->GetSkeletalMesh() && SkeletalMeshComponent->GetComponentSpaceTransforms().Num() > 0)
+		if (SkeletalMeshComponent->GetSkeletalMeshAsset() && SkeletalMeshComponent->GetComponentSpaceTransforms().Num() > 0)
 		{
 			SavePoseSnapshot(UAnimSequencerInstance::SequencerPoseName);
 		}
diff --git a/Engine/Source/Runtime/AnimGraphRuntime/Private/BoneControllers/AnimNode_RigidBody.cpp b/Engine/Source/Runtime/AnimGraphRuntime/Private/BoneControllers/AnimNode_RigidBody.cpp
--- a/Engine/Source/Runtime/AnimGraphRuntime/Private/BoneControllers/AnimNode_RigidBody.cpp
+++ b/Engine/Source/Runtime/AnimGraphRuntime/Private/BoneControllers/AnimNode_RigidBody.cpp
@@ -959,7 +959,7 @@ void ComputeBodyInsertionOrder(TArray<FBoneIndexType>& InsertionOrder, const USk
 
 	InsertionOrder.Reset();
 
-	if (SKC.GetSkeletalMesh() == nullptr)
+	if (SKC.GetSkeletalMeshAsset() == nullptr)
 	{
 		return;
 	}
@@ -972,7 +972,7 @@ void ComputeBodyInsertionOrder(TArray<FBoneIndexType>& InsertionOrder, const USk
 		SKC.ComputeRequiredBones(RequiredBones0, ComponentSpaceTMs0, 0, /*bIgnorePhysicsAsset=*/ true);
 
 		TArray<bool> InSortedOrder;
-		InSortedOrder.AddZeroed(SKC.GetSkeletalMesh()->GetRefSkeleton().GetNum());
+		InSortedOrder.AddZeroed(SKC.GetSkeletalMeshAsset()->GetRefSkeleton().GetNum());
 
 		auto MergeIndices = [&InsertionOrder, &InSortedOrder](const TArray<FBoneIndexType>& RequiredBones) -> void
 		{
@@ -1009,7 +1009,7 @@ void FAnimNode_RigidBody::InitPhysics(const UAnimInstance* InAnimInstance)
 	DestroyPhysicsSimulation();
 
 	const USkeletalMeshComponent* SkeletalMeshComp = InAnimInstance->GetSkelMeshComponent();
-	const USkeletalMesh* SkeletalMeshAsset = SkeletalMeshComp->GetSkeletalMesh();
+	const USkeletalMesh* SkeletalMeshAsset = SkeletalMeshComp->GetSkeletalMeshAsset();
 	USkeleton* SkeletonAsset = InAnimInstance->CurrentSkeleton;
 
 	if(!SkeletalMeshAsset || !SkeletonAsset)
diff --git a/Engine/Source/Runtime/AnimGraphRuntime/Private/BoneControllers/AnimNode_SplineIK.cpp b/Engine/Source/Runtime/AnimGraphRuntime/Private/BoneControllers/AnimNode_SplineIK.cpp
--- a/Engine/Source/Runtime/AnimGraphRuntime/Private/BoneControllers/AnimNode_SplineIK.cpp
+++ b/Engine/Source/Runtime/AnimGraphRuntime/Private/BoneControllers/AnimNode_SplineIK.cpp
@@ -35,9 +35,9 @@ void FAnimNode_SplineIK::GatherDebugData(FNodeDebugData& DebugData)
 
 void FAnimNode_SplineIK::OnInitializeAnimInstance(const FAnimInstanceProxy* InProxy, const UAnimInstance* InAnimInstance)
 {
-	if (InAnimInstance->GetSkelMeshComponent() && InAnimInstance->GetSkelMeshComponent()->GetSkeletalMesh())
+	if (InAnimInstance->GetSkelMeshComponent() && InAnimInstance->GetSkelMeshComponent()->GetSkeletalMeshAsset())
 	{
-		GatherBoneReferences(InAnimInstance->GetSkelMeshComponent()->GetSkeletalMesh()->GetRefSkeleton());
+		GatherBoneReferences(InAnimInstance->GetSkelMeshComponent()->GetSkeletalMeshAsset()->GetRefSkeleton());
 	}
 }
 
diff --git a/Engine/Source/Runtime/ClothingSystemRuntimeCommon/Private/ClothingAsset.cpp b/Engine/Source/Runtime/ClothingSystemRuntimeCommon/Private/ClothingAsset.cpp
--- a/Engine/Source/Runtime/ClothingSystemRuntimeCommon/Private/ClothingAsset.cpp
+++ b/Engine/Source/Runtime/ClothingSystemRuntimeCommon/Private/ClothingAsset.cpp
@@ -833,7 +833,7 @@ void UClothingAssetCommon::ReregisterComponentsUsingClothing()
 		{
 			if(USkeletalMeshComponent* Component = *It)
 			{
-				if(Component->GetSkeletalMesh() == OwnerMesh)
+				if(Component->GetSkeletalMeshAsset() == OwnerMesh)
 				{
 					FComponentReregisterContext Context(Component);
 					// Context goes out of scope, causing Component to be re-registered
@@ -851,7 +851,7 @@ void UClothingAssetCommon::ForEachInteractorUsingClothing(TFunction<void(UClothi
 		{
 			if(USkeletalMeshComponent* Component = *It)
 			{
-				if(Component->GetSkeletalMesh() == OwnerMesh)
+				if(Component->GetSkeletalMeshAsset() == OwnerMesh)
 				{
 					if(UClothingSimulationInteractor* CurInteractor = Component->GetClothingSimulationInteractor())
 					{
diff --git a/Engine/Source/Runtime/ClothingSystemRuntimeCommon/Private/ClothingSimulation.cpp b/Engine/Source/Runtime/ClothingSystemRuntimeCommon/Private/ClothingSimulation.cpp
--- a/Engine/Source/Runtime/ClothingSystemRuntimeCommon/Private/ClothingSimulation.cpp
+++ b/Engine/Source/Runtime/ClothingSystemRuntimeCommon/Private/ClothingSimulation.cpp
@@ -60,7 +60,7 @@ void FClothingSimulationContextCommon::Fill(const USkeletalMeshComponent* InComp
 
 void FClothingSimulationContextCommon::FillBoneTransforms(const USkeletalMeshComponent* InComponent)
 {
-	const USkeletalMesh* const SkeletalMesh = InComponent->GetSkeletalMesh();
+	const USkeletalMesh* const SkeletalMesh = InComponent->GetSkeletalMeshAsset();
 
 	if (USkinnedMeshComponent* const LeaderComponent = InComponent->LeaderPoseComponent.Get())
 	{
@@ -122,7 +122,7 @@ void FClothingSimulationContextCommon::FillRefToLocals(const USkeletalMeshCompon
 	// Constraints are initialized using bone distances upon initialization, so fill out reference pose
 	if (bIsInitialization)
 	{
-		const USkeletalMesh* const SkeletalMesh = InComponent->GetSkeletalMesh();
+		const USkeletalMesh* const SkeletalMesh = InComponent->GetSkeletalMeshAsset();
 		if (SkeletalMesh)
 		{
 			const int32 NumBones = SkeletalMesh->GetRefSkeleton().GetNum();
diff --git a/Engine/Source/Runtime/Engine/Private/Animation/AnimInstance.cpp b/Engine/Source/Runtime/Engine/Private/Animation/AnimInstance.cpp
--- a/Engine/Source/Runtime/Engine/Private/Animation/AnimInstance.cpp
+++ b/Engine/Source/Runtime/Engine/Private/Animation/AnimInstance.cpp
@@ -217,9 +217,9 @@ void UAnimInstance::InitializeAnimation(bool bInDeferRootNodeInitialization)
 	// make sure your skeleton is initialized
 	// you can overwrite different skeleton
 	USkeletalMeshComponent* OwnerComponent = GetSkelMeshComponent();
-	if (OwnerComponent->GetSkeletalMesh() != NULL)
+	if (OwnerComponent->GetSkeletalMeshAsset() != NULL)
 	{
-		CurrentSkeleton = OwnerComponent->GetSkeletalMesh()->GetSkeleton();
+		CurrentSkeleton = OwnerComponent->GetSkeletalMeshAsset()->GetSkeleton();
 	}
 	else
 	{
@@ -1242,9 +1242,9 @@ void UAnimInstance::RecalcRequiredBones()
 	USkeletalMeshComponent* SkelMeshComp = GetSkelMeshComponent();
 	check( SkelMeshComp )
 
-	if( SkelMeshComp->GetSkeletalMesh() && SkelMeshComp->GetSkeletalMesh()->GetSkeleton() )
+	if( SkelMeshComp->GetSkeletalMeshAsset() && SkelMeshComp->GetSkeletalMeshAsset()->GetSkeleton() )
 	{
-		GetProxyOnGameThread<FAnimInstanceProxy>().RecalcRequiredBones(SkelMeshComp, SkelMeshComp->GetSkeletalMesh());
+		GetProxyOnGameThread<FAnimInstanceProxy>().RecalcRequiredBones(SkelMeshComp, SkelMeshComp->GetSkeletalMeshAsset());
 	}
 	else if( CurrentSkeleton != NULL )
 	{
@@ -1683,9 +1683,9 @@ void UAnimInstance::GetActiveCurveNames(EAnimCurveType CurveType, TArray<FName>&
 void UAnimInstance::GetAllCurveNames(TArray<FName>& OutNames) const
 {
 	USkeletalMeshComponent* SkelMeshComp = GetOwningComponent();
-	if (SkelMeshComp && SkelMeshComp->GetSkeletalMesh() && SkelMeshComp->GetSkeletalMesh()->GetSkeleton())
+	if (SkelMeshComp && SkelMeshComp->GetSkeletalMeshAsset() && SkelMeshComp->GetSkeletalMeshAsset()->GetSkeleton())
 	{
-		const USkeleton* CurSkeleton = SkelMeshComp->GetSkeletalMesh()->GetSkeleton();
+		const USkeleton* CurSkeleton = SkelMeshComp->GetSkeletalMeshAsset()->GetSkeleton();
 
 		const FSmartNameMapping* Mapping = CurSkeleton->GetSmartNameContainer(USkeleton::AnimCurveMappingName);
 		if (Mapping)
diff --git a/Engine/Source/Runtime/Engine/Private/Animation/AnimInstanceProxy.cpp b/Engine/Source/Runtime/Engine/Private/Animation/AnimInstanceProxy.cpp
--- a/Engine/Source/Runtime/Engine/Private/Animation/AnimInstanceProxy.cpp
+++ b/Engine/Source/Runtime/Engine/Private/Animation/AnimInstanceProxy.cpp
@@ -649,9 +649,9 @@ void FAnimInstanceProxy::InitializeObjects(UAnimInstance* InAnimInstance)
 		MainMontageEvaluationData = InAnimInstance->IsUsingMainInstanceMontageEvaluationData() ? &MainInstanceProxy->MontageEvaluationData : &MontageEvaluationData;
 	}
 
-	if (SkeletalMeshComponent->GetSkeletalMesh() != nullptr)
+	if (SkeletalMeshComponent->GetSkeletalMeshAsset() != nullptr)
 	{
-		Skeleton = SkeletalMeshComponent->GetSkeletalMesh()->GetSkeleton();
+		Skeleton = SkeletalMeshComponent->GetSkeletalMeshAsset()->GetSkeleton();
 	}
 	else
 	{
diff --git a/Engine/Source/Runtime/Engine/Private/Animation/AnimSequence.cpp b/Engine/Source/Runtime/Engine/Private/Animation/AnimSequence.cpp
--- a/Engine/Source/Runtime/Engine/Private/Animation/AnimSequence.cpp
+++ b/Engine/Source/Runtime/Engine/Private/Animation/AnimSequence.cpp
@@ -3922,9 +3922,9 @@ bool UAnimSequence::CreateAnimation(USkeletalMesh* Mesh)
 
 bool UAnimSequence::CreateAnimation(USkeletalMeshComponent* MeshComponent)
 {
-	if(MeshComponent && MeshComponent->GetSkeletalMesh())
+	if(MeshComponent && MeshComponent->GetSkeletalMeshAsset())
 	{
-		USkeletalMesh * Mesh = MeshComponent->GetSkeletalMesh();
+		USkeletalMesh * Mesh = MeshComponent->GetSkeletalMeshAsset();
 
 		ValidateModel();
 		const FReferenceSkeleton& RefSkeleton = Mesh->GetRefSkeleton();
diff --git a/Engine/Source/Runtime/Engine/Private/Animation/AnimSingleNodeInstance.cpp b/Engine/Source/Runtime/Engine/Private/Animation/AnimSingleNodeInstance.cpp
--- a/Engine/Source/Runtime/Engine/Private/Animation/AnimSingleNodeInstance.cpp
+++ b/Engine/Source/Runtime/Engine/Private/Animation/AnimSingleNodeInstance.cpp
@@ -40,7 +40,7 @@ void UAnimSingleNodeInstance::SetAnimationAsset(class UAnimationAsset* NewAsset,
 	USkeletalMeshComponent* MeshComponent = GetSkelMeshComponent();
 	if (MeshComponent)
 	{
-		if (MeshComponent->GetSkeletalMesh() == nullptr)
+		if (MeshComponent->GetSkeletalMeshAsset() == nullptr)
 		{
 			// if it does not have SkeletalMesh, we nullify it
 			CurrentAsset = nullptr;
@@ -48,7 +48,7 @@ void UAnimSingleNodeInstance::SetAnimationAsset(class UAnimationAsset* NewAsset,
 		else if (CurrentAsset != nullptr)
 		{
 			// if we have an asset, make sure their skeleton matches, otherwise, null it
-			if (!MeshComponent->GetSkeletalMesh()->GetSkeleton()->IsCompatible(CurrentAsset->GetSkeleton()))
+			if (!MeshComponent->GetSkeletalMeshAsset()->GetSkeleton()->IsCompatible(CurrentAsset->GetSkeleton()))
 			{
 				// clear asset since we do not have matching skeleton
 				CurrentAsset = nullptr;
diff --git a/Engine/Source/Runtime/Engine/Private/Animation/AnimTrace.cpp b/Engine/Source/Runtime/Engine/Private/Animation/AnimTrace.cpp
--- a/Engine/Source/Runtime/Engine/Private/Animation/AnimTrace.cpp
+++ b/Engine/Source/Runtime/Engine/Private/Animation/AnimTrace.cpp
@@ -534,7 +534,7 @@ void FAnimTrace::OutputSkeletalMeshComponent(const USkeletalMeshComponent* InCom
 	if(BoneCount > 0 || CurveCount > 0)
 	{
 		TRACE_OBJECT(InComponent);
-		TRACE_SKELETAL_MESH(InComponent->GetSkeletalMesh());
+		TRACE_SKELETAL_MESH(InComponent->GetSkeletalMeshAsset());
 
 		TArray<float>& CurveValues = FAnimTraceScratchBuffers::Get().CurveValues;
 		CurveValues.Reset();
@@ -561,7 +561,7 @@ void FAnimTrace::OutputSkeletalMeshComponent(const USkeletalMeshComponent* InCom
 			<< SkeletalMeshComponent3.Cycle(FPlatformTime::Cycles64())
 			<< SkeletalMeshComponent3.RecordingTime(FObjectTrace::GetWorldElapsedTime(InComponent->GetWorld()))
 			<< SkeletalMeshComponent3.ComponentId(FObjectTrace::GetObjectId(InComponent))
-			<< SkeletalMeshComponent3.MeshId(FObjectTrace::GetObjectId(InComponent->GetSkeletalMesh()))
+			<< SkeletalMeshComponent3.MeshId(FObjectTrace::GetObjectId(InComponent->GetSkeletalMeshAsset()))
 			<< SkeletalMeshComponent3.ComponentToWorld(reinterpret_cast<const float*>(&InComponent->GetComponentToWorld()), sizeof(FTransform) / sizeof(float))
 			<< SkeletalMeshComponent3.Pose(reinterpret_cast<const float*>(InComponent->GetComponentSpaceTransforms().GetData()), BoneCount * (sizeof(FTransform) / sizeof(float)))
 			<< SkeletalMeshComponent3.CurveIds(CurveIds.GetData(), CurveIds.Num())
diff --git a/Engine/Source/Runtime/Engine/Private/Animation/PoseAsset.cpp b/Engine/Source/Runtime/Engine/Private/Animation/PoseAsset.cpp
--- a/Engine/Source/Runtime/Engine/Private/Animation/PoseAsset.cpp
+++ b/Engine/Source/Runtime/Engine/Private/Animation/PoseAsset.cpp
@@ -1191,12 +1191,12 @@ bool UPoseAsset::AddOrUpdatePoseWithUniqueName(USkeletalMeshComponent* MeshCompo
 void UPoseAsset::AddOrUpdatePose(const FSmartName& PoseName, USkeletalMeshComponent* MeshComponent, bool bUpdateCurves)
 {
 	USkeleton* MySkeleton = GetSkeleton();
-	if (MySkeleton && MeshComponent && MeshComponent->GetSkeletalMesh())
+	if (MySkeleton && MeshComponent && MeshComponent->GetSkeletalMeshAsset())
 	{
 		TArray<FName> TrackNames;
 		// note this ignores root motion
 		TArray<FTransform> BoneTransform = MeshComponent->GetComponentSpaceTransforms();
-		const FReferenceSkeleton& RefSkeleton = MeshComponent->GetSkeletalMesh()->GetRefSkeleton();
+		const FReferenceSkeleton& RefSkeleton = MeshComponent->GetSkeletalMeshAsset()->GetRefSkeleton();
 		for (int32 BoneIndex = 0; BoneIndex < RefSkeleton.GetNum(); ++BoneIndex)
 		{
 			TrackNames.Add(RefSkeleton.GetBoneName(BoneIndex));
@@ -1212,7 +1212,7 @@ void UPoseAsset::AddOrUpdatePose(const FSmartName& PoseName, USkeletalMeshCompon
 			}
 		}
 
-		const USkeleton* MeshSkeleton = MeshComponent->GetSkeletalMesh()->GetSkeleton();
+		const USkeleton* MeshSkeleton = MeshComponent->GetSkeletalMeshAsset()->GetSkeleton();
 		const FSmartNameMapping* Mapping = MeshSkeleton->GetSmartNameContainer(USkeleton::AnimCurveMappingName);
 		
 		TArray<float> NewCurveValues;
diff --git a/Engine/Source/Runtime/Engine/Private/Animation/SkeletalMeshActor.cpp b/Engine/Source/Runtime/Engine/Private/Animation/SkeletalMeshActor.cpp
--- a/Engine/Source/Runtime/Engine/Private/Animation/SkeletalMeshActor.cpp
+++ b/Engine/Source/Runtime/Engine/Private/Animation/SkeletalMeshActor.cpp
@@ -74,7 +74,7 @@ void ASkeletalMeshActor::CheckForErrors()
 				->AddToken(FMapErrorToken::Create(FMapErrors::ActorLargeShadowCaster));
 		}
 
-		if (SkeletalMeshComponent->GetSkeletalMesh() == NULL)
+		if (SkeletalMeshComponent->GetSkeletalMeshAsset() == NULL)
 		{
 			FMessageLog("MapCheck").Warning()
 				->AddToken(FUObjectToken::Create(this))
@@ -116,7 +116,7 @@ void ASkeletalMeshActor::PostInitializeComponents()
 	// grab the current mesh for replication
 	if (GetLocalRole() == ROLE_Authority && SkeletalMeshComponent)
 	{
-		ReplicatedMesh = SkeletalMeshComponent->GetSkeletalMesh();
+		ReplicatedMesh = SkeletalMeshComponent->GetSkeletalMeshAsset();
 	}
 
 	// Unfix bodies flagged as 'full anim weight'
@@ -148,8 +148,8 @@ void ASkeletalMeshActor::OnRep_ReplicatedMaterial1()
 
 bool ASkeletalMeshActor::CanPlayAnimation(class UAnimSequenceBase* AnimAssetBase/*=NULL*/) const
 {
-	return (SkeletalMeshComponent->GetSkeletalMesh() && SkeletalMeshComponent->GetSkeletalMesh()->GetSkeleton() &&
-		(!AnimAssetBase || SkeletalMeshComponent->GetSkeletalMesh()->GetSkeleton()->IsCompatible(AnimAssetBase->GetSkeleton())));
+	return (SkeletalMeshComponent->GetSkeletalMeshAsset() && SkeletalMeshComponent->GetSkeletalMeshAsset()->GetSkeleton() &&
+		(!AnimAssetBase || SkeletalMeshComponent->GetSkeletalMeshAsset()->GetSkeleton()->IsCompatible(AnimAssetBase->GetSkeleton())));
 }
 
 #if WITH_EDITOR
@@ -158,9 +158,9 @@ bool ASkeletalMeshActor::GetReferencedContentObjects(TArray<UObject*>& Objects)
 {
 	Super::GetReferencedContentObjects(Objects);
 
-	if (SkeletalMeshComponent->GetSkeletalMesh())
+	if (SkeletalMeshComponent->GetSkeletalMeshAsset())
 	{
-		Objects.Add(SkeletalMeshComponent->GetSkeletalMesh());
+		Objects.Add(SkeletalMeshComponent->GetSkeletalMeshAsset());
 	}
 
 	if (SkeletalMeshComponent->GetAnimationMode() == EAnimationMode::Type::AnimationSingleNode &&
@@ -180,9 +180,9 @@ void ASkeletalMeshActor::EditorReplacedActor(AActor* OldActor)
 	{
 		// if no skeletal mesh set, take one from previous actor
 		if (SkeletalMeshComponent && OldSkelMeshActor->SkeletalMeshComponent &&
-			SkeletalMeshComponent->GetSkeletalMesh() == NULL)
+			SkeletalMeshComponent->GetSkeletalMeshAsset() == NULL)
 		{
-			SkeletalMeshComponent->SetSkeletalMesh(OldSkelMeshActor->SkeletalMeshComponent->GetSkeletalMesh());
+			SkeletalMeshComponent->SetSkeletalMesh(OldSkelMeshActor->SkeletalMeshComponent->GetSkeletalMeshAsset());
 		}
 	}
 }
