Fixed WorldPartitionConvertCommandlet to properly support DataLayer Assets.
- DataLayerAsset folder now takes into account commandlet world suffix
- DataLayerAsset folder follows world package mount point to support worlds inside plugins
- DataLayer assets are now properly deleted in the clean-up pass at the beginning of the conversion
- DataLayer name is now sanitized when converting from a Layer
- Conversion of Layers to DataLayers for actors part of sublevels moved after actors being moved to the partitioned level

diff --git a/Engine/Source/Editor/UnrealEd/Private/Commandlets/WorldPartitionConvertCommandlet.cpp b/Engine/Source/Editor/UnrealEd/Private/Commandlets/WorldPartitionConvertCommandlet.cpp
--- a/Engine/Source/Editor/UnrealEd/Private/Commandlets/WorldPartitionConvertCommandlet.cpp
+++ b/Engine/Source/Editor/UnrealEd/Private/Commandlets/WorldPartitionConvertCommandlet.cpp
@@ -165,7 +165,6 @@ UWorldPartitionConvertCommandlet::UWorldPartitionConvertCommandlet(const FObject
 	, WorldOrigin(FVector::ZeroVector)
 	, WorldExtent(HALF_WORLD_MAX)
 	, LandscapeGridSize(4)
-	, DataLayerAssetFolder(TEXT("Game/DataLayers/"))
 	, DataLayerFactory(NewObject<UDataLayerFactory>())
 {}
 
@@ -729,11 +728,6 @@ int32 UWorldPartitionConvertCommandlet::Main(const FString& Params)
 		return 1;
 	}
 
-	if (FString* DataLayerAssetFolderValue = Arguments.Find(TEXT("DataLayerAssetFolder")))
-	{
-		DataLayerAssetFolder = *DataLayerAssetFolderValue;
-	}
-
 	ReadAdditionalTokensAndSwitches(Tokens, Switches);
 
 	if (bVerbose)
@@ -774,19 +768,46 @@ int32 UWorldPartitionConvertCommandlet::Main(const FString& Params)
 
 	SetupHLOD();
 
+	// Load world
+	UWorld* MainWorld = LoadWorld(Tokens[0]);
+	if (!MainWorld)
+	{
+		UE_LOG(LogWorldPartitionConvertCommandlet, Error, TEXT("Unknown world '%s'"), *Tokens[0]);
+		return 1;
+	}
+
+	// Setup Folder for DataLayer assets
+	if (FString* DataLayerAssetFolderValue = Arguments.Find(TEXT("DataLayerAssetFolder")))
+	{
+		DataLayerAssetFolder = *DataLayerAssetFolderValue;
+	}
+	else
+	{
+		UPackage* Package = MainWorld->GetPackage();
+		FName PackageMountPoint = FPackageName::GetPackageMountPoint(Package->GetLoadedPath().GetPackageName());
+		if (PackageMountPoint.IsNone())
+		{
+			PackageMountPoint = TEXT("Game");
+		}
+		DataLayerAssetFolder = FPaths::RemoveDuplicateSlashes(FString::Printf(TEXT("/%s/DataLayers/%s%s/"), *PackageMountPoint.ToString(), *MainWorld->GetName(), *ConversionSuffix));
+	}
+
 	// Delete existing result from running the commandlet, even if not using the suffix mode to cleanup previous conversion
 	if (!bReportOnly)
 	{
 		UE_SCOPED_TIMER(TEXT("Deleting existing conversion results"), LogWorldPartitionConvertCommandlet, Display);
 
 		FString OldLevelName = Tokens[0] + ConversionSuffix;
-		TArray<FString> ExternalObjectsPaths = ULevel::GetExternalObjectsPaths(OldLevelName);
-		for (const FString& ExternalObjectsPath : ExternalObjectsPaths)
+		TArray<FString> CleanupPaths = ULevel::GetExternalObjectsPaths(OldLevelName);
+		// Append DataLayer assets folder
+		CleanupPaths.Add(DataLayerAssetFolder);
+
+		for (const FString& CleanupPath : CleanupPaths)
 		{
-			FString ExternalObjectsFilePath = FPackageName::LongPackageNameToFilename(ExternalObjectsPath);
-			if (IFileManager::Get().DirectoryExists(*ExternalObjectsFilePath))
+			FString Directory = FPackageName::LongPackageNameToFilename(CleanupPath);
+			if (IFileManager::Get().DirectoryExists(*Directory))
 			{
-				bool bResult = IFileManager::Get().IterateDirectoryRecursively(*ExternalObjectsFilePath, [this, &PackageHelper](const TCHAR* FilenameOrDirectory, bool bIsDirectory)
+				bool bResult = IFileManager::Get().IterateDirectoryRecursively(*Directory, [this, &PackageHelper](const TCHAR* FilenameOrDirectory, bool bIsDirectory)
 				{
 					if (!bIsDirectory)
 					{
@@ -801,7 +822,7 @@ int32 UWorldPartitionConvertCommandlet::Main(const FString& Params)
 
 				if (!bResult)
 				{
-					UE_LOG(LogWorldPartitionConvertCommandlet, Error, TEXT("Failed to delete external package(s)"));
+					UE_LOG(LogWorldPartitionConvertCommandlet, Error, TEXT("Failed to delete previous conversion package(s)"));
 					return 1;
 				}
 			}
@@ -817,14 +838,6 @@ int32 UWorldPartitionConvertCommandlet::Main(const FString& Params)
 		}
 	}
 
-	// Load world
-	UWorld* MainWorld = LoadWorld(Tokens[0]);
-	if (!MainWorld)
-	{
-		UE_LOG(LogWorldPartitionConvertCommandlet, Error, TEXT("Unknown world '%s'"), *Tokens[0]);
-		return 1;
-	}
-
 	// Make sure the world isn't already partitioned
 	AWorldSettings* MainWorldSettings = MainWorld->GetWorldSettings();
 	if (MainWorldSettings->IsPartitionedWorld())
@@ -1004,7 +1017,42 @@ int32 UWorldPartitionConvertCommandlet::Main(const FString& Params)
 		}
 	};
 
-	auto PrepareLevelActors = [this, PartitionFoliage, PartitionLandscape, MainWorldDataLayers, MainWorld](ULevel* Level, TArray<AActor*>& Actors, bool bMainLevel) -> bool
+	IAssetTools& AssetTools = FModuleManager::GetModuleChecked<FAssetToolsModule>("AssetTools").Get();
+
+	auto ConvertActorLayersToDataLayers = [this, MainWorldDataLayers, &AssetTools](AActor* Actor)
+	{
+		// Convert Layers into DataLayers with DynamicallyLoaded flag disabled
+		if (Actor->SupportsDataLayer())
+		{
+			for (FName Layer : Actor->Layers)
+			{
+				FString DataLayerAssetName = SlugStringForValidName(Layer.ToString());
+				FName DataLayerAssetPathName(DataLayerAssetFolder + DataLayerAssetName + TEXT(".") + DataLayerAssetName);
+				UDataLayerInstance* DataLayerInstance = const_cast<UDataLayerInstance*>(MainWorldDataLayers->GetDataLayerInstanceFromAssetName(DataLayerAssetPathName));
+				if (!DataLayerInstance)
+				{
+					if (UObject* Asset = AssetTools.CreateAsset(DataLayerAssetName, DataLayerAssetFolder, UDataLayerAsset::StaticClass(), DataLayerFactory))
+					{
+						PackagesToSave.Add(Asset->GetPackage());
+						UDataLayerAsset* DataLayerAsset = CastChecked<UDataLayerAsset>(Asset);
+						DataLayerAsset->SetType(EDataLayerType::Editor);
+						DataLayerInstance = MainWorldDataLayers->CreateDataLayer<UDataLayerInstanceWithAsset>(DataLayerAsset);
+					}
+				}
+				if (DataLayerInstance)
+				{
+					Actor->AddDataLayer(DataLayerInstance);
+				}
+			}
+		}
+		// Clear actor layers as they are replaced by data layers, keep them if only merging
+		if (!bOnlyMergeSubLevels)
+		{
+			Actor->Layers.Empty();
+		}
+	};
+
+	auto PrepareLevelActors = [this, PartitionFoliage, PartitionLandscape, MainWorld, ConvertActorLayersToDataLayers](ULevel* Level, TArray<AActor*>& Actors, bool bMainLevel) -> bool
 	{
 		TRACE_CPUPROFILER_EVENT_SCOPE(PrepareLevelActors);
 
@@ -1012,9 +1060,7 @@ int32 UWorldPartitionConvertCommandlet::Main(const FString& Params)
 
 		TArray<AInstancedFoliageActor*> IFAs;
 		TSet<ULandscapeInfo*> LandscapeInfos;
-		IAssetTools& AssetTools = FModuleManager::GetModuleChecked<FAssetToolsModule>("AssetTools").Get();
-		FString DataLayerFolder = DataLayerAssetFolder + MainWorld->GetName() + TEXT("/");
-		for (AActor* Actor: Actors)
+		for (AActor* Actor : Actors)
 		{
 			if (Actor && IsValidChecked(Actor))
 			{
@@ -1047,29 +1093,9 @@ int32 UWorldPartitionConvertCommandlet::Main(const FString& Params)
 						}
 					}
 
-					// Convert Layers into DataLayers with DynamicallyLoaded flag disabled
-					if (Actor->SupportsDataLayer())
+					if (bMainLevel)
 					{
-						for (FName Layer : Actor->Layers)
-						{
-							FName DataLayerAssetFullName(DataLayerAssetFolder + Layer.ToString());
-							UDataLayerInstance* DataLayerInstance = const_cast<UDataLayerInstance*>(MainWorldDataLayers->GetDataLayerInstanceFromAssetName(DataLayerAssetFullName));
-							if (!DataLayerInstance)
-							{
-								if (UObject* Asset = AssetTools.CreateAsset(Layer.ToString(), DataLayerFolder, UDataLayerAsset::StaticClass(), DataLayerFactory))
-								{
-									UDataLayerAsset* DataLayerAsset = CastChecked<UDataLayerAsset>(Asset);
-									DataLayerAsset->SetType(EDataLayerType::Editor);
-									DataLayerInstance = MainWorldDataLayers->CreateDataLayer<UDataLayerInstanceWithAsset>(DataLayerAsset);
-								}
-							}
-							Actor->AddDataLayer(DataLayerInstance);
-						}
-					}
-					// Clear actor layers as they are not supported yet in world partition, keep them if only merging
-					if (!bOnlyMergeSubLevels)
-					{
-						Actor->Layers.Empty();
+						ConvertActorLayersToDataLayers(Actor);
 					}
 				}
 			}
@@ -1350,6 +1376,9 @@ int32 UWorldPartitionConvertCommandlet::Main(const FString& Params)
 					ChangeObjectOuter(ActorClass, MainPackage);
 					UE_LOG(LogWorldPartitionConvertCommandlet, Log, TEXT("Extracted non-native class %s"), *ActorClass->GetName());
 				}
+
+				// Actor is now in main level, we can create data layers for it
+				ConvertActorLayersToDataLayers(Actor);
 			}
 		}
 
