Fixed editor crash when user creates a blank Material Layer Instance, Material Function Instance or Material Laye rBlen dInstance from the Asset Browser.

diff --git a/Engine/Source/Editor/UnrealEd/Private/Factories/EditorFactories.cpp b/Engine/Source/Editor/UnrealEd/Private/Factories/EditorFactories.cpp
--- a/Engine/Source/Editor/UnrealEd/Private/Factories/EditorFactories.cpp
+++ b/Engine/Source/Editor/UnrealEd/Private/Factories/EditorFactories.cpp
@@ -580,7 +580,7 @@ UObject* UMaterialFunctionInstanceFactory::FactoryCreateNew(UClass* Class,UObjec
 {
 	auto MFI = NewObject<UMaterialFunctionInstance>(InParent, Class, Name, Flags);
 
-	if (MFI)
+	if (MFI && InitialParent)
 	{
 		MFI->SetParent(InitialParent);
 	}
@@ -603,7 +603,7 @@ UObject* UMaterialFunctionMaterialLayerInstanceFactory::FactoryCreateNew(UClass*
 {
 	auto MFI = NewObject<UMaterialFunctionMaterialLayerInstance>(InParent, Class, Name, Flags);
 
-	if (MFI)
+	if (MFI && InitialParent)
 	{
 		MFI->SetParent(InitialParent);
 	}
@@ -626,7 +626,7 @@ UObject* UMaterialFunctionMaterialLayerBlendInstanceFactory::FactoryCreateNew(UC
 {
 	auto MFI = NewObject<UMaterialFunctionMaterialLayerBlendInstance>(InParent, Class, Name, Flags);
 
-	if (MFI)
+	if (MFI && InitialParent) 
 	{
 		MFI->SetParent(InitialParent);
 	}
