Cooker: FullLoadAndSave: AnimSequence cooking now requires FTickableCookObject::TickObjects for IsCachedCookedPlatformData to complete. Fix FullLoadAndSave to call TickObjects periodically when waiting for IsCachedCookedPlatformDatas.
#rnx

diff --git a/Engine/Source/Editor/UnrealEd/Private/CookOnTheFlyServer.cpp b/Engine/Source/Editor/UnrealEd/Private/CookOnTheFlyServer.cpp
--- a/Engine/Source/Editor/UnrealEd/Private/CookOnTheFlyServer.cpp
+++ b/Engine/Source/Editor/UnrealEd/Private/CookOnTheFlyServer.cpp
@@ -10403,6 +10403,14 @@ uint32 UCookOnTheFlyServer::CookFullLoadAndSave()
 			}
 
 			FAssetCompilingManager::Get().ProcessAsyncTasks(true);
+			double CurrentTime = FPlatformTime::Seconds();
+			if (LastCookableObjectTickTime + TickCookableObjectsFrameTime <= CurrentTime)
+			{
+				UE_SCOPED_COOKTIMER(TickCookableObjects);
+				FTickableCookObject::TickObjects(CurrentTime - LastCookableObjectTickTime, false /* bTickComplete */);
+				LastCookableObjectTickTime = CurrentTime;
+			}
+
 			FPlatformProcess::Sleep(0.001f);
 		}
 
