Added option in WorldPartitionConvertCommandlet to set an existing default HLOD layer when converting a world.

diff --git a/Engine/Source/Editor/UnrealEd/Classes/Commandlets/WorldPartitionConvertCommandlet.h b/Engine/Source/Editor/UnrealEd/Classes/Commandlets/WorldPartitionConvertCommandlet.h
--- a/Engine/Source/Editor/UnrealEd/Classes/Commandlets/WorldPartitionConvertCommandlet.h
+++ b/Engine/Source/Editor/UnrealEd/Classes/Commandlets/WorldPartitionConvertCommandlet.h
@@ -122,6 +122,9 @@ protected:
 	UPROPERTY(Config)
 	FString DefaultHLODLayerName;
 
+	UPROPERTY(Config)
+	FString DefaultHLODLayerAsset;
+
 	UPROPERTY(Config)
 	FString FoliageTypePath;
 
diff --git a/Engine/Source/Editor/UnrealEd/Private/Commandlets/WorldPartitionConvertCommandlet.cpp b/Engine/Source/Editor/UnrealEd/Private/Commandlets/WorldPartitionConvertCommandlet.cpp
--- a/Engine/Source/Editor/UnrealEd/Private/Commandlets/WorldPartitionConvertCommandlet.cpp
+++ b/Engine/Source/Editor/UnrealEd/Private/Commandlets/WorldPartitionConvertCommandlet.cpp
@@ -235,7 +235,15 @@ UWorldPartition* UWorldPartitionConvertCommandlet::CreateWorldPartition(AWorldSe
 	{
 		WorldPartition->EditorHash->LoadConfig(*EditorHashClass, *LevelConfigFilename);
 		WorldPartition->RuntimeHash->LoadConfig(*RuntimeHashClass, *LevelConfigFilename);
-		WorldPartition->DefaultHLODLayer = HLODLayers.FindRef(DefaultHLODLayerName);
+		// Use specified existing default HLOD layer if valid 
+		if (UHLODLayer* ExistingHLODLayer = LoadObject<UHLODLayer>(NULL, *DefaultHLODLayerAsset))
+		{
+			WorldPartition->DefaultHLODLayer = ExistingHLODLayer;
+		}
+		else
+		{
+			WorldPartition->DefaultHLODLayer = HLODLayers.FindRef(DefaultHLODLayerName);
+		}
 	}
 
 	if ((WorldPartition->DefaultHLODLayer == UHLODLayer::GetEngineDefaultHLODLayersSetup()) && !bDisableStreaming)
