SummarizeTraceCommandlet: Fix a crash if ModuleStats is empty and refactored Main() into smaller functions for the individual csvs.

diff --git a/Engine/Source/Editor/UnrealEd/Classes/Commandlets/SummarizeTraceCommandlet.h b/Engine/Source/Editor/UnrealEd/Classes/Commandlets/SummarizeTraceCommandlet.h
--- a/Engine/Source/Editor/UnrealEd/Classes/Commandlets/SummarizeTraceCommandlet.h
+++ b/Engine/Source/Editor/UnrealEd/Classes/Commandlets/SummarizeTraceCommandlet.h
@@ -7,6 +7,7 @@
 
 #pragma once
 #include "Commandlets/Commandlet.h"
+#include "TraceServices/Model/Counters.h"
 #include "SummarizeTraceCommandlet.generated.h"
 
 UCLASS(config=Editor)
@@ -18,6 +19,17 @@ public:
 	//~ Begin UCommandlet Interface
 
 	virtual int32 Main(const FString& CmdLineParams) override;
-	
+
 	//~ End UCommandlet Interface
+private:
+	void WriteAsUTF8String(IFileHandle* Handle, const FString& String);
+	TUniquePtr<IFileHandle> OpenCSVFile(const FString& Name);
+	bool GenerateScopesCSV(const TArray<struct FSummarizeScope>& SortedScopes);
+	bool GenerateCountersCSV(const class FSummarizeCountersProvider& CountersProvider);
+	bool GenerateBookmarksCSV(const class FSummarizeBookmarksProvider& BookmarksProvider);
+	bool GenerateTelemetryCSV(const FString& TestName, bool bAllTelemetry, const TArray<FSummarizeScope>& SortedScopes, const
+	                          FSummarizeCountersProvider& CountersProvider, const FSummarizeBookmarksProvider& BookmarksProvider, bool SkipBaseline);
+
+	FString TracePath;
+	FString TraceFileBasename;
 };
diff --git a/Engine/Source/Editor/UnrealEd/Private/Commandlets/SummarizeTraceCommandlet.cpp b/Engine/Source/Editor/UnrealEd/Private/Commandlets/SummarizeTraceCommandlet.cpp
--- a/Engine/Source/Editor/UnrealEd/Private/Commandlets/SummarizeTraceCommandlet.cpp
+++ b/Engine/Source/Editor/UnrealEd/Private/Commandlets/SummarizeTraceCommandlet.cpp
@@ -1824,29 +1824,7 @@ int32 USummarizeTraceCommandlet::Main(const FString& CmdLineParams)
 		return 1;
 	}
 
-	bool bAllTelemetry = FParse::Param(*CmdLineParams, TEXT("alltelemetry"));
-
-	// load the stats file to know which event name and statistic name to generate in the telemetry csv
-	// the telemetry csv is ingested completely, so this just highlights specific data elements we want to track
-	TMultiMap<FString, StatisticDefinition> NameToDefinitionMap;
-	TSet<FString> CpuScopeNamesWithWildcards;
-	if (!bAllTelemetry)
-	{
-		FString GlobalStatisticsFileName = FPaths::RootDir() / TEXT("Engine") / TEXT("Build") / TEXT("EditorPerfStats.csv");
-		if (FPaths::FileExists(GlobalStatisticsFileName))
-		{
-			UE_LOG(LogSummarizeTrace, Display, TEXT("Loading global statistics from %s"), *GlobalStatisticsFileName);
-			bool bCSVOk = StatisticDefinition::LoadFromCSV(GlobalStatisticsFileName, NameToDefinitionMap, CpuScopeNamesWithWildcards);
-			check(bCSVOk);
-		}
-		FString ProjectStatisticsFileName = FPaths::ProjectDir() / TEXT("Build") / TEXT("EditorPerfStats.csv");
-		if (FPaths::FileExists(ProjectStatisticsFileName))
-		{
-			UE_LOG(LogSummarizeTrace, Display, TEXT("Loading project statistics from %s"), *ProjectStatisticsFileName);
-			bool bCSVOk = StatisticDefinition::LoadFromCSV(ProjectStatisticsFileName, NameToDefinitionMap, CpuScopeNamesWithWildcards);
-		}
-	}
-
+	
 	bool bFound;
 	if (FPaths::FileExists(TraceFileName))
 	{
@@ -1877,6 +1855,7 @@ int32 USummarizeTraceCommandlet::Main(const FString& CmdLineParams)
 		return 1;
 	}
 
+	
 	UE::Trace::FFileDataStream* DataStream = new UE::Trace::FFileDataStream();
 	if (!DataStream->Open(*TraceFileName))
 	{
@@ -1920,7 +1899,7 @@ int32 USummarizeTraceCommandlet::Main(const FString& CmdLineParams)
 			ModuleStats.Sort([](const FSummarizeScope& Lhs, const FSummarizeScope& Rhs) { return Lhs.TotalDurationSeconds >= Rhs.TotalDurationSeconds; });
 
 			// Publish top N longuest load module. The ModuleStats are pre-sorted from the longest to the shorted timer.
-			CollectedScopeSummaries.Append(ModuleStats.GetData(), 10);
+			CollectedScopeSummaries.Append(ModuleStats.GetData(), FMath::Min(10, ModuleStats.Num()));
 		});
 
 	TSharedPtr<TraceServices::IAnalysisSession> Session = TraceServices::CreateAnalysisSession(0, nullptr, TUniquePtr<UE::Trace::IInDataStream>(DataStream));
@@ -1986,122 +1965,179 @@ int32 USummarizeTraceCommandlet::Main(const FString& CmdLineParams)
 		SortedScopes.Add(Scope);
 	}
 	SortedScopes.Sort();
-
-	// csv is UTF-8, so encode every string we print
-	auto WriteUTF8 = [](IFileHandle* Handle, const FString& String)
-	{
-		const auto& UTF8String = StringCast<ANSICHAR>(*String);
-		Handle->Write(reinterpret_cast<const uint8*>(UTF8String.Get()), UTF8String.Length());
-	};
+	
 
 	// some locals to help with all the derived files we are about to generate
-	const FString TracePath = FPaths::GetPath(TraceFileName);
-	const FString TraceFileBasename = FPaths::GetBaseFilename(TraceFileName);
+	TracePath = FPaths::GetPath(TraceFileName);
+	TraceFileBasename = FPaths::GetBaseFilename(TraceFileName);
 
 	// generate a summary csv files, always
-	FString CsvFileName = TraceFileBasename + TEXT("Scopes");
-	CsvFileName = FPaths::Combine(TracePath, FPaths::SetExtension(CsvFileName, "csv"));
-	UE_LOG(LogSummarizeTrace, Display, TEXT("Writing %s..."), *CsvFileName);
-	IFileHandle* CsvHandle = FPlatformFileManager::Get().GetPlatformFile().OpenWrite(*CsvFileName);
-	if (!CsvHandle)
+	if (!GenerateScopesCSV(SortedScopes))
+	{
+		return 1;
+	}		
+	if (!GenerateCountersCSV(CountersProvider))
+	{
+		return 1;
+	}		
+	if (!GenerateBookmarksCSV(BookmarksProvider))
+	{
+		return 1;
+	}		
+
+		
+	// override the test name
+	FString TestName = TraceFileBasename;
+	FParse::Value(*CmdLineParams, TEXT("testname="), TestName, true);
+	
+	bool bAllTelemetry = FParse::Param(*CmdLineParams, TEXT("alltelemetry"));
+	bool SkipBaseline = FParse::Param(*CmdLineParams, TEXT("skipbaseline"));
+	
+	if (!GenerateTelemetryCSV(TestName, bAllTelemetry, SortedScopes, CountersProvider, BookmarksProvider, SkipBaseline))
 	{
-		UE_LOG(LogSummarizeTrace, Error, TEXT("Unable to open csv '%s' for write"), *CsvFileName);
 		return 1;
 	}
-	else
+
+	return 0;
+}
+
+void USummarizeTraceCommandlet::WriteAsUTF8String(IFileHandle* Handle, const FString& String)
+{
+	const auto& UTF8String = StringCast<ANSICHAR>(*String);
+	Handle->Write(reinterpret_cast<const uint8*>(UTF8String.Get()), UTF8String.Length());
+};
+
+TUniquePtr<IFileHandle> USummarizeTraceCommandlet::OpenCSVFile(const FString& Name)
+{
+	FString CsvFileName = TraceFileBasename + Name;
+    CsvFileName = FPaths::Combine(TracePath, FPaths::SetExtension(CsvFileName, "csv"));
+    UE_LOG(LogSummarizeTrace, Display, TEXT("Writing %s..."), *CsvFileName);
+    TUniquePtr<IFileHandle> CsvHandle(FPlatformFileManager::Get().GetPlatformFile().OpenWrite(*CsvFileName));
+
+    if (!CsvHandle)
+    {
+    	UE_LOG(LogSummarizeTrace, Error, TEXT("Unable to open csv '%s' for write"), *CsvFileName);
+    }
+
+	return CsvHandle;
+}
+
+bool USummarizeTraceCommandlet::GenerateScopesCSV(const TArray<FSummarizeScope>& SortedScopes)
+{
+	TUniquePtr<IFileHandle> CsvHandle = OpenCSVFile(TEXT("Scopes"));
+	if (!CsvHandle)
 	{
-		// no newline, see row printfs
-		WriteUTF8(CsvHandle, FString::Printf(TEXT("Name,Count,CountPerSecond,TotalDurationSeconds,FirstStartSeconds,FirstFinishSeconds,FirstDurationSeconds,LastStartSeconds,LastFinishSeconds,LastDurationSeconds,MinDurationSeconds,MaxDurationSeconds,MeanDurationSeconds,DeviationDurationSeconds,")));
-		for (const FSummarizeScope& Scope : SortedScopes)
+		return false;
+	}
+	
+	// no newline, see row printfs
+	WriteAsUTF8String(CsvHandle.Get(), FString::Printf(TEXT("Name,Count,CountPerSecond,TotalDurationSeconds,FirstStartSeconds,FirstFinishSeconds,FirstDurationSeconds,LastStartSeconds,LastFinishSeconds,LastDurationSeconds,MinDurationSeconds,MaxDurationSeconds,MeanDurationSeconds,DeviationDurationSeconds,")));
+	for (const FSummarizeScope& Scope : SortedScopes)
+	{
+		if (!IsCsvSafeString(Scope.Name))
 		{
-			if (!IsCsvSafeString(Scope.Name))
-			{
-				continue;
-			}
-
-			// note newline is at the front of every data line to prevent final extraneous newline, per customary for csv
-			WriteUTF8(CsvHandle, FString::Printf(TEXT("\n%s,%llu,%f,%f,%f,%f,%f,%f,%f,%f,%f,%f,%f,%f,"),
-				*Scope.Name, Scope.GetCount(), Scope.GetCountPerSecond(),
-				Scope.TotalDurationSeconds,
-				Scope.FirstStartSeconds, Scope.FirstFinishSeconds, Scope.FirstDurationSeconds,
-				Scope.LastStartSeconds, Scope.LastFinishSeconds, Scope.LastDurationSeconds,
-				Scope.MinDurationSeconds, Scope.MaxDurationSeconds,
-				Scope.GetMeanDurationSeconds(), Scope.GetDeviationDurationSeconds()));
+			continue;
 		}
-		CsvHandle->Flush();
-		delete CsvHandle;
-		CsvHandle = nullptr;
+
+		// note newline is at the front of every data line to prevent final extraneous newline, per customary for csv
+		WriteAsUTF8String(CsvHandle.Get(), FString::Printf(TEXT("\n%s,%llu,%f,%f,%f,%f,%f,%f,%f,%f,%f,%f,%f,%f,"),
+			*Scope.Name, Scope.GetCount(), Scope.GetCountPerSecond(),
+			Scope.TotalDurationSeconds,
+			Scope.FirstStartSeconds, Scope.FirstFinishSeconds, Scope.FirstDurationSeconds,
+			Scope.LastStartSeconds, Scope.LastFinishSeconds, Scope.LastDurationSeconds,
+			Scope.MinDurationSeconds, Scope.MaxDurationSeconds,
+			Scope.GetMeanDurationSeconds(), Scope.GetDeviationDurationSeconds()));
 	}
+	CsvHandle->Flush();
+
+	return true;
+}
 
-	CsvFileName = TraceFileBasename + TEXT("Counters");
-	CsvFileName = FPaths::Combine(TracePath, FPaths::SetExtension(CsvFileName, "csv"));
-	UE_LOG(LogSummarizeTrace, Display, TEXT("Writing %s..."), *CsvFileName);
-	CsvHandle = FPlatformFileManager::Get().GetPlatformFile().OpenWrite(*CsvFileName);
+bool USummarizeTraceCommandlet::GenerateCountersCSV(const FSummarizeCountersProvider& CountersProvider)
+{
+	TUniquePtr<IFileHandle> CsvHandle = OpenCSVFile(TEXT("Counters"));
 	if (!CsvHandle)
 	{
-		UE_LOG(LogSummarizeTrace, Error, TEXT("Unable to open csv '%s' for write"), *CsvFileName);
-		return 1;
+		return false;
 	}
-	else
+	
+	// no newline, see row printfs
+	WriteAsUTF8String(CsvHandle.Get(), FString::Printf(TEXT("Name,Count,CountPerSecond,First,FirstSeconds,Last,LastSeconds,Minimum,Maximum,Mean,Deviation,")));
+	for (const TUniquePtr<FSummarizeCounter>& Counter : CountersProvider.Counters)
 	{
-		// no newline, see row printfs
-		WriteUTF8(CsvHandle, FString::Printf(TEXT("Name,Count,CountPerSecond,First,FirstSeconds,Last,LastSeconds,Minimum,Maximum,Mean,Deviation,")));
-		for (const TUniquePtr<FSummarizeCounter>& Counter : CountersProvider.Counters)
+		if (!IsCsvSafeString(Counter->Name))
 		{
-			if (!IsCsvSafeString(Counter->Name))
-			{
-				continue;
-			}
-
-			// note newline is at the front of every data line to prevent final extraneous newline, per customary for csv
-			WriteUTF8(CsvHandle, FString::Printf(TEXT("\n%s,%llu,%f,%s,%f,%s,%f,%s,%s,%f,%f,"),
-				*Counter->Name, Counter->GetCount(), Counter->GetCountPerSecond(),
-				*Counter->GetValue(TEXT("First")), Counter->FirstSeconds,
-				*Counter->GetValue(TEXT("Last")), Counter->LastSeconds,
-				*Counter->GetValue(TEXT("Minimum")), *Counter->GetValue(TEXT("Maximum")),
-				Counter->GetMean(), Counter->GetDeviation()));
+			continue;
 		}
-		CsvHandle->Flush();
-		delete CsvHandle;
-		CsvHandle = nullptr;
+
+		// note newline is at the front of every data line to prevent final extraneous newline, per customary for csv
+		WriteAsUTF8String(CsvHandle.Get(), FString::Printf(TEXT("\n%s,%llu,%f,%s,%f,%s,%f,%s,%s,%f,%f,"),
+			*Counter->Name, Counter->GetCount(), Counter->GetCountPerSecond(),
+			*Counter->GetValue(TEXT("First")), Counter->FirstSeconds,
+			*Counter->GetValue(TEXT("Last")), Counter->LastSeconds,
+			*Counter->GetValue(TEXT("Minimum")), *Counter->GetValue(TEXT("Maximum")),
+			Counter->GetMean(), Counter->GetDeviation()));
 	}
+	CsvHandle->Flush();
+
+	return true;
+}
 
-	CsvFileName = TraceFileBasename + TEXT("Bookmarks");
-	CsvFileName = FPaths::Combine(TracePath, FPaths::SetExtension(CsvFileName, "csv"));
-	UE_LOG(LogSummarizeTrace, Display, TEXT("Writing %s..."), *CsvFileName);
-	CsvHandle = FPlatformFileManager::Get().GetPlatformFile().OpenWrite(*CsvFileName);
+bool USummarizeTraceCommandlet::GenerateBookmarksCSV(const FSummarizeBookmarksProvider& BookmarksProvider)
+{
+	TUniquePtr<IFileHandle> CsvHandle = OpenCSVFile(TEXT("Bookmarks"));
 	if (!CsvHandle)
 	{
-		UE_LOG(LogSummarizeTrace, Error, TEXT("Unable to open csv '%s' for write"), *CsvFileName);
-		return 1;
+		return false;
 	}
-	else
+	
+	// no newline, see row printfs
+	WriteAsUTF8String(CsvHandle.Get(), FString::Printf(TEXT("Name,Count,FirstSeconds,LastSeconds,")));
+	for (const TMap<FString, FSummarizeBookmark>::ElementType& Bookmark : BookmarksProvider.Bookmarks)
 	{
-		// no newline, see row printfs
-		WriteUTF8(CsvHandle, FString::Printf(TEXT("Name,Count,FirstSeconds,LastSeconds,")));
-		for (const TMap<FString, FSummarizeBookmark>::ElementType& Bookmark : BookmarksProvider.Bookmarks)
+		if (!IsCsvSafeString(Bookmark.Value.Name))
 		{
-			if (!IsCsvSafeString(Bookmark.Value.Name))
-			{
-				continue;
-			}
-
-			// note newline is at the front of every data line to prevent final extraneous newline, per customary for csv
-			WriteUTF8(CsvHandle, FString::Printf(TEXT("\n%s,%d,%f,%f,"),
-				*Bookmark.Value.Name, Bookmark.Value.Count,
-				Bookmark.Value.FirstSeconds, Bookmark.Value.LastSeconds));
+			continue;
 		}
-		CsvHandle->Flush();
-		delete CsvHandle;
-		CsvHandle = nullptr;
+
+		// note newline is at the front of every data line to prevent final extraneous newline, per customary for csv
+		WriteAsUTF8String(CsvHandle.Get(), FString::Printf(TEXT("\n%s,%d,%f,%f,"),
+			*Bookmark.Value.Name, Bookmark.Value.Count,
+			Bookmark.Value.FirstSeconds, Bookmark.Value.LastSeconds));
 	}
+	CsvHandle->Flush();
+
+	return true;
+}
 
-	FString TelemetryCsvFileName = TraceFileBasename + TEXT("Telemetry");
-	TelemetryCsvFileName = FPaths::Combine(TracePath, FPaths::SetExtension(TelemetryCsvFileName, "csv"));
+bool USummarizeTraceCommandlet::GenerateTelemetryCSV(const FString& TestName,
+	bool bAllTelemetry,
+	const TArray<FSummarizeScope>& SortedScopes,
+	const FSummarizeCountersProvider& CountersProvider,
+	const FSummarizeBookmarksProvider& BookmarksProvider,
+	bool SkipBaseline)
+{
+	// load the stats file to know which event name and statistic name to generate in the telemetry csv
+	// the telemetry csv is ingested completely, so this just highlights specific data elements we want to track
+	TMultiMap<FString, StatisticDefinition> NameToDefinitionMap;
+	TSet<FString> CpuScopeNamesWithWildcards;
+	if (!bAllTelemetry)
+	{
+		FString GlobalStatisticsFileName = FPaths::RootDir() / TEXT("Engine") / TEXT("Build") / TEXT("EditorPerfStats.csv");
+		if (FPaths::FileExists(GlobalStatisticsFileName))
+		{
+			UE_LOG(LogSummarizeTrace, Display, TEXT("Loading global statistics from %s"), *GlobalStatisticsFileName);
+			bool bCSVOk = StatisticDefinition::LoadFromCSV(GlobalStatisticsFileName, NameToDefinitionMap, CpuScopeNamesWithWildcards);
+			check(bCSVOk);
+		}
+		FString ProjectStatisticsFileName = FPaths::ProjectDir() / TEXT("Build") / TEXT("EditorPerfStats.csv");
+		if (FPaths::FileExists(ProjectStatisticsFileName))
+		{
+			UE_LOG(LogSummarizeTrace, Display, TEXT("Loading project statistics from %s"), *ProjectStatisticsFileName);
+			bool bCSVOk = StatisticDefinition::LoadFromCSV(ProjectStatisticsFileName, NameToDefinitionMap, CpuScopeNamesWithWildcards);
+		}
+	}
 
-	// override the test name
-	FString TestName = TraceFileBasename;
-	FParse::Value(*CmdLineParams, TEXT("testname="), TestName, true);
 
 	TArray<TelemetryDefinition> TelemetryData;
 	TSet<FString> ResolvedStatistics;
@@ -2226,13 +2262,13 @@ int32 USummarizeTraceCommandlet::Main(const FString& CmdLineParams)
 		// compare vs. baseline telemetry file, if it exists
 		// note this does assume that the tracefile basename is directly comparable to a file in the baseline folder
 		FString BaselineTelemetryCsvFilePath = FPaths::Combine(FPaths::EngineDir(), TEXT("Build"), TEXT("Baseline"), FPaths::SetExtension(TraceFileBasename + TEXT("Telemetry"), "csv"));
-		if (FParse::Param(*CmdLineParams, TEXT("skipbaseline")))
+		if (SkipBaseline)
 		{
 			BaselineTelemetryCsvFilePath.Empty();
 		}
 		if (FPaths::FileExists(BaselineTelemetryCsvFilePath))
 		{
-			UE_LOG(LogSummarizeTrace, Display, TEXT("Comparing telemetry to baseline telemetry %s..."), *TelemetryCsvFileName);
+			UE_LOG(LogSummarizeTrace, Display, TEXT("Comparing telemetry to baseline telemetry %s..."), *BaselineTelemetryCsvFilePath);
 
 			// each context (scope name or coutner name) and data point (statistic name) pair form a key, an item to check
 			TMap<TPair<FString, FString>, TelemetryDefinition> ContextAndDataPointToDefinitionMap;
@@ -2333,31 +2369,26 @@ int32 USummarizeTraceCommandlet::Main(const FString& CmdLineParams)
 
 		if (!NameToDefinitionMap.IsEmpty())
 		{
-			return 1;
+			UE_LOG(LogSummarizeTrace, Error, TEXT("Exiting..."));
+			return false;
 		}
 	}
 
-	UE_LOG(LogSummarizeTrace, Display, TEXT("Writing telemetry to %s..."), *TelemetryCsvFileName);
-	IFileHandle* TelemetryCsvHandle = FPlatformFileManager::Get().GetPlatformFile().OpenWrite(*TelemetryCsvFileName);
-	if (TelemetryCsvHandle)
-	{
-		// no newline, see row printfs
-		WriteUTF8(TelemetryCsvHandle, FString::Printf(TEXT("TestName,Context,DataPoint,Unit,Measurement,Baseline,")));
-		for (const TelemetryDefinition& Telemetry : TelemetryData)
-		{
-			// note newline is at the front of every data line to prevent final extraneous newline, per customary for csv
-			WriteUTF8(TelemetryCsvHandle, FString::Printf(TEXT("\n%s,%s,%s,%s,%s,%s,"), *Telemetry.TestName, *Telemetry.Context, *Telemetry.DataPoint, *Telemetry.Unit, *Telemetry.Measurement, *Telemetry.Baseline));
-		}
 
-		TelemetryCsvHandle->Flush();
-		delete TelemetryCsvHandle;
-		TelemetryCsvHandle = nullptr;
+	TUniquePtr<IFileHandle> TelemetryCsvHandle = OpenCSVFile(TEXT("Telemetry"));
+	if (!TelemetryCsvHandle)
+	{
+		return false;
 	}
-	else
+	
+	// no newline, see row printfs
+	WriteAsUTF8String(TelemetryCsvHandle.Get(), FString::Printf(TEXT("TestName,Context,DataPoint,Unit,Measurement,Baseline,")));
+	for (const TelemetryDefinition& Telemetry : TelemetryData)
 	{
-		UE_LOG(LogSummarizeTrace, Error, TEXT("Unable to open telemetry csv '%s' for write"), *TelemetryCsvFileName);
-		return 1;
+		// note newline is at the front of every data line to prevent final extraneous newline, per customary for csv
+		WriteAsUTF8String(TelemetryCsvHandle.Get(), FString::Printf(TEXT("\n%s,%s,%s,%s,%s,%s,"), *Telemetry.TestName, *Telemetry.Context, *Telemetry.DataPoint, *Telemetry.Unit, *Telemetry.Measurement, *Telemetry.Baseline));
 	}
 
-	return 0;
+	TelemetryCsvHandle->Flush();
+	return true;
 }
\ No newline at end of file
