Additional menu filtering for BP-based factories

diff --git a/Engine/Plugins/Animation/LiveLink/Source/LiveLinkEditor/Private/VirtualSubjects/LiveLinkBlueprintVirtualSubjectFactory.cpp b/Engine/Plugins/Animation/LiveLink/Source/LiveLinkEditor/Private/VirtualSubjects/LiveLinkBlueprintVirtualSubjectFactory.cpp
--- a/Engine/Plugins/Animation/LiveLink/Source/LiveLinkEditor/Private/VirtualSubjects/LiveLinkBlueprintVirtualSubjectFactory.cpp
+++ b/Engine/Plugins/Animation/LiveLink/Source/LiveLinkEditor/Private/VirtualSubjects/LiveLinkBlueprintVirtualSubjectFactory.cpp
@@ -219,13 +219,13 @@ bool ULiveLinkBlueprintVirtualSubjectFactory::ConfigureProperties()
 	return Role != nullptr;
 }
 
-UObject* ULiveLinkBlueprintVirtualSubjectFactory::FactoryCreateNew(UClass* Class, UObject* InParent, FName Name, EObjectFlags Flags, UObject* Context, FFeedbackContext* Warn)
+UObject* ULiveLinkBlueprintVirtualSubjectFactory::FactoryCreateNew(UClass* Class, UObject* InParent, FName Name, EObjectFlags Flags, UObject* Context, FFeedbackContext* Warn, FName CallingContext)
 {
 	UBlueprint* VirtualSubject = nullptr;
 	if (Role && ensure(SupportedClass == Class))
 	{
 		ensure(0 != (RF_Public & Flags));
-		VirtualSubject = FKismetEditorUtilities::CreateBlueprint(ParentClass, InParent, Name, BPTYPE_Normal, UBlueprint::StaticClass(), UBlueprintGeneratedClass::StaticClass(), NAME_None);
+		VirtualSubject = FKismetEditorUtilities::CreateBlueprint(ParentClass, InParent, Name, BlueprintType, UBlueprint::StaticClass(), UBlueprintGeneratedClass::StaticClass(), CallingContext);
 		if (TSubclassOf<UObject> GeneratedClass = VirtualSubject->GeneratedClass)
 		{
 			if (ULiveLinkBlueprintVirtualSubject* DefaultSubject = GeneratedClass->GetDefaultObject<ULiveLinkBlueprintVirtualSubject>())
diff --git a/Engine/Plugins/Animation/LiveLink/Source/LiveLinkEditor/Public/VirtualSubjects/LiveLinkBlueprintVirtualSubjectFactory.h b/Engine/Plugins/Animation/LiveLink/Source/LiveLinkEditor/Public/VirtualSubjects/LiveLinkBlueprintVirtualSubjectFactory.h
--- a/Engine/Plugins/Animation/LiveLink/Source/LiveLinkEditor/Public/VirtualSubjects/LiveLinkBlueprintVirtualSubjectFactory.h
+++ b/Engine/Plugins/Animation/LiveLink/Source/LiveLinkEditor/Public/VirtualSubjects/LiveLinkBlueprintVirtualSubjectFactory.h
@@ -3,7 +3,7 @@
 #pragma once
 
 #include "CoreMinimal.h"
-#include "Factories/Factory.h"
+#include "Factories/BlueprintFactory.h"
 #include "UObject/ObjectMacros.h"
 
 #include "LiveLinkBlueprintVirtualSubjectFactory.generated.h"
@@ -11,16 +11,12 @@
 class ULiveLinkRole;
 
 UCLASS(hidecategories = Object)
-class LIVELINKEDITOR_API ULiveLinkBlueprintVirtualSubjectFactory : public UFactory
+class LIVELINKEDITOR_API ULiveLinkBlueprintVirtualSubjectFactory : public UBlueprintFactory
 {
 	GENERATED_BODY()
 	
 public:
 	ULiveLinkBlueprintVirtualSubjectFactory();
-	
-	// The parent class of the created blueprint
-	UPROPERTY(EditAnywhere, Category="BlueprintVirtualSubjectFactory", meta=(AllowAbstract = "", BlueprintBaseOnly = ""))
-	TSubclassOf<UObject> ParentClass;
 
 	UPROPERTY(BlueprintReadWrite, Category = "Live Link Blueprint Virtual Subject Factory")
 	TSubclassOf<ULiveLinkRole> Role;
@@ -28,7 +24,7 @@ public:
 	//~ Begin UFactory Interface
 	virtual FText GetDisplayName() const override;
 	virtual bool ConfigureProperties() override;
-	virtual UObject* FactoryCreateNew(UClass* Class, UObject* InParent, FName Name, EObjectFlags Flags, UObject* Context, FFeedbackContext* Warn) override;
+	virtual UObject* FactoryCreateNew(UClass* Class, UObject* InParent, FName Name, EObjectFlags Flags, UObject* Context, FFeedbackContext* Warn, FName CallingContext) override;
 	virtual bool ShouldShowInNewMenu() const override;
 	virtual uint32 GetMenuCategories() const override;
 	//~ Begin UFactory Interface
diff --git a/Engine/Source/Developer/AssetTools/AssetTools.Build.cs b/Engine/Source/Developer/AssetTools/AssetTools.Build.cs
--- a/Engine/Source/Developer/AssetTools/AssetTools.Build.cs
+++ b/Engine/Source/Developer/AssetTools/AssetTools.Build.cs
@@ -36,6 +36,7 @@ public class AssetTools : ModuleRules
 				"ToolMenus",
 				"PhysicsCore",
 				"DeveloperSettings",
+				"ClassViewer",
 				"EngineSettings",
 				"InterchangeCore",
 				"InterchangeEngine",
diff --git a/Engine/Source/Developer/AssetTools/Private/AssetTools.cpp b/Engine/Source/Developer/AssetTools/Private/AssetTools.cpp
--- a/Engine/Source/Developer/AssetTools/Private/AssetTools.cpp
+++ b/Engine/Source/Developer/AssetTools/Private/AssetTools.cpp
@@ -2,6 +2,7 @@
 
 #include "AssetTools.h"
 #include "Factories/Factory.h"
+#include "Factories/BlueprintFactory.h"
 #include "Misc/MessageDialog.h"
 #include "HAL/FileManager.h"
 #include "Misc/FileHelper.h"
@@ -29,6 +30,8 @@
 #include "Framework/MultiBox/MultiBoxBuilder.h"
 #include "ToolMenus.h"
 #include "IClassTypeActions.h"
+#include "ClassViewerModule.h"
+#include "ClassViewerFilter.h"
 #include "AssetTypeActions/AssetTypeActions_Actor.h"
 #include "AssetTypeActions/AssetTypeActions_ActorFolder.h"
 #include "AssetTypeActions/AssetTypeActions_Blueprint.h"
@@ -3980,18 +3983,44 @@ TArray<UFactory*> UAssetToolsImpl::GetNewAssetFactories() const
 	for (TObjectIterator<UClass> It; It; ++It)
 	{
 		UClass* Class = *It;
-		if (Class->IsChildOf(UFactory::StaticClass()) &&
-			!Class->HasAnyClassFlags(CLASS_Abstract))
+		if (!Class->IsChildOf(UFactory::StaticClass()) || Class->HasAnyClassFlags(CLASS_Abstract))
 		{
-			UFactory* Factory = Class->GetDefaultObject<UFactory>();
+			continue;
+		}
+
+		UFactory* Factory = Class->GetDefaultObject<UFactory>();
+		if (!Factory->ShouldShowInNewMenu() || !ensure(!Factory->GetDisplayName().IsEmpty()) || !IsAssetClassSupported(Factory->GetSupportedClass()))
+		{
+			continue;
+		}
 
-			if (Factory->ShouldShowInNewMenu() &&
-				ensure(!Factory->GetDisplayName().IsEmpty()) &&
-				IsAssetClassSupported(Factory->GetSupportedClass()))
+		// For Blueprints, add sub-type filtering using the BP Editor's allow list
+		// Otherwise, there's no way to distinguish between generic BP Actors vs BlueprintFunctionLibraries, etc.
+		if (UBlueprintFactory* BPFactory = Cast<UBlueprintFactory>(Factory))
+		{
+			// Restrict BP factories based on their BlueprintType
+			if (AllowedBlueprintTypes.Num() > 0 && !AllowedBlueprintTypes.Contains(BPFactory->BlueprintType))
 			{
-				Factories.Add(Factory);
+				continue;
+			}
+
+			// BPFactor->ParentClass is stored as a UObject instead of a class... most things should be UClass but we'll just skip it if not
+			if (UClass* BPParentClass = Cast<UClass>(BPFactory->ParentClass))
+			{
+				FClassViewerModule& ClassViewerModule = FModuleManager::LoadModuleChecked<FClassViewerModule>("ClassViewer");
+				if (const TSharedPtr<IClassViewerFilter>& GlobalClassFilter = ClassViewerModule.GetGlobalClassViewerFilter())
+				{
+					TSharedRef<FClassViewerFilterFuncs> ClassFilterFuncs = ClassViewerModule.CreateFilterFuncs();
+					FClassViewerInitializationOptions ClassViewerOptions = {};
+					if (!GlobalClassFilter->IsClassAllowed(ClassViewerOptions, BPParentClass, ClassFilterFuncs))
+					{
+						continue;
+					}
+				}
 			}
 		}
+
+		Factories.Add(Factory);
 	}
 
 	return MoveTemp(Factories);
@@ -4062,6 +4091,11 @@ TSharedRef<FPathPermissionList>& UAssetToolsImpl::GetAssetClassPathPermissionLis
 	return Empty;
 }
 
+TSet<EBlueprintType>& UAssetToolsImpl::GetAllowedBlueprintTypes()
+{
+	return AllowedBlueprintTypes;
+}
+
 void UAssetToolsImpl::AssetClassPermissionListChanged(EAssetClassAction AssetClassAction)
 {
 	switch (AssetClassAction)
diff --git a/Engine/Source/Developer/AssetTools/Private/AssetTools.h b/Engine/Source/Developer/AssetTools/Private/AssetTools.h
--- a/Engine/Source/Developer/AssetTools/Private/AssetTools.h
+++ b/Engine/Source/Developer/AssetTools/Private/AssetTools.h
@@ -128,6 +128,7 @@ public:
 	UE_DEPRECATED(5.1, "Class names are now represented by path names. Please use GetAssetClassPathPermissionList.")
 	TSharedRef<FNamePermissionList>& GetAssetClassPermissionList(EAssetClassAction AssetClassAction);
 	virtual TSharedRef<FPathPermissionList>& GetAssetClassPathPermissionList(EAssetClassAction AssetClassAction) override;
+	virtual TSet<EBlueprintType>& GetAllowedBlueprintTypes() override;
 	virtual TSharedRef<FPathPermissionList>& GetFolderPermissionList() override;
 	virtual TSharedRef<FPathPermissionList>& GetWritableFolderPermissionList() override;
 	virtual bool AllPassWritableFolderFilter(const TArray<FString>& InPaths) const override;
@@ -224,6 +225,9 @@ private:
 	/** Permission lists of assets by class path name, one for each EAssetClassAction */
 	TArray<TSharedRef<FPathPermissionList>> AssetClassPermissionList;
 
+	/** Which types of BlueprintFactories are allowed to be returned through GetNewAssetFactories. Empty set allows all types. */
+	TSet<EBlueprintType> AllowedBlueprintTypes;
+
 	/** Permission list of folder paths */
 	TSharedRef<FPathPermissionList> FolderPermissionList;
 
diff --git a/Engine/Source/Developer/AssetTools/Public/IAssetTools.h b/Engine/Source/Developer/AssetTools/Public/IAssetTools.h
--- a/Engine/Source/Developer/AssetTools/Public/IAssetTools.h
+++ b/Engine/Source/Developer/AssetTools/Public/IAssetTools.h
@@ -15,7 +15,7 @@
 #include "AssetRegistry/ARFilter.h"
 #include "IAssetTools.generated.h"
 
-
+enum EBlueprintType;
 struct FAssetData;
 class IAssetTypeActions;
 class IClassTypeActions;
@@ -478,6 +478,9 @@ public:
 	/** Get asset class permission list for content browser and other systems */
 	virtual TSharedRef<FPathPermissionList>& GetAssetClassPathPermissionList(EAssetClassAction AssetClassAction) = 0;
 
+	/** Which BlueprintTypes are allowed to be created. An empty list should allow everything. */
+	virtual TSet<EBlueprintType>& GetAllowedBlueprintTypes() = 0;
+
 	/** Get folder permission list for content browser and other systems */
 	virtual TSharedRef<FPathPermissionList>& GetFolderPermissionList() = 0;
 
diff --git a/Engine/Source/Editor/UnrealEd/Classes/Factories/BlueprintFactory.h b/Engine/Source/Editor/UnrealEd/Classes/Factories/BlueprintFactory.h
--- a/Engine/Source/Editor/UnrealEd/Classes/Factories/BlueprintFactory.h
+++ b/Engine/Source/Editor/UnrealEd/Classes/Factories/BlueprintFactory.h
@@ -10,6 +10,7 @@
 #include "UObject/ObjectMacros.h"
 #include "Templates/SubclassOf.h"
 #include "Factories/Factory.h"
+#include "Engine/Blueprint.h"
 #include "BlueprintFactory.generated.h"
 
 class FClassViewerInitializationOptions;
@@ -23,6 +24,10 @@ class UNREALED_API UBlueprintFactory : public UFactory
 	UPROPERTY(EditAnywhere, Category=BlueprintFactory, meta=(AllowAbstract = "", BlueprintBaseOnly = ""))
 	TSubclassOf<class UObject> ParentClass;
 
+	// The type of Blueprint to create in FactoryCreateNew
+	UPROPERTY()
+	TEnumAsByte<EBlueprintType> BlueprintType = EBlueprintType::BPTYPE_Normal;
+
 	// Skips the class choosing dialog and uses the ParentClass as the blueprint base class
 	UPROPERTY()
 	bool bSkipClassPicker = false;
@@ -36,10 +41,6 @@ class UNREALED_API UBlueprintFactory : public UFactory
 	virtual UObject* FactoryCreateNew(UClass* Class, UObject* InParent, FName Name, EObjectFlags Flags, UObject* Context, FFeedbackContext* Warn, FName CallingContext) override;
 	virtual UObject* FactoryCreateNew(UClass* Class, UObject* InParent, FName Name, EObjectFlags Flags, UObject* Context, FFeedbackContext* Warn) override;
 	// End of UFactory Interface
-
-protected:
-	// Hook for derived classes to indicate if they want to skip some filtering that is unnecessary for macro libraries
-	virtual bool IsMacroFactory() const { return false; }
 };
 
 
diff --git a/Engine/Source/Editor/UnrealEd/Classes/Factories/BlueprintInterfaceFactory.h b/Engine/Source/Editor/UnrealEd/Classes/Factories/BlueprintInterfaceFactory.h
--- a/Engine/Source/Editor/UnrealEd/Classes/Factories/BlueprintInterfaceFactory.h
+++ b/Engine/Source/Editor/UnrealEd/Classes/Factories/BlueprintInterfaceFactory.h
@@ -8,11 +8,11 @@
 
 #include "CoreMinimal.h"
 #include "UObject/ObjectMacros.h"
-#include "Factories/Factory.h"
+#include "Factories/BlueprintFactory.h"
 #include "BlueprintInterfaceFactory.generated.h"
 
 UCLASS(hidecategories=Object, collapsecategories)
-class UBlueprintInterfaceFactory : public UFactory
+class UBlueprintInterfaceFactory : public UBlueprintFactory
 {
 	GENERATED_UCLASS_BODY()
 
@@ -23,7 +23,7 @@ class UBlueprintInterfaceFactory : public UFactory
 	virtual FText GetToolTip() const override;
 	virtual FString GetToolTipDocumentationExcerpt() const override;
 	virtual UObject* FactoryCreateNew(UClass* Class, UObject* InParent, FName Name, EObjectFlags Flags, UObject* Context, FFeedbackContext* Warn, FName CallingContext) override;
-	virtual UObject* FactoryCreateNew(UClass* Class, UObject* InParent, FName Name, EObjectFlags Flags, UObject* Context, FFeedbackContext* Warn) override;
+	virtual bool ConfigureProperties() override;
 	virtual FString GetDefaultNewAssetName() const override;
 	// End of UFactory interface
 };
diff --git a/Engine/Source/Editor/UnrealEd/Classes/Factories/BlueprintMacroFactory.h b/Engine/Source/Editor/UnrealEd/Classes/Factories/BlueprintMacroFactory.h
--- a/Engine/Source/Editor/UnrealEd/Classes/Factories/BlueprintMacroFactory.h
+++ b/Engine/Source/Editor/UnrealEd/Classes/Factories/BlueprintMacroFactory.h
@@ -25,9 +25,4 @@ class UBlueprintMacroFactory : public UBlueprintFactory
 	virtual UObject* FactoryCreateNew(UClass* Class, UObject* InParent, FName Name, EObjectFlags Flags, UObject* Context, FFeedbackContext* Warn, FName CallingContext) override;
 	virtual FString GetDefaultNewAssetName() const override;
 	// End of UFactory interface
-
-protected:
-	// UBlueprintFactory interface
-	virtual bool IsMacroFactory() const override { return true; }
-	// End of UBlueprintFactory interface
 };
diff --git a/Engine/Source/Editor/UnrealEd/Private/Factories/EditorFactories.cpp b/Engine/Source/Editor/UnrealEd/Private/Factories/EditorFactories.cpp
--- a/Engine/Source/Editor/UnrealEd/Private/Factories/EditorFactories.cpp
+++ b/Engine/Source/Editor/UnrealEd/Private/Factories/EditorFactories.cpp
@@ -7135,7 +7135,7 @@ bool UBlueprintFactory::ConfigureProperties()
 	// Prevent creating blueprints of classes that require special setup (they'll be allowed in the corresponding factories / via other means)
 	TSharedPtr<FBlueprintParentFilter> Filter = MakeShareable(new FBlueprintParentFilter);
 	Options.ClassFilters.Add(Filter.ToSharedRef());
-	if (!IsMacroFactory())
+	if (BlueprintType == EBlueprintType::BPTYPE_MacroLibrary)
 	{
 		Filter->DisallowedChildrenOfClasses.Add(ALevelScriptActor::StaticClass());
 		Filter->DisallowedChildrenOfClasses.Add(UAnimInstance::StaticClass());
@@ -7181,7 +7181,7 @@ UObject* UBlueprintFactory::FactoryCreateNew(UClass* Class, UObject* InParent, F
 		IKismetCompilerInterface& KismetCompilerModule = FModuleManager::LoadModuleChecked<IKismetCompilerInterface>("KismetCompiler");
 		KismetCompilerModule.GetBlueprintTypesForClass(ParentClass, BlueprintClass, BlueprintGeneratedClass);
 
-		return FKismetEditorUtilities::CreateBlueprint(ParentClass, InParent, Name, BPTYPE_Normal, BlueprintClass, BlueprintGeneratedClass, CallingContext);
+		return FKismetEditorUtilities::CreateBlueprint(ParentClass, InParent, Name, BlueprintType, BlueprintClass, BlueprintGeneratedClass, CallingContext);
 	}
 }
 
@@ -7197,11 +7197,11 @@ UObject* UBlueprintFactory::FactoryCreateNew(UClass* Class, UObject* InParent, F
 UBlueprintMacroFactory::UBlueprintMacroFactory(const FObjectInitializer& ObjectInitializer)
 	: Super(ObjectInitializer)
 {
-	
 	bCreateNew = true;
 	bEditAfterNew = true;
 	SupportedClass = UBlueprint::StaticClass();
 	ParentClass = AActor::StaticClass();
+	BlueprintType = BPTYPE_MacroLibrary;
 }
 
 FText UBlueprintMacroFactory::GetDisplayName() const
@@ -7243,7 +7243,7 @@ UObject* UBlueprintMacroFactory::FactoryCreateNew(UClass* Class, UObject* InPare
 	}
 	else
 	{
-		return FKismetEditorUtilities::CreateBlueprint(ParentClass, InParent, Name, BPTYPE_MacroLibrary, UBlueprint::StaticClass(), UBlueprintGeneratedClass::StaticClass(), CallingContext);
+		return FKismetEditorUtilities::CreateBlueprint(ParentClass, InParent, Name, BlueprintType, UBlueprint::StaticClass(), UBlueprintGeneratedClass::StaticClass(), CallingContext);
 	}
 }
 
@@ -7264,6 +7264,7 @@ UBlueprintFunctionLibraryFactory::UBlueprintFunctionLibraryFactory(const FObject
 	bEditAfterNew = true;
 	SupportedClass = UBlueprint::StaticClass();
 	ParentClass = UBlueprintFunctionLibrary::StaticClass();
+	BlueprintType = BPTYPE_FunctionLibrary;
 }
 
 FText UBlueprintFunctionLibraryFactory::GetDisplayName() const
@@ -7305,7 +7306,7 @@ UObject* UBlueprintFunctionLibraryFactory::FactoryCreateNew(UClass* Class, UObje
 	}
 	else
 	{
-		return FKismetEditorUtilities::CreateBlueprint(ParentClass, InParent, Name, BPTYPE_FunctionLibrary, UBlueprint::StaticClass(), UBlueprintGeneratedClass::StaticClass(), CallingContext);
+		return FKismetEditorUtilities::CreateBlueprint(ParentClass, InParent, Name, BlueprintType, UBlueprint::StaticClass(), UBlueprintGeneratedClass::StaticClass(), CallingContext);
 	}
 }
 
@@ -7325,10 +7326,11 @@ FString UBlueprintFunctionLibraryFactory::GetDefaultNewAssetName() const
 UBlueprintInterfaceFactory::UBlueprintInterfaceFactory(const FObjectInitializer& ObjectInitializer)
 	: Super(ObjectInitializer)
 {
-
 	bCreateNew = true;
 	bEditAfterNew = true;
 	SupportedClass = UBlueprint::StaticClass();
+	ParentClass = UInterface::StaticClass();
+	BlueprintType = BPTYPE_Interface;
 }
 
 FText UBlueprintInterfaceFactory::GetDisplayName() const
@@ -7361,9 +7363,6 @@ UObject* UBlueprintInterfaceFactory::FactoryCreateNew(UClass* Class, UObject* In
 	// Make sure we are trying to factory a blueprint, then create and init one
 	check(Class->IsChildOf(UBlueprint::StaticClass()));
 
-	// Force the parent class to be UInterface as per original code
-	UClass* ParentClass = UInterface::StaticClass();
-
 	if ((ParentClass == nullptr) || !FKismetEditorUtilities::CanCreateBlueprintOfClass(ParentClass))
 	{
 		FFormatNamedArguments Args;
@@ -7373,13 +7372,13 @@ UObject* UBlueprintInterfaceFactory::FactoryCreateNew(UClass* Class, UObject* In
 	}
 	else
 	{
-		return FKismetEditorUtilities::CreateBlueprint(ParentClass, InParent, Name, BPTYPE_Interface, UBlueprint::StaticClass(), UBlueprintGeneratedClass::StaticClass(), CallingContext);
+		return FKismetEditorUtilities::CreateBlueprint(ParentClass, InParent, Name, BlueprintType, UBlueprint::StaticClass(), UBlueprintGeneratedClass::StaticClass(), CallingContext);
 	}
 }
 
-UObject* UBlueprintInterfaceFactory::FactoryCreateNew(UClass* Class, UObject* InParent, FName Name, EObjectFlags Flags, UObject* Context, FFeedbackContext* Warn)
+bool UBlueprintInterfaceFactory::ConfigureProperties()
 {
-	return FactoryCreateNew(Class, InParent, Name, Flags, Context, Warn, NAME_None);
+	return true;
 }
 
 FString UBlueprintInterfaceFactory::GetDefaultNewAssetName() const
