EditorDomain: Restore message severity to Warning for EditorDomain unable to load a package. The case of a class declared in the packagedigest being not found is not a natural circumstance and indicates a bug will occur on loading the package.
#rnx

diff --git a/Engine/Source/Editor/UnrealEd/Private/EditorDomain/EditorDomain.cpp b/Engine/Source/Editor/UnrealEd/Private/EditorDomain/EditorDomain.cpp
--- a/Engine/Source/Editor/UnrealEd/Private/EditorDomain/EditorDomain.cpp
+++ b/Engine/Source/Editor/UnrealEd/Private/EditorDomain/EditorDomain.cpp
@@ -221,7 +221,7 @@ bool FEditorDomain::TryFindOrAddPackageSource(FScopeLock& ScopeLock, bool& bOutR
 		}
 		return false;
 	default:
-		UE_LOG(LogEditorDomain, Display,
+		UE_LOG(LogEditorDomain, Warning,
 			TEXT("Could not load package %s from EditorDomain; it will be loaded from the WorkspaceDomain: %s"),
 			*WriteToString<256>(PackageName), *PackageDigest.GetStatusString());
 		PackageSource = new FPackageSource();
