Added IsUnattended and !CanEverRender to the reasons not to throttle editor performance.

This is causing slow-down / timeouts in some automation if the window is not in the foreground

diff --git a/Engine/Source/Editor/UnrealEd/Private/EditorEngine.cpp b/Engine/Source/Editor/UnrealEd/Private/EditorEngine.cpp
--- a/Engine/Source/Editor/UnrealEd/Private/EditorEngine.cpp
+++ b/Engine/Source/Editor/UnrealEd/Private/EditorEngine.cpp
@@ -6288,7 +6288,11 @@ bool UEditorEngine::IsAnyViewportRealtime()
 
 bool UEditorEngine::ShouldThrottleCPUUsage() const
 {
-	if (IsRunningCommandlet() || FApp::IsBenchmarking())
+	// Don't throttle here. Benchmarking, automation, and terminal shouldn't throttle 
+	if (IsRunningCommandlet() 
+		|| FApp::IsBenchmarking()
+		|| FApp::IsUnattended()
+		|| !FApp::CanEverRender())
 	{
 		return false;
 	}
