AssetRegistry Cooking: Copy AssetData updates due to loading into the cooked AssetRegistry after each package is saved rather than at the end of the cook.
This makes it easier to find where the updates occur and easier to support incremental updates from other processes during multiprocess cooking.
#rnx

diff --git a/Engine/Source/Editor/UnrealEd/Private/Commandlets/AssetRegistryGenerator.cpp b/Engine/Source/Editor/UnrealEd/Private/Commandlets/AssetRegistryGenerator.cpp
--- a/Engine/Source/Editor/UnrealEd/Private/Commandlets/AssetRegistryGenerator.cpp
+++ b/Engine/Source/Editor/UnrealEd/Private/Commandlets/AssetRegistryGenerator.cpp
@@ -40,7 +40,9 @@
 #include "Settings/ProjectPackagingSettings.h"
 #include "Stats/StatsMisc.h"
 #include "String/ParseTokens.h"
+#include "Templates/AreTypesEqual.h"
 #include "Templates/UniquePtr.h"
+#include "UObject/SoftObjectPath.h"
 
 #if WITH_EDITOR
 #include "HAL/ThreadHeartBeat.h"
@@ -1526,39 +1528,6 @@ void FAssetRegistryGenerator::AddAssetToFileOrderRecursive(const FName& InPackag
 	}
 }
 
-void FAssetRegistryGenerator::UpdateAssetDatas(bool bForceNoFilter)
-{
-	// Keep the existing AssetData->PackageFlags. These flags were set by UpdateAssetDataPackageFlags after saving
-	// each Asset and the existing values are authoritative for the cooked AssetRegistry.
-	// TODO: Remove this end-of-cook updating of the AssetDatas. Update them as we go along,
-	// in UpdateAssetDataPackageFlags, and set the PackageFlags afterwards. Then we won't have to
-	// have the cost of the many map lookups we do here.
-	TSet<FName> SkipNone;
-	TArray<TPair<FName, uint32>> PackageNameAndFlags;
-	PackageNameAndFlags.Reserve(State.GetNumAssets());
-	State.EnumerateAllAssets(SkipNone, [&PackageNameAndFlags](const FAssetData& AssetData)
-		{
-			if (!AssetData.PackageName.IsNone())
-			{
-				PackageNameAndFlags.Emplace(AssetData.PackageName, AssetData.PackageFlags);
-			}
-			return true;
-		});
-
-	FAssetRegistrySerializationOptions SaveOptions;
-	AssetRegistry.InitializeSerializationOptions(SaveOptions, TargetPlatform->IniPlatformName());
-	if (bForceNoFilter)
-	{
-		SaveOptions.DisableFilters();
-	}
-	AssetRegistry.InitializeTemporaryAssetRegistryState(State, SaveOptions, true);
-
-	for (TPair<FName, uint32>& Pair : PackageNameAndFlags)
-	{
-		State.UpdateAssetDataPackageFlags(Pair.Key, Pair.Value);
-	}
-}
-
 bool FAssetRegistryGenerator::SaveAssetRegistry(const FString& SandboxPath, bool bSerializeDevelopmentAssetRegistry, bool bForceNoFilter)
 {
 	UE_LOG(LogAssetRegistryGenerator, Display, TEXT("Saving asset registry v%d."), FAssetRegistryVersion::Type::LatestVersion);
@@ -1580,6 +1549,9 @@ bool FAssetRegistryGenerator::SaveAssetRegistry(const FString& SandboxPath, bool
 		SaveOptions.DisableFilters();
 	}
 
+	// Filter tags from both development and saved registries, keeping only the tags applicable to the platform
+	State.FilterTags(SaveOptions);
+
 	// Then possibly apply previous AssetData and AssetPackageData for packages kept from a previous cook
 	UpdateKeptPackages();
 	UpdateCollectionAssetData();
@@ -2544,7 +2516,7 @@ const FAssetData* FAssetRegistryGenerator::CreateOrFindAssetData(UObject& Object
 	return AssetData;
 }
 
-void FAssetRegistryGenerator::UpdateAssetRegistryPackageData(const UPackage& Package,
+void FAssetRegistryGenerator::UpdateAssetRegistryData(const UPackage& Package,
 	FSavePackageResultStruct& SavePackageResult,
 	FCookTagList&& InArchiveCookTagList
 	)
@@ -2555,11 +2527,22 @@ void FAssetRegistryGenerator::UpdateAssetRegistryPackageData(const UPackage& Pac
 	uint32 NewPackageFlags = 0;
 	FAssetPackageData* AssetPackageData = GetAssetPackageData(PackageName);
 	bool bSaveSucceeded = SavePackageResult.IsSuccessful();
-	if (bSaveSucceeded)
+
+	// Copy latest data for all Assets in the package into the cooked registry. This should be done even
+	// if not successful so that editor-only packages are recorded as well
+	TArray<FAssetData> AssetDatas;
+	constexpr bool bIncludeOnlyDiskAssets = true; // Enumerating memory assets is unnecessary; we waited on AR to update
+	AssetRegistry.GetAssetsByPackageName(PackageName, AssetDatas, bIncludeOnlyDiskAssets,
+		false /* SkipARFilteredAssets */);
+	for (FAssetData& AssetData : AssetDatas)
 	{
-		// Ensure all assets in the package are recorded in the registry
-		CreateOrFindAssetDatas(Package);
+		State.UpdateAssetData(MoveTemp(AssetData), true /* bCreateIfNotExists */);
+	}
+	// Create a record for assets that were created during PostLoad or cooking and are not in the global assetregistry
+	CreateOrFindAssetDatas(Package);
 
+	if (bSaveSucceeded)
+	{
 		// Migrate cook tags over
 		CookTagsToAdd.Append(MoveTemp(InArchiveCookTagList.ObjectToTags));
 		InArchiveCookTagList.Reset();
@@ -2585,7 +2568,7 @@ void FAssetRegistryGenerator::UpdateAssetRegistryPackageData(const UPackage& Pac
 		TEXT("Trying to update asset package flags in package '%s' that does not exist"), *PackageName.ToString());
 }
 
-void FAssetRegistryGenerator::UpdateAssetRegistryPackageData(UE::Cook::FPackageData& PackageData, UE::Cook::FAssetRegistryPackageMessage&& Message)
+void FAssetRegistryGenerator::UpdateAssetRegistryData(UE::Cook::FPackageData& PackageData, UE::Cook::FAssetRegistryPackageMessage&& Message)
 {
 	const FName PackageName = PackageData.GetPackageName();
 	PreviousPackagesToUpdate.Remove(PackageName);
@@ -2593,26 +2576,12 @@ void FAssetRegistryGenerator::UpdateAssetRegistryPackageData(UE::Cook::FPackageD
 	for (FAssetData& AssetData : Message.AssetDatas)
 	{
 		AssetData.PackageFlags = Message.PackageFlags;
-
-		// MPCOOKTODO: We are setting the AssetData onto this->State rather than into the global assetregistry,
-		// because the AR does not currently support updating an AssetData and it is better for performance to
-		// edit this->State. This will not work as is, because FAssetRegistryGenerator::UpdateAssetDatas later
-		// clobbers this->State with AssetDatas from the global AssetRegistry. Changing that for SingleProcess
-		// director cooks is something we want to do anyways because its confusing as is. Make that change so
-		// we don't clobber this function's results, or change this function to write to the global asset registry.
-		const FAssetData* ExistingAssetData = State.GetAssetByObjectPath(AssetData.ObjectPath);
-		if (!ExistingAssetData)
-		{
-			FAssetData* const NewAssetData = new FAssetData(MoveTemp(AssetData));
-			State.AddAssetData(NewAssetData);
-		}
-		else
-		{
-			State.UpdateAssetData(const_cast<FAssetData*>(ExistingAssetData), MoveTemp(AssetData));
-		}
+		State.UpdateAssetData(MoveTemp(AssetData), true /* bCreateIfNotExists */);
 	}
 	FAssetPackageData* AssetPackageData = GetAssetPackageData(PackageName);
 	AssetPackageData->DiskSize = Message.DiskSize;
+
+	CookTagsToAdd.Append(MoveTemp(Message.CookTags));
 }
 
 namespace UE::Cook
@@ -2624,7 +2593,7 @@ FAssetRegistryReporterRemote::FAssetRegistryReporterRemote(FCookWorkerClient& In
 {
 }
 
-void FAssetRegistryReporterRemote::UpdateAssetRegistryPackageData(FPackageData& PackageData, const UPackage& Package,
+void FAssetRegistryReporterRemote::UpdateAssetRegistryData(FPackageData& PackageData, const UPackage& Package,
 	FSavePackageResultStruct& SavePackageResult, FCookTagList&& InArchiveCookTagList)
 {
 	uint32 NewPackageFlags = 0;
@@ -2646,18 +2615,34 @@ void FAssetRegistryReporterRemote::UpdateAssetRegistryPackageData(FPackageData&
 	FAssetRegistryPackageMessage Message;
 	Message.PackageFlags = NewPackageFlags;
 	Message.DiskSize = DiskSize;
-	IAssetRegistry::Get()->GetAssetsByPackageName(PackageData.GetPackageName(), Message.AssetDatas);
-	for (FAssetData& AssetData : Message.AssetDatas)
+
+	// Add to the message all the AssetDatas in the package from the global AssetRegistry
+	constexpr bool bIncludeOnlyDiskAssets = true; // Enumerating memory assets is unnecessary; we waited on AR to update
+	IAssetRegistry::Get()->GetAssetsByPackageName(PackageData.GetPackageName(), Message.AssetDatas,
+		bIncludeOnlyDiskAssets, false /* SkipARFilteredAssets */);
+
+	// Also add AssetDatas for any assets that were created during PostLoad or cooking and are not in the global assetregistry
+	TSet<FName> ExistingAssets;
+	for (const FAssetData& AssetData : Message.AssetDatas)
 	{
-		TArray<FCookTagList::FTagNameValuePair>* ExtraCookTags = InArchiveCookTagList.ObjectToTags.Find(FSoftObjectPath(AssetData.ObjectPath));
-		if (ExtraCookTags)
-		{
-			FAssetDataTagMap NewTags = AssetData.TagsAndValues.CopyMap();
-			AppendCookTagsToTagMap(MoveTemp(*ExtraCookTags), NewTags);
-			AssetData.TagsAndValues = FAssetDataTagMapSharedView(MoveTemp(NewTags));
-		}
+		ExistingAssets.Add(AssetData.AssetName);
 	}
+	ForEachObjectWithOuter(&Package, [&Message, &ExistingAssets](UObject* const Object)
+		{
+			if (Object->IsAsset())
+			{
+				FName AssetName = Object->GetFName();
+				if (!ExistingAssets.Contains(AssetName))
+				{
+					Message.AssetDatas.Emplace(Object, true /* bAllowBlueprintClass */);
+				}
+			}
+		}, /*bIncludeNestedObjects*/ false);
+
+	// Add the cooktags that were recorded during serialization
+	Message.CookTags = MoveTemp(InArchiveCookTagList.ObjectToTags);
 
+	// Send the message to the director
 	PackageData.GetOrAddPackageRemoteResult().AddMessage(PackageData, TargetPlatform, Message);
 }
 
@@ -2671,6 +2656,11 @@ void FAssetRegistryPackageMessage::Write(FCbWriter& Writer, const FPackageData&
 
 	}
 	Writer.EndArray();
+	// We cast TMap<FSoftObjectPath,                     ValueType>
+	//      to TMap<FSoftObjectPathSerializationWrapper, ValueType>
+	// to workaround FSoftObjectPath's implicit constructor. See comment in CompactBinaryTCP.h
+	static_assert(ARE_TYPES_EQUAL(decltype(CookTags)::KeyType, FSoftObjectPath), "Expected KeyType of CookTags to be FSoftObjectPath");
+	Writer << "T" << reinterpret_cast<const TMap<FSoftObjectPathSerializationWrapper, decltype(CookTags)::ValueType>&>(CookTags);
 	Writer << "F" << PackageFlags;
 	Writer << "S" << DiskSize;
 }
@@ -2693,6 +2683,11 @@ bool FAssetRegistryPackageMessage::TryRead(FCbObject&& Object, FPackageData& Pac
 		}
 	}
 
+	static_assert(ARE_TYPES_EQUAL(decltype(CookTags)::KeyType, FSoftObjectPath), "Expected KeyType of CookTags to be FSoftObjectPath");
+	if (!LoadFromCompactBinary(Object["T"], reinterpret_cast<TMap<FSoftObjectPathSerializationWrapper, decltype(CookTags)::ValueType>&>(CookTags)))
+	{
+		return false;
+	}
 	if (!LoadFromCompactBinary(Object["F"], PackageFlags))
 	{
 		return false;
diff --git a/Engine/Source/Editor/UnrealEd/Private/CookOnTheFlyServer.cpp b/Engine/Source/Editor/UnrealEd/Private/CookOnTheFlyServer.cpp
--- a/Engine/Source/Editor/UnrealEd/Private/CookOnTheFlyServer.cpp
+++ b/Engine/Source/Editor/UnrealEd/Private/CookOnTheFlyServer.cpp
@@ -976,7 +976,7 @@ void UCookOnTheFlyServer::GetDependentPackages( const TSet<FName>& RootPackages,
 bool UCookOnTheFlyServer::ContainsMap(const FName& PackageName) const
 {
 	TArray<FAssetData> Assets;
-	ensure(AssetRegistry->GetAssetsByPackageName(PackageName, Assets, true));
+	ensure(AssetRegistry->GetAssetsByPackageName(PackageName, Assets, true /* IncludeOnlyDiskAssets */));
 
 	for (const FAssetData& Asset : Assets)
 	{
@@ -993,7 +993,7 @@ bool UCookOnTheFlyServer::ContainsRedirector(const FName& PackageName, TMap<FNam
 {
 	bool bFoundRedirector = false;
 	TArray<FAssetData> Assets;
-	ensure(AssetRegistry->GetAssetsByPackageName(PackageName, Assets, true));
+	ensure(AssetRegistry->GetAssetsByPackageName(PackageName, Assets, true /* IncludeOnlyDiskAssets */));
 
 	for (const FAssetData& Asset : Assets)
 	{
@@ -5527,8 +5527,10 @@ void FSaveCookedPackageContext::FinishPlatform()
 	// Update asset registry
 	if (COTFS.IsDirectorCookByTheBook())
 	{
+		// Flush the AssetRegisty so any AssetData changes from the save are present
+		COTFS.AssetRegistry->WaitForCompletion();
 		IAssetRegistryReporter& Reporter = *(COTFS.PlatformManager->GetPlatformData(TargetPlatform)->RegistryReporter);
-		Reporter.UpdateAssetRegistryPackageData(PackageData, *Package, SavePackageResult, MoveTemp(*ArchiveCookContext.GetCookTagList()));
+		Reporter.UpdateAssetRegistryData(PackageData, *Package, SavePackageResult, MoveTemp(*ArchiveCookContext.GetCookTagList()));
 	}
 
 	// If not retrying, mark the package as cooked, either successfully or with failure
@@ -8191,9 +8193,6 @@ void UCookOnTheFlyServer::CookByTheBookFinished()
 			UE_SCOPED_HIERARCHICAL_COOKTIMER(SavingAssetRegistry);
 			SCOPED_BOOT_TIMING("SavingAssetRegistry");
 
-			// Flush the asset registry to get any remaining updates from late loads
-			AssetRegistry->Tick(-1.0f);
-
 			RegisterLocalizationChunkDataGenerator();
 			RegisterShaderChunkDataGenerator();
 
@@ -8276,9 +8275,6 @@ void UCookOnTheFlyServer::CookByTheBookFinished()
 					}
 				}
 				
-				// Sync the Generator's AssetRegistryState with any changes to AssetDatas made during the cook
-				Generator.UpdateAssetDatas(bForceNoFilterAssetsFromAssetRegistry);
-
 				// Add the package hashes to the relevant AssetPackageDatas.
 				// PackageHashes are gated by requiring UPackage::WaitForAsyncFileWrites(), which is called above.
 				FCookSavePackageContext& SaveContext = FindOrCreateSaveContext(TargetPlatform);
@@ -10563,11 +10559,13 @@ uint32 UCookOnTheFlyServer::CookFullLoadAndSave()
 							}
 						}
 
-						// Update asset registry
+						// Flush the global AssetRegisty so any AssetData changes from the save are present
+						AssetRegistry->WaitForCompletion();
+						// Update cooked asset registry
 						FAssetPackageData* AssetPackageData = nullptr;
 						{
 							FScopeLock AssetGeneratorScopeLock(&AssetGeneratorLock);
-							Generator.UpdateAssetRegistryPackageData(*Package, SaveResult, MoveTemp(*CookContext.GetCookTagList()));
+							Generator.UpdateAssetRegistryData(*Package, SaveResult, MoveTemp(*CookContext.GetCookTagList()));
 							AssetPackageData = Generator.GetAssetPackageData(Package->GetFName());
 						}
 						check(AssetPackageData);
diff --git a/Engine/Source/Editor/UnrealEd/Private/Cooker/CompactBinaryTCP.cpp b/Engine/Source/Editor/UnrealEd/Private/Cooker/CompactBinaryTCP.cpp
--- a/Engine/Source/Editor/UnrealEd/Private/Cooker/CompactBinaryTCP.cpp
+++ b/Engine/Source/Editor/UnrealEd/Private/Cooker/CompactBinaryTCP.cpp
@@ -342,4 +342,22 @@ EConnectionStatus TryFlushBuffer(FSocket* Socket, FSendBuffer& Buffer, uint64 Ma
 	return FCompactBinaryTCPImpl::TryWritePacket(Socket, Buffer, TArrayView<FMarshalledMessage>(), MaxPacketSize);
 }
 
-}
\ No newline at end of file
+}
+
+FCbWriter& operator<<(FCbWriter& Writer, const FSoftObjectPathSerializationWrapper& Path)
+{
+	Writer << Path.Inner.ToString();
+	return Writer;
+}
+
+bool LoadFromCompactBinary(FCbFieldView Field, FSoftObjectPathSerializationWrapper& Path)
+{
+	FString PathString;
+	if (!LoadFromCompactBinary(Field, PathString))
+	{
+		Path.Inner.Reset();
+		return false;
+	}
+	Path.Inner.SetPath(PathString);
+	return true;
+}
diff --git a/Engine/Source/Editor/UnrealEd/Private/Cooker/CompactBinaryTCP.h b/Engine/Source/Editor/UnrealEd/Private/Cooker/CompactBinaryTCP.h
--- a/Engine/Source/Editor/UnrealEd/Private/Cooker/CompactBinaryTCP.h
+++ b/Engine/Source/Editor/UnrealEd/Private/Cooker/CompactBinaryTCP.h
@@ -4,10 +4,12 @@
 
 #include "Containers/Array.h"
 #include "Containers/ArrayView.h"
+#include "Containers/Map.h"
 #include "Memory/SharedBuffer.h"
 #include "Misc/Guid.h"
 #include "Serialization/CompactBinarySerialization.h"
 #include "Serialization/CompactBinaryWriter.h"
+#include "UObject/SoftObjectPath.h"
 
 class FSocket;
 
@@ -201,3 +203,42 @@ inline FCbWriter& operator<<(FCbWriter& Writer, const TMap<KeyType, ValueType, S
 	Writer.EndArray();
 	return Writer;
 }
+
+template <typename Key, typename Value>
+FCbWriter& operator<<(FCbWriter& Writer, const TPair<Key, Value>& Pair)
+{
+	Writer.BeginObject();
+	Writer << "K" << Pair.Key;
+	Writer << "V" << Pair.Value;
+	Writer.EndObject();
+	return Writer;
+}
+
+template <typename Key, typename Value>
+bool LoadFromCompactBinary(FCbFieldView Field, TPair<Key, Value>& Pair)
+{
+	bool bOk = LoadFromCompactBinary(Field["K"], Pair.Key);
+	bOk = LoadFromCompactBinary(Field["V"], Pair.Value) & bOk;
+	return bOk;
+}
+
+// FSoftObjectPath has an implicit constructor from FString for backwards compatibility; if we
+// try to create an operator<< for it, it will cause operator<< to be ambiguous.
+// This prevents us from using it directly in containers. To hack around this,
+// containers of FSoftObjectPath need to be reinterpret_casted to containers of FSoftObjectPathSerializationWrapper
+struct FSoftObjectPathSerializationWrapper
+{
+	FSoftObjectPath Inner;
+	bool operator==(const FSoftObjectPathSerializationWrapper& Other) const
+	{
+		return Inner == Other.Inner;
+	}
+	friend uint32 GetTypeHash(const FSoftObjectPathSerializationWrapper& Wrapper)
+	{
+		return GetTypeHash(Wrapper.Inner);
+	}
+};
+
+FCbWriter& operator<<(FCbWriter& Writer, const FSoftObjectPathSerializationWrapper& Path);
+bool LoadFromCompactBinary(FCbFieldView Field, FSoftObjectPathSerializationWrapper& Path);
+
diff --git a/Engine/Source/Editor/UnrealEd/Private/Cooker/CookWorkerServer.cpp b/Engine/Source/Editor/UnrealEd/Private/Cooker/CookWorkerServer.cpp
--- a/Engine/Source/Editor/UnrealEd/Private/Cooker/CookWorkerServer.cpp
+++ b/Engine/Source/Editor/UnrealEd/Private/Cooker/CookWorkerServer.cpp
@@ -477,7 +477,7 @@ void FCookWorkerServer::HandleReceivedPackagePlatformMessages(FPackageData& Pack
 			{
 				FAssetRegistryGenerator* RegistryGenerator = COTFS.PlatformManager->GetPlatformData(TargetPlatform)->RegistryGenerator.Get();
 				check(RegistryGenerator); // The TargetPlatform came from OrderedSessionPlatforms, and the RegistryGenerator should exist for any of those platforms
-				RegistryGenerator->UpdateAssetRegistryPackageData(PackageData, MoveTemp(ARMessage));
+				RegistryGenerator->UpdateAssetRegistryData(PackageData, MoveTemp(ARMessage));
 			}
 		}
 	}
diff --git a/Engine/Source/Editor/UnrealEd/Public/Commandlets/AssetRegistryGenerator.h b/Engine/Source/Editor/UnrealEd/Public/Commandlets/AssetRegistryGenerator.h
--- a/Engine/Source/Editor/UnrealEd/Public/Commandlets/AssetRegistryGenerator.h
+++ b/Engine/Source/Editor/UnrealEd/Public/Commandlets/AssetRegistryGenerator.h
@@ -6,6 +6,7 @@
 
 #include "AssetRegistry/AssetData.h"
 #include "AssetRegistry/AssetRegistryState.h"
+#include "Containers/Map.h"
 #include "Cooker/PackageResultsMessage.h"
 #include "Misc/AssetRegistryInterface.h"
 #include "Misc/Paths.h"
@@ -20,6 +21,7 @@ class IChunkDataGenerator;
 class UChunkDependencyInfo;
 struct FChunkDependencyTreeNode;
 struct FCookTagList;
+struct FSoftObjectPath;
 namespace UE::Cook { class FAssetRegistryPackageMessage; }
 namespace UE::Cook { class FCookWorkerClient; }
 namespace UE::Cook { struct FPackageData; }
@@ -116,9 +118,6 @@ public:
 	 */
 	void RegisterChunkDataGenerator(TSharedRef<IChunkDataGenerator> InChunkDataGenerator);
 
-	/** Called before PreSave; copy AssetData changes from the global AssetRegistry that were made during the cook. */
-	void UpdateAssetDatas(bool bForceNoFilter);
-
 	/**
 	* PreSave
 	* Notify generator that we are about to save the registry and chunk manifests
@@ -203,8 +202,8 @@ public:
 	 * @param Package The package to update info on
 	 * @param SavePackageResult The metadata to associate with the given package name
 	 */
-	void UpdateAssetRegistryPackageData(const UPackage& Package, FSavePackageResultStruct& SavePackageResult, FCookTagList&& InArchiveCookTagList);
-	void UpdateAssetRegistryPackageData(UE::Cook::FPackageData& PackageData, UE::Cook::FAssetRegistryPackageMessage&& Message);
+	void UpdateAssetRegistryData(const UPackage& Package, FSavePackageResultStruct& SavePackageResult, FCookTagList&& InArchiveCookTagList);
+	void UpdateAssetRegistryData(UE::Cook::FPackageData& PackageData, UE::Cook::FAssetRegistryPackageMessage&& Message);
 
 	/**
 	 * Check config to see whether chunk assignments use the AssetManager. If so, run the once-per-process construction
@@ -434,7 +433,7 @@ class IAssetRegistryReporter
 public:
 	virtual ~IAssetRegistryReporter() {}
 
-	virtual void UpdateAssetRegistryPackageData(FPackageData& PackageData, const UPackage& Package,
+	virtual void UpdateAssetRegistryData(FPackageData& PackageData, const UPackage& Package,
 		FSavePackageResultStruct& SavePackageResult, FCookTagList&& InArchiveCookTagList) = 0;
 
 };
@@ -447,10 +446,10 @@ public:
 	{
 	}
 
-	virtual void UpdateAssetRegistryPackageData(FPackageData& PackageData, const UPackage& Package,
+	virtual void UpdateAssetRegistryData(FPackageData& PackageData, const UPackage& Package,
 		FSavePackageResultStruct& SavePackageResult, FCookTagList&& InArchiveCookTagList) override
 	{
-		Generator.UpdateAssetRegistryPackageData(Package, SavePackageResult, MoveTemp(InArchiveCookTagList));
+		Generator.UpdateAssetRegistryData(Package, SavePackageResult, MoveTemp(InArchiveCookTagList));
 	}
 
 private:
@@ -462,7 +461,7 @@ class FAssetRegistryReporterRemote : public IAssetRegistryReporter
 public:
 	FAssetRegistryReporterRemote(FCookWorkerClient& InClient, const ITargetPlatform* InTargetPlatform);
 
-	virtual void UpdateAssetRegistryPackageData(FPackageData& PackageData, const UPackage& Package,
+	virtual void UpdateAssetRegistryData(FPackageData& PackageData, const UPackage& Package,
 		FSavePackageResultStruct& SavePackageResult, FCookTagList&& InArchiveCookTagList) override;
 
 private:
@@ -478,6 +477,7 @@ public:
 	virtual FGuid GetMessageType() const override { return MessageType; }
 
 	TArray<FAssetData> AssetDatas;
+	TMap<FSoftObjectPath, TArray<TPair<FName, FString>>> CookTags;
 	uint32 PackageFlags = 0;
 	int64 DiskSize = -1;
 
diff --git a/Engine/Source/Runtime/AssetRegistry/Private/AssetRegistry.cpp b/Engine/Source/Runtime/AssetRegistry/Private/AssetRegistry.cpp
--- a/Engine/Source/Runtime/AssetRegistry/Private/AssetRegistry.cpp
+++ b/Engine/Source/Runtime/AssetRegistry/Private/AssetRegistry.cpp
@@ -1627,12 +1627,12 @@ bool FAssetRegistryImpl::HasAssets(const FName PackagePath, const bool bRecursiv
 
 }
 
-bool UAssetRegistryImpl::GetAssetsByPackageName(FName PackageName, TArray<FAssetData>& OutAssetData, bool bIncludeOnlyOnDiskAssets) const
+bool UAssetRegistryImpl::GetAssetsByPackageName(FName PackageName, TArray<FAssetData>& OutAssetData, bool bIncludeOnlyOnDiskAssets, bool bSkipARFilteredAssets) const
 {
 	FARFilter Filter;
 	Filter.PackageNames.Add(PackageName);
 	Filter.bIncludeOnlyOnDiskAssets = bIncludeOnlyOnDiskAssets;
-	return GetAssets(Filter, OutAssetData);
+	return GetAssets(Filter, OutAssetData, bSkipARFilteredAssets);
 }
 
 bool UAssetRegistryImpl::GetAssetsByPath(FName PackagePath, TArray<FAssetData>& OutAssetData, bool bRecursive, bool bIncludeOnlyOnDiskAssets) const
@@ -1681,13 +1681,13 @@ void FAssetRegistryImpl::EnumerateAssetsByPathNoTags(FName PackagePath,
 			[&Callback](const UObject* Object, FAssetData&& PartialAssetData)
 			{
 				return Callback(PartialAssetData);
-			});
+			}, true /* bSkipARFilteredAssets */);
 		if (bStopIteration)
 		{
 			return;
 		}
 	}
-	EnumerateDiskAssets(CompiledFilter, PackagesToSkip, Callback);
+	EnumerateDiskAssets(CompiledFilter, PackagesToSkip, Callback, true /* bSkipARFilteredAssets */);
 }
 
 }
@@ -1738,7 +1738,8 @@ bool UAssetRegistryImpl::GetAssetsByTagValues(const TMultiMap<FName, FString>& A
 	return GetAssets(Filter, OutAssetData);
 }
 
-bool UAssetRegistryImpl::GetAssets(const FARFilter& InFilter, TArray<FAssetData>& OutAssetData) const
+bool UAssetRegistryImpl::GetAssets(const FARFilter& InFilter, TArray<FAssetData>& OutAssetData,
+	bool bSkipARFilteredAssets) const
 {
 	using namespace UE::AssetRegistry::Utils;
 
@@ -1758,7 +1759,7 @@ bool UAssetRegistryImpl::GetAssets(const FARFilter& InFilter, TArray<FAssetData>
 			{
 				OutAssetData.Add(MoveTemp(AssetData));
 				return true;
-			});
+			}, bSkipARFilteredAssets);
 	}
 
 	{
@@ -1767,19 +1768,21 @@ bool UAssetRegistryImpl::GetAssets(const FARFilter& InFilter, TArray<FAssetData>
 			{
 				OutAssetData.Emplace(AssetData);
 				return true;
-			});
+			}, bSkipARFilteredAssets);
 	}
 	return true;
 }
 
-bool UAssetRegistryImpl::EnumerateAssets(const FARFilter& InFilter, TFunctionRef<bool(const FAssetData&)> Callback) const
+bool UAssetRegistryImpl::EnumerateAssets(const FARFilter& InFilter, TFunctionRef<bool(const FAssetData&)> Callback,
+	bool bSkipARFilteredAssets) const
 {
 	FARCompiledFilter CompiledFilter;
 	CompileFilter(InFilter, CompiledFilter);
-	return EnumerateAssets(CompiledFilter, Callback);
+	return EnumerateAssets(CompiledFilter, Callback, bSkipARFilteredAssets);
 }
 
-bool UAssetRegistryImpl::EnumerateAssets(const FARCompiledFilter& InFilter, TFunctionRef<bool(const FAssetData&)> Callback) const
+bool UAssetRegistryImpl::EnumerateAssets(const FARCompiledFilter& InFilter, TFunctionRef<bool(const FAssetData&)> Callback,
+	bool bSkipARFilteredAssets) const
 {
 	using namespace UE::AssetRegistry::Utils;
 
@@ -1797,7 +1800,7 @@ bool UAssetRegistryImpl::EnumerateAssets(const FARCompiledFilter& InFilter, TFun
 			[&Callback](FAssetData&& AssetData)
 			{
 				return Callback(AssetData);
-			});
+			}, bSkipARFilteredAssets);
 		if (bStopIteration)
 		{
 			return true;
@@ -1811,7 +1814,7 @@ bool UAssetRegistryImpl::EnumerateAssets(const FARCompiledFilter& InFilter, TFun
 			{
 				FoundAssets.Emplace(AssetData);
 				return true;
-			});
+			}, bSkipARFilteredAssets);
 	}
 	for (const FAssetData& AssetData : FoundAssets)
 	{
@@ -1830,7 +1833,8 @@ namespace Utils
 {
 
 void EnumerateMemoryAssetsHelper(const FARCompiledFilter& InFilter, TSet<FName>& OutPackageNamesWithAssets,
-	bool& bOutStopIteration, TFunctionRef<bool(const UObject* Object, FAssetData&& PartialAssetData)> Callback)
+	bool& bOutStopIteration, TFunctionRef<bool(const UObject* Object, FAssetData&& PartialAssetData)> Callback,
+	bool bSkipARFilteredAssets)
 {
 	checkf(IsInGameThread(), TEXT("Enumerating in-memory assets can only be done on the game thread; it uses non-threadsafe UE::AssetRegistry::Filtering globals."));
 	bOutStopIteration = false;
@@ -1864,7 +1868,7 @@ void EnumerateMemoryAssetsHelper(const FARCompiledFilter& InFilter, TSet<FName>&
 			}
 
 			// Skip classes that report themselves as assets but that the editor AssetRegistry is currently not counting as assets
-			if (UE::AssetRegistry::FFiltering::ShouldSkipAsset(Obj))
+			if (bSkipARFilteredAssets && UE::AssetRegistry::FFiltering::ShouldSkipAsset(Obj))
 			{
 				return;
 			}
@@ -1953,7 +1957,7 @@ void EnumerateMemoryAssetsHelper(const FARCompiledFilter& InFilter, TSet<FName>&
 }
 
 void EnumerateMemoryAssets(const FARCompiledFilter& InFilter, TSet<FName>& OutPackageNamesWithAssets,
-	bool& bOutStopIteration, TFunctionRef<bool(FAssetData&&)> Callback)
+	bool& bOutStopIteration, TFunctionRef<bool(FAssetData&&)> Callback, bool bSkipARFilteredAssets)
 {
 	check(!InFilter.IsEmpty() && Utils::IsFilterValid(InFilter));
 	EnumerateMemoryAssetsHelper(InFilter, OutPackageNamesWithAssets, bOutStopIteration,
@@ -1985,17 +1989,17 @@ void EnumerateMemoryAssets(const FARCompiledFilter& InFilter, TSet<FName>& OutPa
 
 			// All filters passed
 			return Callback(MoveTemp(PartialAssetData));
-		});
+		}, bSkipARFilteredAssets);
 }
 
 }
 
 void FAssetRegistryImpl::EnumerateDiskAssets(const FARCompiledFilter& InFilter, TSet<FName>& PackagesToSkip,
-	TFunctionRef<bool(const FAssetData&)> Callback) const
+	TFunctionRef<bool(const FAssetData&)> Callback, bool bSkipARFilteredAssets) const
 {
 	check(!InFilter.IsEmpty() && Utils::IsFilterValid(InFilter));
 	PackagesToSkip.Append(CachedEmptyPackages);
-	State.EnumerateAssets(InFilter, PackagesToSkip, Callback, true /*bARFiltering*/);
+	State.EnumerateAssets(InFilter, PackagesToSkip, Callback, bSkipARFilteredAssets);
 }
 
 }
diff --git a/Engine/Source/Runtime/AssetRegistry/Private/AssetRegistry.h b/Engine/Source/Runtime/AssetRegistry/Private/AssetRegistry.h
--- a/Engine/Source/Runtime/AssetRegistry/Private/AssetRegistry.h
+++ b/Engine/Source/Runtime/AssetRegistry/Private/AssetRegistry.h
@@ -31,7 +31,7 @@ public:
 
 	// IAssetRegistry implementation
 	virtual bool HasAssets(const FName PackagePath, const bool bRecursive = false) const override;
-	virtual bool GetAssetsByPackageName(FName PackageName, TArray<FAssetData>& OutAssetData, bool bIncludeOnlyOnDiskAssets = false) const override;
+	virtual bool GetAssetsByPackageName(FName PackageName, TArray<FAssetData>& OutAssetData, bool bIncludeOnlyOnDiskAssets = false, bool bSkipARFilteredAssets=true) const override;
 	virtual bool GetAssetsByPath(FName PackagePath, TArray<FAssetData>& OutAssetData, bool bRecursive = false, bool bIncludeOnlyOnDiskAssets = false) const override;
 	virtual bool GetAssetsByPaths(TArray<FName> PackagePath, TArray<FAssetData>& OutAssetData, bool bRecursive = false, bool bIncludeOnlyOnDiskAssets = false) const override;
 	UE_DEPRECATED(5.1, "Class names are now represented by path names. Please use a version of this function that uses FTopLevelAssetPath.")
@@ -39,9 +39,9 @@ public:
 	virtual bool GetAssetsByClass(FTopLevelAssetPath ClassPathName, TArray<FAssetData>& OutAssetData, bool bSearchSubClasses = false) const override;
 	virtual bool GetAssetsByTags(const TArray<FName>& AssetTags, TArray<FAssetData>& OutAssetData) const override;
 	virtual bool GetAssetsByTagValues(const TMultiMap<FName, FString>& AssetTagsAndValues, TArray<FAssetData>& OutAssetData) const override;
-	virtual bool GetAssets(const FARFilter& Filter, TArray<FAssetData>& OutAssetData) const override;
-	virtual bool EnumerateAssets(const FARFilter& Filter, TFunctionRef<bool(const FAssetData&)> Callback) const override;
-	virtual bool EnumerateAssets(const FARCompiledFilter& Filter, TFunctionRef<bool(const FAssetData&)> Callback) const override;
+	virtual bool GetAssets(const FARFilter& Filter, TArray<FAssetData>& OutAssetData, bool bSkipARFilteredAssets = true) const override;
+	virtual bool EnumerateAssets(const FARFilter& Filter, TFunctionRef<bool(const FAssetData&)> Callback, bool bSkipARFilteredAssets=true) const override;
+	virtual bool EnumerateAssets(const FARCompiledFilter& Filter, TFunctionRef<bool(const FAssetData&)> Callback, bool bSkipARFilteredAssets = true) const override;
 	virtual FAssetData GetAssetByObjectPath( const FName ObjectPath, bool bIncludeOnlyOnDiskAssets = false ) const override;
 	virtual bool GetAllAssets(TArray<FAssetData>& OutAssetData, bool bIncludeOnlyOnDiskAssets = false) const override;
 	virtual bool EnumerateAllAssets(TFunctionRef<bool(const FAssetData&)> Callback, bool bIncludeOnlyOnDiskAssets = false) const override;
diff --git a/Engine/Source/Runtime/AssetRegistry/Private/AssetRegistryImpl.h b/Engine/Source/Runtime/AssetRegistry/Private/AssetRegistryImpl.h
--- a/Engine/Source/Runtime/AssetRegistry/Private/AssetRegistryImpl.h
+++ b/Engine/Source/Runtime/AssetRegistry/Private/AssetRegistryImpl.h
@@ -162,7 +162,7 @@ public:
 
 	/** Enumerate assets in the State, filtering by filter and package not in PackagesToSkip */
 	void EnumerateDiskAssets(const FARCompiledFilter& InFilter, TSet<FName>& PackagesToSkip,
-		TFunctionRef<bool(const FAssetData&)> Callback) const;
+		TFunctionRef<bool(const FAssetData&)> Callback, bool bSkipARFilteredAssets) const;
 	/** Enumerate assets in the State, filtering by package not in PackagesToSkip */
 	void EnumerateAllDiskAssets(TSet<FName>& PackageNamesToSkip, TFunctionRef<bool(const FAssetData&)> Callback) const;
 
@@ -566,13 +566,14 @@ void EnumerateAllMemoryAssets(TSet<FName>& OutPackageNamesWithAssets, bool& bOut
  * Fills in OutPackageNamesWithAssets with names of all packages tested.
  */
 void EnumerateMemoryAssetsHelper(const FARCompiledFilter& InFilter, TSet<FName>& OutPackageNamesWithAssets,
-	bool& bOutStopIteration, TFunctionRef<bool(const UObject* Object, FAssetData&& PartialAssetData)> Callback);
+	bool& bOutStopIteration, TFunctionRef<bool(const UObject* Object, FAssetData&& PartialAssetData)> Callback,
+	bool bSkipARFilteredAssets);
 /**
  * Call the given callback on all UObjects in memory that pass the given filter.
  * Fills in OutPackageNamesWithAssets with names of all packages tested.
 */
 void EnumerateMemoryAssets(const FARCompiledFilter& InFilter, TSet<FName>& OutPackageNamesWithAssets,
-	bool& bOutStopIteration, TFunctionRef<bool(FAssetData&&)> Callback);
+	bool& bOutStopIteration, TFunctionRef<bool(FAssetData&&)> Callback, bool bSkipARFilteredAssets);
 
 }
 
diff --git a/Engine/Source/Runtime/AssetRegistry/Private/AssetRegistryState.cpp b/Engine/Source/Runtime/AssetRegistry/Private/AssetRegistryState.cpp
--- a/Engine/Source/Runtime/AssetRegistry/Private/AssetRegistryState.cpp
+++ b/Engine/Source/Runtime/AssetRegistry/Private/AssetRegistryState.cpp
@@ -588,12 +588,12 @@ void FAssetRegistryState::PruneAssetData(const TSet<FName>& RequiredPackages, co
 	}
 }
 
-bool FAssetRegistryState::HasAssets(const FName PackagePath, bool bARFiltering) const
+bool FAssetRegistryState::HasAssets(const FName PackagePath, bool bSkipARFilteredAssets) const
 {
 	const TArray<FAssetData*>* FoundAssetArray = CachedAssetsByPath.Find(PackagePath);
 	if (FoundAssetArray)
 	{
-		if (bARFiltering)
+		if (bSkipARFilteredAssets)
 		{
 			return FoundAssetArray->ContainsByPredicate([](FAssetData* AssetData)
 			{
@@ -608,14 +608,14 @@ bool FAssetRegistryState::HasAssets(const FName PackagePath, bool bARFiltering)
 	return false;
 }
 
-bool FAssetRegistryState::GetAssets(const FARCompiledFilter& Filter, const TSet<FName>& PackageNamesToSkip, TArray<FAssetData>& OutAssetData, bool bARFiltering) const
+bool FAssetRegistryState::GetAssets(const FARCompiledFilter& Filter, const TSet<FName>& PackageNamesToSkip, TArray<FAssetData>& OutAssetData, bool bSkipARFilteredAssets) const
 {
 	return EnumerateAssets(Filter, PackageNamesToSkip, [&OutAssetData](const FAssetData& AssetData)
 	{
 		OutAssetData.Emplace(AssetData);
 		return true;
 	},
-	bARFiltering);
+	bSkipARFilteredAssets);
 }
 
 template<class ArrayType, typename KeyType>
@@ -644,7 +644,8 @@ TArray<FAssetData*> FindAssets(const TMap<KeyType, ArrayType>& Map, const TSet<K
 	return Out;
 }
 
-bool FAssetRegistryState::EnumerateAssets(const FARCompiledFilter& Filter, const TSet<FName>& PackageNamesToSkip, TFunctionRef<bool(const FAssetData&)> Callback, bool bARFiltering) const
+bool FAssetRegistryState::EnumerateAssets(const FARCompiledFilter& Filter, const TSet<FName>& PackageNamesToSkip,
+	TFunctionRef<bool(const FAssetData&)> Callback, bool bSkipARFilteredAssets) const
 {
 	// Verify filter input. If all assets are needed, use EnumerateAllAssets() instead.
 	if (Filter.IsEmpty() || !IsFilterValid(Filter))
@@ -743,7 +744,8 @@ bool FAssetRegistryState::EnumerateAssets(const FARCompiledFilter& Filter, const
 				return true;
 			}
 
-			return bARFiltering && UE::AssetRegistry::FFiltering::ShouldSkipAsset(AssetData->AssetClassPath, AssetData->PackageFlags);
+			return bSkipARFilteredAssets &&
+				UE::AssetRegistry::FFiltering::ShouldSkipAsset(AssetData->AssetClassPath, AssetData->PackageFlags);
 		};
 
 		int32 NumFilterResults = FilterResults.Num();
@@ -797,17 +799,17 @@ bool FAssetRegistryState::EnumerateAssets(const FARCompiledFilter& Filter, const
 	return true;
 }
 
-bool FAssetRegistryState::GetAllAssets(const TSet<FName>& PackageNamesToSkip, TArray<FAssetData>& OutAssetData, bool bARFiltering) const
+bool FAssetRegistryState::GetAllAssets(const TSet<FName>& PackageNamesToSkip, TArray<FAssetData>& OutAssetData, bool bSkipARFilteredAssets) const
 {
 	return EnumerateAllAssets(PackageNamesToSkip, [&OutAssetData](const FAssetData& AssetData)
 	{
 		OutAssetData.Emplace(AssetData);
 		return true;
 	},
-	bARFiltering);
+	bSkipARFilteredAssets);
 }
 
-bool FAssetRegistryState::EnumerateAllAssets(const TSet<FName>& PackageNamesToSkip, TFunctionRef<bool(const FAssetData&)> Callback, bool bARFiltering) const
+bool FAssetRegistryState::EnumerateAllAssets(const TSet<FName>& PackageNamesToSkip, TFunctionRef<bool(const FAssetData&)> Callback, bool bSkipARFilteredAssets) const
 {
 	// All unloaded disk assets
 	for (const TPair<FName, FAssetData*>& AssetDataPair : CachedAssetsByObjectPath)
@@ -816,7 +818,7 @@ bool FAssetRegistryState::EnumerateAllAssets(const TSet<FName>& PackageNamesToSk
 
 		if (AssetData &&
 			!PackageNamesToSkip.Contains(AssetData->PackageName) &&
-			(!bARFiltering || !UE::AssetRegistry::FFiltering::ShouldSkipAsset(AssetData->AssetClassPath, AssetData->PackageFlags)))
+			(!bSkipARFilteredAssets || !UE::AssetRegistry::FFiltering::ShouldSkipAsset(AssetData->AssetClassPath, AssetData->PackageFlags)))
 		{
 			if (!Callback(*AssetData))
 			{
@@ -1688,11 +1690,45 @@ void FAssetRegistryState::AddTagsToAssetData(const FSoftObjectPath& InObjectPath
 		UE_LOG(LogAssetRegistry, Error, TEXT("AddTagsToAssetData called with asset data that doesn't exist! Tags not added. ObjectPath: %s"), *InObjectPath.GetAssetPathName().ToString());
 		return;
 	}
+	FAssetDataTagMap Tags = AssetData->TagsAndValues.CopyMap();
+	Tags.Append(MoveTemp(InTagsAndValues));
+	SetTagsOnExistingAsset(AssetData, MoveTemp(Tags));
+}
+
+
+void FAssetRegistryState::FilterTags(const FAssetRegistrySerializationOptions& Options)
+{
+	for (TPair<FName, FAssetData*>& Pair : CachedAssetsByObjectPath)
+	{
+		FAssetData* AssetData = Pair.Value;
+		check(AssetData);
+
+		FAssetDataTagMap LocalTagsAndValues;
+		FAssetRegistryState::FilterTags(AssetData->TagsAndValues, LocalTagsAndValues, Options.CookFilterlistTagsByClass.Find(AssetData->AssetClassPath), Options);
+		if (LocalTagsAndValues != AssetData->TagsAndValues)
+		{
+			SetTagsOnExistingAsset(AssetData, MoveTemp(LocalTagsAndValues));
+		}
+	}
+}
 
-	// Update the tag cache map with the new tags.
-	for (auto TagIt : InTagsAndValues)
+void FAssetRegistryState::SetTagsOnExistingAsset(FAssetData* AssetData, FAssetDataTagMap&& NewTags)
+{
+	// Update the tag cache map to remove deleted tags
+	for (auto TagIt = AssetData->TagsAndValues.CreateConstIterator(); TagIt; ++TagIt)
 	{
-		const FName FNameKey = TagIt.Key;
+		const FName FNameKey = TagIt.Key();
+
+		if (!NewTags.Contains(FNameKey))
+		{
+			TArray<FAssetData*>* OldTagAssets = CachedAssetsByTag.Find(FNameKey);
+			OldTagAssets->RemoveSingleSwap(AssetData);
+		}
+	}
+	// Update the tag cache map to add added tags
+	for (auto TagIt = NewTags.CreateConstIterator(); TagIt; ++TagIt)
+	{
+		const FName FNameKey = TagIt.Key();
 
 		if (!AssetData->TagsAndValues.Contains(FNameKey))
 		{
@@ -1700,19 +1736,31 @@ void FAssetRegistryState::AddTagsToAssetData(const FSoftObjectPath& InObjectPath
 			NewTagAssets.Add(AssetData);
 		}
 	}
-
-	FAssetDataTagMap OldTags = AssetData->TagsAndValues.CopyMap();
-	OldTags.Append(MoveTemp(InTagsAndValues));
-	AssetData->TagsAndValues = FAssetDataTagMapSharedView(MoveTemp(OldTags));
+	AssetData->TagsAndValues = FAssetDataTagMapSharedView(MoveTemp(NewTags));
 }
 
-
-void FAssetRegistryState::UpdateAssetData(const FAssetData& NewAssetData)
+void FAssetRegistryState::UpdateAssetData(const FAssetData& NewAssetData, bool bCreateIfNotExists)
 {
 	if (FAssetData* AssetData = CachedAssetsByObjectPath.FindRef(NewAssetData.ObjectPath))
 	{
 		UpdateAssetData(AssetData, NewAssetData);
 	}
+	else
+	{
+		AddAssetData(new FAssetData(NewAssetData));
+	}
+}
+
+void FAssetRegistryState::UpdateAssetData(FAssetData&& NewAssetData, bool bCreateIfNotExists)
+{
+	if (FAssetData* AssetData = CachedAssetsByObjectPath.FindRef(NewAssetData.ObjectPath))
+	{
+		UpdateAssetData(AssetData, MoveTemp(NewAssetData));
+	}
+	else if (bCreateIfNotExists)
+	{
+		AddAssetData(new FAssetData(MoveTemp(NewAssetData)));
+	}
 }
 
 void FAssetRegistryState::UpdateAssetData(FAssetData* AssetData, const FAssetData& NewAssetData, bool* bOutModified)
diff --git a/Engine/Source/Runtime/AssetRegistry/Public/AssetRegistry/AssetRegistryState.h b/Engine/Source/Runtime/AssetRegistry/Public/AssetRegistry/AssetRegistryState.h
--- a/Engine/Source/Runtime/AssetRegistry/Public/AssetRegistry/AssetRegistryState.h
+++ b/Engine/Source/Runtime/AssetRegistry/Public/AssetRegistry/AssetRegistryState.h
@@ -154,9 +154,9 @@ public:
 	 * @param Filter filter to apply to the assets in the AssetRegistry
 	 * @param PackageNamesToSkip explicit list of packages to skip, because they were already added
 	 * @param OutAssetData the list of assets in this path
-	 * @param bARFiltering Whether to apply filtering from UE::AssetRegistry::FFiltering (false by default)
+	 * @param bSkipARFilteredAssets If true, skip assets that are skipped by UE::AssetRegistry::FFiltering (false by default)
 	 */
-	bool GetAssets(const FARCompiledFilter& Filter, const TSet<FName>& PackageNamesToSkip, TArray<FAssetData>& OutAssetData, bool bARFiltering = false) const;
+	bool GetAssets(const FARCompiledFilter& Filter, const TSet<FName>& PackageNamesToSkip, TArray<FAssetData>& OutAssetData, bool bSkipARFilteredAssets = false) const;
 
 	/**
 	 * Enumerate asset data for all assets that match the filter.
@@ -358,7 +358,8 @@ public:
 	void AddTagsToAssetData(const FSoftObjectPath& InObjectPath, FAssetDataTagMap&& InTagsAndValues);
 
 	/** Finds an existing asset data based on object path and updates it with the new value and updates lookup maps */
-	void UpdateAssetData(const FAssetData& NewAssetData);
+	void UpdateAssetData(const FAssetData& NewAssetData, bool bCreateIfNotExists=false);
+	void UpdateAssetData(FAssetData&& NewAssetData, bool bCreateIfNotExists = false);
 
 	/** Updates an existing asset data with the new value and updates lookup maps */
 	void UpdateAssetData(FAssetData* AssetData, const FAssetData& NewAssetData, bool* bOutModified = nullptr);
@@ -407,6 +408,9 @@ public:
 	 */
 	void InitializeFromExistingAndPrune(const FAssetRegistryState& ExistingState, const TSet<FName>& RequiredPackages, const TSet<FName>& RemovePackages, const TSet<int32> ChunksToKeep, const FAssetRegistrySerializationOptions& Options);
 
+	/** Edit every AssetData's Tags to remove Tags that are filtered out by the filtering rules in Options */
+	void FilterTags(const FAssetRegistrySerializationOptions& Options);
+
 
 	/** Serialize the registry to/from a file, skipping editor only data */
 	bool Serialize(FArchive& Ar, const FAssetRegistrySerializationOptions& Options);
@@ -470,6 +474,8 @@ private:
 	void LoadDependencies(FArchive& Ar);
 	void LoadDependencies_BeforeFlags(FArchive& Ar, bool bSerializeDependencies, FAssetRegistryVersion::Type Version);
 
+	void SetTagsOnExistingAsset(FAssetData* AssetData, FAssetDataTagMap&& NewTags);
+
 	/** The map of ObjectPath names to asset data for assets saved to disk */
 	TMap<FName, FAssetData*> CachedAssetsByObjectPath;
 
diff --git a/Engine/Source/Runtime/AssetRegistry/Public/AssetRegistry/IAssetRegistry.h b/Engine/Source/Runtime/AssetRegistry/Public/AssetRegistry/IAssetRegistry.h
--- a/Engine/Source/Runtime/AssetRegistry/Public/AssetRegistry/IAssetRegistry.h
+++ b/Engine/Source/Runtime/AssetRegistry/Public/AssetRegistry/IAssetRegistry.h
@@ -127,9 +127,10 @@ public:
 	 *
 	 * @param PackageName the package name for the requested assets (eg, /Game/MyFolder/MyAsset)
 	 * @param OutAssetData the list of assets in this path
+	 * @param bSkipARFilteredAssets If true, skips Objects that return true for IsAsset but are not assets in the current platform.
 	 */
 	UFUNCTION(BlueprintCallable, BlueprintPure=false, Category="AssetRegistry")
-	virtual bool GetAssetsByPackageName(FName PackageName, TArray<FAssetData>& OutAssetData, bool bIncludeOnlyOnDiskAssets = false) const = 0;
+	virtual bool GetAssetsByPackageName(FName PackageName, TArray<FAssetData>& OutAssetData, bool bIncludeOnlyOnDiskAssets = false, bool bSkipARFilteredAssets=true) const = 0;
 
 	/**
 	 * Gets asset data for all assets in the supplied folder path
@@ -187,9 +188,10 @@ public:
 	 *
 	 * @param Filter filter to apply to the assets in the AssetRegistry
 	 * @param OutAssetData the list of assets in this path
+	 * @param bSkipARFilteredAssets If true, skips Objects that return true for IsAsset but are not assets in the current platform.
 	 */
 	UFUNCTION(BlueprintCallable, BlueprintPure=false, Category="AssetRegistry")
-	virtual bool GetAssets(const FARFilter& Filter, TArray<FAssetData>& OutAssetData) const = 0;
+	virtual bool GetAssets(const FARFilter& Filter, TArray<FAssetData>& OutAssetData, bool bSkipARFilteredAssets=true) const = 0;
 
 	/**
 	 * Enumerate asset data for all assets that match the filter.
@@ -198,9 +200,12 @@ public:
 	 *
 	 * @param Filter filter to apply to the assets in the AssetRegistry
 	 * @param Callback function to call for each asset data enumerated
+	 * @param bSkipARFilteredAssets If true, skips Objects that return true for IsAsset but are not assets in the current platform.
 	 */
-	virtual bool EnumerateAssets(const FARFilter& Filter, TFunctionRef<bool(const FAssetData&)> Callback) const = 0;
-	virtual bool EnumerateAssets(const FARCompiledFilter& Filter, TFunctionRef<bool(const FAssetData&)> Callback) const = 0;
+	virtual bool EnumerateAssets(const FARFilter& Filter, TFunctionRef<bool(const FAssetData&)> Callback,
+		bool bSkipARFilteredAssets = true) const = 0;
+	virtual bool EnumerateAssets(const FARCompiledFilter& Filter, TFunctionRef<bool(const FAssetData&)> Callback,
+		bool bSkipARFilteredAssets = true) const = 0;
 
 	/**
 	 * Gets the asset data for the specified object path
