Fix an issue where the selection set was cleared, even if the mode handled the selection change

#rnx

diff --git a/Engine/Source/Editor/UnrealEd/Private/DragTool_BoxSelect.cpp b/Engine/Source/Editor/UnrealEd/Private/DragTool_BoxSelect.cpp
--- a/Engine/Source/Editor/UnrealEd/Private/DragTool_BoxSelect.cpp
+++ b/Engine/Source/Editor/UnrealEd/Private/DragTool_BoxSelect.cpp
@@ -178,9 +178,6 @@ void FDragTool_ActorBoxSelect::EndDrag()
 	UBrushEditingSubsystem* BrushSubsystem = GEditor->GetEditorSubsystem<UBrushEditingSubsystem>();
 	const bool bGeometryMode = BrushSubsystem ? BrushSubsystem->IsGeometryEditorModeActive() : false;
 
-	UTypedElementSelectionSet* SelectionSet = GEditor->GetEditorSubsystem<ULevelEditorSubsystem>()->GetSelectionSet();
-	FTypedElementSelectionOptions ElementSelectionOption;
-
 	FScopedTransaction Transaction( NSLOCTEXT("ActorFrustumSelect", "MarqueeSelectTransation", "Marquee Select" ) );
 
 	bool bShouldSelect = true;
@@ -194,9 +191,8 @@ void FDragTool_ActorBoxSelect::EndDrag()
 	}
 	else if( !bShiftDown )
 	{
-		// If the user is selecting, but isn't hold down SHIFT, remove all current selections.
+		// If the user is selecting, but isn't holding down SHIFT, give modes a chance to clear selection
 		ModeTools->SelectNone();
-		SelectionSet->ClearSelection(ElementSelectionOption);
 	}
 
 	// Let the editor mode try to handle the box selection.
@@ -210,11 +206,19 @@ void FDragTool_ActorBoxSelect::EndDrag()
 	{
 		const bool bStrictDragSelection = GetDefault<ULevelEditorViewportSettings>()->bStrictBoxSelection;
 
+		UTypedElementSelectionSet* SelectionSet = GEditor->GetEditorSubsystem<ULevelEditorSubsystem>()->GetSelectionSet();
+		FTypedElementSelectionOptions ElementSelectionOption;
+		if (!bControlDown && !bShiftDown)
+		{
+			// If the user is selecting, but isn't holding down SHIFT, remove all current selections from the selection set
+			SelectionSet->ClearSelection(ElementSelectionOption);
+		}
+
 		FWorldSelectionElementArgs SeletionArgs
 		{
 			SelectionSet,
 			ETypedElementSelectionMethod::Primary,
-			FTypedElementSelectionOptions(),
+			ElementSelectionOption,
 			&(LevelViewportClient->EngineShowFlags),
 			bStrictDragSelection,
 			bGeometryMode
