FBX Scene Import: Init StaticMesh DistanceFieldScale to 1 instead of 0

diff --git a/Engine/Source/Editor/UnrealEd/Private/Fbx/FbxSceneImportFactory.cpp b/Engine/Source/Editor/UnrealEd/Private/Fbx/FbxSceneImportFactory.cpp
--- a/Engine/Source/Editor/UnrealEd/Private/Fbx/FbxSceneImportFactory.cpp
+++ b/Engine/Source/Editor/UnrealEd/Private/Fbx/FbxSceneImportFactory.cpp
@@ -124,6 +124,7 @@ bool GetFbxSceneImportOptions(UnFbx::FFbxImporter* FbxImporter
 	GlobalImportSettings->ImportTranslation = FVector(0);
 	GlobalImportSettings->ImportRotation = FRotator(0);
 	GlobalImportSettings->ImportUniformScale = 1.0f;
+	GlobalImportSettings->DistanceFieldResolutionScale = 1.0f;
 
 	GlobalImportSettings->bConvertScene = true;
 	GlobalImportSettings->bConvertSceneUnit = true;
