When using bShouldBuildNavigationData, UResavePackagesCommandlet should not checkout non-navigation data related assets
Fixing issue introduced at CL 15888156.

diff --git a/Engine/Source/Editor/UnrealEd/Private/Commandlets/ContentCommandlets.cpp b/Engine/Source/Editor/UnrealEd/Private/Commandlets/ContentCommandlets.cpp
--- a/Engine/Source/Editor/UnrealEd/Private/Commandlets/ContentCommandlets.cpp
+++ b/Engine/Source/Editor/UnrealEd/Private/Commandlets/ContentCommandlets.cpp
@@ -1803,7 +1803,7 @@ void UResavePackagesCommandlet::PerformAdditionalOperations(class UWorld* World,
 	FScopedEditorWorld EditorWorld(World, IVS);
 
 	// Load and Save world partition actor packages
-	if (bResaveWorldPartitionExternalActors)
+	if (bResaveWorldPartitionExternalActors && !bShouldBuildNavigationData)
 	{
 		FWorldPartitionHelpers::ForEachActorDesc(WorldPartition, [this, WorldPartition](const FWorldPartitionActorDesc* ActorDesc)
 		{
