Keep a combined list of all properties that have been requested to snapshot during a given object record

This means that snapshotting A and then B will correctly produce a diff for A and B, rather than flipping between having data for A and having data for B

#rnx

diff --git a/Engine/Source/Editor/UnrealEd/Classes/Editor/Transactor.h b/Engine/Source/Editor/UnrealEd/Classes/Editor/Transactor.h
--- a/Engine/Source/Editor/UnrealEd/Classes/Editor/Transactor.h
+++ b/Engine/Source/Editor/UnrealEd/Classes/Editor/Transactor.h
@@ -111,6 +111,8 @@ protected:
 		TUniquePtr<UE::Transaction::FDiffableObject> DiffableObject;
 		/** The diffable object data when it was last snapshot (always null once finalized) */
 		TUniquePtr<UE::Transaction::FDiffableObject> DiffableObjectSnapshot;
+		/** The combined list of properties that have been passed any Snapshot call for this object (always empty once finalized) */
+		TArray<const FProperty*> AllPropertiesSnapshot;
 		/** Delta change information between the diffable state of the object when the transaction started, and the diffable state of the object when the transaction ended */
 		FTransactionObjectDeltaChange DeltaChange;
 
diff --git a/Engine/Source/Editor/UnrealEd/Private/EditorTransaction.cpp b/Engine/Source/Editor/UnrealEd/Private/EditorTransaction.cpp
--- a/Engine/Source/Editor/UnrealEd/Private/EditorTransaction.cpp
+++ b/Engine/Source/Editor/UnrealEd/Private/EditorTransaction.cpp
@@ -312,6 +312,7 @@ void FTransaction::FObjectRecord::Finalize( FTransaction* Owner, TSharedPtr<ITra
 		// Clear out any diff data now as we won't be getting any more diff requests once finalized
 		DiffableObject.Reset();
 		DiffableObjectSnapshot.Reset();
+		AllPropertiesSnapshot.Reset();
 	}
 }
 
@@ -338,11 +339,30 @@ void FTransaction::FObjectRecord::Snapshot( FTransaction* Owner, TArrayView<cons
 	UObject* CurrentObject = Object.Get();
 	if (CurrentObject)
 	{
+		// Merge these properties into the combined list of properties to snapshot
+		// This means that one snapshot call for A followed by B will actually make a diff for both A and B
+		{
+			// If AllPropertiesSnapshot is empty after the first snapshot, then it means that a previous snapshot was unfiltered and we should respect that
+			const bool bCanMergeProperties = !DiffableObjectSnapshot || AllPropertiesSnapshot.Num() > 0;
+			if (bCanMergeProperties && Properties.Num() > 0)
+			{
+				for (const FProperty* Property : Properties)
+				{
+					AllPropertiesSnapshot.AddUnique(Property);
+				}
+			}
+			else
+			{
+				// No properties means a request to snapshot everything
+				AllPropertiesSnapshot.Reset();
+			}
+		}
+
 		// Serialize the object so we can diff it
 		UE::Transaction::FDiffableObject CurrentDiffableObject;
 		{
 			CurrentDiffableObject.SetObject(CurrentObject);
-			UE::Transaction::FDiffableObjectDataWriter Writer(CurrentDiffableObject, Properties);
+			UE::Transaction::FDiffableObjectDataWriter Writer(CurrentDiffableObject, AllPropertiesSnapshot);
 			// although it would be preferable to use SerializeScriptProperties, this cause a false diff between the first snapshot and the base object
 			// since they were serialized with different algo and we don't record enough context to make the comparison appropriately
 			SerializeObject(Writer);
