Removed outdated #define SourceCondeNavigation that was added for clang 3.7-alpha on Windows in 2015 at CL 2532260 because the code would crash the compiler. The contributor reports that this codes compile fine with clang 14.0.4.

PR #9358: Update SourceCodeNavigation.cpp (Contributed by [REDACTED_PERSON])

diff --git a/Engine/Source/Editor/UnrealEd/Private/SourceCodeNavigation.cpp b/Engine/Source/Editor/UnrealEd/Private/SourceCodeNavigation.cpp
--- a/Engine/Source/Editor/UnrealEd/Private/SourceCodeNavigation.cpp
+++ b/Engine/Source/Editor/UnrealEd/Private/SourceCodeNavigation.cpp
@@ -769,11 +769,9 @@ void FSourceCodeNavigation::Initialize()
 
 const FSourceFileDatabase& FSourceCodeNavigation::GetSourceFileDatabase()
 {
-#if !( PLATFORM_WINDOWS && defined(__clang__) )		// @todo clang: This code causes a strange stack overflow issue when compiling using Clang on Windows
 	// Lock so that nothing may proceed while the AsyncTask is constructing the FSourceFileDatabase for the first time
 	FScopeLock Lock(&CriticalSection);
 	Instance.UpdateIfNeeded();
-#endif
 
 	return Instance;
 }
