MPCook bugfixes: (1) Handle Socket->Send returning false due to  EWOULDBLOCK as a need for try again later rather than a connection failure.
(2) Call ClearHierarchyTimers in CookWorker mode by making it called from ShutdownCookSession no matter what the mode is. Clear is required to avoid an assertion on shutdown.
#rnx

diff --git a/Engine/Source/Editor/UnrealEd/Private/CookOnTheFlyServer.cpp b/Engine/Source/Editor/UnrealEd/Private/CookOnTheFlyServer.cpp
--- a/Engine/Source/Editor/UnrealEd/Private/CookOnTheFlyServer.cpp
+++ b/Engine/Source/Editor/UnrealEd/Private/CookOnTheFlyServer.cpp
@@ -1299,6 +1299,7 @@ uint32 UCookOnTheFlyServer::TickCookWorker()
 	if (StackData.bCookCancelled)
 	{
 		CancelAllQueues();
+		// Make sure no UE_SCOPED_HIERARCHICAL_COOKTIMERs are around ShutdownCookSession, as ShutdownCookSession deletes memory for them
 		ShutdownCookSession();
 		SetIdleStatus(StackData, EIdleStatus::Done);
 	}
@@ -4807,7 +4808,6 @@ void UCookOnTheFlyServer::ShutdownCookOnTheFly()
 		UE_LOG(LogCook, Display, TEXT("Shutting down cook on the fly server"));
 		CookOnTheFlyRequestManager->Shutdown();
 		CookOnTheFlyRequestManager.Reset();
-		ClearHierarchyTimers();
 
 		ShutdownCookSession();
 
@@ -8406,10 +8406,10 @@ void UCookOnTheFlyServer::ShutdownCookSession()
 		PrintFinishStats();
 
 		OutputHierarchyTimers();
-		ClearHierarchyTimers();
 	}
 	CookByTheBookOptions->ClearSessionData();
 	PlatformManager->ClearSessionPlatforms(*this);
+	ClearHierarchyTimers();
 }
 
 void UCookOnTheFlyServer::PrintFinishStats()
diff --git a/Engine/Source/Editor/UnrealEd/Private/Cooker/CompactBinaryTCP.cpp b/Engine/Source/Editor/UnrealEd/Private/Cooker/CompactBinaryTCP.cpp
--- a/Engine/Source/Editor/UnrealEd/Private/Cooker/CompactBinaryTCP.cpp
+++ b/Engine/Source/Editor/UnrealEd/Private/Cooker/CompactBinaryTCP.cpp
@@ -229,9 +229,13 @@ EConnectionStatus FCompactBinaryTCPImpl::PollSendBytes(FSocket* Socket, const vo
 		{
 			InOutBytesWritten += BytesSent;
 		}
-		if (!bResult)
+
+		// Return values of false indicate the socket send failed with an error, but the error
+		// might be temporary, such as EWOULDBLOCK. Even when return value is false, return incomplete
+		// unless socket's ConnectionState has been set to error
+		if (!bResult && Socket->GetConnectionState() == SCS_ConnectionError)
 		{
-			UE_LOG(LogSockets, Error, TEXT("Could not WritePacket to Socket. Socket->Send failed."));
+			UE_LOG(LogSockets, Error, TEXT("Could not WritePacket to Socket. Socket is in the error state."));
 			return EConnectionStatus::Failed;
 		}
 		else
