Allow to generate packages at a different root than the one from the OwnerPackage.

#test cook ancient game, cook city sample, cook qaentry

diff --git a/Engine/Source/Editor/UnrealEd/Private/CookOnTheFlyServer.cpp b/Engine/Source/Editor/UnrealEd/Private/CookOnTheFlyServer.cpp
--- a/Engine/Source/Editor/UnrealEd/Private/CookOnTheFlyServer.cpp
+++ b/Engine/Source/Editor/UnrealEd/Private/CookOnTheFlyServer.cpp
@@ -3148,7 +3148,7 @@ UE::Cook::EPollStatus UCookOnTheFlyServer::BeginCacheObjectsToMove(UE::Cook::FGe
 		}
 		else
 		{
-			ICookPackageSplitter::FGeneratedPackageForPopulate SplitterInfo{ Info.RelativePath, Package, Info.IsCreateAsMap() };
+			ICookPackageSplitter::FGeneratedPackageForPopulate SplitterInfo{ Info.RelativePath, Info.GeneratedRootPath, Package, Info.IsCreateAsMap() };
 			bPopulateSucceeded = Splitter->PopulateGeneratedPackage(Package, SplitDataObject, SplitterInfo,
 				ObjectsToMove, ModifiedPackages);
 		}
@@ -3229,6 +3229,7 @@ void UCookOnTheFlyServer::ConstructGeneratedPackagesForPresave(UE::Cook::FPackag
 	{
 		ICookPackageSplitter::FGeneratedPackageForPreSave& SplitterData = GeneratedPackagesForPresave.Emplace_GetRef();
 		SplitterData.RelativePath = Info.RelativePath;
+		SplitterData.GeneratedRootPath = Info.GeneratedRootPath;
 		SplitterData.bCreatedAsMap = Info.IsCreateAsMap();
 
 		const FString GeneratedPackageName = Info.PackageData->GetPackageName().ToString();
@@ -3374,6 +3375,7 @@ UE::Cook::EPollStatus UCookOnTheFlyServer::TryPopulateGeneratedPackage(UE::Cook:
 	// properly setup any internal reference to this package (SoftObjectPaths or others)
 	ICookPackageSplitter::FGeneratedPackageForPopulate PopulateData;
 	PopulateData.RelativePath = GeneratedInfo.RelativePath;
+	PopulateData.GeneratedRootPath = GeneratedInfo.GeneratedRootPath;
 	PopulateData.Package = GeneratedPackage;
 	PopulateData.bCreatedAsMap = GeneratedInfo.IsCreateAsMap();
 	TArray<UPackage*> ModifiedPackages;
diff --git a/Engine/Source/Editor/UnrealEd/Private/Cooker/CookPackageData.cpp b/Engine/Source/Editor/UnrealEd/Private/Cooker/CookPackageData.cpp
--- a/Engine/Source/Editor/UnrealEd/Private/Cooker/CookPackageData.cpp
+++ b/Engine/Source/Editor/UnrealEd/Private/Cooker/CookPackageData.cpp
@@ -1375,8 +1375,10 @@ bool FGeneratorPackage::TryGenerateList(UObject* OwnerObject, FPackageDatas& Pac
 		}
 		bool bCreateAsMap = *SplitterData.GetCreateAsMap();
 
-		FString PackageName = FPaths::RemoveDuplicateSlashes(FString::Printf(TEXT("%s/%s/%s"),
-			*OwnerPackageData.GetPackageName().ToString(), GeneratedPackageSubPath, *SplitterData.RelativePath));
+		FString PackageRoot = SplitterData.GeneratedRootPath.IsEmpty() ? OwnerPackageData.GetPackageName().ToString() : SplitterData.GeneratedRootPath;
+
+		FString PackageName = FPaths::RemoveDuplicateSlashes(FString::Printf(TEXT("/%s/%s/%s"),
+			*PackageRoot, GeneratedPackageSubPath, *SplitterData.RelativePath));
 		const FName PackageFName(*PackageName);
 		UE::Cook::FPackageData* PackageData = PackageDatas.TryAddPackageDataByPackageName(PackageFName,
 			false /* bRequireExists */, bCreateAsMap);
@@ -1400,6 +1402,7 @@ bool FGeneratorPackage::TryGenerateList(UObject* OwnerObject, FPackageDatas& Pac
 
 		FCookGenerationInfo& GeneratedInfo = PackagesToGenerate.Emplace_GetRef(*PackageData, false /* bInGenerator */);
 		GeneratedInfo.RelativePath = MoveTemp(SplitterData.RelativePath);
+		GeneratedInfo.GeneratedRootPath = MoveTemp(SplitterData.GeneratedRootPath);
 		GeneratedInfo.Dependencies = MoveTemp(SplitterData.Dependencies);
 		GeneratedInfo.SetIsCreateAsMap(bCreateAsMap);
 		PackageData->SetGeneratedOwner(this);
@@ -1601,6 +1604,7 @@ void FGeneratorPackage::ResetSaveState(FCookGenerationInfo& Info, UPackage* Pack
 			{
 				ICookPackageSplitter::FGeneratedPackageForPopulate PopulateInfo;
 				PopulateInfo.RelativePath = Info.RelativePath;
+				PopulateInfo.GeneratedRootPath = Info.GeneratedRootPath;
 				PopulateInfo.bCreatedAsMap = Info.IsCreateAsMap();
 				PopulateInfo.Package = Package;
 				GetCookPackageSplitterInstance()->PostSaveGeneratedPackage(GetOwnerPackage(), SplitObject, PopulateInfo);
diff --git a/Engine/Source/Editor/UnrealEd/Private/Cooker/CookPackageData.h b/Engine/Source/Editor/UnrealEd/Private/Cooker/CookPackageData.h
--- a/Engine/Source/Editor/UnrealEd/Private/Cooker/CookPackageData.h
+++ b/Engine/Source/Editor/UnrealEd/Private/Cooker/CookPackageData.h
@@ -708,6 +708,7 @@ public:
 public:
 	FBeginCacheObjects BeginCacheObjects;
 	FString RelativePath;
+	FString GeneratedRootPath;
 	TArray<FName> Dependencies;
 	FPackageData* PackageData = nullptr;
 	TArray<UPackage*> ModifiedPackages;
diff --git a/Engine/Source/Editor/UnrealEd/Public/CookPackageSplitter.h b/Engine/Source/Editor/UnrealEd/Public/CookPackageSplitter.h
--- a/Engine/Source/Editor/UnrealEd/Public/CookPackageSplitter.h
+++ b/Engine/Source/Editor/UnrealEd/Public/CookPackageSplitter.h
@@ -40,7 +40,8 @@ public:
 	/** Data sent to the cooker to describe each desired generated package */
 	struct FGeneratedPackage
 	{
-		FString RelativePath;		// Generated package relative to the Parent/_Generated_ root
+		FString GeneratedRootPath;
+		FString RelativePath;		// Generated package relative to GeneratedRootPath/Generated_ if GeneratedRootPath member is specified, relative to  Parent/Generated otherwise.
 		TArray<FName> Dependencies; // LongPackageNames that the generated package references
 		/* GetGenerateList must specify true if the package will be a map (.umap, contains a UWorld or ULevel), else false */
 		void SetCreateAsMap(bool bInCreateAsMap) { bCreateAsMap = bInCreateAsMap; }
@@ -68,6 +69,8 @@ public:
 	{
 		/** RelativePath returned from GetGenerateList */
 		FString RelativePath;
+		/** Root returned from GetGenerateList */
+		FString GeneratedRootPath;
 		/**
 		 * Non-null, constructed package. May have been previously passed into PreSaveGeneratedPackage,
 		 * or may be completely empty if that previous package was garbage collected before this call.
@@ -116,6 +119,8 @@ public:
 	{
 		/** RelativePath returned from GetGenerateList */
 		FString RelativePath;
+		/** Root returned from GetGenerateList */
+		FString GeneratedRootPath;
 		/**
 		 * Non-null, constructed package. May have been previously passed into PreSaveGeneratedPackage,
 		 * or may be completely empty if that previous package was garbage collected before this call.
diff --git a/Engine/Source/Runtime/Engine/Private/WorldPartition/WorldPartitionCookPackageSplitter.cpp b/Engine/Source/Runtime/Engine/Private/WorldPartition/WorldPartitionCookPackageSplitter.cpp
--- a/Engine/Source/Runtime/Engine/Private/WorldPartition/WorldPartitionCookPackageSplitter.cpp
+++ b/Engine/Source/Runtime/Engine/Private/WorldPartition/WorldPartitionCookPackageSplitter.cpp
@@ -110,6 +110,7 @@ TArray<ICookPackageSplitter::FGeneratedPackage> FWorldPartitionCookPackageSplitt
 	{
 		ICookPackageSplitter::FGeneratedPackage& GeneratedPackage = PackagesToGenerate.Emplace_GetRef();
 		GeneratedPackage.RelativePath = PackageName;
+		GeneratedPackage.GeneratedRootPath = OwnerPackage->GetName();
 		// all packages we generate get a ULevel from CreateEmptyLevelForRuntimeCell and are hence maps
 		GeneratedPackage.SetCreateAsMap(true);
 		// @todo_ow: Set dependencies once we get iterative cooking working
