Remove Go/Edit To Blueprint menu options when the blueprint is not accessible

diff --git a/Engine/Source/Editor/SubobjectEditor/Private/SSubobjectEditor.cpp b/Engine/Source/Editor/SubobjectEditor/Private/SSubobjectEditor.cpp
--- a/Engine/Source/Editor/SubobjectEditor/Private/SSubobjectEditor.cpp
+++ b/Engine/Source/Editor/SubobjectEditor/Private/SSubobjectEditor.cpp
@@ -1724,11 +1724,10 @@ TSharedRef<SWidget> SSubobject_RowWidget::GetInheritedLinkWidget()
 {
 	FSubobjectEditorTreeNodePtrType NodePtr = GetSubobjectPtr();
     const FSubobjectData* Data = NodePtr ? NodePtr->GetDataSource() : nullptr;
-
 	if(!Data)
 	{
 		return SNullWidget::NullWidget;
-	}	
+	}
 	
 	// Native components are inherited and have a gray hyperlink to their C++ class
 	if(Data->IsNativeComponent())
@@ -1752,7 +1751,8 @@ TSharedRef<SWidget> SSubobject_RowWidget::GetInheritedLinkWidget()
 				.Style(FAppStyle::Get(), "Common.GotoBlueprintHyperlink")
 				.OnNavigate(this, &SSubobject_RowWidget::OnEditBlueprintClicked)
 				.Text(InheritedBPLabel)
-				.ToolTipText(LOCTEXT("EditBlueprint_ToolTip", "Click to edit the blueprint"));
+				.ToolTipText(LOCTEXT("EditBlueprint_ToolTip", "Click to edit the blueprint"))
+				.Visibility(this, &SSubobject_RowWidget::GetEditBlueprintVisibility);
 		}
 	}
 
@@ -1796,6 +1796,24 @@ void SSubobject_RowWidget::OnEditBlueprintClicked()
 	}
 }
 
+EVisibility SSubobject_RowWidget::GetEditBlueprintVisibility() const
+{
+	FSubobjectEditorTreeNodePtrType NodePtr = GetSubobjectPtr();
+	const FSubobjectData* Data = NodePtr ? NodePtr->GetDataSource() : nullptr;
+
+	if (!Data)
+	{
+		return EVisibility::Collapsed;
+	}
+
+	if (const UBlueprint* BP = Data->GetBlueprint())
+	{
+		return GEditor->GetEditorSubsystem<UAssetEditorSubsystem>()->IsAssetEditable(BP) ? EVisibility::Visible : EVisibility::Collapsed;
+	}
+
+	return EVisibility::Collapsed;
+}
+
 EVisibility SSubobject_RowWidget::GetEditNativeCppVisibility() const
 {
 	return ensure(GUnrealEd) && GUnrealEd->GetUnrealEdOptions()->IsCPPAllowed() ? EVisibility::Visible : EVisibility::Collapsed;
diff --git a/Engine/Source/Editor/SubobjectEditor/Public/SSubobjectEditor.h b/Engine/Source/Editor/SubobjectEditor/Public/SSubobjectEditor.h
--- a/Engine/Source/Editor/SubobjectEditor/Public/SSubobjectEditor.h
+++ b/Engine/Source/Editor/SubobjectEditor/Public/SSubobjectEditor.h
@@ -322,6 +322,7 @@ private:
 	/** Callback used when the user clicks on a blueprint inherited variable */
 	void OnEditBlueprintClicked();
 
+	EVisibility GetEditBlueprintVisibility() const;
 	EVisibility GetEditNativeCppVisibility() const;
 	
 	/** Callback used when the user clicks on a native inherited variable */
diff --git a/Engine/Source/Editor/UnrealEd/Private/EditorClassUtils.cpp b/Engine/Source/Editor/UnrealEd/Private/EditorClassUtils.cpp
--- a/Engine/Source/Editor/UnrealEd/Private/EditorClassUtils.cpp
+++ b/Engine/Source/Editor/UnrealEd/Private/EditorClassUtils.cpp
@@ -7,6 +7,7 @@
 #include "Engine/Blueprint.h"
 #include "Editor.h"
 #include "AssetRegistry/AssetRegistryModule.h"
+#include "Subsystems/AssetEditorSubsystem.h"
 
 #include "IDocumentationPage.h"
 #include "IDocumentation.h"
@@ -149,22 +150,66 @@ TSharedRef<SWidget> FEditorClassUtils::GetSourceLink(const UClass* Class, const
 
 		TWeakObjectPtr<UBlueprint> BlueprintPtr = Blueprint;
 
-		FText Text = FText::FromName(Blueprint->GetFName());
+		auto CanEditBlueprint = [BlueprintPtr]()
+		{
+			if (const UBlueprint* Blueprint = BlueprintPtr.Get())
+			{
+				return GEditor->GetEditorSubsystem<UAssetEditorSubsystem>()->IsAssetEditable(Blueprint);
+			}
+
+			return false;
+		};
+
+		FText BlueprintNameText = FText::FromName(Blueprint->GetFName());
+
+		FText FormattedText;
 		if (Params.BlueprintFormat)
 		{
-			Text = FText::Format(*Params.BlueprintFormat, Text);
+			FormattedText = FText::Format(*Params.BlueprintFormat, BlueprintNameText);
 		}
 		else if (Params.bUseDefaultFormat)
 		{
 			const FText DefaultBlueprintFormat = NSLOCTEXT("SourceHyperlink", "EditBlueprint", "Edit {0}");
-			Text = FText::Format(DefaultBlueprintFormat, Text);
+			FormattedText = FText::Format(DefaultBlueprintFormat, BlueprintNameText);
+		}
+		else
+		{
+			FormattedText = BlueprintNameText;
 		}
 
-		Link = SNew(SHyperlink)
-			.Style(FAppStyle::Get(), "Common.GotoBlueprintHyperlink")
-			.OnNavigate_Static(&Local::OnEditBlueprintClicked, BlueprintPtr, Params.Object)
-			.Text(Text)
-			.ToolTipText(NSLOCTEXT("SourceHyperlink", "EditBlueprint_ToolTip", "Click to edit the blueprint"));
+		TSharedRef<SWidget> NoLinkWidget = SNullWidget::NullWidget;
+		if (Params.bEmptyIfNoLink)
+		{
+			NoLinkWidget = SNew(SSpacer)
+				.Visibility_Lambda([CanEditBlueprint]() { return CanEditBlueprint() ? EVisibility::Collapsed : EVisibility::Visible; });
+		}
+		else
+		{
+			NoLinkWidget = SNew(STextBlock)
+				.Text(Params.bUseFormatIfNoLink ? FormattedText : BlueprintNameText)
+				.ColorAndOpacity(FSlateColor::UseSubduedForeground())
+				.Visibility_Lambda([CanEditBlueprint]() { return CanEditBlueprint() ? EVisibility::Collapsed : EVisibility::Visible; });
+		}
+
+		Link = SNew(SHorizontalBox)
+				+SHorizontalBox::Slot()
+				.AutoWidth()
+				.VAlign(VAlign_Center)
+				[
+					SNew(SHyperlink)
+					.Style(FAppStyle::Get(), "Common.GotoBlueprintHyperlink")
+					.OnNavigate_Static(&Local::OnEditBlueprintClicked, BlueprintPtr, Params.Object)
+					.Text(FormattedText)
+					.ToolTipText(NSLOCTEXT("SourceHyperlink", "EditBlueprint_ToolTip", "Click to edit the blueprint"))
+					.Visibility_Lambda([CanEditBlueprint]() { return CanEditBlueprint() ? EVisibility::Visible : EVisibility::Collapsed; })
+				]
+				+SHorizontalBox::Slot()
+				.AutoWidth()
+				.VAlign(VAlign_Center)
+				[
+					NoLinkWidget
+				];
+
 	}
 	else if (Class)
 	{
diff --git a/Engine/Source/Editor/UnrealEd/Private/Subsystems/AssetEditorSubsystem.cpp b/Engine/Source/Editor/UnrealEd/Private/Subsystems/AssetEditorSubsystem.cpp
--- a/Engine/Source/Editor/UnrealEd/Private/Subsystems/AssetEditorSubsystem.cpp
+++ b/Engine/Source/Editor/UnrealEd/Private/Subsystems/AssetEditorSubsystem.cpp
@@ -284,6 +284,30 @@ bool UAssetEditorSubsystem::CloseAllAssetEditors()
 	return bAllEditorsClosed;
 }
 
+bool UAssetEditorSubsystem::IsAssetEditable(const UObject* Asset)
+{
+	FAssetToolsModule& AssetToolsModule = FModuleManager::LoadModuleChecked<FAssetToolsModule>(TEXT("AssetTools"));
+
+	if (!Asset)
+	{
+		return false;
+	}
+
+	if (UPackage* Package = Asset->GetPackage())
+	{
+		if (Package->bIsCookedForEditor)
+		{
+			return false;
+		}
+
+		if (!AssetToolsModule.Get().GetWritableFolderPermissionList()->PassesStartsWithFilter(Package->GetName()))
+		{
+			return false;
+		}
+	}
+
+	return true;
+}
 
 bool UAssetEditorSubsystem::OpenEditorForAsset(UObject* Asset, const EToolkitMode::Type ToolkitMode, TSharedPtr< IToolkitHost > OpenedFromLevelEditor, const bool bShowProgressWindow)
 {
diff --git a/Engine/Source/Editor/UnrealEd/Public/Subsystems/AssetEditorSubsystem.h b/Engine/Source/Editor/UnrealEd/Public/Subsystems/AssetEditorSubsystem.h
--- a/Engine/Source/Editor/UnrealEd/Public/Subsystems/AssetEditorSubsystem.h
+++ b/Engine/Source/Editor/UnrealEd/Public/Subsystems/AssetEditorSubsystem.h
@@ -70,6 +70,8 @@ public:
 	virtual void Initialize(FSubsystemCollectionBase& Collection) override;
 	virtual void Deinitialize() override;
 
+	bool IsAssetEditable(const UObject* Asset);
+
 	/** Opens an asset by path */
 	void OpenEditorForAsset(const FString& AssetPathName);
 
