Adds previously made writable files to Uncontrolled Changelist when modified.
Removes writable files not checkout warning when Uncontrolled Changelist module is enabled

diff --git a/Engine/Source/Editor/UnrealEd/Private/FileHelpers.cpp b/Engine/Source/Editor/UnrealEd/Private/FileHelpers.cpp
--- a/Engine/Source/Editor/UnrealEd/Private/FileHelpers.cpp
+++ b/Engine/Source/Editor/UnrealEd/Private/FileHelpers.cpp
@@ -3403,8 +3403,19 @@ static InternalSavePackageResult InternalSavePackage(UPackage* PackageToSave, bo
 		}
 		else
 		{
-			// If source control is disabled then we don't care if the package is locally writable
-			bOutPackageLocallyWritable = false;
+			FUncontrolledChangelistsModule& UncontrolledChangelistsModule = FUncontrolledChangelistsModule::Get();
+
+			// If we are in offline mode, automatically add the modified package to an Uncontrolled Changelist
+			if (UncontrolledChangelistsModule.IsEnabled() && (!IFileManager::Get().IsReadOnly(*FinalPackageSavePath)))
+			{
+				UncontrolledChangelistsModule.OnMakeWritable({ FinalPackageSavePath });
+				bOutPackageLocallyWritable = true;
+			}
+			else
+			{
+				// If source control is disabled then we don't care if the package is locally writable
+				bOutPackageLocallyWritable = false;
+			}
 		}
 
 		// Handle all failures the same way.
@@ -4033,8 +4044,9 @@ FEditorFileUtils::EPromptReturnCode InternalPromptForCheckoutAndSave(const TArra
 	}
 
 	// If any packages were saved that weren't actually in source control but instead forcibly made writable,
-	// then warn the user about those packages
-	if (WritablePackageFiles.Num() > 0)
+	// then warn the user about those packages. We do not warn if the Uncontrolled Changelists are enabled since
+	// the file will be picked up.
+	if (!FUncontrolledChangelistsModule::Get().IsEnabled() && (WritablePackageFiles.Num() > 0))
 	{
 		FString WritableFiles;
 		for (UPackage* PackageIter : WritablePackageFiles)
