Changed "Compressing Animations" toast to say "Preparing Animations", unifying the message with other asset processing tasks

diff --git a/Engine/Source/Editor/UnrealEd/Private/Animation/AsyncAnimCompressionUI.cpp b/Engine/Source/Editor/UnrealEd/Private/Animation/AsyncAnimCompressionUI.cpp
--- a/Engine/Source/Editor/UnrealEd/Private/Animation/AsyncAnimCompressionUI.cpp
+++ b/Engine/Source/Editor/UnrealEd/Private/Animation/AsyncAnimCompressionUI.cpp
@@ -8,12 +8,12 @@
 
 /** Global notification object. */
 
-/** Notification class for asynchronous shader compiling. */
+/** Notification class for asynchronous animation processing. */
 class FAnimCompressionNotificationImpl : public FGlobalEditorProgressNotification
 {
 public:
 	FAnimCompressionNotificationImpl()
-		: FGlobalEditorProgressNotification(NSLOCTEXT("AsyncAnimCompression", "AnimCompressionInProgress", "Compressing Animations"))
+		: FGlobalEditorProgressNotification(NSLOCTEXT("AsyncAnimCompression", "AnimProcessing", "Preparing Animations"))
 	{}
 protected:
 	/** FGlobalEditorProgressNotification interface */
