Added an "Assign Skeleton" action to AnimationBlueprints

diff --git a/Engine/Source/Developer/AssetTools/Private/AssetTypeActions/AssetTypeActions_AnimBlueprint.cpp b/Engine/Source/Developer/AssetTools/Private/AssetTypeActions/AssetTypeActions_AnimBlueprint.cpp
--- a/Engine/Source/Developer/AssetTools/Private/AssetTypeActions/AssetTypeActions_AnimBlueprint.cpp
+++ b/Engine/Source/Developer/AssetTools/Private/AssetTypeActions/AssetTypeActions_AnimBlueprint.cpp
@@ -141,6 +141,18 @@ void FAssetTypeActions_AnimBlueprint::GetActions(const TArray<UObject*>& InObjec
 				FCanExecuteAction()
 				)
 			);
+
+		Section.AddMenuEntry(
+			"AnimBlueprint_AssignSkeleton",
+			LOCTEXT("SkeletalMesh_AssignSkeleton", "Assign Skeleton"),
+			LOCTEXT("SkeletalMesh_AssignSkeletonTooltip", "Assigns a skeleton to the selected Animation Blueprint(s)."),
+			FSlateIcon(FAppStyle::GetAppStyleSetName(), "Persona.AssetActions.AssignSkeleton"),
+			FUIAction(
+				FExecuteAction::CreateSP(this, &FAssetTypeActions_AnimBlueprint::ExecuteAssignSkeleton, AnimBlueprints),
+				FCanExecuteAction()
+			)
+		);
+
 	}
 }
 
@@ -286,6 +298,40 @@ void FAssetTypeActions_AnimBlueprint::ExecuteFindSkeleton(TArray<TWeakObjectPtr<
 	}
 }
 
+void FAssetTypeActions_AnimBlueprint::ExecuteAssignSkeleton(TArray<TWeakObjectPtr<UAnimBlueprint>> Objects)
+{
+	if (Objects.Num() > 0)
+	{
+		TSharedRef<SWindow> WidgetWindow = SNew(SWindow)
+			.Title(LOCTEXT("ChooseSkeletonWindowTitle", "Choose Skeleton"))
+			.ClientSize(FVector2D(400, 600));
+
+		TSharedPtr<SSkeletonSelectorWindow> SkeletonSelectorWindow;
+		WidgetWindow->SetContent
+		(
+			SNew(SBorder)
+			.BorderImage(FAppStyle::GetBrush("ToolPanel.GroupBorder"))
+			[
+				SAssignNew(SkeletonSelectorWindow, SSkeletonSelectorWindow)
+				.Object(Objects[0].Get())
+			.WidgetWindow(WidgetWindow)
+			]
+		);
+
+		GEditor->EditorAddModalWindow(WidgetWindow);
+		USkeleton* SelectedSkeleton = SkeletonSelectorWindow->GetSelectedSkeleton();
+
+		// only do this if not same
+		if (SelectedSkeleton)
+		{
+			for (TWeakObjectPtr<UAnimBlueprint>& AnimBP : Objects)
+			{
+				AnimBP->TargetSkeleton = SelectedSkeleton;
+			}
+		}
+	}
+}
+
 bool FAssetTypeActions_AnimBlueprint::ReplaceMissingSkeleton(TArray<UObject*> InAnimBlueprints) const
 {
 	// record anim assets that need skeleton replaced
diff --git a/Engine/Source/Developer/AssetTools/Private/AssetTypeActions/AssetTypeActions_AnimBlueprint.h b/Engine/Source/Developer/AssetTools/Private/AssetTypeActions/AssetTypeActions_AnimBlueprint.h
--- a/Engine/Source/Developer/AssetTools/Private/AssetTypeActions/AssetTypeActions_AnimBlueprint.h
+++ b/Engine/Source/Developer/AssetTools/Private/AssetTypeActions/AssetTypeActions_AnimBlueprint.h
@@ -31,6 +31,9 @@ private:
 
 	/** Handler for when FindSkeleton is selected */
 	void ExecuteFindSkeleton(TArray<TWeakObjectPtr<UAnimBlueprint>> Objects);
+
+	/** Handler for when AssignSkeleton is selected */
+	void ExecuteAssignSkeleton(TArray<TWeakObjectPtr<UAnimBlueprint>> Objects);
 	
 	/** Certain options are not available for template anim BPs */
 	bool AreAnyNonTemplateAnimBlueprintsSelected(TArray<TWeakObjectPtr<UAnimBlueprint>> Objects) const;
diff --git a/Engine/Source/Editor/UnrealEd/Private/SSkeletonWidget.cpp b/Engine/Source/Editor/UnrealEd/Private/SSkeletonWidget.cpp
--- a/Engine/Source/Editor/UnrealEd/Private/SSkeletonWidget.cpp
+++ b/Engine/Source/Editor/UnrealEd/Private/SSkeletonWidget.cpp
@@ -15,7 +15,6 @@
 #include "Editor.h"
 #include "Animation/AnimSet.h"
 #include "Interfaces/IMainFrameModule.h"
-#include "IContentBrowserSingleton.h"
 #include "ContentBrowserModule.h"
 #include "IDocumentation.h"
 #include "Animation/Rig.h"
@@ -30,6 +29,8 @@
 
 void SSkeletonListWidget::Construct(const FArguments& InArgs)
 {
+	bShowBones = InArgs._ShowBones;
+	InitialViewType = InArgs._InitialViewType;
 	CurSelectedSkeleton = nullptr;
 
 	FContentBrowserModule& ContentBrowserModule = FModuleManager::Get().LoadModuleChecked<FContentBrowserModule>(TEXT("ContentBrowser"));
@@ -37,34 +38,37 @@ void SSkeletonListWidget::Construct(const FArguments& InArgs)
 	FAssetPickerConfig AssetPickerConfig;
 	AssetPickerConfig.Filter.ClassPaths.Add(USkeleton::StaticClass()->GetClassPathName());
 	AssetPickerConfig.OnAssetSelected = FOnAssetSelected::CreateSP(this, &SSkeletonListWidget::SkeletonSelectionChanged);
-	AssetPickerConfig.InitialAssetViewType = EAssetViewType::Column;
+	AssetPickerConfig.InitialAssetViewType = InitialViewType;
 	AssetPickerConfig.SelectionMode = ESelectionMode::Single;
 	AssetPickerConfig.bShowPathInColumnView = true;
 	AssetPickerConfig.bShowTypeInColumnView = false;
 
-	this->ChildSlot
+	TSharedRef<SVerticalBox> ContentBox = SNew(SVerticalBox)
+		+ SVerticalBox::Slot()
+		.AutoHeight()
 		[
-			SNew(SVerticalBox)
-
-			+SVerticalBox::Slot()
-			.AutoHeight()
+			SNew(STextBlock)
+			.Text(LOCTEXT("SelectSkeletonLabel", "Select Skeleton: "))
+		]
+		+ SVerticalBox::Slot().FillHeight(1).Padding(2)
 			[
-				SNew(STextBlock)
-				.Text(LOCTEXT("SelectSkeletonLabel", "Select Skeleton: "))
-			]
-
-			+SVerticalBox::Slot() .FillHeight(1) .Padding(2)
+				SNew(SBorder)
+				.Content()
 				[
-					SNew(SBorder)
-					.Content()
-					[
-						ContentBrowserModule.Get().CreateAssetPicker(AssetPickerConfig)
-					]
+					ContentBrowserModule.Get().CreateAssetPicker(AssetPickerConfig)
 				]
+			];
 
-			+SVerticalBox::Slot() .FillHeight(1) .Padding(2)
-				.Expose(BoneListSlot)
+	if (bShowBones)
+	{
+		ContentBox->AddSlot()
+			.FillHeight(1).Padding(2)
+			.Expose(BoneListSlot);
+	}
 
+	ChildSlot
+		[
+			ContentBox
 		];
 
 	// Construct the BoneListSlot by clearing the skeleton selection. 
@@ -76,43 +80,46 @@ void SSkeletonListWidget::SkeletonSelectionChanged(const FAssetData& AssetData)
 	BoneList.Empty();
 	CurSelectedSkeleton = Cast<USkeleton>(AssetData.GetAsset());
 
-	if (CurSelectedSkeleton != nullptr)
+	if (bShowBones)
 	{
-		const FReferenceSkeleton& RefSkeleton = CurSelectedSkeleton->GetReferenceSkeleton();
-
-		for (int32 I=0; I<RefSkeleton.GetNum(); ++I)
+		if (CurSelectedSkeleton != nullptr)
 		{
-			BoneList.Add( MakeShareable(new FName(RefSkeleton.GetBoneName(I))) ) ;
-		}
+			const FReferenceSkeleton& RefSkeleton = CurSelectedSkeleton->GetReferenceSkeleton();
 
-		(*BoneListSlot)
-			[
-				SNew(SBorder) .Padding(2)
-				.Content()
+			for (int32 I = 0; I < RefSkeleton.GetNum(); ++I)
+			{
+				BoneList.Add(MakeShareable(new FName(RefSkeleton.GetBoneName(I))));
+			}
+
+			(*BoneListSlot)
+				[
+					SNew(SBorder).Padding(2)
+					.Content()
 				[
 					SNew(SListView< TSharedPtr<FName> >)
 					.OnGenerateRow(this, &SSkeletonListWidget::GenerateSkeletonBoneRow)
-					.ListItemsSource(&BoneList)
-					.HeaderRow
-					(
+				.ListItemsSource(&BoneList)
+				.HeaderRow
+				(
 					SNew(SHeaderRow)
-					+SHeaderRow::Column(TEXT("Bone Name"))
+					+ SHeaderRow::Column(TEXT("Bone Name"))
 					.DefaultLabel(NSLOCTEXT("SkeletonWidget", "BoneName", "Bone Name"))
-					)
+				)
 				]
-			];
-	}
-	else
-	{
-		(*BoneListSlot)
-			[
-				SNew(SBorder) .Padding(2)
-				.Content()
+				];
+		}
+		else
+		{
+			(*BoneListSlot)
+				[
+					SNew(SBorder).Padding(2)
+					.Content()
 				[
 					SNew(STextBlock)
 					.Text(NSLOCTEXT("SkeletonWidget", "NoSkeletonIsSelected", "No skeleton is selected!"))
 				]
-			];
+				];
+		}
 	}
 }
 
@@ -282,6 +289,10 @@ void SSkeletonSelectorWindow::Construct(const FArguments& InArgs)
 	{
 		ConstructWindowFromAnimSet(CastChecked<UAnimSet>(Object));
 	}
+	else if (Object->IsA(UAnimBlueprint::StaticClass()))
+	{
+		ConstructWindowFromAnimBlueprint(CastChecked<UAnimBlueprint>(Object));
+	}
 }
 
 void SSkeletonSelectorWindow::ConstructWindowFromAnimSet(UAnimSet* InAnimSet)
@@ -331,12 +342,32 @@ void SSkeletonSelectorWindow::ConstructWindowFromMesh(USkeletalMesh* InSkeletalM
 		];
 }
 
+void SSkeletonSelectorWindow::ConstructWindowFromAnimBlueprint(UAnimBlueprint* AnimBlueprint)
+{
+	TSharedRef<SVerticalBox> ContentBox = SNew(SVerticalBox)
+		+ SVerticalBox::Slot().FillHeight(1).Padding(2)
+		[
+			SAssignNew(SkeletonWidget, SSkeletonListWidget)
+				.ShowBones(false)
+				.InitialViewType(EAssetViewType::List)
+		];
+
+	ConstructButtons(ContentBox);
+
+	ChildSlot
+		[
+			ContentBox
+		];
+}
+
 void SSkeletonSelectorWindow::ConstructWindow()
 {
 	TSharedRef<SVerticalBox> ContentBox = SNew(SVerticalBox)
 		+SVerticalBox::Slot() .FillHeight(1) .Padding(2)
 		[
 			SAssignNew(SkeletonWidget, SSkeletonListWidget)
+				.ShowBones(true)
+				.InitialViewType(EAssetViewType::Column)
 		];
 
 	ConstructButtons(ContentBox);
diff --git a/Engine/Source/Editor/UnrealEd/Public/SSkeletonWidget.h b/Engine/Source/Editor/UnrealEd/Public/SSkeletonWidget.h
--- a/Engine/Source/Editor/UnrealEd/Public/SSkeletonWidget.h
+++ b/Engine/Source/Editor/UnrealEd/Public/SSkeletonWidget.h
@@ -3,6 +3,7 @@
 #pragma once
 
 #include "CoreMinimal.h"
+#include "IContentBrowserSingleton.h"
 #include "Input/Reply.h"
 #include "Widgets/DeclarativeSyntaxSupport.h"
 #include "Widgets/SWidget.h"
@@ -121,6 +122,8 @@ class SSkeletonListWidget : public SSkeletonWidget
 {
 public:
 	SLATE_BEGIN_ARGS( SSkeletonListWidget ){}
+		SLATE_ARGUMENT(bool, ShowBones)
+		SLATE_ARGUMENT(EAssetViewType::Type, InitialViewType)
 	SLATE_END_ARGS()
 public:
 	// WIDGETS
@@ -153,6 +156,8 @@ public:
 	}
 
 private:
+	bool bShowBones = true;
+	EAssetViewType::Type InitialViewType = EAssetViewType::Column;
 	SVerticalBox::FSlot* BoneListSlot;
 	TArray< TSharedPtr<FName> > BoneList;
 };
@@ -205,6 +210,8 @@ public:
 
 	void ConstructWindowFromMesh(USkeletalMesh* InSkeletalMesh);
 
+	void ConstructWindowFromAnimBlueprint(UAnimBlueprint* AnimBlueprint);
+
 	void ConstructWindow();
 
 	void ConstructButtons(TSharedRef<SVerticalBox> ParentBox)
