Added import/export buttons to Active Theme:
- Clicking "Import..." opens "document" from the OS window. Importing a theme will update the current active theme. After import, this theme should be found where user imported it from (i.e., document) AND the themes folder within the engine. Cannot import a theme with names that already exist. Also cannot import an invalid .json file.

- Clicking "Export..." opens "document" from the OS window. After export, this theme should be found where user exported it to (i.e., document) AND remain in the themes folder within the engine.

- Both import and export will prompt success/failure message within the engine.

UI Changes:
- "Color" header for Active Theme is now called "Theme"
-  Uploaded new icons for import / export (both 16x16 and 20x20 versions)
-  ThemeID now also gets exported when exporting Appearance setting.
- "Viewport Selection Color" and "Editor Window Background Color" are now under "User Interface"

diff --git a/Engine/Source/Editor/EditorStyle/Private/StarshipStyle.cpp b/Engine/Source/Editor/EditorStyle/Private/StarshipStyle.cpp
--- a/Engine/Source/Editor/EditorStyle/Private/StarshipStyle.cpp
+++ b/Engine/Source/Editor/EditorStyle/Private/StarshipStyle.cpp
@@ -433,7 +433,14 @@ void FStarshipEditorStyle::FStyle::SetupGeneralStyles()
 		Set("Icons.Toolbar.Stop", new CORE_IMAGE_BRUSH_SVG("Starship/Common/stop", Icon20x20));
 		Set("Icons.Toolbar.Settings", new CORE_IMAGE_BRUSH_SVG( "Starship/Common/Settings", Icon20x20));
 		Set("Icons.Toolbar.Details", new IMAGE_BRUSH_SVG("Starship/Common/Details", Icon16x16));
-		Set("Icons.Toolbar.Import", new CORE_IMAGE_BRUSH_SVG( "Starship/Common/import", Icon20x20) );
+		Set("Icons.Toolbar.Import", new CORE_IMAGE_BRUSH_SVG("Starship/Common/import_20", Icon20x20));
+		Set("Icons.Toolbar.Export", new CORE_IMAGE_BRUSH_SVG("Starship/Common/export_20", Icon20x20));
+	}
+
+	// Theme Icons
+	{
+		Set("Themes.Import", new CORE_IMAGE_BRUSH_SVG("Starship/Common/import", Icon16x16));
+		Set("Themes.Export", new CORE_IMAGE_BRUSH_SVG("Starship/Common/export", Icon16x16));
 	}
 
 	// Typed Elements Icons
@@ -4108,7 +4115,7 @@ void FStarshipEditorStyle::FStyle::SetupLevelEditorStyle()
 		Set( "LevelEditor.SaveAllLevels", new IMAGE_BRUSH_SVG( "Starship/Common/LevelSaveAll", Icon16x16 ) );
 
 		Set( "LevelEditor.ImportScene",    new IMAGE_BRUSH_SVG( "Starship/Common/LevelImportInto", Icon16x16 ) );
-		Set( "LevelEditor.ExportAll",      new IMAGE_BRUSH_SVG( "Starship/Common/Export", Icon16x16 ) );
+		Set( "LevelEditor.ExportAll",      new CORE_IMAGE_BRUSH_SVG( "Starship/Common/export", Icon16x16 ) );
 		Set( "LevelEditor.ExportSelected", new IMAGE_BRUSH_SVG( "Starship/Common/ExportSelected", Icon16x16 ) );
 
 		Set( "LevelEditor.Recompile", new IMAGE_BRUSH_SVG( "Starship/MainToolbar/compile", Icon40x40 ) );
@@ -4691,7 +4698,7 @@ void FStarshipEditorStyle::FStyle::SetupPersonaStyle()
 		Set("Persona.ImportAnimation", new CORE_IMAGE_BRUSH_SVG("Starship/Common/import", Icon20x20));
 		Set("Persona.ReimportAnimation", new CORE_IMAGE_BRUSH_SVG("Starship/Common/import", Icon20x20));
 		Set("Persona.ApplyCompression", new IMAGE_BRUSH_SVG("Starship/Common/Compress", Icon20x20));
-		Set("Persona.ExportToFBX", new IMAGE_BRUSH_SVG("Starship/Common/Export", Icon20x20));
+		Set("Persona.ExportToFBX", new CORE_IMAGE_BRUSH_SVG("Starship/Common/export_20", Icon20x20));
 		Set("Persona.CreateAsset", new IMAGE_BRUSH_SVG("Starship/Persona/PersonaCreateAsset", Icon20x20));
 		Set("Persona.StartRecordAnimation", new IMAGE_BRUSH(TEXT("Icons/icon_Persona_StartRecord_40x"), Icon40x40));
 		Set("Persona.StopRecordAnimation", new IMAGE_BRUSH(TEXT("Icons/icon_Persona_StopRecord_40x"), Icon40x40));
@@ -4789,7 +4796,7 @@ void FStarshipEditorStyle::FStyle::SetupPersonaStyle()
 	// Animation editor
 	{
 		Set("AnimationEditor.ApplyCompression", new IMAGE_BRUSH_SVG("Starship/Common/Compress", Icon20x20));
-		Set("AnimationEditor.ExportToFBX", new IMAGE_BRUSH_SVG("Starship/Common/Export", Icon20x20));
+		Set("AnimationEditor.ExportToFBX", new CORE_IMAGE_BRUSH_SVG("Starship/Common/export_20", Icon20x20));
 		Set("AnimationEditor.ReimportAnimation", new CORE_IMAGE_BRUSH_SVG("Starship/Common/import", Icon20x20));
 		Set("AnimationEditor.CreateAsset", new IMAGE_BRUSH_SVG("Starship/Persona/PersonaCreateAsset", Icon20x20));
 		Set("AnimationEditor.SetKey", new CORE_IMAGE_BRUSH_SVG("Starship/Common/plus", Icon20x20));
@@ -6558,7 +6565,7 @@ void FStarshipEditorStyle::FStyle::SetupToolkitStyles()
 
 		Set( "FontEditor.Update", new IMAGE_BRUSH_SVG( "Starship/Common/Update", Icon20x20 ) );
 		Set( "FontEditor.UpdateAll", new IMAGE_BRUSH_SVG( "Starship/Common/UpdateAll", Icon20x20 ) );
-		Set( "FontEditor.ExportPage", new IMAGE_BRUSH_SVG( "Starship/Common/Export", Icon20x20 ) );
+		Set( "FontEditor.ExportPage", new CORE_IMAGE_BRUSH_SVG( "Starship/Common/export_20", Icon20x20 ) );
 		Set( "FontEditor.ExportAllPages", new IMAGE_BRUSH_SVG( "Starship/Common/ExportAll", Icon20x20 ) );
 		Set( "FontEditor.FontBackgroundColor", new IMAGE_BRUSH_SVG( "Starship/FontEditor/FontBackground", Icon20x20 ) );
 		Set( "FontEditor.FontForegroundColor", new IMAGE_BRUSH_SVG( "Starship/FontEditor/FontForeground", Icon20x20 ) );
diff --git a/Engine/Source/Editor/UnrealEd/Classes/Settings/EditorStyleSettings.h b/Engine/Source/Editor/UnrealEd/Classes/Settings/EditorStyleSettings.h
--- a/Engine/Source/Editor/UnrealEd/Classes/Settings/EditorStyleSettings.h
+++ b/Engine/Source/Editor/UnrealEd/Classes/Settings/EditorStyleSettings.h
@@ -71,14 +71,14 @@ public:
 	bool bColorVisionDeficiencyCorrectionPreviewWithDeficiency;
 
 	/** The color used to represent selection */
-	UPROPERTY(EditAnywhere, config, Category=Colors, meta=(DisplayName="Viewport Selection Color"))
+	UPROPERTY(EditAnywhere, config, Category=UserInterface, meta=(DisplayName="Viewport Selection Color"))
 	FLinearColor SelectionColor;
 
 	UPROPERTY(config)
 	bool bEnableEditorWindowBackgroundColor;
 
 	/** The color used to tint the editor window backgrounds */
-	UPROPERTY(EditAnywhere, config, Category=Colors, meta=(EditCondition="bEnableEditorWindowBackgroundColor"))
+	UPROPERTY(EditAnywhere, config, Category=UserInterface, meta=(EditCondition="bEnableEditorWindowBackgroundColor"))
 	FLinearColor EditorWindowBackgroundColor;
 
 	/** Whether to use small toolbar icons without labels or not. */
@@ -145,6 +145,9 @@ public:
 	UPROPERTY(EditAnywhere, config, Category=UserInterface)
 	uint32 bEnableColorizedEditorTabs : 1;
 
+	UPROPERTY(config)
+	FGuid CurrentAppliedTheme; 
+
 public:
 
 	/**
diff --git a/Engine/Source/Editor/UnrealEd/Private/Settings/EditorStyleClasses.cpp b/Engine/Source/Editor/UnrealEd/Private/Settings/EditorStyleClasses.cpp
--- a/Engine/Source/Editor/UnrealEd/Private/Settings/EditorStyleClasses.cpp
+++ b/Engine/Source/Editor/UnrealEd/Private/Settings/EditorStyleClasses.cpp
@@ -6,6 +6,7 @@
 #include "Misc/ConfigCacheIni.h"
 #include "Modules/ModuleManager.h"
 #include "UObject/UnrealType.h"
+#include "Styling/StyleColors.h"
 
 
 /* UEditorStyleSettings interface
@@ -37,6 +38,17 @@ UEditorStyleSettings::UEditorStyleSettings( const FObjectInitializer& ObjectInit
 
 void UEditorStyleSettings::Init()
 {
+	// if it's a valid id, set current theme ID in USlateThemeManager to CurrentAppliedTheme. 
+	if (CurrentAppliedTheme.IsValid())
+	{
+		USlateThemeManager::Get().ApplyTheme(CurrentAppliedTheme); 
+	}
+	// if it's invalid id, set current theme to default. 
+	else
+	{
+		CurrentAppliedTheme = USlateThemeManager::Get().GetCurrentThemeID();
+		SaveConfig(); 
+	}
 
 	// Set from CVar 
 	IConsoleVariable* CVar = IConsoleManager::Get().FindConsoleVariable(TEXT("EnableHighDPIAwareness"));
diff --git a/Engine/Source/Editor/UnrealEd/Private/Settings/EditorStyleSettingsCustomization.cpp b/Engine/Source/Editor/UnrealEd/Private/Settings/EditorStyleSettingsCustomization.cpp
--- a/Engine/Source/Editor/UnrealEd/Private/Settings/EditorStyleSettingsCustomization.cpp
+++ b/Engine/Source/Editor/UnrealEd/Private/Settings/EditorStyleSettingsCustomization.cpp
@@ -11,6 +11,12 @@
 #include "Widgets/Input/SButton.h"
 #include "Modules/ModuleManager.h"
 #include "Widgets/Input/SEditableTextBox.h"
+#include "Framework/Notifications/NotificationManager.h"
+#include "SSimpleButton.h"
+#include "Serialization/JsonReader.h"
+#include "Serialization/JsonSerializer.h"
+#include "Misc/FileHelper.h"
+
 
 #if ALLOW_THEMES
 #include "IDesktopPlatform.h"
@@ -18,11 +24,13 @@
 #include "SPrimaryButton.h"
 #include "HAL/FileManager.h"
 #include "Misc/MessageDialog.h"
+#include "Settings/EditorStyleSettings.h"
 
 #define LOCTEXT_NAMESPACE "ThemeEditor"
 
 TWeakPtr<SWindow> ThemeEditorWindow;
 FString CurrentActiveThemeDisplayName;
+FString OriginalThemeName;
 
 
 class SThemeEditor : public SCompoundWidget
@@ -159,18 +167,21 @@ private:
 		return USlateThemeManager::Get().GetCurrentTheme().DisplayName;
 	}
 
+	// Validate new theme name against existing theme names to avoid duplicate. 
 	bool ValidateThemeName(const FText& ThemeName)
 	{
 		FText OutErrorMessage;
-		FString	Filename = USlateThemeManager::Get().GetEngineThemeDir() / ThemeName.ToString() + TEXT(".json");
+		const TArray<FStyleTheme> ThemeOptions = USlateThemeManager::Get().GetThemes();
 
-		// in cases the name user entered already exists when creating a new theme: 
-		if (FPaths::FileExists(Filename))
+		for (const FStyleTheme& Theme : ThemeOptions)
 		{
-			// show error message whenever there's duplicate 
-			OutErrorMessage = FText::Format(LOCTEXT("RenameThemeAlreadyExists", "A theme already exists with the name '{0}'."), ThemeName);
-			EditableThemeName->SetError(OutErrorMessage);
-			return false; 
+			// show error message whenever there's duplicate (and different from the previous name) 
+			if (Theme.DisplayName.EqualTo(ThemeName) && !CurrentActiveThemeDisplayName.Equals(ThemeName.ToString()))
+			{
+				OutErrorMessage = FText::Format(LOCTEXT("RenameThemeAlreadyExists", "A theme already exists with the name '{0}'."), ThemeName);
+				EditableThemeName->SetError(OutErrorMessage);
+				return false;
+			}
 		}
 		EditableThemeName->SetError(FText::GetEmpty());
 		return true; 
@@ -184,10 +195,15 @@ private:
 
 	void OnThemeNameCommitted(const FText& NewName, ETextCommit::Type = ETextCommit::Default)
 	{
-		// if the new name is a valid name, write to display. Else, go back to the previous name. 
-		if (ValidateThemeName(NewName))
+		if (!ValidateThemeName(NewName))
+		{
+			const FText OriginalTheme = FText::FromString(OriginalThemeName); 
+			EditableThemeName->SetText(OriginalTheme);
+			EditableThemeName->SetError(FText::GetEmpty());
+		}
+		else
 		{
-			USlateThemeManager::Get().SetCurrentThemeDisplayName(NewName);
+			EditableThemeName->SetText(NewName);
 		}
 	}
 
@@ -195,52 +211,41 @@ private:
 	{
 		FString Filename;
 		bool bSuccess = true; 
+		FString PreviousFilename; 
 
 		const FStyleTheme& Theme = USlateThemeManager::Get().GetCurrentTheme();
 
-		// New theme being created: 
-		if (Theme.Filename.IsEmpty())
+		// updated name is taken: DO NOT SAVE. 
+		if (!ValidateThemeName(EditableThemeName->GetText()))
 		{
-			TSharedPtr<SWindow> MyWindow = FSlateApplication::Get().FindWidgetWindow(SharedThis(this));
-
-			TArray<FString> Filenames;
-			IDesktopPlatform* DesktopPlatform = FDesktopPlatformModule::Get();
-			Filename = USlateThemeManager::Get().GetEngineThemeDir() / Theme.DisplayName.ToString() + TEXT(".json");
-
-			// if not a valid name, DO NOT SAVE. 
-			if (!ValidateThemeName(Theme.DisplayName))
-			{
-				bSuccess = false;
-			}
+			bSuccess = false;
 		}
-		// Modifying an existing theme: 
+		// Duplicating a theme: 
+		else if (Theme.Filename.IsEmpty())
+		{
+			// updated name is not taken: SAVE. 
+			USlateThemeManager::Get().SetCurrentThemeDisplayName(EditableThemeName->GetText());
+			Filename = USlateThemeManager::Get().GetUserThemeDir() / Theme.DisplayName.ToString() + TEXT(".json");
+			EditableThemeName->SetError(FText::GetEmpty()); 
+		}
+		// Modifying a theme: would only be here if the user is modifying a user-specific theme. 
 		else
 		{
 			// updated name is not taken: SAVE. 
-			if (ValidateThemeName(Theme.DisplayName))
-			{
-				Filename = Theme.Filename;
-			}
-			// updated name either already exists, or is the same as before: 
-			else
-			{		
-				// updated name is the same as the current theme: SAVE anyway. 
-				if (Theme.DisplayName.ToString().Equals(CurrentActiveThemeDisplayName))
-				{
-					Filename = Theme.Filename;
-				}
-				// updated name already exists: DON'T SAVE. 
-				else
-				{
-					bSuccess = false;
-				}
-			}
+			PreviousFilename = Theme.Filename; 
+			USlateThemeManager::Get().SetCurrentThemeDisplayName(EditableThemeName->GetText());
+			Filename = USlateThemeManager::Get().GetUserThemeDir() / Theme.DisplayName.ToString() + TEXT(".json");
+			EditableThemeName->SetError(FText::GetEmpty());
 		}
 
-		if(!Filename.IsEmpty() && bSuccess)
+		if (!Filename.IsEmpty() && bSuccess)
 		{
-			USlateThemeManager::Get().SaveCurrentThemeAs(Filename);
-
+			USlateThemeManager::Get().SaveCurrentThemeAs(Filename); 
+			// if user modified an existing user-specific theme name, delete the old one. 
+			if (!PreviousFilename.IsEmpty() && !PreviousFilename.Equals(Filename))
+			{
+				IPlatformFile::GetPlatformPhysical().DeleteFile(*PreviousFilename);
+			}
 			EditableThemeName->SetError(FText::GetEmpty()); 
 
 			ParentWindow.Pin()->SetOnWindowClosed(FOnWindowClosed());
@@ -348,7 +353,7 @@ TSharedRef<IDetailCustomization> FEditorStyleSettingsCustomization::MakeInstance
 
 void FEditorStyleSettingsCustomization::CustomizeDetails(IDetailLayoutBuilder& DetailLayout)
 {
-	IDetailCategoryBuilder& ColorCategory = DetailLayout.EditCategory("Colors");
+	IDetailCategoryBuilder& ColorCategory = DetailLayout.EditCategory("Theme");
 
 	TArray<UObject*> Objects = { &USlateThemeManager::Get() };
 
@@ -428,57 +433,85 @@ void FEditorStyleSettingsCustomization::MakeThemePickerRow(IDetailPropertyRow& P
 		.HAlign(HAlign_Center)
 		.AutoWidth()
 		[
-			SNew(SButton)
-			.ButtonStyle(FAppStyle::Get(), "SimpleButton")
-			.Visibility_Lambda(
+			SNew(SSimpleButton)
+			.Icon(FAppStyle::Get().GetBrush("Icons.Edit"))
+			.IsEnabled_Lambda(
 				[]()
 				{
-					// if current active theme is the default dark theme: collapse the edit button. 
-					return USlateThemeManager::Get().IsDefaultThemeActive() ? EVisibility::Collapsed : EVisibility::Visible;
+					return !(USlateThemeManager::Get().IsEngineTheme() || USlateThemeManager::Get().IsProjectTheme());
 				})
-			.ToolTipText(LOCTEXT("EditThemeToolTip", "Edit this theme"))
+			.ToolTipText_Lambda(
+				[]()
+				{	
+					if (USlateThemeManager::Get().IsEngineTheme())
+					{
+						return LOCTEXT("CannotEditEngineThemeToolTip", "Engine themes can't be edited");
+					}
+					else if (USlateThemeManager::Get().IsProjectTheme())
+					{
+						return LOCTEXT("CannotEditProjectThemeToolTip", "Project themes can't be edited");
+					}
+					return LOCTEXT("EditThemeToolTip", "Edit this theme");
+				 })
 			.OnClicked(this, &FEditorStyleSettingsCustomization::OnEditThemeClicked)
-			[
-				SNew(SImage)
-				.ColorAndOpacity(FSlateColor::UseForeground())
-				.Image(FAppStyle::Get().GetBrush("Icons.Edit"))
-			]
 		]
 		+SHorizontalBox::Slot()
 		.VAlign(VAlign_Center)
 		.HAlign(HAlign_Center)
 		.AutoWidth()
 		[
-			SNew(SButton)
-			.ButtonStyle(FAppStyle::Get(), "SimpleButton")
+			SNew(SSimpleButton)
+			.Icon(FAppStyle::Get().GetBrush("Icons.Duplicate"))
 			.ToolTipText(LOCTEXT("DuplicateThemeToolTip", "Duplicate this theme and edit it"))
 			.OnClicked(this, &FEditorStyleSettingsCustomization::OnDuplicateAndEditThemeClicked)
-			[
-				SNew(SImage)
-				.ColorAndOpacity(FSlateColor::UseForeground())
-				.Image(FAppStyle::Get().GetBrush("Icons.Duplicate"))
-			]
 		]
 		+ SHorizontalBox::Slot()
-			.VAlign(VAlign_Center)
-			.HAlign(HAlign_Center)
-			.AutoWidth()
-			[
-			SNew(SButton)
-			.ButtonStyle(FAppStyle::Get(), "SimpleButton")
-			.Visibility_Lambda(
+		.AutoWidth()
+		.Padding(8.0f, 0.0f, 0.0f, 0.0f)
+		[
+			// export button
+			SNew(SSimpleButton)
+			.Icon(FAppStyle::Get().GetBrush("Themes.Export"))
+			.ToolTipText(LOCTEXT("ExportButtonTooltip", "Export this theme to a file on your computer"))
+			.OnClicked(this, &FEditorStyleSettingsCustomization::OnExportThemeClicked)
+		]
+		+ SHorizontalBox::Slot()
+		.AutoWidth()
+		.Padding(8.0f, 0.0f, 0.0f, 0.0f)
+		[
+			// import button
+			SNew(SSimpleButton)
+			.Icon(FAppStyle::Get().GetBrush("Themes.Import"))
+			.ToolTipText(LOCTEXT("ImportButtonTooltip", "Import a theme from a file on your computer"))
+			.OnClicked(this, &FEditorStyleSettingsCustomization::OnImportThemeClicked)
+		]
+		+ SHorizontalBox::Slot()
+		.VAlign(VAlign_Center)
+		.HAlign(HAlign_Center)
+		.AutoWidth()
+		[
+			// delete button
+			SNew(SSimpleButton)
+			.Icon(FAppStyle::Get().GetBrush("Icons.Delete"))
+			.IsEnabled_Lambda(
 				[]()
 				{
-					// if current active theme is the default dark theme: collapse the delete button. 
-					return USlateThemeManager::Get().IsDefaultThemeActive() ? EVisibility::Collapsed : EVisibility::Visible;
+					return !(USlateThemeManager::Get().IsEngineTheme() || USlateThemeManager::Get().IsProjectTheme()); 
+				})
+			.ToolTipText_Lambda(
+				[]()
+				{
+					if (USlateThemeManager::Get().IsEngineTheme())
+					{
+						return LOCTEXT("CannotDeleteEngineThemeToolTip", "Engine themes can't be deleted");
+					}
+					else if (USlateThemeManager::Get().IsProjectTheme())
+					{
+						return LOCTEXT("CannotDeleteProjectThemeToolTip", "Project themes can't be deleted");
+					}
+					return LOCTEXT("DeleteThemeToolTip", "Delete this theme");
 				})
-			.ToolTipText(LOCTEXT("DeleteThemeToolTip", "Delete this theme"))
 			.OnClicked(this, &FEditorStyleSettingsCustomization::OnDeleteThemeClicked)
-			[
-				SNew(SImage)
-				.ColorAndOpacity(FSlateColor::UseForeground())
-				.Image(FAppStyle::Get().GetBrush("Icons.Delete"))
-			]
 		]
 	];
 }
@@ -513,6 +546,129 @@ static void OnThemeEditorClosed(bool bSaved, TWeakPtr<FEditorStyleSettingsCustom
 	}
 }
 
+FReply FEditorStyleSettingsCustomization::OnExportThemeClicked()
+{
+	TArray<FString> OutFiles; 
+	const void* ParentWindowHandle = FSlateApplication::Get().FindBestParentWindowHandleForDialogs(nullptr);
+	const FString ExportPath = FPlatformProcess::UserDir();
+	const FString DefaultFileName = USlateThemeManager::Get().GetCurrentTheme().DisplayName.ToString();
+
+	if (FDesktopPlatformModule::Get()->SaveFileDialog(ParentWindowHandle, LOCTEXT("ExportThemeDialogTitle", "Export current theme...").ToString(), FPaths::GetPath(ExportPath), DefaultFileName, TEXT("JSON files (*.json)|*.json"), EFileDialogFlags::None, OutFiles))
+	{
+		const FString SourcePath = USlateThemeManager::Get().GetCurrentTheme().Filename;
+		const FString DestPath = OutFiles[0]; 
+
+		if (IPlatformFile::GetPlatformPhysical().CopyFile(*DestPath, *SourcePath)) 
+		{
+			ShowNotification(LOCTEXT("ExportThemeSuccess", "Export theme succeeded"), SNotificationItem::CS_Success);
+		}
+		else
+		{
+			ShowNotification(LOCTEXT("ExportThemeFailure", "Export theme failed"), SNotificationItem::CS_Fail);
+		}
+	}
+	return FReply::Handled();
+}
+
+// validate theme name without error messages: 
+bool IsThemeNameValid(const FString& ThemeName)
+{
+	const TArray<FStyleTheme> ThemeOptions = USlateThemeManager::Get().GetThemes();
+
+	for (const FStyleTheme& Theme : ThemeOptions)
+	{
+		// show error message whenever there's duplicate (and different from the previous name) 
+		if (Theme.DisplayName.ToString().Equals(ThemeName))
+		{
+			return false;
+		}
+	}
+	return true;
+}
+
+
+void GetThemeIdFromPath(FString& ThemePath, FString& ImportedThemeID)
+{
+	FString ThemeData;
+
+	if (FFileHelper::LoadFileToString(ThemeData, *ThemePath))
+	{
+		TSharedRef<TJsonReader<>> ReaderRef = TJsonReaderFactory<>::Create(ThemeData);
+		TJsonReader<>& Reader = ReaderRef.Get();
+
+		TSharedPtr<FJsonObject> ObjectPtr = MakeShareable(new FJsonObject()); 
+
+		if (FJsonSerializer::Deserialize(Reader, ObjectPtr) && ObjectPtr.IsValid())
+		{
+			// Just check that the theme has Id. We won't load them unless the theme is used
+			ObjectPtr->TryGetStringField(TEXT("Id"), ImportedThemeID); 
+		}
+	}
+}
+
+FReply FEditorStyleSettingsCustomization::OnImportThemeClicked()
+{
+	TArray<FString> OutFiles;
+	const void* ParentWindowHandle = FSlateApplication::Get().FindBestParentWindowHandleForDialogs(nullptr); 
+	const FString ImportFrom = FPlatformProcess::UserDir();
+
+	if (FDesktopPlatformModule::Get()->OpenFileDialog(ParentWindowHandle, LOCTEXT("ImportThemeDialogTitle", "Import theme...").ToString(), FPaths::GetPath(ImportFrom), TEXT(""), TEXT("JSON files (*.json)|*.json"), EFileDialogFlags::None, OutFiles))
+	{
+		FString SourcePath = OutFiles[0];
+		const FString DestPath = USlateThemeManager::Get().GetUserThemeDir() / FPaths::GetCleanFilename(SourcePath);
+
+		FString PathPart; 
+		FString Extension; 
+		FString FilenameWithoutExtension; 
+		FPaths::Split(SourcePath, PathPart, FilenameWithoutExtension, Extension);
+
+		// if theme name exists, don't import (to prevent from overwriting existing theme files)
+		if (!IsThemeNameValid(FilenameWithoutExtension))
+		{
+			ShowNotification(LOCTEXT("ImportThemeFailure", "Import theme failed: Theme name already exists"), SNotificationItem::CS_Fail);
+		}
+		// if theme name is valid: copying the file is safe (as it will not overwrite existing theme files)
+		else 
+		{	
+			const int32 NumOfThemesBefore = USlateThemeManager::Get().GetThemes().Num();
+
+			if  (IPlatformFile::GetPlatformPhysical().CopyFile(*DestPath, *SourcePath)) 
+			{	
+				// update the number of valid themes: 
+				USlateThemeManager::Get().LoadThemes();
+
+				// if valid theme: the theme Num will be updated. 
+				if (USlateThemeManager::Get().GetThemes().Num() != NumOfThemesBefore)
+				{
+					// Extract ID as a FString directly from a JSON file. 
+					FString ImportedThemeID; 
+					GetThemeIdFromPath(SourcePath, ImportedThemeID);
+
+					// convert FString ID to a FGuid
+					FGuid ImportedThemeGUID = FGuid(ImportedThemeID); 
+					USlateThemeManager::Get().ApplyTheme(ImportedThemeGUID);
+					RefreshComboBox();
+
+					ShowNotification(LOCTEXT("ImportThemeSuccess", "Import theme succeeded"), SNotificationItem::CS_Success);
+				}
+				// if invalid theme: delete the copied file. 
+				else
+				{
+					// incomplete themes will not reach here. 
+					IPlatformFile::GetPlatformPhysical().DeleteFile(*DestPath); 
+					ShowNotification(LOCTEXT("ImportThemeFailure", "Import theme failed: Invalid theme"), SNotificationItem::CS_Fail);
+				}
+			}
+			// if unable to copy the file to user-specific theme location, do nothing. 
+			else
+			{
+				ShowNotification(LOCTEXT("ImportThemeFailure", "Import theme failed"), SNotificationItem::CS_Fail);
+			}
+		}
+	}
+	return FReply::Handled();
+}
+
 FReply FEditorStyleSettingsCustomization::OnDeleteThemeClicked()
 {
 	const FStyleTheme PreviouslyActiveTheme = USlateThemeManager::Get().GetCurrentTheme();
@@ -529,7 +685,7 @@ FReply FEditorStyleSettingsCustomization::OnDeleteThemeClicked()
 		USlateThemeManager::Get().ApplyDefaultTheme();
 
 		// remove previously active theme
-		const FString Filename = USlateThemeManager::Get().GetEngineThemeDir() / PreviouslyActiveTheme.DisplayName.ToString() + TEXT(".json");
+		const FString Filename = USlateThemeManager::Get().GetUserThemeDir() / PreviouslyActiveTheme.DisplayName.ToString() + TEXT(".json");
 		IFileManager::Get().Delete(*Filename);
 		USlateThemeManager::Get().RemoveTheme(PreviouslyActiveTheme.Id);
 		RefreshComboBox();
@@ -542,10 +698,10 @@ FReply FEditorStyleSettingsCustomization::OnDuplicateAndEditThemeClicked()
 {
 	FGuid PreviouslyActiveTheme = USlateThemeManager::Get().GetCurrentTheme().Id;
 
-	CurrentActiveThemeDisplayName = USlateThemeManager::Get().GetCurrentTheme().DisplayName.ToString(); 
-
 	FGuid NewThemeId = USlateThemeManager::Get().DuplicateActiveTheme();
 	USlateThemeManager::Get().ApplyTheme(NewThemeId);
+	CurrentActiveThemeDisplayName = USlateThemeManager::Get().GetCurrentTheme().DisplayName.ToString();
+	OriginalThemeName = USlateThemeManager::Get().GetCurrentTheme().DisplayName.ToString(); 
 
 	RefreshComboBox();
 
@@ -558,6 +714,7 @@ FReply FEditorStyleSettingsCustomization::OnEditThemeClicked()
 {
 
 	CurrentActiveThemeDisplayName = USlateThemeManager::Get().GetCurrentTheme().DisplayName.ToString();
+	OriginalThemeName = CurrentActiveThemeDisplayName; 
 
 	OpenThemeEditorWindow(FOnThemeEditorClosed::CreateStatic(&OnThemeEditorClosed, TWeakPtr<FEditorStyleSettingsCustomization>(SharedThis(this)), FGuid(), FGuid()));
 
@@ -572,12 +729,16 @@ FString FEditorStyleSettingsCustomization::GetTextLabelForThemeEntry(TSharedPtr<
 
 void FEditorStyleSettingsCustomization::OnThemePicked(TSharedPtr<FString> NewSelection, ESelectInfo::Type SelectInfo)
 {
+	UEditorStyleSettings* StyleSetting = GetMutableDefault<UEditorStyleSettings>();
+
+	// set current applied theme to selected theme. 
+	const TArray<FStyleTheme>& Themes = USlateThemeManager::Get().GetThemes();
+	StyleSetting->CurrentAppliedTheme = Themes[TCString<TCHAR>::Atoi(**NewSelection)].Id;
+	
 	// If set directly in code, the theme was already applied
 	if(SelectInfo != ESelectInfo::Direct)
 	{
-		const TArray<FStyleTheme>& Themes = USlateThemeManager::Get().GetThemes();
-
-		USlateThemeManager::Get().ApplyTheme(Themes[TCString<TCHAR>::Atoi(**NewSelection)].Id);
+		USlateThemeManager::Get().ApplyTheme(StyleSetting->CurrentAppliedTheme);
 	}
 }
 
@@ -621,6 +782,16 @@ bool FEditorStyleSettingsCustomization::IsThemeEditingEnabled() const
 	return !ThemeEditorWindow.IsValid();
 }
 
+void FEditorStyleSettingsCustomization::ShowNotification(const FText& Text, SNotificationItem::ECompletionState CompletionState) const
+{
+	FNotificationInfo Notification(Text);
+	Notification.ExpireDuration = 3.f;
+	Notification.bUseSuccessFailIcons = false;
+
+	FSlateNotificationManager::Get().AddNotification(Notification)->SetCompletionState(CompletionState);
+}
+
+
 #undef LOCTEXT_NAMESPACE
 
 #endif // ALLOW_THEMES
diff --git a/Engine/Source/Editor/UnrealEd/Private/Settings/EditorStyleSettingsCustomization.h b/Engine/Source/Editor/UnrealEd/Private/Settings/EditorStyleSettingsCustomization.h
--- a/Engine/Source/Editor/UnrealEd/Private/Settings/EditorStyleSettingsCustomization.h
+++ b/Engine/Source/Editor/UnrealEd/Private/Settings/EditorStyleSettingsCustomization.h
@@ -5,6 +5,7 @@
 #include "IDetailCustomization.h"
 #include "IPropertyTypeCustomization.h"
 #include "Styling/StyleColors.h"
+#include "Widgets/Notifications/SNotificationList.h"
 
 class IDetailLayoutBuilder;
 class STextComboBox;
@@ -40,6 +41,8 @@ private:
 	void GenerateThemeOptions(TSharedPtr<FString>& OutSelectedTheme);
 
 	void MakeThemePickerRow(IDetailPropertyRow& PropertyRow);
+	FReply OnExportThemeClicked(); 
+	FReply OnImportThemeClicked(); 
 	FReply OnDeleteThemeClicked();
 	FReply OnDuplicateAndEditThemeClicked();
 	FReply OnEditThemeClicked();
@@ -47,6 +50,7 @@ private:
 	void OnThemePicked(TSharedPtr<FString> NewSelection, ESelectInfo::Type SelectInfo);
 	void OpenThemeEditorWindow(FOnThemeEditorClosed OnThemeEditorClosed);
 	bool IsThemeEditingEnabled() const;
+	void ShowNotification(const FText& Text, SNotificationItem::ECompletionState CompletionState) const;
 private:
 	TArray<TSharedPtr<FString>> ThemeOptions;
 	TSharedPtr<STextComboBox> ComboBox;
diff --git a/Engine/Source/Runtime/SlateCore/Private/Styling/StyleColors.cpp b/Engine/Source/Runtime/SlateCore/Private/Styling/StyleColors.cpp
--- a/Engine/Source/Runtime/SlateCore/Private/Styling/StyleColors.cpp
+++ b/Engine/Source/Runtime/SlateCore/Private/Styling/StyleColors.cpp
@@ -146,7 +146,7 @@ void USlateThemeManager::LoadThemes()
 	LoadThemesFromDirectory(GetProjectThemeDir());
 
 	// User specific themes
-	LoadThemesFromDirectory(FPaths::EngineVersionAgnosticUserDir());
+	LoadThemesFromDirectory(GetUserThemeDir());
 
 	EnsureValidCurrentTheme();
 
@@ -186,7 +186,7 @@ void USlateThemeManager::SaveCurrentThemeAs(const FString& Filename)
 		{
 			FPlatformFileManager::Get().GetPlatformFile().SetReadOnly(*Filename, false);
 			// create a new path if the filename has been changed. 
-			NewPath = USlateThemeManager::Get().GetEngineThemeDir() / CurrentTheme.DisplayName.ToString() + TEXT(".json");
+			NewPath = USlateThemeManager::Get().GetUserThemeDir() / CurrentTheme.DisplayName.ToString() + TEXT(".json");
 
 			if (!NewPath.Equals(CurrentTheme.Filename))
 			{
@@ -237,9 +237,32 @@ void USlateThemeManager::ApplyDefaultTheme()
 	ApplyTheme(DefaultDarkTheme.Id); 
 }
 
-bool USlateThemeManager::IsDefaultThemeActive() const 
+bool USlateThemeManager::IsEngineTheme() const
 {
-	return GetCurrentTheme() == DefaultDarkTheme; 
+	// users cannot edit/delete engine-specific themes: 
+	const FString EnginePath = GetEngineThemeDir() / GetCurrentTheme().DisplayName.ToString() + TEXT(".json");
+
+	IPlatformFile& FileManager = FPlatformFileManager::Get().GetPlatformFile();
+
+	if (FileManager.FileExists(*EnginePath))
+	{
+		return true;
+	}
+	return false;
+}
+
+bool USlateThemeManager::IsProjectTheme() const
+{
+	// users cannot edit/delete project-specific themes: 
+	const FString ProjectPath = GetProjectThemeDir() / GetCurrentTheme().DisplayName.ToString() + TEXT(".json");
+
+	IPlatformFile& FileManager = FPlatformFileManager::Get().GetPlatformFile(); 
+
+	if (FileManager.FileExists(*ProjectPath))
+	{
+		return true; 
+	}
+	return false; 
 }
 
 void USlateThemeManager::RemoveTheme(FGuid ThemeId)
diff --git a/Engine/Source/Runtime/SlateCore/Public/Styling/StyleColors.h b/Engine/Source/Runtime/SlateCore/Public/Styling/StyleColors.h
--- a/Engine/Source/Runtime/SlateCore/Public/Styling/StyleColors.h
+++ b/Engine/Source/Runtime/SlateCore/Public/Styling/StyleColors.h
@@ -156,6 +156,10 @@ public:
 		return ActiveColors.StyleColors[static_cast<int32>(Color)];
 	}
 
+	const FGuid& GetCurrentThemeID()
+	{
+		return CurrentThemeId; 
+	}
 
 	USlateThemeManager();
 
@@ -213,9 +217,14 @@ public:
 	void ApplyDefaultTheme();
 
 	/**
-	* Returns true if the default dark theme is the active theme
+	* Returns true if the active theme is an engine-specific theme
+	*/
+	bool IsEngineTheme() const;
+
+	/**
+	* Returns true if the active theme is a project-specific theme
 	*/
-	bool IsDefaultThemeActive() const;
+	bool IsProjectTheme() const;
 
 	/**
 	 * Removes a theme. 
