AssetPrivatizeModel and Dialogue

diff --git a/Engine/Plugins/Editor/ContentBrowser/ContentBrowserAssetDataSource/Source/ContentBrowserAssetDataSource/Private/ContentBrowserAssetDataCore.cpp b/Engine/Plugins/Editor/ContentBrowser/ContentBrowserAssetDataSource/Source/ContentBrowserAssetDataSource/Private/ContentBrowserAssetDataCore.cpp
--- a/Engine/Plugins/Editor/ContentBrowser/ContentBrowserAssetDataSource/Source/ContentBrowserAssetDataSource/Private/ContentBrowserAssetDataCore.cpp
+++ b/Engine/Plugins/Editor/ContentBrowser/ContentBrowserAssetDataSource/Source/ContentBrowserAssetDataSource/Private/ContentBrowserAssetDataCore.cpp
@@ -587,6 +587,42 @@ bool CanDeleteAssetFileItem(IAssetTools* InAssetTools, const FContentBrowserAsse
 	return true;
 }
 
+bool CanPrivatizeItem(IAssetTools* InAssetTools, IAssetRegistry* InAssetRegistry, const UContentBrowserDataSource* InOwnerDataSource, const FContentBrowserItemData& InItem, FText* OutErrorMsg)
+{
+	if (TSharedPtr<const FContentBrowserAssetFolderItemDataPayload> FolderPayload = GetAssetFolderItemPayload(InOwnerDataSource, InItem))
+	{
+		return false;
+	}
+
+	if (TSharedPtr<const FContentBrowserAssetFileItemDataPayload> AssetPayload = GetAssetFileItemPayload(InOwnerDataSource, InItem))
+	{
+		return CanPrivatizeAssetFileItem(InAssetTools, *AssetPayload, OutErrorMsg);
+	}
+
+	return false;
+}
+
+bool CanPrivatizeAssetFileItem(IAssetTools* InAssetTools, const FContentBrowserAssetFileItemDataPayload& InAssetPayload, FText* OutErrorMsg)
+{
+	if (!CanModifyAssetFileItem(InAssetTools, InAssetPayload, OutErrorMsg))
+	{
+		return false;
+	}
+
+	if (IsRunningPIE(OutErrorMsg))
+	{
+		return false;
+	}
+
+	if (InAssetPayload.GetAssetData().IsRedirector())
+	{
+		SetOptionalErrorMessage(OutErrorMsg, LOCTEXT("Error_CannotPrivatizeRedirectors", "Cannot make redirectors private"));
+		return false;
+	}
+
+	return true;
+}
+
 bool DeleteItems(IAssetTools* InAssetTools, IAssetRegistry* InAssetRegistry, const UContentBrowserDataSource* InOwnerDataSource, TArrayView<const FContentBrowserItemData> InItems)
 {
 	TArray<TSharedRef<const FContentBrowserAssetFolderItemDataPayload>, TInlineAllocator<16>> FolderPayloads;
@@ -627,6 +663,41 @@ bool DeleteItems(IAssetTools* InAssetTools, IAssetRegistry* InAssetRegistry, con
 	return bDidDelete;
 }
 
+bool PrivatizeItems(IAssetTools* InAssetTools, IAssetRegistry* InAssetRegistry, const UContentBrowserDataSource* InOwnerDataSource, TArrayView<const FContentBrowserItemData> InItems)
+{
+	TArray<TSharedRef<const FContentBrowserAssetFileItemDataPayload>, TInlineAllocator<16>> AssetPayloads;
+
+	auto ProcessAssetFolderItem = [](const TSharedRef<const FContentBrowserAssetFolderItemDataPayload>& InFolderPayload)
+	{
+		return true;
+	};
+
+	auto ProcessAssetFileItem = [InAssetTools, &AssetPayloads](const TSharedRef<const FContentBrowserAssetFileItemDataPayload>& InAssetPayload)
+	{
+		if (CanPrivatizeAssetFileItem(InAssetTools, *InAssetPayload, nullptr))
+		{
+			AssetPayloads.Add(InAssetPayload);
+		}
+		return true;
+	};
+
+	EnumerateAssetItemPayloads(InOwnerDataSource, InItems, ProcessAssetFolderItem, ProcessAssetFileItem);
+
+	if (AssetPayloads.IsEmpty())
+	{
+		return false;
+	}
+
+	TArray<FAssetData> AssetsToPrivatize;
+
+	for (const TSharedRef<const FContentBrowserAssetFileItemDataPayload>& AssetPayload : AssetPayloads)
+	{
+		AssetsToPrivatize.Add(AssetPayload->GetAssetData());
+	}
+
+	return ObjectTools::PrivatizeAssets(AssetsToPrivatize) > 0;
+}
+
 bool DeleteAssetFolderItems(TArrayView<const TSharedRef<const FContentBrowserAssetFolderItemDataPayload>> InFolderPayloads)
 {
 	TArray<FString> FoldersToDelete;
diff --git a/Engine/Plugins/Editor/ContentBrowser/ContentBrowserAssetDataSource/Source/ContentBrowserAssetDataSource/Private/ContentBrowserAssetDataSource.cpp b/Engine/Plugins/Editor/ContentBrowser/ContentBrowserAssetDataSource/Source/ContentBrowserAssetDataSource/Private/ContentBrowserAssetDataSource.cpp
--- a/Engine/Plugins/Editor/ContentBrowser/ContentBrowserAssetDataSource/Source/ContentBrowserAssetDataSource/Private/ContentBrowserAssetDataSource.cpp
+++ b/Engine/Plugins/Editor/ContentBrowser/ContentBrowserAssetDataSource/Source/ContentBrowserAssetDataSource/Private/ContentBrowserAssetDataSource.cpp
@@ -1323,6 +1323,21 @@ bool UContentBrowserAssetDataSource::BulkDeleteItems(TArrayView<const FContentBr
 	return ContentBrowserAssetData::DeleteItems(AssetTools, AssetRegistry, this, InItems);
 }
 
+bool UContentBrowserAssetDataSource::CanPrivatizeItem(const FContentBrowserItemData& InItem, FText* OutErrorMsg)
+{
+	return ContentBrowserAssetData::CanPrivatizeItem(AssetTools, AssetRegistry, this, InItem, OutErrorMsg);
+}
+
+bool UContentBrowserAssetDataSource::PrivatizeItem(const FContentBrowserItemData& InItem)
+{
+	return ContentBrowserAssetData::PrivatizeItems(AssetTools, AssetRegistry, this, MakeArrayView(&InItem, 1));
+}
+
+bool UContentBrowserAssetDataSource::BulkPrivatizeItems(TArrayView<const FContentBrowserItemData> InItems)
+{
+	return ContentBrowserAssetData::PrivatizeItems(AssetTools, AssetRegistry, this, InItems);
+}
+
 bool UContentBrowserAssetDataSource::CanRenameItem(const FContentBrowserItemData& InItem, const FString* InNewName, FText* OutErrorMsg)
 {
 	return ContentBrowserAssetData::CanRenameItem(AssetTools, this, InItem, InNewName, OutErrorMsg);
diff --git a/Engine/Plugins/Editor/ContentBrowser/ContentBrowserAssetDataSource/Source/ContentBrowserAssetDataSource/Public/ContentBrowserAssetDataCore.h b/Engine/Plugins/Editor/ContentBrowser/ContentBrowserAssetDataSource/Source/ContentBrowserAssetDataSource/Public/ContentBrowserAssetDataCore.h
--- a/Engine/Plugins/Editor/ContentBrowser/ContentBrowserAssetDataSource/Source/ContentBrowserAssetDataSource/Public/ContentBrowserAssetDataCore.h
+++ b/Engine/Plugins/Editor/ContentBrowser/ContentBrowserAssetDataSource/Source/ContentBrowserAssetDataSource/Public/ContentBrowserAssetDataCore.h
@@ -93,6 +93,12 @@ namespace ContentBrowserAssetData
 
 	CONTENTBROWSERASSETDATASOURCE_API bool DeleteAssetFileItems(TArrayView<const TSharedRef<const FContentBrowserAssetFileItemDataPayload>> InAssetPayloads);
 
+	CONTENTBROWSERASSETDATASOURCE_API bool CanPrivatizeItem(IAssetTools* InAssetTools, IAssetRegistry* InAssetRegistry, const UContentBrowserDataSource* InOwnerDataSource, const FContentBrowserItemData& InItem, FText* OutErrorMsg);
+
+	CONTENTBROWSERASSETDATASOURCE_API bool CanPrivatizeAssetFileItem(IAssetTools* InAssetTools, const FContentBrowserAssetFileItemDataPayload& InAssetPayload, FText* OutErrorMsg);
+
+	CONTENTBROWSERASSETDATASOURCE_API bool PrivatizeItems(IAssetTools* InAssetTools, IAssetRegistry* InAssetRegistry, const UContentBrowserDataSource* InOwnerDataSource, TArrayView<const FContentBrowserItemData> InItems);
+
 	CONTENTBROWSERASSETDATASOURCE_API bool CanRenameItem(IAssetTools* InAssetTools, const UContentBrowserDataSource* InOwnerDataSource, const FContentBrowserItemData& InItem, const FString* InNewName, FText* OutErrorMsg);
 
 	CONTENTBROWSERASSETDATASOURCE_API bool CanRenameAssetFolderItem(IAssetTools* InAssetTools, const FContentBrowserAssetFolderItemDataPayload& InFolderPayload, const FString* InNewName, FText* OutErrorMsg);
diff --git a/Engine/Plugins/Editor/ContentBrowser/ContentBrowserAssetDataSource/Source/ContentBrowserAssetDataSource/Public/ContentBrowserAssetDataSource.h b/Engine/Plugins/Editor/ContentBrowser/ContentBrowserAssetDataSource/Source/ContentBrowserAssetDataSource/Public/ContentBrowserAssetDataSource.h
--- a/Engine/Plugins/Editor/ContentBrowser/ContentBrowserAssetDataSource/Source/ContentBrowserAssetDataSource/Public/ContentBrowserAssetDataSource.h
+++ b/Engine/Plugins/Editor/ContentBrowser/ContentBrowserAssetDataSource/Source/ContentBrowserAssetDataSource/Public/ContentBrowserAssetDataSource.h
@@ -212,6 +212,12 @@ public:
 
 	virtual bool BulkDeleteItems(TArrayView<const FContentBrowserItemData> InItems) override;
 
+	virtual bool CanPrivatizeItem(const FContentBrowserItemData& InItem, FText* OutErrorMsg);
+
+	virtual bool PrivatizeItem(const FContentBrowserItemData& InItem) override;
+
+	virtual bool BulkPrivatizeItems(TArrayView<const FContentBrowserItemData> InItems) override;
+
 	virtual bool CanRenameItem(const FContentBrowserItemData& InItem, const FString* InNewName, FText* OutErrorMsg) override;
 
 	virtual bool RenameItem(const FContentBrowserItemData& InItem, const FString& InNewName, FContentBrowserItemData& OutNewItem) override;
diff --git a/Engine/Source/Editor/ContentBrowser/Private/AssetContextMenu.cpp b/Engine/Source/Editor/ContentBrowser/Private/AssetContextMenu.cpp
--- a/Engine/Source/Editor/ContentBrowser/Private/AssetContextMenu.cpp
+++ b/Engine/Source/Editor/ContentBrowser/Private/AssetContextMenu.cpp
@@ -941,7 +941,15 @@ void FAssetContextMenu::ExecutePublicAssetToggle()
 				return;
 			}
 
-			ItemAssetPackage->SetIsExternallyReferenceable(!ItemAssetPackage->IsExternallyReferenceable());
+			if (ItemAssetPackage->IsExternallyReferenceable())
+			{
+				ExecuteBulkUnsetPublicAsset();
+				return;
+			}
+			else
+			{
+				ItemAssetPackage->SetIsExternallyReferenceable(true);
+			}
 
 			ItemAssetPackage->Modify();
 
@@ -983,27 +991,34 @@ void FAssetContextMenu::ExecuteBulkSetPublicAsset()
 
 void FAssetContextMenu::ExecuteBulkUnsetPublicAsset()
 {
+	// Batch these by their data sources
+	TMap<UContentBrowserDataSource*, TArray<FContentBrowserItemData>> SourcesAndItems;
 	for (const FContentBrowserItem& SelectedItem : SelectedFiles)
 	{
-		FAssetData ItemAssetData;
-		if (SelectedItem.Legacy_TryGetAssetData(ItemAssetData))
+		FContentBrowserItem::FItemDataArrayView ItemDataArray = SelectedItem.GetInternalItems();
+		for (const FContentBrowserItemData& ItemData : ItemDataArray)
 		{
-			UPackage* ItemAssetPackage = ItemAssetData.GetPackage();
-
-			if (!ItemAssetPackage)
-			{
-				continue;
-			}
-
-			if (ItemAssetPackage->IsExternallyReferenceable())
+			if (UContentBrowserDataSource* ItemDataSource = ItemData.GetOwnerDataSource())
 			{
-				ItemAssetPackage->SetIsExternallyReferenceable(false);
-
-				ItemAssetPackage->Modify();
+				FText PrivateErrorMsg;
+				if (ItemDataSource->CanPrivatizeItem(ItemData, &PrivateErrorMsg))
+				{
+					TArray<FContentBrowserItemData>& ItemsForSource = SourcesAndItems.FindOrAdd(ItemDataSource);
+					ItemsForSource.Add(ItemData);
+				}
+				else
+				{
+					AssetViewUtils::ShowErrorNotifcation(PrivateErrorMsg);
+				}
 			}
 		}
 	}
 
+	for (const auto& SourceAndItemsPair : SourcesAndItems)
+	{
+		SourceAndItemsPair.Key->BulkPrivatizeItems(SourceAndItemsPair.Value);
+	}
+
 	OnAssetViewRefreshRequested.ExecuteIfBound();
 }
 
diff --git a/Engine/Source/Editor/ContentBrowserData/Private/ContentBrowserDataSource.cpp b/Engine/Source/Editor/ContentBrowserData/Private/ContentBrowserDataSource.cpp
--- a/Engine/Source/Editor/ContentBrowserData/Private/ContentBrowserDataSource.cpp
+++ b/Engine/Source/Editor/ContentBrowserData/Private/ContentBrowserDataSource.cpp
@@ -367,6 +367,26 @@ bool UContentBrowserDataSource::BulkDeleteItems(TArrayView<const FContentBrowser
 	return bSuccess;
 }
 
+bool UContentBrowserDataSource::CanPrivatizeItem(const FContentBrowserItemData& InItem, FText* OutErrorMsg)
+{
+	return false;
+}
+
+bool UContentBrowserDataSource::PrivatizeItem(const FContentBrowserItemData& InItem)
+{
+	return false;
+}
+
+bool UContentBrowserDataSource::BulkPrivatizeItems(TArrayView<const FContentBrowserItemData> InItems)
+{
+	bool bSuccess = false;
+	for (const FContentBrowserItemData& Item : InItems)
+	{
+		bSuccess |= PrivatizeItem(Item);
+	}
+	return bSuccess;
+}
+
 bool UContentBrowserDataSource::CanRenameItem(const FContentBrowserItemData& InItem, const FString* InNewName, FText* OutErrorMsg)
 {
 	return false;
diff --git a/Engine/Source/Editor/ContentBrowserData/Public/ContentBrowserDataSource.h b/Engine/Source/Editor/ContentBrowserData/Public/ContentBrowserDataSource.h
--- a/Engine/Source/Editor/ContentBrowserData/Public/ContentBrowserDataSource.h
+++ b/Engine/Source/Editor/ContentBrowserData/Public/ContentBrowserDataSource.h
@@ -463,6 +463,34 @@ public:
 	 */
 	virtual bool BulkDeleteItems(TArrayView<const FContentBrowserItemData> InItems);
 
+	/**
+	* Query whether the given item can be privatized, optionally providing error information if it cannot.
+	* 
+	* @param InItem The item to query.
+	* @param OutErrorMessage Optional error message to fill on failure.
+	* 
+	* @return True if the item was deleted, false otherwise.
+	*/
+	virtual bool CanPrivatizeItem(const FContentBrowserItemData& InItem, FText* OutErrorMsg);
+
+	/**
+	* Attempt to mark the given item as private (NotExternallyReferenceable).
+	* 
+	* @param InItem The item to mark private.
+	*
+	* @return True if the item was marked private, false otherwise
+	*/
+	virtual bool PrivatizeItem(const FContentBrowserItemData& InItem);
+
+	/**
+	* Attempt to mark the given items as private (NotExternallyReferenceable)
+	* 
+	* @param InItems The items to be marked private.
+	* 
+	* @return True if any items were marked private, false otherwise
+	*/
+	virtual bool BulkPrivatizeItems(TArrayView<const FContentBrowserItemData> InItems);
+
 	/**
 	 * Query whether the given item is can be renamed, optionally providing error information if it cannot.
 	 *
diff --git a/Engine/Source/Editor/UnrealEd/Private/AssetPrivatizeModel.cpp b/Engine/Source/Editor/UnrealEd/Private/AssetPrivatizeModel.cpp
new file mode 100644
--- /dev/null
+++ b/Engine/Source/Editor/UnrealEd/Private/AssetPrivatizeModel.cpp
@@ -0,0 +1,229 @@
+// Copyright Epic Games, Inc. All Rights Reserved.
+
+#include "AssetPrivatizeModel.h"
+#include "Engine/Blueprint.h"
+
+#include "ObjectTools.h"
+#include "AssetRegistry/AssetRegistryModule.h"
+
+#define LOCTEXT_NAMESPACE "FAssetPrivatizeModel"
+
+FPendingPrivateAsset::FPendingPrivateAsset(UObject* InObject, UPackage* InOwningPackage)
+	: Object(InObject)
+	, OwningPackage(InOwningPackage)
+	, bIsReferencedInMemoryByNonUndo(false)
+	, bIsReferencedInMemoryByUndo(false)
+	, bReferencesChecked(false)
+{
+	if (!OwningPackage)
+	{
+		OwningPackage = Object->GetPackage();
+	}
+
+	OwningPackageMountPointName = FPackageName::GetPackageMountPoint(OwningPackage->GetFName().ToString());
+}
+
+bool FPendingPrivateAsset::IsReferenceIllegal(const FName& InReference)
+{
+	// A reference is illegal if it refers to a private asset that is not in the same mount point
+	FName ReferenceMountPointName = FPackageName::GetPackageMountPoint(InReference.ToString());
+	return !ReferenceMountPointName.IsNone() && (ReferenceMountPointName != OwningPackageMountPointName);
+}
+
+void FPendingPrivateAsset::CheckForIllegalReferences()
+{
+	if (bReferencesChecked)
+	{
+		return;
+	}
+
+	TRACE_CPUPROFILER_EVENT_SCOPE(FPendingPrivateAsset::CheckForIllegalReferences)
+		bReferencesChecked = true;
+
+	FAssetRegistryModule& AssetRegistryModule = FModuleManager::LoadModuleChecked<FAssetRegistryModule>(TEXT("AssetRegistry"));
+
+	AssetRegistryModule.Get().GetReferencers(OwningPackage->GetFName(), IllegalDiskReferences);
+
+	IllegalDiskReferences.RemoveAll([=](const FName& Reference)
+		{
+			return !IsReferenceIllegal(Reference);
+		});
+
+	ObjectTools::GatherObjectReferencersForDeletion(Object, bIsReferencedInMemoryByNonUndo, bIsReferencedInMemoryByUndo, &IllegalMemoryReferences);
+
+	IllegalMemoryReferences.ExternalReferences.RemoveAll([=](const FReferencerInformation& ReferenceInfo)
+		{
+			return !IsReferenceIllegal(ReferenceInfo.Referencer->GetPackage()->GetFName());
+		});
+
+	IllegalMemoryReferences.InternalReferences.RemoveAll([=](const FReferencerInformation& ReferenceInfo)
+		{
+			return !IsReferenceIllegal(ReferenceInfo.Referencer->GetPackage()->GetFName());
+		});
+}
+
+FAssetPrivatizeModel::FAssetPrivatizeModel(const TArray<UObject*>& InObjectsToPrivatize)
+	: bIsAnythingReferencedInMemoryByNonUndo(false)
+	, bIsAnythingReferencedInMemoryByUndo(false)
+	, PendingPrivateIndex(0)
+	, State(StartScanning)
+	, ObjectsPrivatized(0)
+{
+	for (UObject* ObjectToPrivatize : InObjectsToPrivatize)
+	{
+		if (ObjectToPrivatize)
+		{
+			UPackage* ObjectPackage = ObjectToPrivatize->GetPackage();
+
+			if (ObjectPackage && ObjectPackage->IsExternallyReferenceable())
+			{
+				AddObjectToPrivatize(ObjectToPrivatize, ObjectPackage);
+			}
+		}
+	}
+}
+
+void FAssetPrivatizeModel::AddObjectToPrivatize(UObject* InObject, UPackage* InOwningPackage)
+{
+	bool NewPendingAsset = !PendingPrivateAssets.ContainsByPredicate([=](const TSharedPtr<FPendingPrivateAsset>& PendingPrivateAsset)
+		{
+			return PendingPrivateAsset->GetObject() == InObject;
+		});
+
+	if (NewPendingAsset)
+	{
+		PendingPrivateAssets.Add(MakeShareable(new FPendingPrivateAsset(InObject, InOwningPackage)));
+	}
+
+	SetState(StartScanning);
+}
+
+void FAssetPrivatizeModel::SetState(EState NewState)
+{
+	if (State != NewState)
+	{
+		State = NewState;
+
+		if (StateChanged.IsBound())
+		{
+			StateChanged.Broadcast(NewState);
+		}
+	}
+}
+
+float FAssetPrivatizeModel::GetProgress() const
+{
+	return PendingPrivateIndex / (float)PendingPrivateAssets.Num();
+}
+
+FText FAssetPrivatizeModel::GetProgressText() const
+{
+	if (PendingPrivateIndex < PendingPrivateAssets.Num())
+	{
+		return FText::FromString(PendingPrivateAssets[PendingPrivateIndex]->GetObject()->GetName());
+	}
+	else
+	{
+		return LOCTEXT("Done", "Done!");
+	}
+}
+
+bool FAssetPrivatizeModel::CanPrivatize()
+{
+	return !CanForcePrivatize();
+}
+
+bool FAssetPrivatizeModel::DoPrivatize()
+{
+	if (!CanPrivatize())
+	{
+		return false;
+	}
+
+	for (TSharedPtr<FPendingPrivateAsset>& PendingPrivateAsset : PendingPrivateAssets)
+	{
+		PendingPrivateAsset->GetOwningPackage()->SetIsExternallyReferenceable(false);
+		PendingPrivateAsset->GetOwningPackage()->Modify();
+		ObjectsPrivatized++;
+	}
+
+	return true;
+}
+
+bool FAssetPrivatizeModel::CanForcePrivatize()
+{
+	return bIsAnythingReferencedInMemoryByNonUndo || bIsAnythingReferencedInMemoryByUndo || (IllegalOnDiskReferences.Num() > 0);
+}
+
+bool FAssetPrivatizeModel::DoForcePrivatize()
+{
+	TArray<UObject*> ObjectsToPrivatize;
+
+	TSet<UObject*> ObjectsToPrivatizeWithin;
+	for (TSharedPtr<FPendingPrivateAsset>& PendingPrivateAsset : PendingPrivateAssets)
+	{
+		ObjectsToPrivatize.Add(PendingPrivateAsset->GetObject());
+
+		FReferencerInformationList& AssetIllegalReferencers = PendingPrivateAsset->IllegalMemoryReferences;
+		for (FReferencerInformation& ExternalReference : AssetIllegalReferencers.ExternalReferences)
+		{
+			ObjectsToPrivatizeWithin.Add(ExternalReference.Referencer);
+		}
+
+		PendingPrivateAsset->GetOwningPackage()->SetIsExternallyReferenceable(false);
+		PendingPrivateAsset->GetOwningPackage()->Modify();
+		ObjectsPrivatized++;
+	}
+
+    // Null out the illegal references
+	ObjectTools::ForceReplaceReferences(nullptr, ObjectsToPrivatize, ObjectsToPrivatizeWithin);
+
+	return true;
+}
+
+void FAssetPrivatizeModel::ScanForReferences()
+{
+	double StartTickSeconds = FPlatformTime::Seconds();
+
+	while (PendingPrivateIndex < PendingPrivateAssets.Num() && (FPlatformTime::Seconds() - StartTickSeconds) < MaxTickSeconds)
+	{
+		TSharedPtr<FPendingPrivateAsset>& PendingPrivateAsset = PendingPrivateAssets[PendingPrivateIndex];
+
+		PendingPrivateAsset->CheckForIllegalReferences();
+
+		IllegalOnDiskReferences.Append(PendingPrivateAsset->IllegalDiskReferences);
+
+		bIsAnythingReferencedInMemoryByUndo |= PendingPrivateAsset->IsReferencedInMemoryByUndo();
+		bIsAnythingReferencedInMemoryByNonUndo |= PendingPrivateAsset->IsReferencedInMemoryByNonUndo();
+		PendingPrivateIndex++;
+	}
+
+	if (PendingPrivateIndex >= PendingPrivateAssets.Num())
+	{
+		SetState(Finished);
+	}
+}
+
+void FAssetPrivatizeModel::Tick(const float InDeltaTime)
+{
+	switch (State)
+	{
+	case Waiting:
+		break;
+	case StartScanning:
+		IllegalOnDiskReferences = TSet<FName>();
+		bIsAnythingReferencedInMemoryByNonUndo = false;
+		bIsAnythingReferencedInMemoryByUndo = false;
+		PendingPrivateIndex = 0;
+		SetState(Scanning);
+		break;
+	case Scanning:
+		ScanForReferences();
+		break;
+	case Finished:
+		break;
+	default:
+		break;
+	}
+}
+#undef LOCTEXT_NAMESPACE
\ No newline at end of file
diff --git a/Engine/Source/Editor/UnrealEd/Private/Dialogs/SPrivateAssetsDialog.cpp b/Engine/Source/Editor/UnrealEd/Private/Dialogs/SPrivateAssetsDialog.cpp
new file mode 100644
--- /dev/null
+++ b/Engine/Source/Editor/UnrealEd/Private/Dialogs/SPrivateAssetsDialog.cpp
@@ -0,0 +1,453 @@
+// Copyright Epic Games, Inc. All Rights Reserved.
+
+#include "Dialogs/SPrivateAssetsDialog.h"
+#include "Widgets/Notifications/SProgressBar.h"
+#include "Widgets/Input/SButton.h"
+#include "Widgets/Text/STextBlock.h"
+#include "SPrimaryButton.h"
+#include "SWarningOrErrorBox.h"
+
+#include "IContentBrowserSingleton.h"
+#include "ContentBrowserModule.h"
+
+#include "Editor.h"
+#include "Editor/Transactor.h"
+
+#define LOCTEXT_NAMESPACE "SPrivateAssetsDialog"
+
+namespace PrivatizeAssetsView
+{
+	static const FName ColumnID_Asset("Asset");
+	static const FName ColumnID_AssetClass("Class");
+	static const FName ColumnID_DiskReferences("DiskReferences");
+	static const FName ColumnID_MemoryReferences("MemoryReferences");
+}
+
+//////////////////////////////////////////////////////////////////////////
+// SPendingPrivateAssetRow
+
+class SPendingPrivateAssetRow : public SMultiColumnTableRow<TSharedPtr<FPendingPrivateAsset>>
+{
+public:
+
+	SLATE_BEGIN_ARGS(SPendingPrivateAssetRow) {}
+	SLATE_END_ARGS()
+
+	void Construct(const FArguments& InArgs, const TSharedRef<STableViewBase>& InOwnerTableView, TSharedPtr<FPendingPrivateAsset> InItem)
+	{
+		Item = InItem;
+
+		SMultiColumnTableRow<TSharedPtr<FPendingPrivateAsset>>::Construct(FSuperRowType::FArguments(), InOwnerTableView);
+	}
+
+	virtual TSharedRef<SWidget> GenerateWidgetForColumn(const FName& ColumnName)
+	{
+		if (ColumnName == PrivatizeAssetsView::ColumnID_Asset)
+		{
+			return SNew(SHorizontalBox)
+				+ SHorizontalBox::Slot()
+				.AutoWidth()
+				.Padding(3, 0, 0, 0)
+				[
+					SNew(STextBlock)
+					.Text(FText::FromString(Item->GetObject()->GetName()))
+				];
+		}
+		else if (ColumnName == PrivatizeAssetsView::ColumnID_AssetClass)
+		{
+			return SNew(STextBlock)
+				.Text(FText::FromString(Item->GetObject()->GetClass()->GetName()));
+		}
+		else if (ColumnName == PrivatizeAssetsView::ColumnID_DiskReferences)
+		{
+			int32 IllegalDiskReferenceCount = Item->IllegalDiskReferences.Num();
+			FFormatNamedArguments Args;
+			Args.Add(TEXT("AssetCount"), FText::AsNumber(IllegalDiskReferenceCount));
+
+			FText OnDiskCountText = IllegalDiskReferenceCount > 1 ?
+				FText::Format(LOCTEXT("InAssetReferences", "{AssetCount} References"), Args) : FText::Format(LOCTEXT("InAssetReference", "{AssetCount} Reference"), Args);
+
+			return SNew(STextBlock)
+				.Text(OnDiskCountText)
+				.Visibility(IllegalDiskReferenceCount > 0 ? EVisibility::Visible : EVisibility::Hidden);
+		}
+		else if (ColumnName == PrivatizeAssetsView::ColumnID_MemoryReferences)
+		{
+			int32 IllegalMemoryReferenceCount = Item->IllegalMemoryReferences.ExternalReferences.Num();
+			FFormatNamedArguments Args;
+			Args.Add(TEXT("ReferenceCount"), FText::AsNumber(IllegalMemoryReferenceCount));
+
+			FText InMemoryCountText = IllegalMemoryReferenceCount > 1 ? 
+				FText::Format(LOCTEXT("InMemoryReferences", "{ReferenceCount} References"), Args) : FText::Format(LOCTEXT("InMemoryReference", "{ReferenceCount} Reference"), Args);
+
+			return SNew(STextBlock)
+				.Text(InMemoryCountText)
+				.Visibility(IllegalMemoryReferenceCount > 0 ? EVisibility::Visible : EVisibility::Hidden);
+		}
+
+		return SNullWidget::NullWidget;
+	}
+
+private:
+	TSharedPtr<FPendingPrivateAsset> Item;
+};
+
+//////////////////////////////////////////////////////////////////////////
+// SPrivateAssetsDialog
+
+SPrivateAssetsDialog::~SPrivateAssetsDialog()
+{
+	PrivatizeModel->OnStateChanged().RemoveAll(this);
+}
+
+void SPrivateAssetsDialog::Construct(const FArguments& InArgs, const TSharedRef<FAssetPrivatizeModel> InPrivatizeModel)
+{
+	bIsActiveTimerRegistered = true;
+	RegisterActiveTimer(0.0f, FWidgetActiveTimerDelegate::CreateSP(this, &SPrivateAssetsDialog::TickPrivatizeModel));
+
+	PrivatizeModel = InPrivatizeModel;
+
+	ParentWindow = InArgs._ParentWindow;
+
+	ChildSlot
+	[
+		SAssignNew(RootContainer, SBorder)
+		.BorderImage(FAppStyle::GetBrush("AssetDeleteDialog.Background"))
+		.Padding(10)
+	];
+
+	PrivatizeModel->OnStateChanged().AddRaw(this, &SPrivateAssetsDialog::HandlePrivatizeModelStateChanged);
+
+	HandlePrivatizeModelStateChanged(PrivatizeModel->GetState());
+}
+
+FReply SPrivateAssetsDialog::OnKeyDown(const FGeometry& MyGeometry, const FKeyEvent& InKeyEvent)
+{
+
+	if (InKeyEvent.GetKey() == EKeys::Escape)
+	{
+		ParentWindow.Get()->RequestDestroyWindow();
+		return FReply::Handled();
+	}
+	return FReply::Unhandled();
+}
+
+TSharedRef<SWidget> SPrivateAssetsDialog::BuildProgressDialog()
+{
+	// Shows a progress bar during the scanning of references for each pending private asset
+	return SNew(SVerticalBox)
+	+ SVerticalBox::Slot()
+	.VAlign(VAlign_Center)
+	.FillHeight(1.0f)
+	[
+		SNew(SVerticalBox)
+		+ SVerticalBox::Slot()
+		.Padding(5.0f, 0)
+		[
+			SNew(STextBlock)
+			.Text(this, &SPrivateAssetsDialog::ScanningText)
+		]
+
+		+ SVerticalBox::Slot()
+		.AutoHeight()
+		.Padding(5.0f, 10.0f)
+		[
+			SNew(SProgressBar)
+			.Percent(this, &SPrivateAssetsDialog::ScanningProgressFraction)
+		]
+	];
+}
+
+TSharedRef<SWidget> SPrivateAssetsDialog::BuildPrivatizeDialog()
+{
+	
+	TSharedRef<SHeaderRow> HeaderRowWidget =
+		SNew(SHeaderRow)
+		+SHeaderRow::Column(PrivatizeAssetsView::ColumnID_Asset)
+		.DefaultLabel(LOCTEXT("Column_AssetName", "Asset"))
+		.HAlignHeader(EHorizontalAlignment::HAlign_Left)
+		.FillWidth(0.5f)
+		+SHeaderRow::Column(PrivatizeAssetsView::ColumnID_AssetClass)
+		.DefaultLabel(LOCTEXT("Column_AssetClass", "Class"))
+		.HAlignHeader(HAlign_Left)
+		.FillWidth(0.25f)
+		+SHeaderRow::Column(PrivatizeAssetsView::ColumnID_DiskReferences)
+		.DefaultLabel(LOCTEXT("Column_DiskReferences", "Asset Referencers"))
+		.HAlignHeader(HAlign_Left)
+		.FillWidth(0.25f)
+		+SHeaderRow::Column(PrivatizeAssetsView::ColumnID_MemoryReferences)
+		.DefaultLabel(LOCTEXT("Column_MemoryReferences", "Memory Referencers"))
+		.HAlignHeader(HAlign_Left)
+		.FillWidth(0.25f);
+
+	return SNew(SVerticalBox)
+	// The to be made private assets and their reference counts
+	+SVerticalBox::Slot()
+	.FillHeight(0.5f)
+	.Padding(5.0f)
+	[
+		SNew(SBorder)
+		.BorderImage(FAppStyle::GetBrush("ToolPanel.GroupBorder"))
+		.Padding(FMargin(0, 0, 0, 3))
+		[
+			SNew(SVerticalBox)
+			+SVerticalBox::Slot()
+			.AutoHeight()
+			[
+				SNew(SBorder)
+				.BorderImage(FAppStyle::GetBrush("DetailsView.CategoryTop"))
+				.BorderBackgroundColor(FLinearColor(.6f, .6f, .6f, 1.0f))
+				.Padding(3.0f)
+				[
+					SNew(STextBlock)
+					.Text(LOCTEXT("AttemptingPrivatize", "The following assets will be made private:"))
+					.Font(FAppStyle::GetFontStyle("BoldFont"))
+				]
+			]
+			+SVerticalBox::Slot()
+			.FillHeight(1.0f)
+			[
+				SAssignNew(ObjectsToPrivatizeList, SListView<TSharedPtr<FPendingPrivateAsset>>)
+				.ListItemsSource(PrivatizeModel->GetPendingPrivateAssets())
+				.OnGenerateRow(this, &SPrivateAssetsDialog::HandleGenerateAssetRow)
+				.HeaderRow(HeaderRowWidget)
+			]
+		]
+	]
+	// An asset list of all on disk references that would become illegal if the pending assets are made private
+	+SVerticalBox::Slot()
+	.FillHeight(1.0f)
+	.Padding(5.0f)
+	[
+		SNew(SBorder)
+		.BorderImage(FAppStyle::GetBrush("Brushes.Panel"))
+		.Visibility(this, &SPrivateAssetsDialog::GetAssetReferencesVisibility)
+		[
+			SNew(SVerticalBox)
+			+SVerticalBox::Slot()
+			.AutoHeight()
+			[
+				SNew(SBorder)
+				.BorderImage(FAppStyle::GetBrush("DetailsView.CategoryTop"))
+				.BorderBackgroundColor(FLinearColor(.6f, .6f, .6f, 1.0f))
+				.Padding(3.0f)
+				[
+					SNew(STextBlock)
+					.Text(LOCTEXT("AssetsReferencingPendingPrivateAssets", "References to the above assets will be removed from the following:"))
+					.Font(FAppStyle::GetFontStyle("BoldFont"))
+				]
+			]
+			+SVerticalBox::Slot()
+			.FillHeight(1.0f)
+			[
+				BuildAssetViewForReferencerAssets()
+			]
+		]
+	]
+	// Details any extra warning information about the operation
+	+SVerticalBox::Slot()
+	.AutoHeight()
+	.Padding(FMargin(8.0f, 5.0f))
+	[
+		SNew(SWarningOrErrorBox)
+		.Visibility(this, &SPrivateAssetsDialog::GetWarningTextVisibility)
+		.MessageStyle(EMessageStyle::Warning)
+		.Message(this, &SPrivateAssetsDialog::GetWarningText)
+	]
+	// Confirm or cancel buttons
+	+SVerticalBox::Slot()
+	.AutoHeight()
+	[
+		SNew(SBorder)
+		.BorderImage(FAppStyle::GetBrush("Brushes.Panel"))
+		.Padding(0)
+		[
+			SNew(SVerticalBox)
+			+SVerticalBox::Slot()
+			.AutoHeight()
+			.Padding(0, 4)
+			[
+				SNew(SHorizontalBox)
+				+SHorizontalBox::Slot()
+				.FillWidth(1.0)
+				.Padding(6, 0)
+				[
+					SNew(SBorder)
+					.BorderImage(FAppStyle::GetBrush("NoBorder"))
+					[
+						SNew(SPrimaryButton)
+						.Text(LOCTEXT("Privatize", "Make Private"))
+						.ToolTipText(LOCTEXT("PrivatizeTooltipText", "Make private"))
+						.OnClicked(this, &SPrivateAssetsDialog::ForcePrivatize)
+					]
+				]
+				+SHorizontalBox::Slot()
+				.FillWidth(1.0)
+				.Padding(6, 0)
+				[
+					SNew(SBorder)
+					.BorderImage(FAppStyle::GetBrush("NoBorder"))
+					.VAlign(VAlign_Bottom)
+					[
+						SNew(SButton)
+						.HAlign(HAlign_Center)
+						.Text(LOCTEXT("Cancel", "Cancel"))
+						.ToolTipText(LOCTEXT("CancelPrivatizeTooltipText", "Cancel making selection private"))
+						.OnClicked(this, &SPrivateAssetsDialog::Cancel)
+					]
+				]
+			]
+		]
+	];
+}
+
+TSharedRef<SWidget> SPrivateAssetsDialog::BuildAssetViewForReferencerAssets()
+{
+	FContentBrowserModule& ContentBrowserModule = FModuleManager::Get().LoadModuleChecked<FContentBrowserModule>(TEXT("ContentBrowser"));
+
+	FAssetPickerConfig AssetPickerConfig;
+
+	AssetPickerConfig.bAllowDragging = false;
+	AssetPickerConfig.bCanShowClasses = false;
+	AssetPickerConfig.bAllowNullSelection = false;
+	AssetPickerConfig.bShowBottomToolbar = false;
+	AssetPickerConfig.bAutohideSearchBar = true;
+	AssetPickerConfig.bPreloadAssetsForContextMenu = false;
+
+	AssetPickerConfig.InitialAssetViewType = EAssetViewType::Column;
+	AssetPickerConfig.OnShouldFilterAsset = FOnShouldFilterAsset::CreateSP(this, &SPrivateAssetsDialog::OnShouldFilterAsset);
+
+	return ContentBrowserModule.Get().CreateAssetPicker(AssetPickerConfig);
+}
+
+EActiveTimerReturnType SPrivateAssetsDialog::TickPrivatizeModel(double InCurrentTime, float InDeltaTime)
+{
+	PrivatizeModel->Tick(InDeltaTime);
+
+	if (PrivatizeModel->GetState() == FAssetPrivatizeModel::EState::Finished)
+	{
+		bIsActiveTimerRegistered = false;
+		return EActiveTimerReturnType::Stop;
+	}
+
+	return EActiveTimerReturnType::Continue;
+}
+
+void SPrivateAssetsDialog::HandlePrivatizeModelStateChanged(FAssetPrivatizeModel::EState NewState)
+{
+	switch (NewState)
+	{
+	case FAssetPrivatizeModel::StartScanning:
+		RootContainer->SetContent(BuildProgressDialog());
+		break;
+	case FAssetPrivatizeModel::Finished:
+		// If in our scan we realize we can cleanly perform the privatize operation without any illegal reference fix up
+		// We can skip the confirmation dialog and close down
+		if (PrivatizeModel->CanPrivatize())
+		{
+			Privatize();
+		}
+		RootContainer->SetContent(BuildPrivatizeDialog());
+		break;
+	default:
+		break;
+	}
+}
+
+bool SPrivateAssetsDialog::OnShouldFilterAsset(const FAssetData& InAssetData) const
+{
+	if (InAssetData.IsRedirector() && !InAssetData.IsUAsset())
+	{
+		return true;
+	}
+
+	if (PrivatizeModel->GetIllegalAssetReferences().Contains(InAssetData.PackageName))
+	{
+		return false;
+	}
+
+	return true;
+}
+
+FReply SPrivateAssetsDialog::Privatize()
+{
+	ParentWindow.Get()->RequestDestroyWindow();
+
+	PrivatizeModel->DoPrivatize();
+
+	return FReply::Handled();
+}
+
+FReply SPrivateAssetsDialog::ForcePrivatize()
+{
+	ParentWindow.Get()->RequestDestroyWindow();
+
+	// If anything is referenced in undo we need to purge the undo stack
+	// Otherwise you could undo back to restoring a reference which would now be illegal
+	if (PrivatizeModel->IsAnythingReferencedInMemoryByUndo())
+	{
+		GEditor->Trans->Reset(LOCTEXT("PrivatizeSelectedItem", "Privatize Selected Item"));
+	}
+
+	PrivatizeModel->DoForcePrivatize();
+
+	return FReply::Handled();
+}
+
+FReply SPrivateAssetsDialog::Cancel()
+{
+	ParentWindow.Get()->RequestDestroyWindow();
+
+	return FReply::Handled();
+}
+
+FText SPrivateAssetsDialog::GetWarningText() const
+{
+	bool bReferencesFoundInUndo = PrivatizeModel->IsAnythingReferencedInMemoryByUndo();
+	bool bReferencesFoundInNonUndo = PrivatizeModel->IsAnythingReferencedInMemoryByNonUndo();
+
+	FText MemoryWarning = LOCTEXT("PrivateAssetMemoryReferenceWarning", "- Some assets are still referenced in memory");
+	FText UndoWarning = LOCTEXT("PrivateAssetUndoWarning", "- Undo stack will be cleared as a result of this action");
+
+	FText WarningText;
+
+	if (bReferencesFoundInNonUndo)
+	{
+		WarningText = MemoryWarning;
+	}
+
+	if (bReferencesFoundInUndo)
+	{
+		WarningText = FText::Join(FText::FromString("\n"), WarningText, UndoWarning);
+	}
+
+	return WarningText;
+}
+
+EVisibility SPrivateAssetsDialog::GetWarningTextVisibility() const
+{
+	return (PrivatizeModel->IsAnythingReferencedInMemoryByUndo() || 
+		PrivatizeModel->IsAnythingReferencedInMemoryByNonUndo()) ? EVisibility::Visible : EVisibility::Collapsed;
+}
+
+EVisibility SPrivateAssetsDialog::GetAssetReferencesVisibility() const
+{
+	return PrivatizeModel->GetIllegalAssetReferences().Num() == 0 ? EVisibility::Collapsed : EVisibility::Visible;
+}
+
+FText SPrivateAssetsDialog::ScanningText() const
+{
+	return PrivatizeModel->GetProgressText();
+}
+
+TOptional<float> SPrivateAssetsDialog::ScanningProgressFraction() const
+{
+	return PrivatizeModel->GetProgress();
+}
+
+TSharedRef<ITableRow> SPrivateAssetsDialog::HandleGenerateAssetRow(TSharedPtr<FPendingPrivateAsset> InItem, const TSharedRef<STableViewBase>& OwnerTable)
+{
+	return SNew(SPendingPrivateAssetRow, OwnerTable, InItem);
+}
+
+#undef LOCTEXT_NAMESPACE
\ No newline at end of file
diff --git a/Engine/Source/Editor/UnrealEd/Private/ObjectTools.cpp b/Engine/Source/Editor/UnrealEd/Private/ObjectTools.cpp
--- a/Engine/Source/Editor/UnrealEd/Private/ObjectTools.cpp
+++ b/Engine/Source/Editor/UnrealEd/Private/ObjectTools.cpp
@@ -71,6 +71,7 @@
 #include "Dialogs/DlgReferenceTree.h"
 #include "AssetRegistry/ARFilter.h"
 #include "AssetDeleteModel.h"
+#include "Dialogs/SPrivateAssetsDialog.h"
 #include "Dialogs/SDeleteAssetsDialog.h"
 #include "AudioDevice.h"
 #include "ReferencedAssetsUtils.h"
@@ -2523,6 +2524,27 @@ namespace ObjectTools
 		return NumPackagesToDelete + NumObjectsToDelete;
 	}
 
+	int32 PrivatizeAssets(const TArray<FAssetData>& AssetsToPrivatize, bool bShowConfirmation)
+	{
+		TArray<UObject*> ObjectsToPrivatize;
+		for (const FAssetData& AssetToPrivatize : AssetsToPrivatize)
+		{
+			UObject* ObjectToPrivatize = AssetToPrivatize.GetAsset({ ULevel::LoadAllExternalObjectsTag });
+
+			if (ObjectToPrivatize)
+			{
+				ObjectsToPrivatize.Add(ObjectToPrivatize);
+			}
+		}
+
+		if (!ObjectsToPrivatize.IsEmpty())
+		{
+			return PrivatizeObjects(ObjectsToPrivatize, bShowConfirmation);
+		}
+
+		return 0;
+	}
+
 	void AddExtraObjectsToDelete(TArray< UObject* >& ObjectsToDelete)
 	{
 		const int32 OriginalNum = ObjectsToDelete.Num();
@@ -2742,6 +2764,76 @@ namespace ObjectTools
 		return DeleteModel->GetDeletedObjectCount();
 	}
 
+	int32 PrivatizeObjects(const TArray<UObject*>& InObjectsToPrivatize, bool bShowConfirmation, EAllowCancelDuringPrivatize AllowCancelDuringPrivatize)
+	{
+		const FScopedBusyCursor BusyCursor;
+		TArray<UObject*> ObjectsToPrivatize = InObjectsToPrivatize;
+
+		if (!HandleFullyLoadingPackages(ObjectsToPrivatize, NSLOCTEXT("UnrealEd", "Privatize", "Privatize")))
+		{
+			return 0;
+		}
+
+		FAssetRegistryModule& AssetRegistryModule = FModuleManager::LoadModuleChecked<FAssetRegistryModule>(TEXT("AssetRegistry"));
+
+		if (AssetRegistryModule.Get().IsLoadingAssets())
+		{
+			FNotificationInfo Info(NSLOCTEXT("UnrealEd", "Warning_CantPrivatizeRebuildingAssetRegistry", "Unable To Mark Private While Discovering Assets"));
+			Info.ExpireDuration = 3.0f;
+			FSlateNotificationManager::Get().AddNotification(Info);
+			return 0;
+		}
+
+		TSharedRef<FAssetPrivatizeModel> PrivatizeModel = MakeShared<FAssetPrivatizeModel>(ObjectsToPrivatize);
+
+		if (bShowConfirmation)
+		{
+			const FVector2D DEFAULT_WINDOW_SIZE = FVector2D(600, 700);
+
+			TSharedRef<SWindow> PrivatizeAssetsWindow = SNew(SWindow)
+				.Title(NSLOCTEXT("UnrealED", "Privatize Assets", "Make Assets Private"))
+				.ClientSize(DEFAULT_WINDOW_SIZE);
+
+			TSharedRef<SPrivateAssetsDialog> PrivatizeDialog =
+				SNew(SPrivateAssetsDialog, PrivatizeModel)
+				.ParentWindow(PrivatizeAssetsWindow);
+
+			PrivatizeAssetsWindow->SetContent(PrivatizeDialog);
+
+			GEditor->EditorAddModalWindow(PrivatizeAssetsWindow);
+
+			return PrivatizeModel->GetObjectsPrivatizedCount();
+		}
+
+		bool bUserCanceled = false;
+		const bool bAllowCancelDuringPrivatize = (AllowCancelDuringPrivatize == EAllowCancelDuringPrivatize::AllowCancel);
+		GWarn->BeginSlowTask(NSLOCTEXT("UnrealEd", "VerifyingPrivatize", "Verifying Privatize"), true, bAllowCancelDuringPrivatize);
+		while (!bUserCanceled && PrivatizeModel->GetState() != FAssetPrivatizeModel::Finished)
+		{
+			PrivatizeModel->Tick(0);
+			GWarn->StatusUpdate((int32)(PrivatizeModel->GetProgress() * 100), 100, PrivatizeModel->GetProgressText());
+
+			if (bAllowCancelDuringPrivatize)
+			{
+				bUserCanceled = GWarn->ReceivedUserCancel();
+			}
+		}
+		GWarn->EndSlowTask();
+
+		if (bUserCanceled)
+		{
+			UE_LOG(LogUObjectGlobals, Warning, TEXT("User cancelled privatize operation"));
+			return 0;
+		}
+
+		if (!PrivatizeModel->DoPrivatize())
+		{
+			UE_LOG(LogUObjectGlobals, Warning, TEXT("Could not mark private"));
+		}
+
+		return PrivatizeModel->GetObjectsPrivatizedCount();
+	}
+
 	static bool MakeReadOnlyPackageWritable(UObject* ObjectToDelete, bool& bMakeWritable, bool& bSilent)
 	{
 		// If an object's package is read only, and source control is not enabled, ask the user whether they wish
diff --git a/Engine/Source/Editor/UnrealEd/Public/AssetPrivatizeModel.h b/Engine/Source/Editor/UnrealEd/Public/AssetPrivatizeModel.h
new file mode 100644
--- /dev/null
+++ b/Engine/Source/Editor/UnrealEd/Public/AssetPrivatizeModel.h
@@ -0,0 +1,160 @@
+// Copyright Epic Games, Inc. All Rights Reserved.
+
+#pragma once
+
+#include "CoreMinimal.h"
+#include "UObject/ObjectMacros.h"
+
+struct FPendingPrivateAsset : TSharedFromThis<FPendingPrivateAsset>
+{
+public:
+
+	FPendingPrivateAsset(UObject* InObject, UPackage* InOwningPackage = nullptr);
+
+	/** Checks if a given reference does not share the same mount point as the object. Which would be illegal if the object is made private */
+	bool IsReferenceIllegal(const FName& InReference);
+
+	/** Checks for and records references on disk and in memory for this object that would be considered illegal if the object became private */
+	void CheckForIllegalReferences();
+
+	/** Gets the object being privatized */
+	UObject* GetObject() { return Object; }
+
+	/** Gets the objects owning package */
+	UPackage* GetOwningPackage() { return OwningPackage; }
+
+	/** Is the pending object referenced in memory by the undo stack */
+	bool IsReferencedInMemoryByNonUndo() { return bIsReferencedInMemoryByNonUndo; }
+
+	/** Is the pending private object referenced in memory by something other than the undo stack */
+	bool IsReferencedInMemoryByUndo() { return bIsReferencedInMemoryByUndo; }
+
+	/** The on disk references to this object that would be illegal if the object were made private */
+	TArray<FName> IllegalDiskReferences;
+
+	/** The in memory references to this object that would be illegal if the object were made private (*excluding the undo buffer) */
+	FReferencerInformationList IllegalMemoryReferences;
+
+private:
+
+	/** The object to make private */
+	UObject* Object;
+
+	/** The owning package of the object*/
+	UPackage* OwningPackage;
+	/** The cached mount pount name of the owning package */
+	FName OwningPackageMountPointName;
+
+	/** Flag indicating if this object is referenced in memory by the engine (excluding the undo buffer) */
+	bool bIsReferencedInMemoryByNonUndo;
+
+	/** Flag indicating if this object is referenced by the undo stack */
+	bool bIsReferencedInMemoryByUndo;
+
+	/** Flag indicating that references have been checked, so don't check again. */
+	bool bReferencesChecked;
+};
+
+class UNREALED_API FAssetPrivatizeModel
+{
+public:
+	enum EState
+	{
+		// Waiting to start scanning
+		Waiting = 0,
+		// Begin scanning for references
+		StartScanning,
+		// Scan for references to the pending private assets
+		Scanning,
+		// Finished scanning
+		Finished,
+	};
+
+	FAssetPrivatizeModel(const TArray<UObject*>& InObjectsToPrivatize);
+
+	/** Add an object to the list of pending private assets, this will invalidate the scanning state */
+	void AddObjectToPrivatize(UObject* InObject, UPackage* InOwningPackage);
+
+	/** Returns the pending private assets */
+	const TArray<TSharedPtr<FPendingPrivateAsset>>* GetPendingPrivateAssets() const { return &PendingPrivateAssets; }
+
+	/** Gets the packages of the assets on disk that would illegally reference the pending private objects if they were made private; won't be accurate until the scanning process completes. */
+	const TSet<FName>& GetIllegalAssetReferences() const { return IllegalOnDiskReferences; }
+
+	/** Ticks the privatize model which does a timeslice of work before returning */
+	void Tick(const float InDeltaTime);
+
+	/** Returns the current state of the scanning process */
+	EState GetState() { return State; }
+
+	/** Gets the 0..1 progress of the scan */
+	float GetProgress() const;
+
+	/** Gets the current text to display for the current progress of the scan */
+	FText GetProgressText() const;
+
+	/** Are any of the pending assets being referenced in memory. */
+	bool IsAnythingReferencedInMemoryByNonUndo() const { return bIsAnythingReferencedInMemoryByNonUndo; }
+
+	/** Are any of the pending assets being referenced in the undo stack. */
+	bool IsAnythingReferencedInMemoryByUndo() const { return bIsAnythingReferencedInMemoryByUndo; }
+
+	/** Gets the number of objects successfully made private */
+	int32 GetObjectsPrivatizedCount() const { return ObjectsPrivatized; }
+
+	/** Returns true if it is valid to mark private the pending objects with no problems (there are no illegal references detected) */
+	bool CanPrivatize();
+
+	/** Performs the privatization action if possible */
+	bool DoPrivatize();
+
+	/** Returns true if we can't cleanly mark the pending objects as private (there are illegal references detected) */
+	bool CanForcePrivatize();
+
+	/** Marks the pending private objects as private and nulls out all detected illegal references */
+	bool DoForcePrivatize();
+
+	/** Fires whenever the state changes */
+	DECLARE_EVENT_OneParam(FAssetPrivatizeModel, FOnStateChanged, EState /*NewState*/);
+	FOnStateChanged& OnStateChanged()
+	{
+		return StateChanged;
+	}
+
+private:
+
+	/** Sets the current state of the model */
+	void SetState(EState NewState);
+
+	/** Helper that does the reference detection during scanning */
+	void ScanForReferences();
+
+private:
+
+	/** Holds an event delegate that is executed when the state changes */
+	FOnStateChanged StateChanged;
+
+	/** The assets being marked private */
+	TArray<TSharedPtr<FPendingPrivateAsset>> PendingPrivateAssets;
+
+	/** Are any of the pending private assets being referenced in memory */
+	bool bIsAnythingReferencedInMemoryByNonUndo;
+
+	/** Are any of the pending private assets being referenced in the undo stack */
+	bool bIsAnythingReferencedInMemoryByUndo;
+
+	/** A tick-to-tick state tracking variable so we know what pending private object we checked last */
+	int32 PendingPrivateIndex;
+
+	/** On disk references to the current pending private objects that would be illegal if they are made private */
+	TSet<FName> IllegalOnDiskReferences;
+
+	/** The internal progress/state of the privatize model which can take several frames to recalculate validity */
+	EState State;
+
+	/** The maximum timeslice per tick we are willing to scan for*/
+	const double MaxTickSeconds = 0.100;
+
+	/** The number of objects successfully marked private */
+	int32 ObjectsPrivatized;
+};
\ No newline at end of file
diff --git a/Engine/Source/Editor/UnrealEd/Public/Dialogs/SPrivateAssetsDialog.h b/Engine/Source/Editor/UnrealEd/Public/Dialogs/SPrivateAssetsDialog.h
new file mode 100644
--- /dev/null
+++ b/Engine/Source/Editor/UnrealEd/Public/Dialogs/SPrivateAssetsDialog.h
@@ -0,0 +1,86 @@
+// Copyright Epic Games, Inc. All Rights Reserved.
+
+#pragma once
+
+#include "CoreMinimal.h"
+#include "AssetPrivatizeModel.h"
+#include "Widgets/SCompoundWidget.h"
+#include "Types/SlateStructs.h"
+#include "Widgets/Views/STableRow.h"
+#include "Widgets/Layout/SBorder.h"
+
+
+
+/**
+* The dialog that appears to help users through the process of marking assets private in the editor.
+* It helps them find references to assets being marked private and notifies them of which references would become illegal
+* and be nulled out.
+*/
+class SPrivateAssetsDialog : public SCompoundWidget
+{
+public:
+	SLATE_BEGIN_ARGS(SPrivateAssetsDialog)
+	{}
+		// The parent window hosting this dialog
+		SLATE_ATTRIBUTE(TSharedPtr<SWindow>, ParentWindow)
+
+	SLATE_END_ARGS()
+
+public:
+	virtual ~SPrivateAssetsDialog();
+
+	void Construct(const FArguments& InArgs, const TSharedRef<FAssetPrivatizeModel> InPrivatizeModel);
+
+	virtual FReply OnKeyDown(const FGeometry& MyGeometry, const FKeyEvent& InKeyEvent) override;
+
+private:
+
+	TSharedRef<SWidget> BuildProgressDialog();
+	TSharedRef<SWidget> BuildPrivatizeDialog();
+	TSharedRef<SWidget> BuildAssetViewForReferencerAssets();
+
+private:
+	/** Active timer to tick the privatize model until it reaches a "Finished" state */
+	EActiveTimerReturnType TickPrivatizeModel(double InCurrentTime, float InDeltaTime);
+
+	/** Bound to the privatize model state change event */
+	void HandlePrivatizeModelStateChanged(FAssetPrivatizeModel::EState NewState);
+
+	/** Handles filtering the asset picker to only show the references affected by the operation */
+	bool OnShouldFilterAsset(const FAssetData& InAssetData) const;
+
+	FReply Privatize();
+	FReply ForcePrivatize();
+	FReply Cancel();
+
+	FText GetWarningText() const;
+
+	/** Gets the visibility of any warning text to show at the bottom of the dialog */
+	EVisibility GetWarningTextVisibility() const;
+
+	/** Gets the visibility of if there are any on disk assets to show */
+	EVisibility GetAssetReferencesVisibility() const;
+
+	/** Gets the scanning text to display for the progress bar */
+	FText ScanningText() const;
+
+	/** Gets the scanning progress for the progress bar */
+	TOptional<float> ScanningProgressFraction() const;
+
+	/** Handles generating the columns per pending private asset */
+	TSharedRef<ITableRow> HandleGenerateAssetRow(TSharedPtr<FPendingPrivateAsset> InItem, const TSharedRef<STableViewBase>& OwnerTable);
+
+private:
+	/** Whether the active timer is currently registered */
+	bool bIsActiveTimerRegistered;
+
+	/** The model used for privatizing assets */
+	TSharedPtr<FAssetPrivatizeModel> PrivatizeModel;
+
+	/** Attributes*/
+	TAttribute<TSharedPtr<SWindow>> ParentWindow;
+
+	/** Widgets */
+	TSharedPtr<SBorder> RootContainer;
+	TSharedPtr<SListView<TSharedPtr<FPendingPrivateAsset>>> ObjectsToPrivatizeList;
+};
\ No newline at end of file
diff --git a/Engine/Source/Editor/UnrealEd/Public/ObjectTools.h b/Engine/Source/Editor/UnrealEd/Public/ObjectTools.h
--- a/Engine/Source/Editor/UnrealEd/Public/ObjectTools.h
+++ b/Engine/Source/Editor/UnrealEd/Public/ObjectTools.h
@@ -133,6 +133,12 @@ namespace ObjectTools
 		CancelNotAllowed
 	};
 
+	enum class EAllowCancelDuringPrivatize : uint8
+	{
+		AllowCancel,
+		CancelNotAllowed
+	};
+
 	/**
 	 * Handles fully loading packages for a set of passed in objects.
 	 *
@@ -296,6 +302,17 @@ namespace ObjectTools
 	 */
 	UNREALED_API int32 DeleteObjects( const TArray< UObject* >& ObjectsToDelete, bool bShowConfirmation = true, EAllowCancelDuringDelete AllowCancelDuringDelete = EAllowCancelDuringDelete::AllowCancel);
 
+	/**
+	* Privatizes the list of objects (marks their packages as NotExternallyReferencable)
+	* 
+	* @param InObjectsToPrivatize The list of objects to privatize
+	* @param bShowConfirmation True when the dialog should prompt the user that they are about to privatize something and doing so would break references
+	* @param AllowCancelDuringPrivatize Whether or not canceling is allowed when not showing the confirmation dialog
+	* 
+	* @return The number of objects successfully privatized
+	*/
+	UNREALED_API int32 PrivatizeObjects(const TArray<UObject*>& InObjectsToPrivatize, bool bShowConfirmation = true, EAllowCancelDuringPrivatize AllowCancelDuringPrivatize = EAllowCancelDuringPrivatize::AllowCancel);
+
 	/**
 	* Deletes the list of objects without checking if they are still being used.  This should not be called directly
 	* this is primarily used by the delete system after it has done the work of making sure it's safe to delete.
@@ -315,6 +332,16 @@ namespace ObjectTools
 	* @return The number of assets successfully deleted
 	*/
 	UNREALED_API int32 DeleteAssets( const TArray<FAssetData>& AssetsToDelete, bool bShowConfirmation = true );
+	
+	/**
+	* Privatizes the list of Assets (marks their packages as NotExternallyReferenceable)
+	* 
+	* @param AssetsToPrivatize The list of assets to privatize
+	* @param bShowConfirmation True when a dialog should prompt the user that they are about to privatize something and going to break references
+	* 
+	* @return The number of assets successfully privatized
+	*/
+	UNREALED_API int32 PrivatizeAssets(const TArray<FAssetData>& AssetsToPrivatize, bool bShowConfirmation = true);
 
 	/**
 	 * Delete a single object
