MPCook: Local sockets will lose not-yet-parsed data in the local process's memory for the socket if the notification that the remote connection terminated the socket is received before the data is parsed. This was causing CookWorkerServer to lose the final messages sent by the CookWorkerClient during CookComplete flushing of final messages, if the Director process was slowed by e.g. sitting at a breakpoint.
To fix this, before terminating the socket and process after a CookComplete message, the CookWorker client needs to wait for the Director to report it has read all of the CookWorker's flushed messages.
#rnx

diff --git a/Engine/Source/Editor/UnrealEd/Private/Cooker/CookDirector.cpp b/Engine/Source/Editor/UnrealEd/Private/Cooker/CookDirector.cpp
--- a/Engine/Source/Editor/UnrealEd/Private/Cooker/CookDirector.cpp
+++ b/Engine/Source/Editor/UnrealEd/Private/Cooker/CookDirector.cpp
@@ -215,27 +215,51 @@ void FCookDirector::TickFromSchedulerThread()
 void FCookDirector::PumpCookComplete(bool& bCompleted)
 {
 	TickWorkerConnects();
+
+	// MPCOOKTODO: Messages sent from the CookWorkers about discovered packages might come in after
+	// the last message sent about completed saves. These discovered packages can cause the cook to fall
+	// back to incomplete. Don't send CookComplete messages to the CookWorkers until all CookWorkers have
+	// reported they are idle.
 	for (TPair<int32, TUniquePtr<FCookWorkerServer>>& Pair : RemoteWorkers)
 	{
 		FCookWorkerServer& RemoteWorker = *Pair.Value;
-		RemoteWorker.PumpCookComplete();
-		if (RemoteWorker.IsShuttingDown())
-		{
-			TUniquePtr<FCookWorkerServer>& Existing = ShuttingDownWorkers.FindOrAdd(&RemoteWorker);
-			check(!Existing); // We should not be able to send the same pointer into ShuttingDown twice
-		}
+		RemoteWorker.SignalCookComplete();
+		check(RemoteWorker.IsShuttingDown());
+		TUniquePtr<FCookWorkerServer>& Existing = ShuttingDownWorkers.FindOrAdd(&RemoteWorker);
+		check(!Existing); // We should not be able to send the same pointer into ShuttingDown twice
 	}
 	TickWorkerShutdowns();
-	bCompleted = RemoteWorkers.Num() == 0;
+	check(RemoteWorkers.Num() == 0);
+
+	bCompleted = true;
+	for (const TPair <FCookWorkerServer*, TUniquePtr<FCookWorkerServer>>& Pair : ShuttingDownWorkers)
+	{
+		FCookWorkerServer* RemoteWorker = Pair.Key;
+		check(RemoteWorker->IsShuttingDown());
+		bCompleted = bCompleted && !RemoteWorker->IsFlushingBeforeShutdown();
+	}
 	SetWorkersStalled(!bCompleted);
 }
 
 void FCookDirector::ShutdownCookSession()
 {
+	// Cancel any inprogress workers and move them to the Shutdown list
 	while (!RemoteWorkers.IsEmpty())
 	{
 		AbortWorker(TMap<int32, TUniquePtr<FCookWorkerServer>>::TIterator(RemoteWorkers).Value()->GetWorkerId());
 	}
+
+	// Immediately shutdown any gracefully shutting down workers
+	for (TPair<FCookWorkerServer*, TUniquePtr<FCookWorkerServer>>& Pair : ShuttingDownWorkers)
+	{
+		if (Pair.Key->IsFlushingBeforeShutdown())
+		{
+			TSet<FPackageData*> UnusedPendingPackages;
+			Pair.Key->AbortWorker(UnusedPendingPackages);
+		}
+	}
+
+	// Wait for all the shutdowns to complete
 	for (;;)
 	{
 		TickWorkerShutdowns();
@@ -246,6 +270,8 @@ void FCookDirector::ShutdownCookSession()
 		constexpr double SleepSeconds = 0.010;
 		FPlatformProcess::Sleep(SleepSeconds);
 	}
+
+	// Kill any connections that had just been made and not yet assigned to a Server
 	PendingConnections.Reset();
 
 	// Restore the FCookDirector to its original state so that it is ready for a new session
@@ -591,7 +617,11 @@ void FCookDirector::AbortWorker(FWorkerId WorkerId)
 	}
 	--DesiredNumRemoteWorkers;
 	TSet<FPackageData*> PackagesToReassign;
-	RemoteWorker->AbortWorker(PackagesToReassign);
+	RemoteWorker->AbortAssignments(PackagesToReassign);
+	if (!RemoteWorker->IsShuttingDown())
+	{
+		RemoteWorker->AbortWorker(PackagesToReassign);
+	}
 	for (FPackageData* PackageData : PackagesToReassign)
 	{
 		check(PackageData->IsInProgress()); // Packages that were assigned to a worker should be in the AssignedToWorker state
diff --git a/Engine/Source/Editor/UnrealEd/Private/Cooker/CookWorkerClient.cpp b/Engine/Source/Editor/UnrealEd/Private/Cooker/CookWorkerClient.cpp
--- a/Engine/Source/Editor/UnrealEd/Private/Cooker/CookWorkerClient.cpp
+++ b/Engine/Source/Editor/UnrealEd/Private/Cooker/CookWorkerClient.cpp
@@ -32,7 +32,7 @@ FCookWorkerClient::FCookWorkerClient(UCookOnTheFlyServer& InCOTFS)
 FCookWorkerClient::~FCookWorkerClient()
 {
 	if (ConnectStatus == EConnectStatus::Connected ||
-		(EConnectStatus::WaitForDisconnectFirst <= ConnectStatus && ConnectStatus <= EConnectStatus::WaitForDisconnectLast))
+		(EConnectStatus::FlushAndAbortFirst <= ConnectStatus && ConnectStatus <= EConnectStatus::FlushAndAbortLast))
 	{
 		UE_LOG(LogCook, Warning, TEXT("CookWorkerServer %d was destroyed before it finished Disconnect. The CookDirector may be missing some information."));
 	}
@@ -76,7 +76,7 @@ void FCookWorkerClient::TickFromSchedulerThread(FTickStackData& StackData)
 bool FCookWorkerClient::IsDisconnecting() const
 {
 	return ConnectStatus == EConnectStatus::LostConnection ||
-		(EConnectStatus::WaitForDisconnectFirst <= ConnectStatus && ConnectStatus <= EConnectStatus::WaitForDisconnectLast);
+		(EConnectStatus::FlushAndAbortFirst <= ConnectStatus && ConnectStatus <= EConnectStatus::FlushAndAbortLast);
 }
 
 bool FCookWorkerClient::IsDisconnectComplete() const
@@ -412,35 +412,50 @@ void FCookWorkerClient::HandleReceiveMessages(TArray<UE::CompactBinaryTCP::FMars
 {
 	for (UE::CompactBinaryTCP::FMarshalledMessage& Message : Messages)
 	{
-		if (Message.MessageType == FAbortWorkerMessage::MessageType)
+		if (EConnectStatus::FlushAndAbortFirst <= ConnectStatus && ConnectStatus <= EConnectStatus::FlushAndAbortLast)
 		{
-			FAbortWorkerMessage AbortMessage;
-			AbortMessage.TryRead(MoveTemp(Message.Object));
-			if (AbortMessage.Type == FAbortWorkerMessage::EType::CookComplete)
+			if (Message.MessageType == FAbortWorkerMessage::MessageType)
 			{
-				UE_LOG(LogCook, Display, TEXT("CookWorkerClient received CookComplete message from Director. Flushing messages and shutting down."));
+				UE_LOG(LogCook, Display, TEXT("CookWorkerClient received AbortWorker message from Director. Terminating flush and shutting down."));
+				SendToState(EConnectStatus::LostConnection);
+				break;
 			}
-			else
-			{
-				UE_LOG(LogCook, Display, TEXT("CookWorkerClient received AbortWorker message from Director. Shutting down."));
-			}
-			SendToState(EConnectStatus::WaitForDisconnect);
-			break;
-		}
-		else if (Message.MessageType == FInitialConfigMessage::MessageType)
-		{
-			UE_LOG(LogCook, Warning, TEXT("CookWorkerClient received unexpected repeat of InitialConfigMessage. Ignoring it."));
+			UE_LOG(LogCook, Error, TEXT("CookWorkerClient received message %s from Director after receiving Abort message. Message will be ignored."),
+			*Message.MessageType.ToString());
 		}
-		else if (Message.MessageType == FAssignPackagesMessage::MessageType)
+		else
 		{
-			FAssignPackagesMessage AssignPackagesMessage;
-			if (!AssignPackagesMessage.TryRead(MoveTemp(Message.Object)))
+			if (Message.MessageType == FAbortWorkerMessage::MessageType)
 			{
-				LogInvalidMessage(TEXT("FAssignPackagesMessage"));
+				FAbortWorkerMessage AbortMessage;
+				AbortMessage.TryRead(MoveTemp(Message.Object));
+				if (AbortMessage.Type == FAbortWorkerMessage::EType::CookComplete)
+				{
+					UE_LOG(LogCook, Display, TEXT("CookWorkerClient received CookComplete message from Director. Flushing messages and shutting down."));
+					SendToState(EConnectStatus::FlushAndAbortFirst);
+				}
+				else
+				{
+					UE_LOG(LogCook, Display, TEXT("CookWorkerClient received AbortWorker message from Director. Shutting down."));
+					SendToState(EConnectStatus::LostConnection);
+					break;
+				}
 			}
-			else
+			else if (Message.MessageType == FInitialConfigMessage::MessageType)
 			{
-				AssignPackages(AssignPackagesMessage);
+				UE_LOG(LogCook, Warning, TEXT("CookWorkerClient received unexpected repeat of InitialConfigMessage. Ignoring it."));
+			}
+			else if (Message.MessageType == FAssignPackagesMessage::MessageType)
+			{
+				FAssignPackagesMessage AssignPackagesMessage;
+				if (!AssignPackagesMessage.TryRead(MoveTemp(Message.Object)))
+				{
+					LogInvalidMessage(TEXT("FAssignPackagesMessage"));
+				}
+				else
+				{
+					AssignPackages(AssignPackagesMessage);
+				}
 			}
 		}
 	}
@@ -452,20 +467,22 @@ void FCookWorkerClient::PumpDisconnect(FTickStackData& StackData)
 	{
 		switch (ConnectStatus)
 		{
-		case EConnectStatus::WaitForDisconnect:
+		case EConnectStatus::FlushAndAbortFirst:
 		{
 			TickCollectors(StackData, true /* bFlush */);
 			// Add code here for any waiting we need to do for the local CookOnTheFlyServer to gracefully shutdown
 			SendMessage(FAbortWorkerMessage(FAbortWorkerMessage::EType::Abort));
-			SendToState(EConnectStatus::WaitForDisconnectSocketFlush);
+			SendToState(EConnectStatus::WaitForAbortAcknowledge);
 			break;
 		}
-		case EConnectStatus::WaitForDisconnectSocketFlush:
+		case EConnectStatus::WaitForAbortAcknowledge:
 		{
 			using namespace UE::CompactBinaryTCP;
-			EConnectionStatus SocketStatus = TryFlushBuffer(ServerSocket, SendBuffer);
-			if (SocketStatus == EConnectionStatus::Incomplete)
+			PumpReceiveMessages();
+			if (ConnectStatus == EConnectStatus::WaitForAbortAcknowledge)
 			{
+				PumpSendMessages();
+
 				constexpr float WaitForDisconnectTimeout = 60.f;
 				if (FPlatformTime::Seconds() - ConnectStartTimeSeconds > WaitForDisconnectTimeout && !IsCookIgnoreTimeouts())
 				{
@@ -476,8 +493,11 @@ void FCookWorkerClient::PumpDisconnect(FTickStackData& StackData)
 				}
 				return; // Exit the Pump loop for now and keep waiting
 			}
-			SendToState(EConnectStatus::LostConnection);
-			break;
+			else
+			{
+				check(ConnectStatus == EConnectStatus::LostConnection);
+				return;
+			}
 		}
 		case EConnectStatus::LostConnection:
 		{
@@ -500,7 +520,7 @@ void FCookWorkerClient::SendToState(EConnectStatus TargetStatus)
 {
 	switch (TargetStatus)
 	{
-	case EConnectStatus::WaitForDisconnect:
+	case EConnectStatus::FlushAndAbortFirst:
 		ConnectStartTimeSeconds = FPlatformTime::Seconds();
 		break;
 	case EConnectStatus::LostConnection:
diff --git a/Engine/Source/Editor/UnrealEd/Private/Cooker/CookWorkerClient.h b/Engine/Source/Editor/UnrealEd/Private/Cooker/CookWorkerClient.h
--- a/Engine/Source/Editor/UnrealEd/Private/Cooker/CookWorkerClient.h
+++ b/Engine/Source/Editor/UnrealEd/Private/Cooker/CookWorkerClient.h
@@ -66,10 +66,9 @@ private:
 		PollWriteConnectMessage,
 		PollReceiveConfigMessage,
 		Connected,
-		WaitForDisconnectFirst,
-		WaitForDisconnect = WaitForDisconnectFirst,
-		WaitForDisconnectSocketFlush,
-		WaitForDisconnectLast = WaitForDisconnectSocketFlush,
+		FlushAndAbortFirst,
+		WaitForAbortAcknowledge,
+		FlushAndAbortLast = WaitForAbortAcknowledge,
 		LostConnection,
 	};
 
diff --git a/Engine/Source/Editor/UnrealEd/Private/Cooker/CookWorkerServer.cpp b/Engine/Source/Editor/UnrealEd/Private/Cooker/CookWorkerServer.cpp
--- a/Engine/Source/Editor/UnrealEd/Private/Cooker/CookWorkerServer.cpp
+++ b/Engine/Source/Editor/UnrealEd/Private/Cooker/CookWorkerServer.cpp
@@ -38,7 +38,7 @@ FCookWorkerServer::~FCookWorkerServer()
 		PackageData->SendToState(UE::Cook::EPackageState::Request, ESendFlags::QueueAddAndRemove);
 	}
 
-	if (IsConnected() || ConnectStatus == EConnectStatus::WaitForDisconnect)
+	if (ConnectStatus == EConnectStatus::Connected || ConnectStatus == EConnectStatus::PumpingCookComplete || ConnectStatus == EConnectStatus::WaitForDisconnect)
 	{
 		UE_LOG(LogCook, Error, TEXT("CookWorkerServer %d was destroyed before it finished Disconnect. The remote process may linger and may interfere with writes of future packages."),
 			WorkerId.GetRemoteIndex());
@@ -75,7 +75,7 @@ void FCookWorkerServer::AbortAssignments(TSet<FPackageData*>& OutPendingPackages
 {
 	if (PendingPackages.Num())
 	{
-		if (IsConnected())
+		if (ConnectStatus == EConnectStatus::Connected)
 		{
 			TArray<FName> PackageNames;
 			PackageNames.Reserve(PendingPackages.Num());
@@ -96,7 +96,7 @@ void FCookWorkerServer::AbortAssignment(FPackageData& PackageData)
 {
 	if (PendingPackages.Remove(&PackageData))
 	{
-		if (IsConnected())
+		if (ConnectStatus == EConnectStatus::Connected)
 		{
 			TArray<FName> PackageNames;
 			PackageNames.Add(PackageData.GetPackageName());
@@ -110,10 +110,21 @@ void FCookWorkerServer::AbortAssignment(FPackageData& PackageData)
 void FCookWorkerServer::AbortWorker(TSet<FPackageData*>& OutPendingPackages)
 {
 	AbortAssignments(OutPendingPackages);
-	if (IsConnected())
+	switch (ConnectStatus)
+	{
+	case EConnectStatus::Uninitialized: // Fall through
+	case EConnectStatus::WaitForConnect:
+		SendToState(EConnectStatus::LostConnection);
+		break;
+	case EConnectStatus::Connected: // Fall through
+	case EConnectStatus::PumpingCookComplete:
 	{
 		SendMessage(FAbortWorkerMessage(FAbortWorkerMessage::EType::Abort));
 		SendToState(EConnectStatus::WaitForDisconnect);
+		break;
+	}
+	default:
+		break;
 	}
 }
 
@@ -142,14 +153,14 @@ void FCookWorkerServer::SendToState(EConnectStatus TargetStatus)
 	ConnectStatus = TargetStatus;
 }
 
-bool FCookWorkerServer::IsConnected() const
+bool FCookWorkerServer::IsShuttingDown() const
 {
-	return EConnectStatus::ConnectedFirst <= ConnectStatus && ConnectStatus <= EConnectStatus::ConnectedLast;
+	return ConnectStatus == EConnectStatus::PumpingCookComplete || ConnectStatus == EConnectStatus::WaitForDisconnect || ConnectStatus == EConnectStatus::LostConnection;
 }
 
-bool FCookWorkerServer::IsShuttingDown() const
+bool FCookWorkerServer::IsFlushingBeforeShutdown() const
 {
-	return ConnectStatus == EConnectStatus::WaitForDisconnect || ConnectStatus == EConnectStatus::LostConnection;
+	return ConnectStatus == EConnectStatus::PumpingCookComplete;
 }
 
 bool FCookWorkerServer::IsShutdownComplete() const
@@ -182,60 +193,9 @@ bool FCookWorkerServer::TryHandleConnectMessage(FWorkerConnectMessage& Message,
 }
 
 void FCookWorkerServer::TickFromSchedulerThread()
-{
-	if (IsConnected())
-	{
-		PumpReceiveMessages();
-		if (IsConnected())
-		{
-			SendPendingPackages();
-			PumpSendMessages();
-		}
-	}
-	else
-	{
-		PumpConnect();
-		if (IsConnected())
-		{
-			// Recursively call this function to call PumpReceive and PumpSend
-			TickFromSchedulerThread();
-		}
-	}
-}
-
-void FCookWorkerServer::PumpCookComplete()
-{
-	if (ConnectStatus == EConnectStatus::Connected)
-	{
-		SendMessage(FAbortWorkerMessage(FAbortWorkerMessage::EType::CookComplete));
-		SendToState(EConnectStatus::PumpingCookComplete);
-	}
-	else if (ConnectStatus == EConnectStatus::PumpingCookComplete)
-	{
-		TickFromSchedulerThread();
-		if (IsConnected())
-		{
-			constexpr float WaitForPumpCompleteTimeout = 10.f * 60;
-			if (FPlatformTime::Seconds()  - ConnectStartTimeSeconds > WaitForPumpCompleteTimeout && !IsCookIgnoreTimeouts())
-			{
-				UE_LOG(LogCook, Error, TEXT("CookWorker process of CookWorkerServer %d failed to finalize its cook within %.0f seconds; we will tell it to shutdown."),
-					WorkerId.GetRemoteIndex(), WaitForPumpCompleteTimeout);
-				SendMessage(FAbortWorkerMessage(FAbortWorkerMessage::EType::Abort));
-				SendToState(EConnectStatus::WaitForDisconnect);
-				return;
-			}
-		}
-	}
-}
-void FCookWorkerServer::PumpConnect()
 {
 	for (;;)
 	{
-		if (IsConnected())
-		{
-			// Nothing further to do
-			return;
-		}
 		switch (ConnectStatus)
 		{
 		case EConnectStatus::Uninitialized:
@@ -248,6 +208,33 @@ void FCookWorkerServer::PumpConnect()
 				return; // Try again later
 			}
 			break;
+		case EConnectStatus::Connected:
+			PumpReceiveMessages();
+			if (ConnectStatus == EConnectStatus::Connected)
+			{
+				SendPendingPackages();
+				PumpSendMessages();
+				return; // Tick duties complete; yield the tick
+			}
+			break;
+		case EConnectStatus::PumpingCookComplete:
+		{
+			PumpReceiveMessages();
+			if (ConnectStatus == EConnectStatus::PumpingCookComplete)
+			{
+				PumpSendMessages();
+				constexpr float WaitForPumpCompleteTimeout = 10.f * 60;
+				if (FPlatformTime::Seconds() - ConnectStartTimeSeconds <= WaitForPumpCompleteTimeout || IsCookIgnoreTimeouts())
+				{
+					return; // Try again later
+				}
+				UE_LOG(LogCook, Error, TEXT("CookWorker process of CookWorkerServer %d failed to finalize its cook within %.0f seconds; we will tell it to shutdown."),
+					WorkerId.GetRemoteIndex(), WaitForPumpCompleteTimeout);
+				SendMessage(FAbortWorkerMessage(FAbortWorkerMessage::EType::Abort));
+				SendToState(EConnectStatus::WaitForDisconnect);
+			}
+			break;
+		}
 		case EConnectStatus::WaitForDisconnect:
 			TickWaitForDisconnect();
 			if (ConnectStatus == EConnectStatus::WaitForDisconnect)
@@ -264,6 +251,23 @@ void FCookWorkerServer::PumpConnect()
 	}
 }
 
+void FCookWorkerServer::SignalCookComplete()
+{
+	switch (ConnectStatus)
+	{
+	case EConnectStatus::Uninitialized: // Fall through
+	case EConnectStatus::WaitForConnect:
+		SendToState(EConnectStatus::LostConnection);
+		break;
+	case EConnectStatus::Connected:
+		SendMessage(FAbortWorkerMessage(FAbortWorkerMessage::EType::CookComplete));
+		SendToState(EConnectStatus::PumpingCookComplete);
+		break;
+	default:
+		break; // Already in a disconnecting state
+	}
+}
+
 void FCookWorkerServer::LaunchProcess()
 {
 	bool bShowCookWorkers = Director.GetShowWorkerOption() == FCookDirector::EShowWorker::SeparateWindows;
@@ -339,14 +343,6 @@ void FCookWorkerServer::TickWaitForDisconnect()
 	UE::CompactBinaryTCP::TryFlushBuffer(Socket, SendBuffer);
 	TArray<UE::CompactBinaryTCP::FMarshalledMessage> Messages;
 	TryReadPacket(Socket, ReceiveBuffer, Messages);
-	for (UE::CompactBinaryTCP::FMarshalledMessage& Message : Messages)
-	{
-		if (Message.MessageType == FAbortWorkerMessage::MessageType)
-		{
-			SendToState(EConnectStatus::LostConnection);
-			return;
-		}
-	}
 
 	if (bTerminateImmediately || (CurrentTime - ConnectStartTimeSeconds > WaitForDisconnectTimeout && !IsCookIgnoreTimeouts()))
 	{
@@ -413,7 +409,9 @@ void FCookWorkerServer::HandleReceiveMessages(TArray<UE::CompactBinaryTCP::FMars
 			UE_CLOG(ConnectStatus != EConnectStatus::PumpingCookComplete && ConnectStatus != EConnectStatus::WaitForDisconnect,
 				LogCook, Error, TEXT("CookWorkerServer %d remote process shut down unexpectedly. Assigned packages will be returned to the director."),
 				WorkerId.GetRemoteIndex());
-			SendToState(EConnectStatus::LostConnection);
+
+			SendMessage(FAbortWorkerMessage(FAbortWorkerMessage::AbortAcknowledge));
+			SendToState(EConnectStatus::WaitForDisconnect);
 			break;
 		}
 		else if (Message.MessageType == FPackageResultsMessage::MessageType)
diff --git a/Engine/Source/Editor/UnrealEd/Private/Cooker/CookWorkerServer.h b/Engine/Source/Editor/UnrealEd/Private/Cooker/CookWorkerServer.h
--- a/Engine/Source/Editor/UnrealEd/Private/Cooker/CookWorkerServer.h
+++ b/Engine/Source/Editor/UnrealEd/Private/Cooker/CookWorkerServer.h
@@ -50,12 +50,12 @@ public:
 	/** Periodic Tick function to send and receive messages to the Client. */
 	void TickFromSchedulerThread();
 	/** Called when the COTFS Server has detected all packages are complete. Tell the CookWorker to flush messages and exit. */
-	void PumpCookComplete();
+	void SignalCookComplete();
 
-	/** Is this in one of the connected states (wherein the CookWorker is working and can send arbitrary messages)? */
-	bool IsConnected() const;
 	/** Is this either shutting down or completed shutdown of its remote Client? */
 	bool IsShuttingDown() const;
+	/** Is this executing the portion of graceful shutdown where it waits for the CookWorker to transfer remaining messages? */
+	bool IsFlushingBeforeShutdown() const;
 	/** Is this not yet or no longer connected to a remote Client? */
 	bool IsShutdownComplete() const;
 
@@ -64,17 +64,13 @@ private:
 	{
 		Uninitialized,
 		WaitForConnect,
-		ConnectedFirst,
-		Connected = ConnectedFirst,
+		Connected,
 		PumpingCookComplete,
-		ConnectedLast = PumpingCookComplete,
 		WaitForDisconnect,
 		LostConnection,
 	};
 
 private:
-	/** Helper for Tick, Pump messages when connecting or disconnecting. */
-	void PumpConnect();
 	/** Helper for PumpConnect, launch the remote Client process. */
 	void LaunchProcess();
 	/** Helper for PumpConnect, wait for connect message from Client, set state to LostConnection if we timeout. */
@@ -164,6 +160,7 @@ public:
 	{
 		CookComplete,
 		Abort,
+		AbortAcknowledge,
 	};
 	FAbortWorkerMessage(EType InType = EType::Abort);
 	virtual void Write(FCbWriter& Writer) const override;
