Downgrade the timeout message in the DerivedDataCacheCommandlet from an Error to a Warning.

diff --git a/Engine/Source/Editor/UnrealEd/Private/Commandlets/DerivedDataCacheCommandlet.cpp b/Engine/Source/Editor/UnrealEd/Private/Commandlets/DerivedDataCacheCommandlet.cpp
--- a/Engine/Source/Editor/UnrealEd/Private/Commandlets/DerivedDataCacheCommandlet.cpp
+++ b/Engine/Source/Editor/UnrealEd/Private/Commandlets/DerivedDataCacheCommandlet.cpp
@@ -336,7 +336,7 @@ void UDerivedDataCacheCommandlet::FinishCachingObjects(const TArray<ITargetPlatf
 			if (CurrentTime - LastActivityTime >= DDCCommandletMaxWaitSeconds)
 			{
 				UObject* Object = CachingObjects.CreateIterator()->Key;
-				UE_LOG(LogDerivedDataCacheCommandlet, Error, TEXT("Timed out for %.2lfs waiting for %d objects to finish caching. First object: %s."),
+				UE_LOG(LogDerivedDataCacheCommandlet, Warning, TEXT("Timed out for %.2lfs waiting for %d objects to finish caching. First object: %s."),
 					DDCCommandletMaxWaitSeconds, CachingObjects.Num(), *Object->GetFullName());
 				break;
 			}
