Created BaseIteratePackagesCommandlet and GenerateModerationArtifacts commandlet.
Artifact generation commandlet supports texture2d.

diff --git a/Engine/Source/Editor/UnrealEd/Classes/Commandlets/BaseIteratePackagesCommandlet.h b/Engine/Source/Editor/UnrealEd/Classes/Commandlets/BaseIteratePackagesCommandlet.h
new file mode 100644
--- /dev/null
+++ b/Engine/Source/Editor/UnrealEd/Classes/Commandlets/BaseIteratePackagesCommandlet.h
@@ -0,0 +1,283 @@
+// Copyright Epic Games, Inc. All Rights Reserved.
+
+#pragma once
+
+#include "CoreMinimal.h"
+#include "Templates/PimplPtr.h"
+#include "UObject/ObjectMacros.h"
+#include "Commandlets/Commandlet.h"
+#include "Engine/EngineTypes.h"
+#include "BaseIteratePackagesCommandlet.generated.h"
+
+
+// Log category should be accessible by derived classes
+UNREALED_API DECLARE_LOG_CATEGORY_EXTERN(LogIteratePackagesCommandlet, Log, All);
+
+UCLASS()
+// Added UNREALED_API to expose this to the save packages test
+// Base commandlet used to iterate packages provided on a commandline
+// has callbacks for processing each package
+class UNREALED_API UBaseIteratePackagesCommandlet : public UCommandlet
+{
+    GENERATED_UCLASS_BODY()
+
+protected:
+
+	enum EBrevity
+	{
+		VERY_VERBOSE,
+		INFORMATIVE,
+		ONLY_ERRORS
+	};
+
+	EBrevity Verbosity;
+
+	/** only packages that have this version or higher will be processed; a value of IGNORE_PACKAGE_VERSION indicates that there is no minimum package version */
+	int32 MinResaveUEVersion;
+
+	/**
+	 * Limits processing to packages with this UE4 package version or lower.
+	 * A value of IGNORE_PACKAGE_VERSION (default) removes this limitation.
+	 */
+	int32 MaxResaveUEVersion;
+
+	/**
+	 * Limits processing to packages with this licensee package version or lower.
+	 * A value of IGNORE_PACKAGE_VERSION (default) removes this limitation.
+	 */
+	int32 MaxResaveLicenseeUEVersion;
+
+	/** 
+	 * Maximum number of packages to resave to avoid having a massive sync
+	 * A value of -1 (default) removes this limitation.
+	 */
+	int32 MaxPackagesToResave;
+
+	/** allows users to save only packages with a particular class in them (useful for fixing content) */
+	TArray<FString> ResaveClasses;
+
+	// If non-empty, this substring has to be present in the package name for the commandlet to process it
+	FString PackageSubstring;
+
+	// strip editor only content
+	bool bStripEditorOnlyContent;
+
+	// skip the assert when a package can not be opened
+	bool bCanIgnoreFails;
+
+	/** load all packages, and display warnings for those packages which would have been resaved but were read-only */
+	bool bVerifyContent;
+
+	/** if we should only save dirty packages **/
+	bool bOnlySaveDirtyPackages;
+
+	/** if we should auto checkout packages that need to be saved**/
+	bool bAutoCheckOut;
+
+	/** if we should batch together source control operations rather than submit them one at a time */
+	bool bBatchSourceControl;
+
+	/** if we should simply skip checked out files rather than error-ing out */
+	bool bSkipCheckedOutFiles;
+
+	/** if we should auto checkin packages that were checked out**/
+	bool bAutoCheckIn;
+
+	/** Should we build lighting for the packages we are saving? **/
+	//bool bShouldBuildLighting;
+
+	/** Should we build reflection captures for the packages we are saving? **/
+	//bool bShouldBuildReflectionCaptures;
+
+	/** Should we build texture streaming for the packages we are saving? **/
+	//bool bShouldBuildTextureStreaming;
+
+	/** Similar to above, but applies to all packages rather than just maps **/
+	//bool bShouldBuildTextureStreamingForAll;
+
+	/** only process packages containing materials **/
+	bool bOnlyMaterials;
+
+	/** Only save packages with a changelist of zero **/
+	bool bOnlyUnversioned;
+
+	/** Only save packages that been saved by a licensee **/
+	bool bOnlyLicenseed;
+
+	/** Only save packages containing bulkdata payloads that have been virtualized **/
+	bool bOnlyVirtualized;
+
+	/** Only save packages containing FPayloadTrailers **/
+	bool bOnlyPayloadTrailers;
+
+	/** Should we build navigation data for the packages we are saving? **/
+	//bool bShouldBuildNavigationData;
+
+	/** Ignore package version changelist **/
+	bool bIgnoreChangelist;
+
+	/** Filter packages based on a collection **/
+	TSet<FName> CollectionFilter;
+
+	/** Should we update HLODs */
+	//bool bShouldBuildHLOD;
+	//bool bGenerateClusters;
+	//bool bGenerateMeshProxies;
+	//bool bForceClusterGeneration;
+	//bool bForceProxyGeneration;
+	//bool bForceSingleClusterForLevel;
+	bool bSkipSubLevels;
+	//bool bHLODMapCleanup;
+	//FString ForceHLODSetupAsset;
+	//FString HLODSkipToMap;
+	bool bForceUATEnvironmentVariableSet;
+	bool bResaveOnDemand = false;
+
+	/** Running count of packages that got modified and will need to be resaved */
+	int32 PackagesConsideredForResave;
+	int32 PackagesResaved;
+	int32 PackagesDeleted;
+	int32 TotalPackagesForResave;
+
+	/** Only collect garbage after N packages */
+	int32 GarbageCollectionFrequency;
+
+ 	/** Lighting Build Quality(default: Production) */
+ 	//ELightingBuildQuality LightingBuildQuality;
+ 
+	/** List of files to submit */
+	TArray<FString> FilesToSubmit;
+
+	/** The list of switches that were passed on the commandline */
+	TArray<FString> Switches;
+
+	/** List of redirector packages that should be fixed up at the end */
+	//TArray<FString> RedirectorsToFixup;
+
+	/** A queue containing source control operations to be performed in batches. */
+	TPimplPtr<class FQueuedSourceControlOperations> SourceControlQueue;
+
+	/** If running resaveondemand, only packages reported into this set are resaved. */
+	TSet<FName> ResaveOnDemandPackages;
+	TSet<FName> ResaveOnDemandSystems;
+	FCriticalSection ResaveOnDemandPackagesLock;
+
+	/**
+	 * Evaluates the command-line to determine which maps to check.  By default all maps are checked
+	 * Provides child classes with a chance to initialize any variables, parse the command line, etc.
+	 *
+	 * @param	Tokens			the list of tokens that were passed to the commandlet
+	 * @param	MapPathNames	receives the list of path names for the maps that will be checked.
+	 *
+	 * @return	0 to indicate that the commandlet should continue; otherwise, the error code that should be returned by Main()
+	 */
+	virtual int32 InitializeParameters( const TArray<FString>& Tokens, TArray<FString>& MapPathNames );
+
+	virtual void InitializePackageNames(const TArray<FString>& Tokens, TArray<FString>& MapPathNames, bool& bExplicitPackages) { }
+
+	void ParseSourceControlOptions(const TArray<FString>& Tokens);
+
+	void OnAddResaveOnDemandPackage(FName SystemName, FName PackageName);
+
+	/** Loads and saves a single package */
+	virtual void LoadAndSaveOnePackage(const FString& Filename);
+
+	/** Checks to see if a package should be skipped */
+	virtual bool ShouldSkipPackage(const FString& Filename);
+
+	/** Deletes a single package */
+	virtual void DeleteOnePackage(const FString& Filename);
+
+	/**
+	 * Allow the commandlet to perform any operations on the export/import table of the package before all objects in the package are loaded.
+	 *
+	 * @param	PackageLinker	the linker for the package about to be loaded
+	 * @param	bSavePackage	[in]	indicates whether the package is currently going to be saved
+	 *							[out]	set to true to resave the package
+	 */
+	virtual void PerformPreloadOperations( FLinkerLoad* PackageLinker, bool& bSavePackage );
+
+	/**
+	 * Allows the commandlet to perform any additional operations on the object before it is resaved.
+	 *
+	 * @param	Object			the object in the current package that is currently being processed
+	 * @param	bSavePackage	[in]	indicates whether the package is currently going to be saved
+	 *							[out]	set to true to resave the package
+	 */
+	virtual void PerformAdditionalOperations( class UObject* Object, bool& bSavePackage ) { }
+
+	/**
+	 * Allows the commandlet to perform any additional operations on the package before it is resaved.
+	 *
+	 * @param	Package			the package that is currently being processed
+	 * @param	bSavePackage	[in]	indicates whether the package is currently going to be saved
+	 *							[out]	set to true to resave the package
+	 */
+	virtual void PerformAdditionalOperations( class UPackage* Package, bool& bSavePackage ) { }
+
+	/**
+	* Allows the commandlet to perform any additional operations on the world before it is resaved.
+	*
+	* @param	World			the world that is currently being processed
+	* @param	bSavePackage	[in]	indicates whether the package is currently going to be saved
+	*							[out]	set to true to resave the package
+	*/
+	virtual void PerformAdditionalOperations(class UWorld* World, bool& bSavePackage) { }
+
+	/**
+	 * Allows the commandlet to do something after all the packages have been processed
+	 */
+	virtual void PostSavePackages() { }
+
+	/**
+	 * Removes any UClass exports from packages which aren't script packages.
+	 *
+	 * @param	Package			the package that is currently being processed
+	 *
+	 * @return	true to resave the package
+	 */
+	bool CleanClassesFromContentPackages( class UPackage* Package );
+
+	// Get the changelist description to use if automatically checking packages out
+	virtual FText GetChangelistDescription() const;
+
+	bool CheckoutFile(const FString& Filename, bool bAddFile = false, bool bIgnoreAlreadyCheckedOut = false);
+	bool RevertFile(const FString& Filename);
+	bool CanCheckoutFile(const FString& Filename, FString& CheckedOutUser);
+	void CheckoutAndSavePackage(UPackage* Package, TArray<FString>& SublevelFilenames, bool bIgnoreAlreadyCheckedOut = false);
+	void CheckInFiles(const TArray<FString>& InFilesToSubmit, const FText& InDescription) const;
+
+	/**
+	 * Creates and returns a unique filename in the temporary file directory.
+	 *
+	 * @return	A unique temporary filename
+	 */
+	FString CreateTempFilename();
+
+	/**
+	 * Returns the path of the directory we use to store all temporary files for
+	 * this commandlet.
+	 *
+	 * @return Path to the temporary file directory
+	 */
+	FString GetTempFilesDirectory();
+
+	/**
+	 * Delete the remaining files in the temporary files directory for the commandlet.
+	 */
+	void CleanTempFiles();	
+
+	// Print out a message only if running in very verbose mode
+	void VerboseMessage(const FString& Message);
+
+	/** Parse commandline to decide whether resaveondemand is activated for the given system. */
+	TSet<FName> ParseResaveOnDemandSystems();
+
+private:
+
+
+public:		
+	//~ Begin UCommandlet Interface
+	virtual int32 Main(const FString& Params) override;
+	//~ End UCommandlet Interface
+};
diff --git a/Engine/Source/Editor/UnrealEd/Classes/Commandlets/GenerateModerationArtifactsCommandlet.h b/Engine/Source/Editor/UnrealEd/Classes/Commandlets/GenerateModerationArtifactsCommandlet.h
new file mode 100644
--- /dev/null
+++ b/Engine/Source/Editor/UnrealEd/Classes/Commandlets/GenerateModerationArtifactsCommandlet.h
@@ -0,0 +1,79 @@
+// Copyright Epic Games, Inc. All Rights Reserved.
+
+#pragma once
+
+#include "BaseIteratePackagesCommandlet.h"
+#include "GenerateModerationArtifactsCommandlet.generated.h"
+
+UNREALED_API DECLARE_LOG_CATEGORY_EXTERN(LogModerationArtifactsCommandlet, Log, All);
+
+UCLASS()
+// Added UNREALED_API to expose this to the save packages test
+// Base commandlet used to iterate packages provided on a commandline
+// has callbacks for processing each package
+class UNREALED_API UGenerateModerationArtifactsCommandlet : public UBaseIteratePackagesCommandlet
+{
+    GENERATED_UCLASS_BODY()
+
+	virtual int32 InitializeParameters( const TArray<FString>& Tokens, TArray<FString>& MapPathNames ) override;
+
+	// virtual void InitializePackageNames(const TArray<FString>& Tokens, TArray<FString>& MapPathNames, bool& bExplicitPackages) { }
+
+	/** Loads and saves a single package */
+	//virtual void LoadAndSaveOnePackage(const FString& Filename);
+
+	/** Checks to see if a package should be skipped */
+	// virtual bool ShouldSkipPackage(const FString& Filename);
+
+	/** Deletes a single package */
+	// virtual void DeleteOnePackage(const FString& Filename);
+
+	/**
+	 * Allow the commandlet to perform any operations on the export/import table of the package before all objects in the package are loaded.
+	 *
+	 * @param	PackageLinker	the linker for the package about to be loaded
+	 * @param	bSavePackage	[in]	indicates whether the package is currently going to be saved
+	 *							[out]	set to true to resave the package
+	 */
+	// virtual void PerformPreloadOperations( FLinkerLoad* PackageLinker, bool& bSavePackage );
+
+	/**
+	 * Allows the commandlet to perform any additional operations on the object before it is resaved.
+	 *
+	 * @param	Object			the object in the current package that is currently being processed
+	 * @param	bSavePackage	[in]	indicates whether the package is currently going to be saved
+	 *							[out]	set to true to resave the package
+	 */
+	virtual void PerformAdditionalOperations( class UObject* Object, bool& bSavePackage ) override;
+
+	/**
+	 * Allows the commandlet to perform any additional operations on the package before it is resaved.
+	 *
+	 * @param	Package			the package that is currently being processed
+	 * @param	bSavePackage	[in]	indicates whether the package is currently going to be saved
+	 *							[out]	set to true to resave the package
+	 */
+	virtual void PerformAdditionalOperations( class UPackage* Package, bool& bSavePackage ) override;
+
+	/**
+	* Allows the commandlet to perform any additional operations on the world before it is resaved.
+	*
+	* @param	World			the world that is currently being processed
+	* @param	bSavePackage	[in]	indicates whether the package is currently going to be saved
+	*							[out]	set to true to resave the package
+	*/
+	// virtual void PerformAdditionalOperations(class UWorld* World, bool& bSavePackage) { }
+
+private:
+	void GatherLocalizationFromPackage(class UPackage* Package);
+	void GatherFStringsFromObject(class UObject* Object);
+	
+	// void GenerateArtifact(class UDataTable* DataTable)
+	void GenerateArtifact(class UTexture* Texture);
+	void GenerateArtifact(class UStaticMeshComponent* StaticMesh);
+
+	FString CreateOutputFileName(UObject* Object, const FString& Extension);
+
+	FString OutputPath;
+
+};
diff --git a/Engine/Source/Editor/UnrealEd/Private/Commandlets/BaseIteratePackagesCommandlet.cpp b/Engine/Source/Editor/UnrealEd/Private/Commandlets/BaseIteratePackagesCommandlet.cpp
new file mode 100644
--- /dev/null
+++ b/Engine/Source/Editor/UnrealEd/Private/Commandlets/BaseIteratePackagesCommandlet.cpp
@@ -0,0 +1,1505 @@
+// Copyright Epic Games, Inc. All Rights Reserved.
+
+/*=============================================================================
+	BaseIteratePackageCommandlet.cpp: Commandlet provides basic package iteration functionaliy for derived commandlets
+=============================================================================*/
+#include "Commandlets/BaseIteratePackagesCommandlet.h"
+
+#include "CoreMinimal.h"
+#include "HAL/FileManager.h"
+#include "Misc/CommandLine.h"
+#include "Misc/FileHelper.h"
+#include "Misc/Paths.h"
+#include "Misc/ConfigCacheIni.h"
+#include "Modules/ModuleManager.h"
+#include "UObject/Class.h"
+#include "UObject/MetaData.h"
+#include "UObject/Object.h"
+#include "UObject/ObjectMacros.h"
+#include "UObject/Package.h"
+#include "UObject/PackageResourceManager.h"
+#include "UObject/UObjectHash.h"
+#include "UObject/UObjectIterator.h"
+#include "Misc/PackageName.h"
+#include "Misc/EngineVersion.h"
+#include "Misc/RedirectCollector.h"
+#include "Engine/EngineTypes.h"
+#include "Materials/Material.h"
+#include "ISourceControlOperation.h"
+#include "SourceControlOperations.h"
+#include "SourceControlHelpers.h"
+#include "ISourceControlModule.h"
+#include "Engine/MapBuildDataRegistry.h"
+//#include "Particles/ParticleSystem.h"
+#include "EngineGlobals.h"
+//#include "Particles/ParticleEmitter.h"
+#include "GameFramework/WorldSettings.h"
+//#include "Engine/StaticMesh.h"
+#include "AssetRegistry/AssetData.h"
+//#include "Engine/Brush.h"
+#include "Editor.h"
+#include "EditorWorldUtils.h"
+#include "FileHelpers.h"
+#include "PlatformInfo.h"
+#include "CollectionManagerModule.h"
+#include "ICollectionManager.h"
+#include "CommandletSourceControlUtils.h"
+//#include "WorldPartition/WorldPartition.h"
+//#include "WorldPartition/WorldPartitionHelpers.h"
+//#include "LevelInstance/LevelInstanceInterface.h"
+#include "AssetCompilingManager.h"
+#include "PackageHelperFunctions.h"
+#include "PackageTools.h"
+#include "StaticMeshCompiler.h"
+#include "String/ParseTokens.h"
+#include "UObject/PackageTrailer.h"
+#include "UObject/SavePackage.h"
+
+DEFINE_LOG_CATEGORY(LogIteratePackagesCommandlet);
+
+#include "AssetRegistry/AssetRegistryModule.h"
+#include "IDirectoryWatcher.h"
+#include "DirectoryWatcherModule.h"
+//#include "Particles/Material/ParticleModuleMeshMaterial.h"
+//#include "Particles/ParticleLODLevel.h"
+//#include "Particles/ParticleModuleRequired.h"
+//#include "Particles/TypeData/ParticleModuleTypeDataMesh.h"
+#include "Engine/LevelStreaming.h"
+#include "EditorBuildUtils.h"
+
+// for UBaseIteratePackagesCommandlet::PerformAdditionalOperations building lighting code
+//#include "LightingBuildOptions.h"
+
+// for UBaseIteratePackagesCommandlet::PerformAdditionalOperations building navigation data
+//#include "EngineUtils.h"
+//#include "NavigationData.h"
+//#include "AI/NavigationSystemBase.h"
+
+// For preloading FFindInBlueprintSearchManager
+#include "FindInBlueprintManager.h"
+
+#include "ShaderCompiler.h"
+//#include "IHierarchicalLODUtilities.h"
+//#include "HierarchicalLODUtilitiesModule.h"
+//#include "HierarchicalLOD.h"
+//#include "HierarchicalLODProxyProcessor.h"
+//#include "HLOD/HLODEngineSubsystem.h"
+//#include "GenericPlatform/GenericPlatformProcess.h"
+//#include "HAL/ThreadManager.h"
+
+//#include "ICollectionManager.h"
+//#include "CollectionManagerModule.h"
+//#include "UObject/UObjectThreadContext.h"
+//#include "Engine/LODActor.h"
+//#include "PerQualityLevelProperties.h"
+//#include "Misc/RedirectCollector.h"
+
+/**-----------------------------------------------------------------------------
+ *	UBaseIteratePackagesCommandlet commandlet.
+ *
+ * This commandlet exposes some functionality for iterating packages 
+ *
+ *
+----------------------------------------------------------------------------**/
+
+#define CURRENT_PACKAGE_VERSION 0
+#define IGNORE_PACKAGE_VERSION INDEX_NONE
+
+UBaseIteratePackagesCommandlet::UBaseIteratePackagesCommandlet(const FObjectInitializer& ObjectInitializer)
+	: Super(ObjectInitializer)
+	, bForceUATEnvironmentVariableSet(false)
+{
+
+}
+
+int32 UBaseIteratePackagesCommandlet::InitializeParameters( const TArray<FString>& Tokens, TArray<FString>& PackageNames )
+{
+	TRACE_CPUPROFILER_EVENT_SCOPE(UBaseIteratePackagesCommandlet::InitializeResaveParameters);
+
+	Verbosity = VERY_VERBOSE;
+
+	TArray<FString> Unused;
+	bool bExplicitPackages = false;
+
+	// Check to see if we have an explicit list of packages
+	for( int32 SwitchIdx = 0; SwitchIdx < Switches.Num(); SwitchIdx++ )
+	{
+		FString Package;
+		FString PackageFolder;
+		FString Maps;
+		FString File;
+		const FString& CurrentSwitch = Switches[ SwitchIdx ];
+		if( FParse::Value( *CurrentSwitch, TEXT( "PACKAGE="), Package ) )
+		{
+			FString PackageFile;
+			if (FPackageName::SearchForPackageOnDisk(Package, NULL, &PackageFile))
+			{
+				PackageNames.Add( *PackageFile );
+			}
+			else
+			{
+				UE_LOG(LogIteratePackagesCommandlet, Error, TEXT("Failed to find the package given by the cmdline: PACKAGE=%s"), *Package);
+			}
+			
+			bExplicitPackages = true;
+		}
+		else if( FParse::Value( *CurrentSwitch, TEXT( "PACKAGEFOLDER="), PackageFolder ) )
+		{
+			TArray<FString> FilesInPackageFolder;
+			FPackageName::FindPackagesInDirectory(FilesInPackageFolder, PackageFolder);
+			for( int32 FileIndex = 0; FileIndex < FilesInPackageFolder.Num(); FileIndex++ )
+			{
+				FString PackageFile(FilesInPackageFolder[FileIndex]);
+				FPaths::MakeStandardFilename(PackageFile);
+				PackageNames.Add( *PackageFile );
+			}
+			bExplicitPackages = true;
+		}
+		else if (FParse::Value(*CurrentSwitch, TEXT("MAP="), Maps))
+		{
+			// Allow support for -MAP=Value1+Value2+Value3
+			for (int32 PlusIdx = Maps.Find(TEXT("+"), ESearchCase::CaseSensitive); PlusIdx != INDEX_NONE; PlusIdx = Maps.Find(TEXT("+"), ESearchCase::CaseSensitive))
+			{
+				const FString NextMap = Maps.Left(PlusIdx);
+				if (NextMap.Len() > 0)
+				{
+					FString MapFile;
+					FPackageName::SearchForPackageOnDisk(NextMap, NULL, &MapFile);
+					PackageNames.Add(*MapFile);
+					bExplicitPackages = true;
+				}
+
+				Maps.RightInline(Maps.Len() - (PlusIdx + 1), false);
+			}
+			FString MapFile;
+			FPackageName::SearchForPackageOnDisk(Maps, NULL, &MapFile);
+			PackageNames.Add(*MapFile);
+			bExplicitPackages = true;
+		}
+		else if (FParse::Value(*CurrentSwitch, TEXT("FILE="), File))
+		{
+			FString Text;
+			if (FFileHelper::LoadFileToString(Text, *File))
+			{
+				TArray<FString> Lines;
+
+				// Remove all carriage return characters.
+				Text.ReplaceInline(TEXT("\r"), TEXT(""), ESearchCase::CaseSensitive);
+				// Read all lines
+				Text.ParseIntoArray(Lines, TEXT("\n"), true);
+
+				for (auto Line : Lines)
+				{
+					FString PackageFile;
+					if (FPackageName::SearchForPackageOnDisk(Line, NULL, &PackageFile))
+					{
+						PackageNames.AddUnique(*PackageFile);
+					}
+					else
+					{
+						UE_LOG(LogIteratePackagesCommandlet, Error, TEXT("Failed to find package %s"), *Line);
+					}
+				}
+
+				bExplicitPackages = true;
+				UE_LOG(LogIteratePackagesCommandlet, Display, TEXT("Loaded %d Packages from %s"), PackageNames.Num(), *File);
+			}
+			else
+			{
+				UE_LOG(LogIteratePackagesCommandlet, Error, TEXT("Failed to load file %s"), *File);
+			}
+		}
+	}
+
+	// Check for numeric settings
+	for (const FString& CurrentSwitch : Switches)
+	{
+		if (FParse::Value(*CurrentSwitch, TEXT("GCFREQ="), GarbageCollectionFrequency))
+		{
+			UE_LOG(LogIteratePackagesCommandlet, Display, TEXT("Setting garbage collection to happen every %d packages."), GarbageCollectionFrequency);
+		}
+	}
+
+	InitializePackageNames(Tokens, PackageNames, bExplicitPackages);
+
+	// ... if not, load in all packages
+	if( !bExplicitPackages )
+	{
+		UE_LOG( LogIteratePackagesCommandlet, Display, TEXT( "No maps found to save, checking Project Settings for Directory or Asset Path(s)" ) );
+
+		uint8 PackageFilter = NORMALIZE_DefaultFlags;
+		if ( Switches.Contains(TEXT("SKIPMAPS")) )
+		{
+			PackageFilter |= NORMALIZE_ExcludeMapPackages;
+		}
+		else if ( Switches.Contains(TEXT("MAPSONLY")) )
+		{
+			PackageFilter |= NORMALIZE_ExcludeContentPackages;
+		}
+
+		if (Switches.Contains(TEXT("PROJECTONLY")))
+		{
+			PackageFilter |= NORMALIZE_ExcludeEnginePackages;
+		}
+
+		if ( Switches.Contains(TEXT("SkipDeveloperFolders")) || Switches.Contains(TEXT("NODEV")) )
+		{
+			PackageFilter |= NORMALIZE_ExcludeDeveloperPackages;
+		}
+		else if ( Switches.Contains(TEXT("OnlyDeveloperFolders")) )
+		{
+			PackageFilter |= NORMALIZE_ExcludeNonDeveloperPackages;
+		}
+
+		const FString AssetSearch = TEXT("*") + FPackageName::GetAssetPackageExtension();
+		const FString MapSearch = TEXT("*") + FPackageName::GetMapPackageExtension();
+
+		bool bAnyFound = NormalizePackageNames(Unused, PackageNames, AssetSearch, PackageFilter);
+		bAnyFound |= NormalizePackageNames(Unused, PackageNames, MapSearch, PackageFilter);
+		
+		if (!bAnyFound)
+		{
+			return 1;
+		}
+	}
+
+	// Check for a max package limit
+	MaxPackagesToResave = -1;
+	for ( int32 SwitchIdx = 0; SwitchIdx < Switches.Num(); SwitchIdx++ )
+	{
+		const FString& CurrentSwitch = Switches[SwitchIdx];
+		if( FParse::Value(*CurrentSwitch,TEXT("MAXPACKAGESTORESAVE="),MaxPackagesToResave))
+		{
+			UE_LOG(LogIteratePackagesCommandlet, Warning, TEXT( "Only resaving a maximum of %d packages." ), MaxPackagesToResave );
+			break;
+		}
+	}
+
+	FAssetRegistryModule& AssetRegistryModule = FModuleManager::LoadModuleChecked<FAssetRegistryModule>(TEXT("AssetRegistry"));
+	IAssetRegistry& AssetRegistry = AssetRegistryModule.Get();
+
+	
+
+	// ResaveOnDemand is a generic way for systems to specify during load that the commandlet should resave a package
+	ResaveOnDemandSystems = ParseResaveOnDemandSystems();
+	if (!ResaveOnDemandSystems.IsEmpty())
+	{
+		TStringBuilder<64> SystemNames;
+		for (FName SystemName : ResaveOnDemandSystems)
+		{
+			SystemNames << SystemName << TEXT(",");
+		}
+		SystemNames.RemoveSuffix(1); // Remove the trailing ,
+		UE_LOG(LogIteratePackagesCommandlet, Display, TEXT("ResaveOnDemand=%s. Only saving packages that are reported via UE::SavePackageUtilities::OnAddResaveOnDemandPackage."),
+			*SystemNames);
+		bResaveOnDemand = true;
+		UE::SavePackageUtilities::OnAddResaveOnDemandPackage.BindUObject(this, &UBaseIteratePackagesCommandlet::OnAddResaveOnDemandPackage);
+	}
+
+	// This option allows the dependency graph and soft object path redirect map to be populated. This is useful if you want soft object references to redirectors to be followed to the destination asset at save time.
+	const bool bSearchAllAssets = Switches.Contains(TEXT("SearchAllAssets"));
+
+	// Check for the min and max versions
+	MinResaveUEVersion = IGNORE_PACKAGE_VERSION;
+	MaxResaveUEVersion = IGNORE_PACKAGE_VERSION;
+	MaxResaveLicenseeUEVersion = IGNORE_PACKAGE_VERSION;
+	if ( Switches.Contains(TEXT("CHECKLICENSEEVER")) )
+	{
+		// Limits resaving to packages with this licensee package version or lower.
+		MaxResaveLicenseeUEVersion = FMath::Max<int32>(GPackageFileLicenseeUEVersion - 1, 0);
+	}
+	if ( Switches.Contains(TEXT("CHECKUE4VER")) )
+	{
+		// Limits resaving to packages with this package version or lower.
+		MaxResaveUEVersion = FMath::Max<int32>(GPackageFileUEVersion.ToValue() - 1, 0);
+	}
+	else if ( Switches.Contains(TEXT("RESAVEDEPRECATED")) )
+	{
+		// Limits resaving to packages with this package version or lower.
+		MaxResaveUEVersion = FMath::Max<int32>(VER_UE4_DEPRECATED_PACKAGE - 1, 0);
+	}
+	else
+	{
+		// determine if the resave operation should be constrained to certain package versions
+		for ( int32 SwitchIdx = 0; SwitchIdx < Switches.Num(); SwitchIdx++ )
+		{
+			const FString& CurrentSwitch = Switches[SwitchIdx];
+			if ( MinResaveUEVersion == IGNORE_PACKAGE_VERSION && FParse::Value(*CurrentSwitch,TEXT("MINVER="),MinResaveUEVersion) )
+			{
+				if ( MinResaveUEVersion == CURRENT_PACKAGE_VERSION )
+				{
+					MinResaveUEVersion = GPackageFileUEVersion.ToValue();
+				}
+			}
+
+			if ( MaxResaveUEVersion == IGNORE_PACKAGE_VERSION && FParse::Value(*CurrentSwitch,TEXT("MAXVER="),MaxResaveUEVersion) )
+			{
+				if ( MaxResaveUEVersion == CURRENT_PACKAGE_VERSION )
+				{
+					MaxResaveUEVersion = GPackageFileUEVersion.ToValue();
+				}
+			}
+		}
+	}
+
+	FString ClassList;
+	for (const FString& CurrentSwitch : Switches)
+	{
+		if ( FParse::Value(*CurrentSwitch, TEXT("RESAVECLASS="), ClassList, false) )
+		{
+			TArray<FString> ClassNames;
+			ClassList.ParseIntoArray(ClassNames, TEXT(","), true);
+			for (FString& ClassName : ClassNames)
+			{
+				if (FPackageName::IsShortPackageName(ClassName))
+				{
+					UE_LOG(LogIteratePackagesCommandlet, Warning, TEXT("RESAVECLASS param requires class path names. Short name provided: %s."), *ClassName);
+					FTopLevelAssetPath ClassPathName = UClass::TryConvertShortTypeNameToPathName<UStruct>(ClassName, ELogVerbosity::Warning, TEXT("ResavePackages"));
+					if (!ClassPathName.IsNull())
+					{
+						ClassName = ClassPathName.ToString();
+					}
+				}
+				ResaveClasses.AddUnique(ClassName);
+			}
+
+			break;
+		}
+	}
+
+	/** determine if we should check subclasses of ResaveClasses */
+	bool bIncludeChildClasses = Switches.Contains(TEXT("IncludeChildClasses"));
+	if (bIncludeChildClasses && ResaveClasses.Num() == 0)
+	{
+		// Sanity check fail
+		UE_LOG(LogIteratePackagesCommandlet, Error, TEXT("AllowSubclasses param requires ResaveClass param."));
+		return 1;
+	}
+
+	if (bIncludeChildClasses)
+	{
+		// Can't use ranged for here because the array grows inside of this loop.
+		// Also, no need to iterate over the newly added objects as we know
+		// we have found all of their subclasses too (IsChildOf guarantees that).
+		const int32 NumResaveClasses = ResaveClasses.Num();
+		for (int32 ClassIndex = 0; ClassIndex < NumResaveClasses; ++ClassIndex)
+		{
+			// Find the class object and then all derived classes
+			UClass* ResaveClass = UClass::TryFindTypeSlow<UClass>(ResaveClasses[ClassIndex]);
+			if (ResaveClass)
+			{
+				for (TObjectIterator<UClass> It; It; ++It)
+				{
+					UClass* MaybeChildClass = *It;
+					if (MaybeChildClass->IsChildOf(ResaveClass))
+					{
+						ResaveClasses.AddUnique(MaybeChildClass->GetPathName());
+					}
+				}
+			}
+		}
+	}
+
+	return 0;
+}
+
+void UBaseIteratePackagesCommandlet::ParseSourceControlOptions(const TArray<FString>& Tokens)
+{
+	if (SourceControlQueue == nullptr)
+	{
+		return; // No point parsing the options if we don't have anything enabled
+	}
+
+	int32 QueuedPackageFlushLimit = INDEX_NONE;
+	int64 QueueFileSizeFlushLimit = INDEX_NONE;
+
+	for (const FString& CurrentSwitch : Tokens)
+	{
+		if (FParse::Value(*CurrentSwitch, TEXT("BatchPackageLimit="), QueuedPackageFlushLimit))
+		{
+			if (QueuedPackageFlushLimit >= 0)
+			{
+				UE_LOG(LogIteratePackagesCommandlet, Display, TEXT("Setting source control batches to be limited to %d package(s) at a time."), QueuedPackageFlushLimit);
+				SourceControlQueue->SetMaxNumQueuedPackages(QueuedPackageFlushLimit);
+			}
+			else
+			{
+				// Negative values mean we will not flush the source control batch based on the number of packages
+				UE_LOG(LogIteratePackagesCommandlet, Display, TEXT("Setting source control batches to have no package limit!"));
+			}
+		}
+		else  if (FParse::Value(*CurrentSwitch, TEXT("BatchFileSizeLimit="), QueueFileSizeFlushLimit))
+		{
+			const int64 SizeLimit = TNumericLimits<int64>::Max() / (1024 * 1024);
+			if (QueueFileSizeFlushLimit > SizeLimit)
+			{
+				UE_LOG(LogIteratePackagesCommandlet, Display, TEXT("-BatchFileSizeLimit=%lld is too large! The max value allowed is %lld, clamping..."), QueueFileSizeFlushLimit, SizeLimit);
+				QueueFileSizeFlushLimit = SizeLimit;
+			}
+
+			if (QueueFileSizeFlushLimit >= 0)
+			{
+				UE_LOG(LogIteratePackagesCommandlet, Display, TEXT("Setting source control batches to be limited to %lld MB."), QueueFileSizeFlushLimit);
+
+				SourceControlQueue->SetMaxTemporaryFileTotalSize(QueueFileSizeFlushLimit);
+			}
+			else
+			{
+				// Negative values mean we will not flush the source control batch based on the disk space taken by temp files
+				UE_LOG(LogIteratePackagesCommandlet, Display, TEXT("Setting source control batches to have no disk space limit!"));
+			}
+		}
+	}
+}
+
+void UBaseIteratePackagesCommandlet::OnAddResaveOnDemandPackage(FName SystemName, FName PackageName)
+{
+	FScopeLock ScopeLock(&ResaveOnDemandPackagesLock);
+	if (ResaveOnDemandSystems.Contains(SystemName))
+	{
+		ResaveOnDemandPackages.Add(PackageName);
+	}
+}
+
+
+TSet<FName> UBaseIteratePackagesCommandlet::ParseResaveOnDemandSystems()
+{
+	using namespace UE::String;
+	TSet<FName> SystemNames;
+
+	const TCHAR* CommandLineStream = FCommandLine::Get();
+	for (;;)
+	{
+		FString Token = FParse::Token(CommandLineStream, false /* bUseEscape */);
+		if (Token.IsEmpty())
+		{
+			break;
+		}
+		FString TokenSystemNames;
+		if (FParse::Value(*Token, TEXT("-resaveondemand="), TokenSystemNames))
+		{
+			UE::String::ParseTokensMultiple(TokenSystemNames, TConstArrayView<TCHAR>({ '+', ',' }),
+				[&SystemNames](FStringView TokenSystemName)
+				{
+					SystemNames.Add(FName(TokenSystemName));
+				}, EParseTokensOptions::Trim | EParseTokensOptions::SkipEmpty);
+		}
+	}
+	return SystemNames;
+}
+
+
+bool UBaseIteratePackagesCommandlet::ShouldSkipPackage(const FString& Filename)
+{
+	return false;
+}
+
+void UBaseIteratePackagesCommandlet::LoadAndSaveOnePackage(const FString& Filename)
+{
+	TRACE_CPUPROFILER_EVENT_SCOPE(UBaseIteratePackagesCommandlet::LoadAndSaveOnePackage);
+
+	// Check to see if a derived commandlet wants to skip this package for one reason or another
+	if (ShouldSkipPackage(Filename))
+	{
+		return;
+	}
+
+	// Skip the package if it doesn't have a required substring match
+	if (PackageSubstring.Len() && !Filename.Contains(PackageSubstring) )
+	{
+		VerboseMessage(FString::Printf(TEXT("Skipping %s"), *Filename));
+		return;
+	}
+
+	FPackagePath PackagePath = FPackagePath::FromLocalPath(Filename);
+	if (CollectionFilter.Num())
+	{
+		FString PackageNameToCreate = PackagePath.GetPackageNameOrFallback();
+		FName PackageName(*PackageNameToCreate);
+		if (!CollectionFilter.Contains(PackageName))
+		{
+			return;
+		}
+	}
+
+	if (bOnlyMaterials)
+	{
+		FString PackageNameToCreate = PackagePath.GetPackageNameOrFallback();
+		FName PackageFNameToCreate(*PackageNameToCreate);
+
+		FAssetRegistryModule& AssetRegistryModule = FModuleManager::LoadModuleChecked<FAssetRegistryModule>(TEXT("AssetRegistry"));
+		IAssetRegistry& AssetRegistry = AssetRegistryModule.Get();
+		TArray<FAssetData> PackageAssetData;
+		if (!AssetRegistry.GetAssetsByPackageName(PackageFNameToCreate, PackageAssetData, true))
+		{
+			return;
+		}
+
+		bool bAnyMaterialsPresent = false;
+		for (const FAssetData& AssetData : PackageAssetData)
+		{
+			UClass* AssetClass = AssetData.GetClass();
+			if (AssetClass && AssetClass->IsChildOf(UMaterialInterface::StaticClass()))
+			{
+				bAnyMaterialsPresent = true;
+				break;
+			}
+		}
+
+		if (!bAnyMaterialsPresent)
+		{
+			return;
+		}
+	}
+
+	bool bIsReadOnly = IFileManager::Get().IsReadOnly(*Filename);
+
+	if ( bIsReadOnly && !bVerifyContent && !bAutoCheckOut )
+	{
+		if (Verbosity != ONLY_ERRORS)
+		{
+			UE_LOG(LogIteratePackagesCommandlet, Warning, TEXT("Skipping read-only file %s"), *Filename);
+		}
+	}
+	else
+	{
+		if (Verbosity != ONLY_ERRORS)
+		{
+			UE_LOG(LogIteratePackagesCommandlet, Display, TEXT("Loading %s"), *Filename);
+		}
+
+		static int32 LastErrorCount = 0;
+
+		int32 NumErrorsFromLoading = GWarn->GetNumErrors();
+		if (NumErrorsFromLoading > LastErrorCount)
+		{
+			UE_LOG(LogIteratePackagesCommandlet, Warning, TEXT("%d total errors encountered during loading"), NumErrorsFromLoading);
+		}
+		LastErrorCount = NumErrorsFromLoading;
+
+		// Get the package linker.
+		VerboseMessage(TEXT("Pre GetPackageLinker"));
+
+		FLinkerLoad* Linker = LoadPackageLinker(nullptr, PackagePath, LOAD_NoVerify);
+	
+		// Bail early if we don't have a valid linker (package was out of date, etc)
+		if( !Linker )
+		{
+			VerboseMessage(TEXT("Aborting...package could not be loaded"));
+			CollectGarbage(RF_NoFlags);
+			return;
+		}
+
+		VerboseMessage(TEXT("Post GetPackageLinker"));
+
+		bool bSavePackage = true;
+		PerformPreloadOperations(Linker, bSavePackage);
+
+		VerboseMessage(FString::Printf(TEXT("Post PerformPreloadOperations, Resave? %d"), bSavePackage));
+		
+		if (bSavePackage)
+		{
+			PackagesConsideredForResave++;
+
+			// Only rebuild static meshes on load for the to be saved package.
+			extern ENGINE_API FName GStaticMeshPackageNameToRebuild;
+			FName PackageFName = PackagePath.GetPackageFName();
+			check(!PackageFName.IsNone());
+			GStaticMeshPackageNameToRebuild = PackageFName;
+
+			// Assert if package couldn't be opened so we have no chance of messing up saving later packages.
+			UPackage* Package = nullptr;
+			{
+				TRACE_CPUPROFILER_EVENT_SCOPE(UBaseIteratePackagesCommandlet::LoadAndSaveOnePackage::LoadPackage);
+				Package = LoadPackage(NULL, *Filename, 0);
+			}
+
+			if (Package == NULL)
+			{
+				if (bCanIgnoreFails == true)
+				{
+					return;
+				}
+				else
+				{
+					check(Package);
+				}
+			}
+
+			VerboseMessage(TEXT("Post LoadPackage"));
+
+			// if we are only saving dirty packages and the package is not dirty, then we do not want to save the package (remember the default behavior is to ALWAYS save the package)
+			if( ( bOnlySaveDirtyPackages == true ) && ( Package->IsDirty() == false ) )
+			{
+				bSavePackage = false;
+			}
+
+			// here we want to check and see if we have any loading warnings
+			// if we do then we want to resave this package
+			if( !bSavePackage && FParse::Param(FCommandLine::Get(),TEXT("SavePackagesThatHaveFailedLoads")) == true )
+			{
+				//UE_LOG(LogIteratePackagesCommandlet, Warning, TEXT( "NumErrorsFromLoading: %d GWarn->Errors num: %d" ), NumErrorsFromLoading, GWarn->GetNumErrors() );
+
+				if( NumErrorsFromLoading != GWarn->GetNumErrors() )
+				{
+					bSavePackage = true;
+				}
+			}
+
+			{
+				UWorld* World = UWorld::FindWorldInPackage(Package);
+				if (World)
+				{
+					PerformAdditionalOperations(World, bSavePackage);
+				}
+			}
+
+			// hook to allow performing additional checks without lumping everything into this one function
+			PerformAdditionalOperations(Package,bSavePackage);
+
+			VerboseMessage(TEXT("Post PerformAdditionalOperations"));
+
+			// Check for any special per object operations
+			ForEachObjectWithOuter(Package, [this, &bSavePackage](UObject* Object)
+			{
+				PerformAdditionalOperations(Object, bSavePackage);
+			});
+			
+			VerboseMessage(TEXT("Post PerformAdditionalOperations Loop"));
+
+			if (bStripEditorOnlyContent)
+			{
+				UE_LOG(LogIteratePackagesCommandlet, Log, TEXT("Removing editor only data"));
+				Package->SetPackageFlags(PKG_FilterEditorOnly);
+			}
+
+			if (bResaveOnDemand && bSavePackage)
+			{
+				// Finish all async loading for the package, in case the systems marking it for load do so only from the async loading
+				FAssetCompilingManager::Get().FinishAllCompilation();
+				FlushAsyncLoading();
+				// Skip saving the package if no systems requested it
+				FScopeLock ScopeLock(&ResaveOnDemandPackagesLock);
+				bSavePackage = ResaveOnDemandPackages.Contains(Package->GetFName());
+			}
+
+			if (bSavePackage == true)
+			{
+				bool bIsEmpty = true;
+				// Check to see if this package contains only metadata, and if so delete the package instead of resaving it
+
+				TArray<UObject *> ObjectsInOuter;
+				GetObjectsWithOuter(Package, ObjectsInOuter);
+				for (UObject* Obj : ObjectsInOuter)
+				{
+					if (!Obj->IsA(UMetaData::StaticClass()))
+					{
+						// This package has a real object
+						bIsEmpty = false;
+						break;
+					}
+				}
+
+				if (bIsEmpty)
+				{
+					bSavePackage = false;
+					Package = nullptr;
+					
+					UE_LOG(LogIteratePackagesCommandlet, Display, TEXT("Package %s is empty and will be deleted"), *Filename);
+
+					DeleteOnePackage(Filename);
+				}
+			}
+
+			// Now based on the computation above we will see if we should actually attempt
+			// to save this package
+			if (bSavePackage == true)
+			{
+				if( bIsReadOnly == true && bVerifyContent == true && bAutoCheckOut == false )
+				{
+					UE_LOG(LogIteratePackagesCommandlet, Warning, TEXT("Package [%s] is read-only but needs to be resaved (UE Version: %i, Licensee Version: %i  Current UE Version: %i, Current Licensee Version: %i)"),
+						*Filename, Linker->Summary.GetFileVersionUE().ToValue(), Linker->Summary.GetFileVersionLicenseeUE(), GPackageFileUEVersion.ToValue(), VER_LATEST_ENGINE_LICENSEEUE4 );
+					if( SavePackageHelper(Package, FString(TEXT("Temp.temp"))) )
+					{
+						UE_LOG(LogIteratePackagesCommandlet, Warning, TEXT("Correctly saved:  [Temp.temp].") );
+					}
+				}
+				else
+				{
+					// Check to see if we need to check this package out (but do not check out here if SourceControlQueue is enabled)
+					const bool bAttemptCheckoutNow = bAutoCheckOut && SourceControlQueue == nullptr;
+					if (bAttemptCheckoutNow)
+					{
+						TRACE_CPUPROFILER_EVENT_SCOPE(UBaseIteratePackagesCommandlet::LoadAndSaveOnePackage::AutoCheckOut);
+						if( bIsReadOnly )
+						{
+							VerboseMessage(TEXT("Pre ForceGetStatus1"));
+							ISourceControlProvider& SourceControlProvider = ISourceControlModule::Get().GetProvider();
+							FSourceControlStatePtr SourceControlState = SourceControlProvider.GetState( Package, EStateCacheUsage::ForceUpdate );
+							if(SourceControlState.IsValid())
+							{
+								FString OtherCheckedOutUser;
+								if( SourceControlState->IsCheckedOutOther(&OtherCheckedOutUser) )
+								{
+									UE_LOG(LogIteratePackagesCommandlet, Warning, TEXT("[REPORT] Overwriting package %s already checked out by someone else (%s), will not submit"), *Filename, *OtherCheckedOutUser);
+								}
+								else if( !SourceControlState->IsCurrent() )
+								{
+									UE_LOG(LogIteratePackagesCommandlet, Warning, TEXT("[REPORT] Overwriting package %s (not at head), will not submit"), *Filename);
+								}
+								else
+								{
+									VerboseMessage(TEXT("Pre CheckOut"));
+
+									SourceControlProvider.Execute(ISourceControlOperation::Create<FCheckOut>(), Package);
+
+									VerboseMessage(TEXT("Post CheckOut"));
+
+									FilesToSubmit.AddUnique(*Filename);
+								}
+							}
+							VerboseMessage(TEXT("Post ForceGetStatus2"));
+						}
+					}
+
+					// Update the readonly state now that source control has had a chance to run
+					bIsReadOnly = IFileManager::Get().IsReadOnly(*Filename);
+
+					// If we tried to check out the file but it is still read only then we failed and 
+					// need to emit an error and go to the next package			
+					if (bAttemptCheckoutNow && bIsReadOnly)
+					{
+						if (bSkipCheckedOutFiles)
+						{
+							UE_LOG(LogIteratePackagesCommandlet, Warning, TEXT("Skipping Package: %s (unable to check out)"), *Filename);
+						}
+						else
+						{
+							UE_LOG(LogIteratePackagesCommandlet, Error, TEXT("Unable to check out the Package: %s"), *Filename);
+						}
+						
+						return;
+					}
+
+					if (Verbosity != ONLY_ERRORS)
+					{
+						UE_LOG(LogIteratePackagesCommandlet, Display, TEXT("Resaving package [%s] (UE Version: %i, Licensee Version: %i  Saved UE Version: %i, Saved Licensee Version: %i)"),
+							*Filename,Linker->Summary.GetFileVersionUE().ToValue(), Linker->Summary.GetFileVersionLicenseeUE(), GPackageFileUEVersion.ToValue(), VER_LATEST_ENGINE_LICENSEEUE4 );
+					}
+
+					const static bool bKeepPackageGUIDOnSave = FParse::Param(FCommandLine::Get(), TEXT("KeepPackageGUIDOnSave"));
+					ESaveFlags SaveFlags = bKeepPackageGUIDOnSave ? SAVE_KeepGUID : SAVE_None;
+					
+					if (bIsReadOnly == false || SourceControlQueue == nullptr)
+					{
+						TRACE_CPUPROFILER_EVENT_SCOPE(UBaseIteratePackagesCommandlet::LoadAndSaveOnePackage::SavePackage);
+
+						if (SavePackageHelper(Package, Filename, RF_Standalone, GWarn, SaveFlags))
+						{
+							PackagesResaved++;
+							if (Verbosity == VERY_VERBOSE)
+							{
+								UE_LOG(LogIteratePackagesCommandlet, Display, TEXT("Correctly saved:  [%s]."), *Filename);
+							}
+						}
+					}
+					else
+					{
+						// The target file is still read only and we have SourceControlQueue enabled so we
+						// need to save to a temporary file first, and then queue the result.
+						const FString TempFilename = CreateTempFilename();
+
+						if (SavePackageHelper(Package, TempFilename, RF_Standalone, GWarn, SaveFlags))
+						{
+							SourceControlQueue->QueueCheckoutAndReplaceOperation(Filename, TempFilename, Package);
+						}
+					}
+				}
+			}
+		}
+
+		static int32 Counter = 0;
+
+		if (!GarbageCollectionFrequency || Counter++ % GarbageCollectionFrequency == 0)
+		{
+			if (FApp::CanEverRender())
+			{
+				FAssetCompilingManager::Get().FinishAllCompilation();
+			}
+
+			if (GarbageCollectionFrequency > 1)
+			{
+				UE_LOG(LogIteratePackagesCommandlet, Display, TEXT("GC"));
+			}
+			VerboseMessage(TEXT("Pre CollectGarbage"));
+
+			CollectGarbage(RF_NoFlags);
+
+			VerboseMessage(TEXT("Post CollectGarbage"));
+		}
+	}
+}
+
+void UBaseIteratePackagesCommandlet::DeleteOnePackage(const FString& Filename)
+{
+	TRACE_CPUPROFILER_EVENT_SCOPE(UBaseIteratePackagesCommandlet::DeleteOnePackage);
+
+	bool bIsReadOnly = IFileManager::Get().IsReadOnly(*Filename);
+
+	if (bVerifyContent)
+	{
+		return;
+	}
+
+	if (bIsReadOnly && !bAutoCheckOut)
+	{
+		if (Verbosity != ONLY_ERRORS)
+		{
+			UE_LOG(LogIteratePackagesCommandlet, Warning, TEXT("Skipping read-only file %s"), *Filename);
+		}
+		return;
+	}
+
+	if (SourceControlQueue != nullptr)
+	{
+		// All files (read only and non read only) need to query their source control status so
+		// they should all be queued if avaliable.
+		SourceControlQueue->QueueDeleteOperation(Filename);
+		return;
+	}
+
+	FString PackageName;
+	FPackageName::TryConvertFilenameToLongPackageName(Filename, PackageName);
+
+	UPackage* Package = FindPackage(nullptr, *PackageName);
+
+	if (Package)
+	{
+		// Unload package so we can delete it
+		TArray<UPackage *> PackagesToDelete;
+		PackagesToDelete.Add(Package);
+		UPackageTools::UnloadPackages(PackagesToDelete);
+		PackagesToDelete.Empty();
+		Package = nullptr;
+	}
+
+	FString PackageFilename = SourceControlHelpers::PackageFilename(Filename);
+	ISourceControlProvider& SourceControlProvider = ISourceControlModule::Get().GetProvider();
+	FSourceControlStatePtr SourceControlState = SourceControlProvider.GetState(PackageFilename, EStateCacheUsage::ForceUpdate);
+
+	if (SourceControlState.IsValid() && (SourceControlState->IsCheckedOut() || SourceControlState->IsAdded()))
+	{
+		UE_LOG(LogIteratePackagesCommandlet, Display, TEXT("Revert '%s' from source control..."), *Filename);
+		SourceControlProvider.Execute(ISourceControlOperation::Create<FRevert>(), PackageFilename);
+
+		UE_LOG(LogIteratePackagesCommandlet, Display, TEXT("Deleting '%s' from source control..."), *Filename);
+		SourceControlProvider.Execute(ISourceControlOperation::Create<FDelete>(), PackageFilename);
+
+		PackagesDeleted++;
+
+		FilesToSubmit.AddUnique(Filename);
+	}
+	else if (SourceControlState.IsValid() && SourceControlState->CanCheckout())
+	{
+		UE_LOG(LogIteratePackagesCommandlet, Display, TEXT("Deleting '%s' from source control..."), *Filename);
+		SourceControlProvider.Execute(ISourceControlOperation::Create<FDelete>(), PackageFilename);
+
+		PackagesDeleted++;
+
+		FilesToSubmit.AddUnique(Filename);
+	}
+	else if (SourceControlState.IsValid() && SourceControlState->IsCheckedOutOther())
+	{
+		UE_LOG(LogIteratePackagesCommandlet, Warning, TEXT("Couldn't delete '%s' from source control, someone has it checked out, skipping..."), *Filename);
+	}
+	else if (SourceControlState.IsValid() && !SourceControlState->IsSourceControlled())
+	{
+		UE_LOG(LogIteratePackagesCommandlet, Warning, TEXT("'%s' is not in source control, attempting to delete from disk..."), *Filename);
+		if (IFileManager::Get().Delete(*Filename, false, true) == true)
+		{
+			PackagesDeleted++;
+		}
+		else
+		{
+			UE_LOG(LogIteratePackagesCommandlet, Warning, TEXT("  ... failed to delete from disk."), *Filename);
+		}
+	}
+	else
+	{
+		UE_LOG(LogIteratePackagesCommandlet, Warning, TEXT("'%s' is in an unknown source control state, attempting to delete from disk..."), *Filename);
+		if (IFileManager::Get().Delete(*Filename, false, true)== true)
+		{
+			PackagesDeleted++;
+		}
+		else
+		{
+			UE_LOG(LogIteratePackagesCommandlet, Warning, TEXT("  ... failed to delete from disk."), *Filename);
+		}
+	}
+}
+
+int32 UBaseIteratePackagesCommandlet::Main( const FString& Params )
+{
+	TRACE_CPUPROFILER_EVENT_SCOPE(UBaseIteratePackagesCommandlet);
+
+	const TCHAR* Parms = *Params;
+	TArray<FString> Tokens;
+	ParseCommandLine(Parms, Tokens, Switches);
+
+	// Ensure source control is initialized and shut down properly
+	FScopedSourceControl SourceControl;
+
+	// skip the assert when a package can not be opened
+	bCanIgnoreFails = Switches.Contains(TEXT("SKIPFAILS"));
+	/** load all packages, and display warnings for those packages which would have been resaved but were read-only */
+	bVerifyContent = Switches.Contains(TEXT("VERIFY"));
+	/** if we should only save dirty packages **/
+	bOnlySaveDirtyPackages = Switches.Contains(TEXT("OnlySaveDirtyPackages"));
+	/** if we should auto checkout packages that need to be saved**/
+	bAutoCheckOut = Switches.Contains(TEXT("AutoCheckOutPackages")) || Switches.Contains(TEXT("AutoCheckOut"));
+	/** if we should batch source control operations*/
+	bBatchSourceControl = Switches.Contains(TEXT("BatchSourceControl"));
+	/** if we should simply skip checked out files rather than error-ing out */
+	bSkipCheckedOutFiles = Switches.Contains(TEXT("SkipCheckedOutPackages"));
+	/** if we should auto checkin packages that were checked out**/
+	bAutoCheckIn = bAutoCheckOut && (Switches.Contains(TEXT("AutoCheckIn")) || Switches.Contains(TEXT("AutoSubmit")));
+	/** determine if we can skip the version changelist check */
+	bIgnoreChangelist = Switches.Contains(TEXT("IgnoreChangelist"));
+	/** whether we should only save packages with changelist zero */
+	bOnlyUnversioned = Switches.Contains(TEXT("OnlyUnversioned"));
+	/** whether we should only save packages saved by licensees */
+	bOnlyLicenseed = Switches.Contains(TEXT("OnlyLicenseed"));
+	/** whether we should only save packages containing virtualized bulkdata payloads */
+	bOnlyVirtualized = Switches.Contains(TEXT("OnlyVirtualized"));
+	/** whether we should only save packages containing FPayloadTrailers */
+	bOnlyPayloadTrailers = Switches.Contains(TEXT("OnlyPayloadTrailers"));
+	/** only process packages containing materials */
+	bOnlyMaterials = Switches.Contains(TEXT("onlymaterials"));
+	
+	/** check for filtering packages by collection. **/
+	FString FilterByCollection;
+	FParse::Value(*Params, TEXT("FilterByCollection="), FilterByCollection);
+
+	CollectionFilter = TSet<FName>();
+	if (!FilterByCollection.IsEmpty())
+	{
+		ICollectionManager& CollectionManager = FCollectionManagerModule::GetModule().Get();
+		TArray<FName> CollectionAssets;
+		if (!CollectionManager.GetAssetsInCollection(FName(*FilterByCollection), ECollectionShareType::CST_All, CollectionAssets))
+		{
+			UE_LOG(LogIteratePackagesCommandlet, Warning, TEXT("Could not get assets in collection '%s'. Skipping filter."), *FilterByCollection);
+		}
+		else
+		{
+			//insert all of the collection names into the set for fast filter checks
+			for (const FName &AssetName : CollectionAssets)
+			{
+				CollectionFilter.Add(FName(*FPackageName::ObjectPathToPackageName(AssetName.ToString())));
+			}
+		}
+	}
+
+
+	TArray<FString> PackageNames;
+	int32 ResultCode = InitializeParameters(Tokens, PackageNames);
+	if ( ResultCode != 0 )
+	{
+		return ResultCode;
+	}
+
+
+	int32 GCIndex = 0;
+	PackagesConsideredForResave = 0;
+	PackagesResaved = 0;
+	PackagesDeleted = 0;
+	TotalPackagesForResave = PackageNames.Num();
+
+	// allow for an option to restart at a given package name (in case it dies during a run, etc)
+	bool bCanProcessPackage = true;
+	FString FirstPackageToProcess;
+	if (FParse::Value(*Params, TEXT("FirstPackage="), FirstPackageToProcess))
+	{
+		bCanProcessPackage = false;
+	}
+	FParse::Value(*Params, TEXT("PackageSubString="), PackageSubstring);
+	if (PackageSubstring.Len())
+	{
+		UE_LOG(LogIteratePackagesCommandlet, Display, TEXT( "Restricted to packages containing %s" ), *PackageSubstring );
+	}
+
+	// Avoid crash saving blueprint
+	FFindInBlueprintSearchManager::Get();
+
+	// Tick shader compiler if we are rendering
+	if(IsAllowCommandletRendering() && GShaderCompilingManager)
+	{
+		GShaderCompilingManager->ProcessAsyncResults(true, false);
+	}
+
+	if (bBatchSourceControl)
+	{
+		// Convert the commandlets verbosity to that of FQueuedSourceControlOperations
+		FQueuedSourceControlOperations::EVerbosity LogVerbosity = FQueuedSourceControlOperations::EVerbosity::All;
+		switch (Verbosity)
+		{
+		case VERY_VERBOSE:
+			LogVerbosity = FQueuedSourceControlOperations::EVerbosity::All;
+			break;
+		case INFORMATIVE:
+			LogVerbosity = FQueuedSourceControlOperations::EVerbosity::Info;
+			break;
+		case ONLY_ERRORS:
+			LogVerbosity = FQueuedSourceControlOperations::EVerbosity::ErrorsOnly;
+			break;
+		default:
+			UE_LOG(LogIteratePackagesCommandlet, Warning, TEXT("Unknown verbosity to pass to FQueuedSourceControlOperations!"));
+		}
+
+		SourceControlQueue = MakePimpl<FQueuedSourceControlOperations>(LogVerbosity);	
+		ParseSourceControlOptions(Switches);
+	}
+
+	// Make sure any remaining temp files from previous runs are removed
+	CleanTempFiles();
+
+	// Iterate over all packages.
+	for( int32 PackageIndex = 0; PackageIndex < PackageNames.Num(); PackageIndex++ )
+	{
+		// Make sure we don't rebuild SMs that we're not going to save.
+		extern ENGINE_API FName GStaticMeshPackageNameToRebuild;
+		GStaticMeshPackageNameToRebuild = NAME_None;
+
+		const FString& Filename = PackageNames[PackageIndex];
+
+		// skip over packages before the first one allowed, if it was specified
+		if (!bCanProcessPackage)
+		{
+			if (FPackageName::FilenameToLongPackageName(Filename) == FirstPackageToProcess)
+			{
+				bCanProcessPackage = true;
+			}
+			else
+			{
+				UE_LOG(LogIteratePackagesCommandlet, Display, TEXT("Skipping %s"), *Filename);
+				continue;
+			}
+		}
+
+		// Load and save this package
+		LoadAndSaveOnePackage(Filename);
+
+		// Check if we need to flush any source control operations yet
+		if (SourceControlQueue != nullptr)
+		{
+			SourceControlQueue->FlushPendingOperations(false);
+		}
+
+		// Tick shader compiler if we are rendering
+		if(IsAllowCommandletRendering() && GShaderCompilingManager)
+		{
+			GShaderCompilingManager->ProcessAsyncResults(true, false);
+		}
+
+		// Break out if we've resaved enough packages
+		if( MaxPackagesToResave > -1 && PackagesResaved >= MaxPackagesToResave )
+		{
+			UE_LOG(LogIteratePackagesCommandlet, Warning, TEXT( "Attempting to resave more than MaxPackagesToResave; exiting" ) );
+			break;
+		}
+	}
+	
+	if (SourceControlQueue != nullptr)
+	{
+		// Flush any remaining source control operations
+		SourceControlQueue->FlushPendingOperations(true);
+
+		// Fix up the stats
+		PackagesDeleted += SourceControlQueue->GetNumDeletedFiles();
+		PackagesResaved += SourceControlQueue->GetNumReplacedFiles();
+
+		// Add files that the source control queue modified to the list of files that 
+		// will need to be submitted to source control
+		FilesToSubmit.Append(SourceControlQueue->GetModifiedFiles());
+
+		SourceControlQueue.Reset();
+	}
+
+	// Force a directory watcher and asset registry tick
+	FDirectoryWatcherModule& DirectoryWatcherModule = FModuleManager::Get().LoadModuleChecked<FDirectoryWatcherModule>(TEXT("DirectoryWatcher"));
+	DirectoryWatcherModule.Get()->Tick(-1.0f);
+
+	FAssetRegistryModule& AssetRegistryModule = FModuleManager::LoadModuleChecked<FAssetRegistryModule>(TEXT("AssetRegistry"));
+	IAssetRegistry& AssetRegistry = AssetRegistryModule.Get();
+	AssetRegistry.Tick(-1.0f);
+
+	PostSavePackages();
+
+	// Submit the results to source control
+	CheckInFiles(FilesToSubmit, GetChangelistDescription());
+
+	if (bForceUATEnvironmentVariableSet)
+	{
+		FPlatformMisc::SetEnvironmentVar(TEXT("uebp_UATMutexNoWait"), TEXT("0"));		
+	}
+
+	UE_LOG(LogIteratePackagesCommandlet, Display, TEXT("[REPORT] %d/%d packages were considered for resaving"), PackagesConsideredForResave, TotalPackagesForResave);
+	UE_LOG(LogIteratePackagesCommandlet, Display, TEXT("[REPORT] %d/%d packages were resaved"), PackagesResaved, PackagesConsideredForResave);
+	UE_LOG(LogIteratePackagesCommandlet, Display, TEXT("[REPORT] %d/%d packages were deleted"), PackagesDeleted, PackagesConsideredForResave);
+
+
+	return 0;
+}
+
+FText UBaseIteratePackagesCommandlet::GetChangelistDescription() const
+{
+	FTextBuilder ChangelistDescription;
+	ChangelistDescription.AppendLine(NSLOCTEXT("ContentCmdlets", "ChangelistDescription", "Resave Packages"));
+	return ChangelistDescription.ToText();
+}
+
+
+void UBaseIteratePackagesCommandlet::PerformPreloadOperations( FLinkerLoad* PackageLinker, bool& bSavePackage )
+{
+	const int32 UEPackageVersion = PackageLinker->Summary.GetFileVersionUE().ToValue();
+	const int32 LicenseeUEPackageVersion = PackageLinker->Summary.GetFileVersionLicenseeUE();
+
+	// validate that this package meets the minimum requirement
+	if (MinResaveUEVersion != IGNORE_PACKAGE_VERSION && UEPackageVersion < MinResaveUEVersion)
+	{
+		bSavePackage = false;
+		return;
+	}
+
+	// Check if this package meets the maximum requirements.
+	const bool bNoLimitation = MaxResaveUEVersion == IGNORE_PACKAGE_VERSION && MaxResaveLicenseeUEVersion == IGNORE_PACKAGE_VERSION;
+	const bool bAllowResave = bNoLimitation ||
+						 (MaxResaveUEVersion != IGNORE_PACKAGE_VERSION && UEPackageVersion <= MaxResaveUEVersion) ||
+						 (MaxResaveLicenseeUEVersion != IGNORE_PACKAGE_VERSION && LicenseeUEPackageVersion <= MaxResaveLicenseeUEVersion);
+	// If not, don't resave it.
+	if (!bAllowResave)
+	{
+		bSavePackage = false;
+		return;
+	}
+
+	// If the package was saved with a higher engine version do not try to resave it. This also addresses problem with people 
+	// building editor locally and resaving content with a 0 CL version (e.g. BUILD_FROM_CL == 0)
+	if (!bIgnoreChangelist && PackageLinker->Summary.SavedByEngineVersion.GetChangelist() > FEngineVersion::Current().GetChangelist())
+	{
+		UE_LOG(LogIteratePackagesCommandlet, Warning, TEXT("Skipping resave of %s due to engine version mismatch (Package:%d, Editor:%d) "), 
+			*PackageLinker->GetArchiveName(),
+			PackageLinker->Summary.SavedByEngineVersion.GetChangelist(), 
+			FEngineVersion::Current().GetChangelist());
+		bSavePackage = false;
+		return;
+	}
+
+	// Check if the changelist number is zero
+	if (bOnlyUnversioned && PackageLinker->Summary.SavedByEngineVersion.GetChangelist() != 0)
+	{
+		bSavePackage = false;
+		return;
+	}
+
+	// Check if the package was saved by licensees
+	if (bOnlyLicenseed)
+	{
+		if (!PackageLinker->Summary.SavedByEngineVersion.IsLicenseeVersion())
+		{
+			bSavePackage = false;
+			return;
+		}
+		else
+		{
+			UE_LOG(LogIteratePackagesCommandlet, Display, TEXT("Resaving %s that contains licensee version (%s)"),
+				*PackageLinker->GetArchiveName(),
+				*PackageLinker->Summary.SavedByEngineVersion.ToString());
+		}
+	}
+
+	// Check if the package contains virtualized bulkdata payloads
+	if (bOnlyVirtualized)
+	{
+		const UE::FPackageTrailer* Trailer = PackageLinker->GetPackageTrailer();
+		if (Trailer == nullptr || Trailer->GetNumPayloads(UE::EPayloadStorageType::Virtualized) == 0)
+		{
+			bSavePackage = false;
+			return;
+		}
+	}
+
+	// Check if the package contains a FPackageTrailer or not
+	if (bOnlyPayloadTrailers)
+	{
+		const UE::FPackageTrailer* Trailer = PackageLinker->GetPackageTrailer();
+		if (Trailer == nullptr)
+		{
+			bSavePackage = false;
+			return;
+		}
+	}
+
+	// Check if the package contains any instances of the class that needs to be resaved.
+	if (ResaveClasses.Num() > 0)
+	{
+		bSavePackage = false;
+		for (int32 ExportIndex = 0; !bSavePackage && ExportIndex < PackageLinker->ExportMap.Num(); ExportIndex++)
+		{
+			FTopLevelAssetPath ExportClassPathName(PackageLinker->GetExportClassPackage(ExportIndex), PackageLinker->GetExportClassName(ExportIndex));
+			if (ResaveClasses.Contains(ExportClassPathName.ToString()))
+			{
+				bSavePackage = true;
+				break;
+			}
+		}
+	}
+}
+
+bool UBaseIteratePackagesCommandlet::CheckoutFile(const FString& Filename, bool bAddFile, bool bIgnoreAlreadyCheckedOut)
+{
+	if (!bAutoCheckOut)
+	{
+		return true;
+	}
+
+	bool bIsReadOnly = IFileManager::Get().IsReadOnly(*Filename);
+	if (!bIsReadOnly && !bAddFile)
+	{
+		return true;
+	}
+
+	ISourceControlProvider& SourceControlProvider = ISourceControlModule::Get().GetProvider();
+	FSourceControlStatePtr SourceControlState = SourceControlProvider.GetState(*Filename, EStateCacheUsage::ForceUpdate);
+	if (SourceControlState.IsValid())
+	{
+		// Already checked out/added this file
+		if (SourceControlState->IsCheckedOut() || SourceControlState->IsAdded())
+		{
+			return true;
+		}
+		else if (!SourceControlState->IsSourceControlled())
+		{
+			if ( bAddFile )
+			{
+				if (SourceControlProvider.Execute(ISourceControlOperation::Create<FMarkForAdd>(), *Filename) == ECommandResult::Succeeded)
+				{
+					UE_LOG(LogIteratePackagesCommandlet, Display, TEXT("[REPORT] %s successfully added"), *Filename);
+					return true;
+				}
+				else
+				{
+					if (!bIgnoreAlreadyCheckedOut)
+					{
+						UE_LOG(LogIteratePackagesCommandlet, Error, TEXT("[REPORT] %s could not be added!"), *Filename);
+					}
+					else
+					{
+						UE_LOG(LogIteratePackagesCommandlet, Warning, TEXT("[REPORT] %s could not be added!"), *Filename);
+					}
+				}
+			}
+		}
+		else if (!SourceControlState->IsCurrent())
+		{
+			if (!bIgnoreAlreadyCheckedOut)
+			{
+				UE_LOG(LogIteratePackagesCommandlet, Error, TEXT("[REPORT] %s is not synced to head, can not submit"), *Filename);
+			}
+			else
+			{
+				UE_LOG(LogIteratePackagesCommandlet, Warning, TEXT("[REPORT] %s is not synced to head, can not submit"), *Filename);
+			}
+		}
+		else if (!SourceControlState->CanCheckout())
+		{
+			FString CurrentlyCheckedOutUser;
+			if (SourceControlState->IsCheckedOutOther(&CurrentlyCheckedOutUser))
+			{
+				if (!bIgnoreAlreadyCheckedOut)
+				{
+					UE_LOG(LogIteratePackagesCommandlet, Error, TEXT("[REPORT] %s level is already checked out by someone else (%s), can not submit!"), *Filename, *CurrentlyCheckedOutUser);
+				}
+				else
+				{
+					UE_LOG(LogIteratePackagesCommandlet, Warning, TEXT("[REPORT] %s level is already checked out by someone else (%s), can not submit!"), *Filename, *CurrentlyCheckedOutUser);
+				}
+			}
+			else
+			{
+				UE_LOG(LogIteratePackagesCommandlet, Error, TEXT("[REPORT] Unable to checkout %s, can not submit"), *Filename);
+			}
+		}
+		else 
+		{
+			if (SourceControlProvider.Execute(ISourceControlOperation::Create<FCheckOut>(), *Filename) == ECommandResult::Succeeded)
+			{
+				UE_LOG(LogIteratePackagesCommandlet, Display, TEXT("[REPORT] %s Checked out successfully"), *Filename);
+				return true;
+			}
+			else
+			{
+				UE_LOG(LogIteratePackagesCommandlet, Warning, TEXT("[REPORT] %s could not be checked out!"), *Filename);
+			}
+		}
+	}
+	return false;
+}
+
+bool UBaseIteratePackagesCommandlet::RevertFile(const FString& Filename)
+{
+	ISourceControlProvider& SourceControlProvider = ISourceControlModule::Get().GetProvider();
+	FSourceControlStatePtr SourceControlState = SourceControlProvider.GetState(*Filename, EStateCacheUsage::ForceUpdate);
+	bool bSuccesfullyReverted = false;
+
+	if (SourceControlState.IsValid())
+	{
+		if (SourceControlState->CanRevert() && (SourceControlProvider.Execute(ISourceControlOperation::Create<FRevert>(), *Filename) == ECommandResult::Succeeded))
+		{
+			bSuccesfullyReverted = true;
+			UE_LOG(LogIteratePackagesCommandlet, Display, TEXT("[REPORT] %s Reverted successfully"), *Filename);			
+		}
+		else
+		{
+			UE_LOG(LogIteratePackagesCommandlet, Warning, TEXT("[REPORT] %s could not be reverted!"), *Filename);
+		}
+	}
+
+	return bSuccesfullyReverted;
+}
+
+bool UBaseIteratePackagesCommandlet::CanCheckoutFile(const FString& Filename, FString& CheckedOutUser)
+{
+	ISourceControlProvider& SourceControlProvider = ISourceControlModule::Get().GetProvider();
+	FSourceControlStatePtr SourceControlState = SourceControlProvider.GetState(*Filename, EStateCacheUsage::ForceUpdate);
+	bool bCanCheckout = true;
+	if (SourceControlState.IsValid())
+	{
+		if (!SourceControlState->CanCheckout())
+		{
+			if (!SourceControlState->IsCheckedOutOther(&CheckedOutUser))
+			{
+				CheckedOutUser = "";
+			}
+
+			bCanCheckout = false;
+		}
+	}
+
+	return bCanCheckout;
+}
+
+void UBaseIteratePackagesCommandlet::CheckoutAndSavePackage(UPackage* Package, TArray<FString>& SublevelFilenames, bool bIgnoreAlreadyCheckedOut)
+{
+	check(Package);
+
+	FString PackageFilename;
+	if (FPackageName::TryConvertLongPackageNameToFilename(Package->GetName(), PackageFilename, Package->ContainsMap() ? FPackageName::GetMapPackageExtension() : FPackageName::GetAssetPackageExtension()))
+	{
+		if (IFileManager::Get().FileExists(*PackageFilename))
+		{
+			if (CheckoutFile(PackageFilename, true, bIgnoreAlreadyCheckedOut))
+			{
+				SublevelFilenames.Add(PackageFilename);
+				if (!SavePackageHelper(Package, PackageFilename))
+				{
+					UE_LOG(LogIteratePackagesCommandlet, Error, TEXT("Failed to save existing package %s"), *PackageFilename);
+				}
+			}
+			else
+			{
+				UE_LOG(LogIteratePackagesCommandlet, Error, TEXT("Failed to check out existing package %s"), *PackageFilename);
+			}
+		}
+		else
+		{
+			if (SavePackageHelper(Package, PackageFilename))
+			{
+				if (CheckoutFile(PackageFilename, true, bIgnoreAlreadyCheckedOut))
+				{
+					SublevelFilenames.Add(PackageFilename);
+				}
+				else
+				{
+					UE_LOG(LogIteratePackagesCommandlet, Error, TEXT("Failed to check out new package %s"), *PackageFilename);
+				}
+			}
+			else
+			{
+				UE_LOG(LogIteratePackagesCommandlet, Error, TEXT("Failed to save new package %s"), *PackageFilename);
+			}
+		}
+	}
+}
+
+void UBaseIteratePackagesCommandlet::CheckInFiles(const TArray<FString>& InFilesToSubmit, const FText& InDescription) const
+{
+	if (!bAutoCheckIn)
+	{
+		return;
+	}
+
+	// Check in all changed files
+	if (InFilesToSubmit.Num() > 0)
+	{
+		TSharedRef<FCheckIn, ESPMode::ThreadSafe> CheckInOperation = ISourceControlOperation::Create<FCheckIn>();
+		CheckInOperation->SetDescription(InDescription);
+
+		ISourceControlProvider& SourceControlProvider = ISourceControlModule::Get().GetProvider();
+		SourceControlProvider.Execute(CheckInOperation, SourceControlHelpers::PackageFilenames(InFilesToSubmit));
+	}
+}
+
+
+bool UBaseIteratePackagesCommandlet::CleanClassesFromContentPackages( UPackage* Package )
+{
+	check(Package);
+	bool bResult = false;
+
+	for ( TObjectIterator<UClass> It; It; ++It )
+	{
+		if ( It->IsIn(Package) )
+		{
+			UE_LOG(LogIteratePackagesCommandlet, Warning, TEXT("Removing class '%s' from package [%s]"), *It->GetPathName(), *Package->GetName());
+
+			// mark the class as transient so that it won't be saved into the package
+			It->SetFlags(RF_Transient);
+
+			// clear the standalone flag just to be sure :)
+			It->ClearFlags(RF_Standalone);
+			bResult = true;
+		}
+	}
+
+	return bResult;
+}
+
+void UBaseIteratePackagesCommandlet::VerboseMessage(const FString& Message)
+{
+	if (Verbosity == VERY_VERBOSE)
+	{
+		UE_LOG(LogIteratePackagesCommandlet, Verbose, TEXT("%s"), *Message);
+	}
+}
+
+FString UBaseIteratePackagesCommandlet::CreateTempFilename()
+{
+	return FPaths::CreateTempFilename(*GetTempFilesDirectory());
+}
+
+void UBaseIteratePackagesCommandlet::CleanTempFiles()
+{
+	TRACE_CPUPROFILER_EVENT_SCOPE(UBaseIteratePackagesCommandlet::CleanTempFiles);
+	const FString DirPath = GetTempFilesDirectory();
+
+	UE_LOG(LogIteratePackagesCommandlet, Display, TEXT("Cleaning temp file directory"), *DirPath);
+
+	if (!IFileManager::Get().DeleteDirectory(*DirPath, false, true))
+	{
+		UE_LOG(LogIteratePackagesCommandlet, Warning, TEXT("Failed to clean temp file directory:  %s"), *DirPath);
+	}
+}
+
+FString UBaseIteratePackagesCommandlet::GetTempFilesDirectory()
+{
+	return FPaths::ProjectSavedDir() / TEXT("Temp/ResavePackages");
+}
diff --git a/Engine/Source/Editor/UnrealEd/Private/Commandlets/GeneratemoderationArtifactsCommandlet.cpp b/Engine/Source/Editor/UnrealEd/Private/Commandlets/GeneratemoderationArtifactsCommandlet.cpp
new file mode 100644
--- /dev/null
+++ b/Engine/Source/Editor/UnrealEd/Private/Commandlets/GeneratemoderationArtifactsCommandlet.cpp
@@ -0,0 +1,305 @@
+// Copyright Epic Games, Inc. All Rights Reserved.
+
+/*=============================================================================
+	GenerateModerationArtifactsCommandlet.cpp: Commandlet provides basic package iteration functionaliy for derived commandlets
+=============================================================================*/
+#include "Commandlets/GenerateModerationArtifactsCommandlet.h"
+#include "CoreMinimal.h"
+#include "IImageWrapper.h"
+#include "IImageWrapperModule.h"
+#include "Modules/ModuleManager.h"
+#include "Engine/Texture.h"
+#include "Components/StaticMeshComponent.h"
+#include "Misc/FileHelper.h"
+
+/**-----------------------------------------------------------------------------
+ *	UGenerateModerationArtifactsCommandlet commandlet.
+ *
+ * This commandlet exposes some functionality for iterating packages 
+ *
+ *
+----------------------------------------------------------------------------**/
+
+DEFINE_LOG_CATEGORY(LogModerationArtifactsCommandlet);
+
+UGenerateModerationArtifactsCommandlet::UGenerateModerationArtifactsCommandlet(const FObjectInitializer& ObjectInitializer)
+	: Super(ObjectInitializer)
+{
+}
+
+int32 UGenerateModerationArtifactsCommandlet::InitializeParameters( const TArray<FString>& Tokens, TArray<FString>& PackageNames )
+{
+	
+
+	for (int32 SwitchIdx = 0; SwitchIdx < Switches.Num(); SwitchIdx++)
+	{
+		const FString& CurrentSwitch = Switches[SwitchIdx];
+		if (FParse::Value(*CurrentSwitch, TEXT("OutputDir="), OutputPath))
+		{
+			continue;
+		}
+	}
+
+	return Super::InitializeParameters(Tokens, PackageNames);
+}
+
+void UGenerateModerationArtifactsCommandlet::PerformAdditionalOperations(class UPackage* Package, bool& bSavePackage)
+{
+
+	// need to process localization, fstrings in structs, probably need to do something special with datatables... 
+
+	GatherLocalizationFromPackage(Package);
+}
+
+
+
+void UGenerateModerationArtifactsCommandlet::GatherFStringsFromObject(class UObject* Object)
+{
+	
+}
+
+
+void UGenerateModerationArtifactsCommandlet::GatherLocalizationFromPackage(class UPackage* Package)
+{
+
+/*
+	TArray<FGatherableTextData> GatherableTextDataArray;
+
+	// Gathers from the given package
+	EPropertyLocalizationGathererResultFlags GatherableTextResultFlags = EPropertyLocalizationGathererResultFlags::Empty;
+	FPropertyLocalizationDataGatherer(GatherableTextDataArray, Package, GatherableTextResultFlags);
+
+	if (PackagePendingGather.PackageLocCacheState != EPackageLocCacheState::Cached)
+	{
+		ProcessGatherableTextDataArray(GatherableTextDataArray);
+	}
+*/
+
+#if 0
+	int32 NumPackagesFailed = 0;
+	TArray<FName> PackagesWithStaleGatherCache;
+	TArray<FGatherableTextData> GatherableTextDataArray;
+	while (PackagesPendingGather.Num() > 0)
+	{
+		const FPackagePendingGather PackagePendingGather = PackagesPendingGather.Pop(/ *bAllowShrinking* /false);
+		const FNameBuilder PackageNameStr(PackagePendingGather.PackageName);
+
+		const int32 CurrentPackageNum = ++NumPackagesProcessed + NumPackagesFailed;
+		const float PercentageComplete = (((float)CurrentPackageNum / (float)PackageCount) * 100.0f);
+		UE_LOG(LogGatherTextFromAssetsCommandlet, Display, TEXT("[%6.2f%%] Loading package: '%s'..."), PercentageComplete, *PackageNameStr);
+
+		UPackage* Package = nullptr;
+		{
+			FLoadPackageLogOutputRedirector::FScopedCapture ScopedCapture(&LogOutputRedirector, *PackageNameStr);
+			Package = LoadPackage(nullptr, *PackageNameStr, LOAD_NoWarn | LOAD_Quiet);
+		}
+
+		if (!Package)
+		{
+			UE_LOG(LogGatherTextFromAssetsCommandlet, Warning, TEXT("Failed to load package: '%s'."), *PackageNameStr);
+			++NumPackagesFailed;
+			continue;
+		}
+
+		// Tick background tasks
+		if (GShaderCompilingManager)
+		{
+			GShaderCompilingManager->ProcessAsyncResults(true, false);
+		}
+		if (GDistanceFieldAsyncQueue)
+		{
+			GDistanceFieldAsyncQueue->ProcessAsyncTasks();
+		}
+
+		// Because packages may not have been resaved after this flagging was implemented, we may have added packages to load that weren't flagged - potential false positives.
+		// The loading process should have reflagged said packages so that only true positives will have this flag.
+		if (Package->RequiresLocalizationGather())
+		{
+			UE_LOG(LogGatherTextFromAssetsCommandlet, Display, TEXT("[%6.2f%%] Gathering package: '%s'..."), PercentageComplete, *PackageNameStr);
+
+			// Gathers from the given package
+			EPropertyLocalizationGathererResultFlags GatherableTextResultFlags = EPropertyLocalizationGathererResultFlags::Empty;
+			FPropertyLocalizationDataGatherer(GatherableTextDataArray, Package, GatherableTextResultFlags);
+
+			bool bSavePackage = false;
+
+			// Optionally check to see whether the clean gather we did is in-sync with the gather cache and deal with it accordingly
+			if ((bReportStaleGatherCache || bFixStaleGatherCache) && PackagePendingGather.PackageLocCacheState == EPackageLocCacheState::Cached)
+			{
+				// Look for any structurally significant changes (missing, added, or changed texts) in the cache
+				// Ignore insignificant things (like source changes caused by assets moving or being renamed)
+				if (EnumHasAnyFlags(GatherableTextResultFlags, EPropertyLocalizationGathererResultFlags::HasTextWithInvalidPackageLocalizationID)
+					|| !IsGatherableTextDataIdentical(GatherableTextDataArray, PackagePendingGather.GatherableTextDataArray))
+				{
+					PackagesWithStaleGatherCache.Add(PackagePendingGather.PackageName);
+
+					if (bFixStaleGatherCache)
+					{
+						bSavePackage = true;
+					}
+				}
+			}
+
+			// Optionally save the package if it is missing a gather cache
+			if (bFixMissingGatherCache && PackagePendingGather.PackageLocCacheState == EPackageLocCacheState::Uncached_TooOld)
+			{
+				bSavePackage = true;
+			}
+
+			// Re-save the package to attempt to fix it?
+			if (bSavePackage)
+			{
+				UE_LOG(LogGatherTextFromAssetsCommandlet, Display, TEXT("Resaving package: '%s'..."), *PackageNameStr);
+
+				bool bSavedPackage = false;
+				{
+					FLoadPackageLogOutputRedirector::FScopedCapture ScopedCapture(&LogOutputRedirector, PackageNameStr);
+					bSavedPackage = FLocalizedAssetSCCUtil::SavePackageWithSCC(SourceControlInfo, Package, PackagePendingGather.PackageFilename);
+				}
+
+				if (!bSavedPackage)
+				{
+					UE_LOG(LogGatherTextFromAssetsCommandlet, Warning, TEXT("Failed to resave package: '%s'."), *PackageNameStr);
+				}
+			}
+
+			// This package may have already been cached in cases where we're reporting or fixing assets with a stale gather cache
+			// This check prevents it being gathered a second time
+			if (PackagePendingGather.PackageLocCacheState != EPackageLocCacheState::Cached)
+			{
+				ProcessGatherableTextDataArray(GatherableTextDataArray);
+			}
+
+			GatherableTextDataArray.Reset();
+		}
+	#endif
+}
+
+
+void UGenerateModerationArtifactsCommandlet::PerformAdditionalOperations(class UObject* Object, bool& bSavePackage)
+{
+	bSavePackage = false;
+
+	GatherFStringsFromObject(Object);
+
+	if (Object->GetClass()->IsChildOf(UTexture::StaticClass()))
+	{
+		GenerateArtifact(StaticCast<UTexture*>(Object));
+	}
+	else if (Object->GetClass()->IsChildOf(UStaticMeshComponent::StaticClass()))
+	{
+		GenerateArtifact(StaticCast<UStaticMeshComponent*>(Object));
+	}
+#if 0 // todo add suppport for datatables 
+	else if (Object->GetClass()->IsChildOf(UDataTable::StaticClass()))
+	{
+		GenerateArtifact(StaticCast<UDataTable*>(Object));
+	}
+#endif
+}
+
+FString UGenerateModerationArtifactsCommandlet::CreateOutputFileName(class UObject* Object, const FString& Extension)
+{
+	UPackage* Package = Object->GetOutermost();
+	const FPackagePath& PackagePath = Package->GetLoadedPath();
+	
+	FMD5Hash FileHash = FMD5Hash::HashFile(*PackagePath.GetLocalFullPath());
+	FString FileName = FString::Printf(TEXT("%s-%s.%s"), *Object->GetClass()->GetName(), *LexToString(FileHash), *Extension);
+	FString FullPath = FPaths::Combine(OutputPath, FileName);
+	UE_LOG(LogModerationArtifactsCommandlet, Display, TEXT("Created moderation file %s for asset %s"), *FullPath, *Object->GetPathName());
+	return FullPath;
+}
+
+void UGenerateModerationArtifactsCommandlet::GenerateArtifact(UTexture* Texture)
+{
+	UE_LOG(LogModerationArtifactsCommandlet, Display, TEXT("Found texture %s"), *Texture->GetFullName());
+
+	if (Texture->Source.IsValid())
+	{
+		FString OutputFileName = CreateOutputFileName(Texture, TEXT("png"));
+
+
+		IImageWrapperModule& ImageWrapperModule = FModuleManager::LoadModuleChecked<IImageWrapperModule>(FName("ImageWrapper"));
+		TSharedPtr<IImageWrapper> ImageWrapper = ImageWrapperModule.CreateImageWrapper(EImageFormat::PNG);
+
+
+
+		FImage Image;
+		if (ImageWrapper.IsValid() && Texture->Source.GetMipImage(Image, 0))
+		{
+			ERGBFormat RGBFormat;
+			int BitsPerChannel;
+			switch (Image.Format)
+			{
+			case ERawImageFormat::G8:
+				RGBFormat = ERGBFormat::Gray;
+				BitsPerChannel = 8;
+				break;
+			case ERawImageFormat::BGRA8:
+				RGBFormat = ERGBFormat::BGRA;
+				BitsPerChannel = 8;
+				break;
+			case ERawImageFormat::BGRE8:
+				RGBFormat = ERGBFormat::BGRE;
+				BitsPerChannel = 8;
+				break;
+			case ERawImageFormat::RGBA16:
+				RGBFormat = ERGBFormat::RGBA;
+				BitsPerChannel = 16;
+				break;
+			case ERawImageFormat::RGBA16F:
+				RGBFormat = ERGBFormat::RGBAF;
+				BitsPerChannel = 16;
+				break;
+			case ERawImageFormat::RGBA32F:
+				RGBFormat = ERGBFormat::RGBAF;
+				BitsPerChannel = 32;
+				break;
+			case ERawImageFormat::G16:
+				RGBFormat = ERGBFormat::Gray;
+				BitsPerChannel = 16;
+				break;
+			case ERawImageFormat::R16F:
+				RGBFormat = ERGBFormat::GrayF;
+				BitsPerChannel = 16;
+				break;
+			case ERawImageFormat::R32F:
+				RGBFormat = ERGBFormat::GrayF;
+				BitsPerChannel = 32;
+				break;
+			default:
+				UE_LOG(LogModerationArtifactsCommandlet, Display, TEXT("Texture %s source image format %s is unsupported"), *Texture->GetFullName(), ERawImageFormat::GetName(Image.Format));
+				return;
+			}
+			
+			if (ImageWrapper.IsValid() && ImageWrapper->SetRaw(Image.RawData.GetData(), Image.RawData.Num(), Image.GetWidth(), Image.GetHeight(), RGBFormat, BitsPerChannel))
+			{
+				EImageCompressionQuality PngQuality = EImageCompressionQuality::Default; // 0 means default 
+				TArray64<uint8> CompressedData = ImageWrapper->GetCompressed((int32)PngQuality);
+				if (CompressedData.Num() > 0)
+				{
+					FFileHelper::SaveArrayToFile(CompressedData, *OutputFileName);
+				}
+			}
+		}
+
+	}
+	/*
+
+	ImageWrapper->SetRaw(SurfaceData.GetData(), SurfaceData.GetAllocatedSize(), UnprojectedAtlasWidth, UnprojectedAtlasHeight, ERGBFormat::BGRA, 32);
+	const TArray64<uint8> PNGDataUnprojected = ImageWrapper->GetCompressed(100);
+	FFileHelper::SaveArrayToFile(PNGDataUnprojected, *AtlasNameUnprojected);
+	ImageWrapper.Reset();*/
+	
+}
+
+/*
+void UGenerateModerationArtifactsCommandlet::GenerateArtifact(UDataTable* DataTable)
+{
+	check(0);  // not implemented yet
+	// FString UDataTable::GetTableAsString(const EDataTableExportFlags InDTExportFlags) const
+}*/
+void UGenerateModerationArtifactsCommandlet::GenerateArtifact(UStaticMeshComponent* StaticMesh)
+{
+	UE_LOG(LogModerationArtifactsCommandlet, Display, TEXT("Found staticmesh %s"), *StaticMesh->GetFullName());
+}
\ No newline at end of file
