HLOD: Optimized HLOD cluster generation
* Leverage the fact that if a cluster becomes invalid it will never become valid again and saves all valid merged clusters in a separate array to speed up processing
* Algo went from O(n) to O(log n)... In a test map, it yields a 100x speed up depending on the clustering settings
* Removed iteration bound (was max 3) since clustering will converge a lot faster than it did before

diff --git a/Engine/Source/Editor/UnrealEd/Private/HierarchicalLOD.cpp b/Engine/Source/Editor/UnrealEd/Private/HierarchicalLOD.cpp
--- a/Engine/Source/Editor/UnrealEd/Private/HierarchicalLOD.cpp
+++ b/Engine/Source/Editor/UnrealEd/Private/HierarchicalLOD.cpp
@@ -142,6 +142,8 @@ void FHierarchicalLODBuilder::PreviewBuild()
 
 void FHierarchicalLODBuilder::BuildClusters(ULevel* InLevel)
 {	
+	TRACE_CPUPROFILER_EVENT_SCOPE(FHierarchicalLODBuilder::BuildClusters);
+
 	SCOPE_LOG_TIME(TEXT("STAT_HLOD_BuildClusters"), nullptr);
 
 	const TArray<FHierarchicalSimplification>& BuildLODLevelSettings = InLevel->GetWorldSettings()->GetHierarchicalLODSetup();
@@ -587,7 +589,7 @@ void FHierarchicalLODBuilder::InitializeClusters(ULevel* InLevel, const int32 LO
 
 						if (NewClusterCost <= CullCost)
 						{
-							Clusters.Add(NewClusterCandidate);
+							Clusters.Add(MoveTemp(NewClusterCandidate));
 						}
 					}
 				}
@@ -1142,6 +1144,8 @@ void FHierarchicalLODBuilder::DeleteLODActors(ULevel* InLevel)
 
 void FHierarchicalLODBuilder::BuildMeshForLODActor(ALODActor* LODActor, const uint32 LODLevel)
 {
+	TRACE_CPUPROFILER_EVENT_SCOPE(FHierarchicalLODBuilder::BuildMeshForLODActor);
+
 	const TArray<FHierarchicalSimplification>& BuildLODLevelSettings = LODActor->GetLevel()->GetWorldSettings()->GetHierarchicalLODSetup();
 	UMaterialInterface* BaseMaterial = LODActor->GetLevel()->GetWorldSettings()->GetHierarchicalLODBaseMaterial();
 	
@@ -1163,6 +1167,8 @@ void FHierarchicalLODBuilder::BuildMeshForLODActor(ALODActor* LODActor, const ui
 
 void FHierarchicalLODBuilder::MergeClustersAndBuildActors(ULevel* InLevel, const int32 LODIdx, float HighestCost, int32 MinNumActors)
 {
+	TRACE_CPUPROFILER_EVENT_SCOPE(FHierarchicalLODBuilder::MergeClustersAndBuildActors);
+
 	if (Clusters.Num() > 0 || HLODVolumeClusters.Num() > 0)
 	{
 		FString LevelName = FPackageName::GetShortName(InLevel->GetOutermost()->GetName());
@@ -1172,78 +1178,80 @@ void FHierarchicalLODBuilder::MergeClustersAndBuildActors(ULevel* InLevel, const
 		// merge clusters first
 		{
 			SCOPE_LOG_TIME(TEXT("HLOD_MergeClusters"), nullptr);
-			static int32 TotalIteration = 3;
-			const int32 TotalCluster = Clusters.Num();
+			bool bStable = false;
 
-			FScopedSlowTask SlowTask(100.0f, FText::Format(LOCTEXT("HierarchicalLOD_BuildClusters", "Building Clusters for LOD {LODIndex} of {LevelName}..."), Arguments));
-			SlowTask.MakeDialog();
-
-			for (int32 Iteration = 0; Iteration < TotalIteration; ++Iteration)
+			while (!bStable)
 			{
-				bool bChanged = false;
+				const int32 NumClusters = Clusters.Num();
+
+				FScopedSlowTask SlowTask(NumClusters, FText::Format(LOCTEXT("HierarchicalLOD_BuildClusters", "Building Clusters for LOD {LODIndex} of {LevelName}..."), Arguments));
+				SlowTask.MakeDialog();
+
+				TArray<FLODCluster> ValidMergedClusters;
+				ValidMergedClusters.Reserve(NumClusters);
+
+				bStable = true;
+
 				// now we have minimum Clusters
-				for (int32 ClusterId = 0; ClusterId < TotalCluster; ++ClusterId)
+				for (int32 ClusterId = 0; ClusterId < NumClusters; ++ClusterId)
 				{
+					SlowTask.EnterProgressFrame(1.0f);
+
 					FLODCluster& Cluster = Clusters[ClusterId];
 					UE_LOG(LogLODGenerator, Verbose, TEXT("%d. %0.2f {%s}"), ClusterId + 1, Cluster.GetCost(), *Cluster.ToString());
 
-					// progress bar update every percent, if ClustersPerPercent is zero ignore the progress bar as number of iterations is small.
-					int32 ClustersPerPercent = (TotalCluster / (100.0f / TotalIteration));
-					if ( ClustersPerPercent > 0 && ClusterId % ClustersPerPercent == 0) {
-						SlowTask.EnterProgressFrame(1.0f);
-					}
-
 					if (Cluster.IsValid())
 					{
-						for (int32 MergedClusterId = 0; MergedClusterId < ClusterId; ++MergedClusterId)
+						// compare with previous valid clusters
+						for (FLODCluster& MergedCluster : ValidMergedClusters)
 						{
-							// compare with previous clusters
-							FLODCluster& MergedCluster = Clusters[MergedClusterId];
-							// see if it's valid, if it contains, check the cost
-							if (MergedCluster.IsValid())
+							check(MergedCluster.IsValid())
+
+							// if valid, see if it contains any of this actors
+							if (MergedCluster.Contains(Cluster))
 							{
-								if (MergedCluster.Contains(Cluster))
-								{
-									// if valid, see if it contains any of this actors
-									// merge whole clusters
-									FLODCluster NewCluster = Cluster + MergedCluster;
-									float MergeCost = NewCluster.GetCost();
+								float MergeCost = MergedCluster.GetMergedCost(Cluster);
 
-									// merge two clusters
-									if (MergeCost <= HighestCost)
-									{
-										UE_LOG(LogLODGenerator, Log, TEXT("Merging of Cluster (%d) and (%d) with merge cost (%0.2f) "), ClusterId + 1, MergedClusterId + 1, MergeCost);
+								// merge two clusters
+								if (MergeCost <= HighestCost)
+								{
+									MergedCluster += Cluster;
+									// now this cluster is invalid
+									Cluster.Invalidate();
 
-										MergedCluster = NewCluster;
-										// now this cluster is invalid
-										Cluster.Invalidate();
+									bStable = false;
+									break;
+								}
+								else
+								{
+									Cluster -= MergedCluster;
+									bStable = false;
 
-										bChanged = true;
-										break;
-									}
-									else
+									if (!Cluster.IsValid())
 									{
-										Cluster -= MergedCluster;
-										bChanged = true;
+										// If the cluster becomes invalid, MergedCluster.Contains() will always return false, so exit immediately
+										break;
 									}
 								}
 							}
 						}
 
+						if (Cluster.IsValid())
+						{
+							ValidMergedClusters.Add(Cluster);
+						}
+
 						UE_LOG(LogLODGenerator, Verbose, TEXT("Processed(%s): %0.2f {%s}"), Cluster.IsValid() ? TEXT("Valid") : TEXT("Invalid"), Cluster.GetCost(), *Cluster.ToString());
 					}
 				}
 
-				if (bChanged == false)
-				{
-					break;
-				}
+				Clusters = ValidMergedClusters;
 			}
 		}
 
 		for (TPair<AHierarchicalLODVolume*, FLODCluster>& Cluster : HLODVolumeClusters)
 		{
-			Clusters.Add(Cluster.Value);
+			Clusters.Add(MoveTemp(Cluster.Value));
 		}
 
 
@@ -1292,7 +1300,7 @@ void FHierarchicalLODBuilder::MergeClustersAndBuildActors(ULevel* InLevel, const
 					{
 						// Add them to the RejectedActorsInLevel to be re-considered at the next LOD
 						ValidStaticMeshActorsInLevel.RemoveAllSwap([&Cluster](const AActor* Actor) { return Cluster.Actors.Contains(Actor); });
-						RejectedActorsInLevel.Append(Cluster.Actors);
+						RejectedActorsInLevel.Append(Cluster.Actors.Array());
 					}
 				}
 			}
diff --git a/Engine/Source/Editor/UnrealEd/Private/LODCluster.cpp b/Engine/Source/Editor/UnrealEd/Private/LODCluster.cpp
--- a/Engine/Source/Editor/UnrealEd/Private/LODCluster.cpp
+++ b/Engine/Source/Editor/UnrealEd/Private/LODCluster.cpp
@@ -99,6 +99,16 @@ FLODCluster::FLODCluster(const FLODCluster& Other)
 	
 }
 
+FLODCluster::FLODCluster(FLODCluster&& Other)
+	: Actors(Other.Actors)
+	, Bound(Other.Bound)
+	, FillingFactor(Other.FillingFactor)
+	, ClusterCost(Other.ClusterCost)
+	, bValid(Other.bValid)
+{
+
+}
+
 FLODCluster::FLODCluster(AActor* Actor1)
 : Bound(ForceInit)
 , bValid(true)
@@ -132,7 +142,7 @@ FLODCluster::FLODCluster()
 FSphere FLODCluster::AddActor(AActor* NewActor)
 {
 	bValid = true;
-	ensure (Actors.Contains(NewActor) == false);
+
 	Actors.Add(NewActor);
 	FVector Origin, Extent;
 
@@ -182,9 +192,29 @@ FLODCluster& FLODCluster::operator=(const FLODCluster& Other)
 	return *this;
 }
 
+FLODCluster& FLODCluster::operator=(FLODCluster&& Other)
+{
+	this->bValid = Other.bValid;
+	this->Actors = Other.Actors;
+	this->Bound = Other.Bound;
+	this->FillingFactor = Other.FillingFactor;
+	this->ClusterCost = Other.ClusterCost;
+
+	return *this;
+}
+
 bool FLODCluster::operator==(const FLODCluster& Other) const
 {
-	return Actors == Other.Actors;
+	return Actors.Num() == Other.Actors.Num() && Actors.Includes(Other.Actors);
+}
+
+float FLODCluster::GetMergedCost(const FLODCluster& Other) const
+{
+	float MergedFillingFactor = CalculateFillingFactor(Bound, FillingFactor, Other.Bound, Other.FillingFactor);
+	FSphere MergedBound = Bound + Other.Bound;
+
+	float MergedClusterCost = (MergedBound.W * MergedBound.W * MergedBound.W) / MergedFillingFactor;
+	return MergedClusterCost;
 }
 
 void FLODCluster::MergeClusters(const FLODCluster& Other)
@@ -196,11 +226,7 @@ void FLODCluster::MergeClusters(const FLODCluster& Other)
 
 	ClusterCost = ( Bound.W * Bound.W * Bound.W ) / FillingFactor;
 	
-
-	for (auto& Actor: Other.Actors)
-	{
-		Actors.AddUnique(Actor);
-	}
+	Actors.Append(Other.Actors);
 
 	if (Actors.Num() > 0)
 	{
@@ -210,51 +236,30 @@ void FLODCluster::MergeClusters(const FLODCluster& Other)
 
 void FLODCluster::SubtractCluster(const FLODCluster& Other)
 {
-	for(int32 ActorId=0; ActorId<Actors.Num(); ++ActorId)
-	{
-		if (Other.Actors.Contains(Actors[ActorId]))
-		{
-			Actors.RemoveAt(ActorId);
-			--ActorId;
-		}
-	}
+	Actors = Actors.Difference(Other.Actors);
 
-	TArray<AActor*> NewActors = Actors;
-	Actors.Empty();
-	// need to recalculate parameter
-	if (NewActors.Num() == 0)
-	{
-		Invalidate();
-	}
-	else if (NewActors.Num() == 1)
+	Invalidate();
+
+	// We need to recalculate parameters
+	if (Actors.Num() > 0)
 	{
-		Bound = FSphere(ForceInitToZero);
-		AddActor(NewActors[0]);
+		bValid = true;
 		FillingFactor = 1.f;
-		ClusterCost = ( Bound.W * Bound.W * Bound.W ) / FillingFactor;
-	}
-	else if (NewActors.Num() >= 2)
-	{
-		Bound = FSphere(ForceInit);
-
-		FSphere Actor1Bound = AddActor(NewActors[0]);
-		FSphere Actor2Bound = AddActor(NewActors[1]);
-
-		// calculate new filling factor
-		FillingFactor = CalculateFillingFactor(Actor1Bound, 1.f, Actor2Bound, 1.f);
+		Bound = FSphere(ForceInitToZero);
 
-		// if more actors, we add them manually
-		for (int32 ActorId=2; ActorId<NewActors.Num(); ++ActorId)
+		for (AActor* Actor : Actors)
 		{
-			// if not contained, it shouldn't be
-			check (!Actors.Contains(NewActors[ActorId]));
+			FVector Origin, Extent;
+			Actor->GetActorBounds(false, Origin, Extent);
+
+			// scale 0.01 (change to meter from centimeter)
+			FSphere NewBound = FSphere(Origin * CM_TO_METER, Extent.Size() * CM_TO_METER);
 
-			FSphere NewBound = AddActor(NewActors[ActorId]);
 			FillingFactor = CalculateFillingFactor(NewBound, 1.f, Bound, FillingFactor);
 			Bound += NewBound;
 		}
 
-		ClusterCost = ( Bound.W * Bound.W * Bound.W ) / FillingFactor;
+		ClusterCost = (Bound.W * Bound.W * Bound.W) / FillingFactor;
 	}
 }
 
diff --git a/Engine/Source/Editor/UnrealEd/Public/LODCluster.h b/Engine/Source/Editor/UnrealEd/Public/LODCluster.h
--- a/Engine/Source/Editor/UnrealEd/Public/LODCluster.h
+++ b/Engine/Source/Editor/UnrealEd/Public/LODCluster.h
@@ -3,6 +3,7 @@
 #pragma once
 
 #include "CoreMinimal.h"
+#include "Containers/Set.h"
 
 class AActor;
 class ALODActor;
@@ -23,6 +24,8 @@ struct FLODCluster
 {
 	// constructors
 	FLODCluster(const FLODCluster& Other);
+	FLODCluster(FLODCluster&& Other);
+
 	FLODCluster(AActor* Actor1);
 	FLODCluster(AActor* Actor1, AActor* Actor2);
 	FLODCluster();
@@ -33,6 +36,7 @@ struct FLODCluster
 	FLODCluster operator-(const FLODCluster& Other) const;
 	FLODCluster& operator-=(const FLODCluster& Other);
 	FLODCluster& operator=(const FLODCluster & Other);
+	FLODCluster& operator=(FLODCluster&& Other);
 
 	bool operator==(const FLODCluster& Other) const;
 
@@ -48,6 +52,9 @@ struct FLODCluster
 		return ClusterCost;
 	}
 
+	/** Return cost of the union of this cluster & the other cluster, lower is better */
+	float GetMergedCost(const FLODCluster& Other) const;
+
 	/** Compare clusters and returns true when this contains any of Other's actors */
 	bool Contains(FLODCluster& Other) const;
 	
@@ -56,7 +63,7 @@ struct FLODCluster
 	
 	// member variable
 	/** List of Actors that this cluster contains */
-	TArray<AActor*> Actors;
+	TSet<AActor*, DefaultKeyFuncs<AActor*>, TInlineSetAllocator<2>> Actors;
 	/** Cluster bounds */
 	FSphere	Bound;
 	/** Filling factor for this cluster, determines how much of the cluster's bounds/area is occupied by the contained actors*/
