SourceControl: Don't perform unnecessary UpdateStatus operation (PackageFilesToDelete might have become empty in the code just preceding this change).

diff --git a/Engine/Source/Editor/UnrealEd/Private/ObjectTools.cpp b/Engine/Source/Editor/UnrealEd/Private/ObjectTools.cpp
--- a/Engine/Source/Editor/UnrealEd/Private/ObjectTools.cpp
+++ b/Engine/Source/Editor/UnrealEd/Private/ObjectTools.cpp
@@ -2323,7 +2323,7 @@ namespace ObjectTools
 		}
 
 		// Get the current source control states of all the package files we're deleting at once.
-		if (ISourceControlModule::Get().IsEnabled())
+		if ( PackagesToDelete.Num() && ISourceControlModule::Get().IsEnabled() )
 		{
 			SourceControlProvider.GetState(PackageFilesToDelete, PackageSCCStates, EStateCacheUsage::ForceUpdate);
 		}
