UE-161875:
//UE5/Main - Build Engine Localization - Warning: Text conflict from LOCTEXT macro in EditorStyleSettingsCustomization.cpp

diff --git a/Engine/Source/Editor/UnrealEd/Private/Settings/EditorStyleSettingsCustomization.cpp b/Engine/Source/Editor/UnrealEd/Private/Settings/EditorStyleSettingsCustomization.cpp
--- a/Engine/Source/Editor/UnrealEd/Private/Settings/EditorStyleSettingsCustomization.cpp
+++ b/Engine/Source/Editor/UnrealEd/Private/Settings/EditorStyleSettingsCustomization.cpp
@@ -625,7 +625,7 @@ FReply FEditorStyleSettingsCustomization::OnImportThemeClicked()
 		// if theme name exists, don't import (to prevent from overwriting existing theme files)
 		if (!IsThemeNameValid(FilenameWithoutExtension))
 		{
-			ShowNotification(LOCTEXT("ImportThemeFailure", "Import theme failed: Theme name already exists"), SNotificationItem::CS_Fail);
+			ShowNotification(LOCTEXT("ImportThemeFailureNameExists", "Import theme failed: Theme name already exists"), SNotificationItem::CS_Fail);
 		}
 		// if theme name is valid: copying the file is safe (as it will not overwrite existing theme files)
 		else 
@@ -656,7 +656,7 @@ FReply FEditorStyleSettingsCustomization::OnImportThemeClicked()
 				{
 					// incomplete themes will not reach here. 
 					IPlatformFile::GetPlatformPhysical().DeleteFile(*DestPath); 
-					ShowNotification(LOCTEXT("ImportThemeFailure", "Import theme failed: Invalid theme"), SNotificationItem::CS_Fail);
+					ShowNotification(LOCTEXT("ImportThemeFailureInvalidName", "Import theme failed: Invalid theme"), SNotificationItem::CS_Fail);
 				}
 			}
 			// if unable to copy the file to user-specific theme location, do nothing. 
