Fix saving packages progress dialog in the editor going too fast and add progress dialogs for updating status and checking out files from source control.

#rnx

diff --git a/Engine/Source/Editor/UnrealEd/Private/FileHelpers.cpp b/Engine/Source/Editor/UnrealEd/Private/FileHelpers.cpp
--- a/Engine/Source/Editor/UnrealEd/Private/FileHelpers.cpp
+++ b/Engine/Source/Editor/UnrealEd/Private/FileHelpers.cpp
@@ -1581,7 +1581,10 @@ bool FEditorFileUtils::AddCheckoutPackageItems(bool bCheckDirty, TArray<UPackage
 		if (SourceControlCheckPackages.Num())
 		{
 			// Update the source control status of all potentially relevant packages
+			FScopedSlowTask SlowTask(SourceControlCheckPackages.Num(), LOCTEXT("UpdatingSourceControlStatus", "Updating source control status..."));
+			SlowTask.MakeDialog();
 			SourceControlProvider.Execute(ISourceControlOperation::Create<FUpdateStatus>(), SourceControlCheckPackages);
+			SlowTask.EnterProgressFrame(SourceControlCheckPackages.Num());
 		}
 	}
 
@@ -2073,7 +2076,10 @@ ECommandResult::Type FEditorFileUtils::CheckoutPackages(const TArray<UPackage*>&
 	// Attempt to check out each package the user specified to be checked out that is not read only
 	if(FinalPackageCheckoutList.Num() > 0)
 	{
+		FScopedSlowTask SlowTask(FinalPackageCheckoutList.Num(), LOCTEXT("CheckingOutPackages", "Checking out packages..."));
+		SlowTask.MakeDialog();
 		CheckOutResult = SourceControlProvider.Execute(ISourceControlOperation::Create<FCheckOut>(), FinalPackageCheckoutList);
+		SlowTask.EnterProgressFrame(FinalPackageCheckoutList.Num());
 	}
 
 	// Attempt to mark for add each package the user specified that is not already tracked by source control
@@ -3974,15 +3980,6 @@ FEditorFileUtils::EPromptReturnCode InternalPromptForCheckoutAndSave(const TArra
 				Package->FullyLoad();
 			}
 
-			const UWorld* const AssociatedWorld = UWorld::FindWorldInPackage(Package);
-			const bool bIsMapPackage = AssociatedWorld != nullptr;
-
-			const FText SavingPackageText = (bIsMapPackage)
-				? FText::Format(NSLOCTEXT("UnrealEd", "SavingMapf", "Saving map {0}"), FText::FromString(Package->GetName()))
-				: FText::Format(NSLOCTEXT("UnrealEd", "SavingAssetf", "Saving asset {0}"), FText::FromString(Package->GetName()));
-
-			SlowTask.EnterProgressFrame(1, SavingPackageText);
-
 			// if the package we are saving is considered empty, mark it for deletion on disk instead
 			if (UPackage::IsEmptyPackage(Package))
 			{
@@ -3998,10 +3995,20 @@ FEditorFileUtils::EPromptReturnCode InternalPromptForCheckoutAndSave(const TArra
 		if (PackagesToClean.Num() > 0)
 		{
 			ObjectTools::CleanupAfterSuccessfulDelete(PackagesToClean, true);
+			SlowTask.EnterProgressFrame(PackagesToClean.Num());
 		}
 
 		for (UPackage* Package : PackagesToSave)
 		{
+			const UWorld* const AssociatedWorld = UWorld::FindWorldInPackage(Package);
+			const bool bIsMapPackage = AssociatedWorld != nullptr;
+
+			const FText SavingPackageText = (bIsMapPackage)
+				? FText::Format(NSLOCTEXT("UnrealEd", "SavingMapf", "Saving map {0}"), FText::FromString(Package->GetName()))
+				: FText::Format(NSLOCTEXT("UnrealEd", "SavingAssetf", "Saving asset {0}"), FText::FromString(Package->GetName()));
+
+			SlowTask.EnterProgressFrame(1, SavingPackageText);
+
 			// Save the package
 			bool bPackageLocallyWritable;
 			const InternalSavePackageResult SaveStatus = InternalSavePackage(Package, bUseDialog, bPackageLocallyWritable, SaveErrors);
