Bring back CL21316192:

Original CL Desc
-----------------------------------------------------------------
Bulk reset loaders when saving actor files to avoid going through all imports of all packages once per saved actor.
Saves several seconds when saving lots of actor files.

#rnx

diff --git a/Engine/Source/Editor/UnrealEd/Private/FileHelpers.cpp b/Engine/Source/Editor/UnrealEd/Private/FileHelpers.cpp
--- a/Engine/Source/Editor/UnrealEd/Private/FileHelpers.cpp
+++ b/Engine/Source/Editor/UnrealEd/Private/FileHelpers.cpp
@@ -16,6 +16,7 @@
 #include "Modules/ModuleManager.h"
 #include "UObject/UObjectHash.h"
 #include "UObject/UObjectIterator.h"
+#include "ProfilingDebugging/ScopedTimers.h"
 #include "Misc/Attribute.h"
 #include "Widgets/DeclarativeSyntaxSupport.h"
 #include "Styling/SlateTypes.h"
@@ -3167,6 +3168,11 @@ enum class InternalSavePackageResult : int8
 	Error,
 };
 
+static void PrepareSavePackages(const TArray<UPackage*>& PackagesToSave)
+{
+	ResetLoaders(MakeArrayView<UObject*>((UObject**)PackagesToSave.GetData(), PackagesToSave.Num()));
+}
+
 /**
  * Actually save a package. Prompting for Save as if necessary
  *
@@ -3550,6 +3556,7 @@ static void InternalNotifyNoPackagesSaved(const bool bUseDialog)
 static bool InternalSavePackagesFast(const TArray<UPackage*>& PackagesToSave, bool bUseDialog, TArray<UPackage*>& OutFailedPackages)
 {
 	TRACE_CPUPROFILER_EVENT_SCOPE(InternalSavePackagesFast);
+	UE_SCOPED_TIMER(TEXT("InternalSavePackagesFast"), LogFileHelpers, Log);
 
 	bool bReturnCode = true;
 
@@ -3612,6 +3619,8 @@ static bool InternalSavePackagesFast(const TArray<UPackage*>& PackagesToSave, bo
 		ObjectTools::CleanupAfterSuccessfulDelete(PackagesToClean, true);
 	}
 
+	PrepareSavePackages(FinalPackagesToSave);
+
 	for (UPackage* Package : FinalPackagesToSave)
 	{
 		bool bPackageLocallyWritable;
@@ -3942,6 +3951,8 @@ bool FEditorFileUtils::SaveCurrentLevel()
  */
 FEditorFileUtils::EPromptReturnCode InternalPromptForCheckoutAndSave(const TArray<UPackage*>& FinalSaveList, bool bUseDialog, TArray<UPackage*>& OutFailedPackages)
 {
+	UE_SCOPED_TIMER(TEXT("InternalPromptForCheckoutAndSave"), LogFileHelpers, Log);
+
 	FEditorFileUtils::EPromptReturnCode ReturnResponse = FEditorFileUtils::PR_Success;
 	const FScopedBusyCursor BusyCursor;
 	FSaveErrorOutputDevice SaveErrors;
@@ -3998,6 +4009,8 @@ FEditorFileUtils::EPromptReturnCode InternalPromptForCheckoutAndSave(const TArra
 			SlowTask.EnterProgressFrame(PackagesToClean.Num());
 		}
 
+		PrepareSavePackages(PackagesToSave);
+
 		for (UPackage* Package : PackagesToSave)
 		{
 			const UWorld* const AssociatedWorld = UWorld::FindWorldInPackage(Package);
diff --git a/Engine/Source/Runtime/CoreUObject/Private/UObject/Linker.cpp b/Engine/Source/Runtime/CoreUObject/Private/UObject/Linker.cpp
--- a/Engine/Source/Runtime/CoreUObject/Private/UObject/Linker.cpp
+++ b/Engine/Source/Runtime/CoreUObject/Private/UObject/Linker.cpp
@@ -499,6 +499,18 @@ void ResetLoaders(UObject* InPkg)
 	FLinkerManager::Get().ResetLoaders(InPkg);
 }
 
+void ResetLoaders(TArrayView<UObject*> InOuters)
+{
+	if (IsAsyncLoading())
+	{
+		UE_LOG(LogLinker, Log, TEXT("ResetLoaders is flushing async loading"));
+	}
+
+	// Make sure we're not in the middle of loading something in the background.
+	FlushAsyncLoading();
+	FLinkerManager::Get().ResetLoaders(InOuters);
+}
+
 void DeleteLoaders()
 {
 	FLinkerManager::Get().DeleteLinkers();
diff --git a/Engine/Source/Runtime/CoreUObject/Private/UObject/LinkerManager.cpp b/Engine/Source/Runtime/CoreUObject/Private/UObject/LinkerManager.cpp
--- a/Engine/Source/Runtime/CoreUObject/Private/UObject/LinkerManager.cpp
+++ b/Engine/Source/Runtime/CoreUObject/Private/UObject/LinkerManager.cpp
@@ -167,6 +167,25 @@ void FLinkerManager::ResetLoaders(UObject* InPkg)
 	}
 }
 
+void FLinkerManager::ResetLoaders(TArrayView<UObject*> InPackages)
+{
+	TSet<FLinkerLoad*> LinkerLoads;
+	LinkerLoads.Reserve(InPackages.Num());
+
+	for (UObject* Object : InPackages)
+	{
+		if (UPackage* TopLevelPackage = Object->GetPackage())
+		{
+			if (FLinkerLoad* LinkerToReset = FLinkerLoad::FindExistingLinkerForPackage(TopLevelPackage))
+			{
+				LinkerLoads.Add(LinkerToReset);
+			}
+		}
+	}
+
+	ResetLoaders(LinkerLoads);
+}
+
 void FLinkerManager::ResetLoaders(const TSet<FLinkerLoad*>& InLinkerLoads)
 {
 	// Remove import references
diff --git a/Engine/Source/Runtime/CoreUObject/Private/UObject/LinkerManager.h b/Engine/Source/Runtime/CoreUObject/Private/UObject/LinkerManager.h
--- a/Engine/Source/Runtime/CoreUObject/Private/UObject/LinkerManager.h
+++ b/Engine/Source/Runtime/CoreUObject/Private/UObject/LinkerManager.h
@@ -116,6 +116,9 @@ public:
 	/** Empty the loaders */
 	void ResetLoaders(UObject* InPkg);
 
+	/** Empty the loaders */
+	void ResetLoaders(TArrayView<UObject*> InPackages);
+
 	/** Empty the loaders from the specified set */
 	void ResetLoaders(const TSet<FLinkerLoad*>& InLinkerLoadSet);
 
diff --git a/Engine/Source/Runtime/CoreUObject/Public/UObject/Linker.h b/Engine/Source/Runtime/CoreUObject/Public/UObject/Linker.h
--- a/Engine/Source/Runtime/CoreUObject/Public/UObject/Linker.h
+++ b/Engine/Source/Runtime/CoreUObject/Public/UObject/Linker.h
@@ -697,6 +697,7 @@ COREUOBJECT_API void ResetLinkerExports(UPackage* InPackage);
  * Flushes async loading.
  */
 COREUOBJECT_API void ResetLoaders(UObject* InOuter);
+COREUOBJECT_API void ResetLoaders(TArrayView<UObject*> InOuters);
 
 /** Deletes all linkers that have finished loading */
 COREUOBJECT_API void DeleteLoaders();
