Add Nanite settings to the list of settings to preserve when re-importing a static mesh

diff --git a/Engine/Source/Editor/UnrealEd/Private/ImportUtils/StaticMeshImportUtils.cpp b/Engine/Source/Editor/UnrealEd/Private/ImportUtils/StaticMeshImportUtils.cpp
--- a/Engine/Source/Editor/UnrealEd/Private/ImportUtils/StaticMeshImportUtils.cpp
+++ b/Engine/Source/Editor/UnrealEd/Private/ImportUtils/StaticMeshImportUtils.cpp
@@ -750,6 +750,8 @@ TSharedPtr<FExistingStaticMeshData> StaticMeshImportUtils::SaveExistingStaticMes
 	ExistingMeshDataPtr->ExistingNeverStream = ExistingMesh->NeverStream;
 	ExistingMeshDataPtr->ExistingNumCinematicMipLevels = ExistingMesh->NumCinematicMipLevels;
 
+	ExistingMeshDataPtr->ExistingNaniteSettings = ExistingMesh->NaniteSettings;
+
 	UFbxStaticMeshImportData* ImportData = Cast<UFbxStaticMeshImportData>(ExistingMesh->AssetImportData);
 	if (ImportData && ExistingMeshDataPtr->UseMaterialNameSlotWorkflow)
 	{
@@ -1286,6 +1288,8 @@ void StaticMeshImportUtils::RestoreExistingMeshData(const TSharedPtr<const FExis
 	NewMesh->bGlobalForceMipLevelsToBeResident = ExistingMeshDataPtr->ExistingForceMiplevelsToBeResident;
 	NewMesh->NeverStream = ExistingMeshDataPtr->ExistingNeverStream;
 	NewMesh->NumCinematicMipLevels = ExistingMeshDataPtr->ExistingNumCinematicMipLevels;
+
+	NewMesh->NaniteSettings = ExistingMeshDataPtr->ExistingNaniteSettings;
 }
 
 #undef LOCTEXT_NAMESPACE
diff --git a/Engine/Source/Editor/UnrealEd/Public/ImportUtils/StaticMeshImportUtils.h b/Engine/Source/Editor/UnrealEd/Public/ImportUtils/StaticMeshImportUtils.h
--- a/Engine/Source/Editor/UnrealEd/Public/ImportUtils/StaticMeshImportUtils.h
+++ b/Engine/Source/Editor/UnrealEd/Public/ImportUtils/StaticMeshImportUtils.h
@@ -89,6 +89,8 @@ struct FExistingStaticMeshData
 	bool						ExistingNeverStream;
 	int32						ExistingNumCinematicMipLevels;
 
+	FMeshNaniteSettings			ExistingNaniteSettings;
+
 	UStaticMesh::FOnMeshChanged	ExistingOnMeshChanged;
 	UStaticMesh* ExistingComplexCollisionMesh = nullptr;
 
