Added FDirectoryPath to testprops.

#rnx

diff --git a/Engine/Source/Editor/UnrealEd/Classes/Editor/PropertyEditorTestObject.h b/Engine/Source/Editor/UnrealEd/Classes/Editor/PropertyEditorTestObject.h
--- a/Engine/Source/Editor/UnrealEd/Classes/Editor/PropertyEditorTestObject.h
+++ b/Engine/Source/Editor/UnrealEd/Classes/Editor/PropertyEditorTestObject.h
@@ -707,6 +707,12 @@ class UPropertyEditorTestObject : public UObject
 	UPROPERTY(EditAnywhere, Category=EditCondition, meta=(EditCondition="bEditConditionForFixedArray"))
 	FString FixedArrayWithEditCondition[5];
 
+	UPROPERTY(EditAnywhere, Category=EditCondition)
+	bool bEditConditionForDirectoryPath;
+
+	UPROPERTY(EditAnywhere, Category=EditCondition, meta=(EditCondition="bEditConditionForDirectoryPath"))
+	FDirectoryPath DirectoryPath;
+
 	UPROPERTY(EditAnywhere, Category=EditCondition)
 	int64 EditConditionFlags;
 
@@ -907,4 +913,4 @@ private:
 	TArray<TSharedPtr<IDetailTreeNode>> DetailsNodes;
 	TSharedPtr<SListView<TSharedPtr<IDetailTreeNode>>> ListView;
 	TSharedPtr<IPropertyRowGenerator> PropertyRowGenerator;
-};
\ No newline at end of file
+};
