Mark as garbage and move aside graphs that are orphaned due to no longer being present in the list of interfaces, we could add more logic to avoid losing the graphs, but the existing code does not maintain ownership and instead leaks the interface function name - we want to reclaim that name when orphaning the graph

diff --git a/Engine/Source/Editor/UnrealEd/Private/Kismet2/BlueprintEditorUtils.cpp b/Engine/Source/Editor/UnrealEd/Private/Kismet2/BlueprintEditorUtils.cpp
--- a/Engine/Source/Editor/UnrealEd/Private/Kismet2/BlueprintEditorUtils.cpp
+++ b/Engine/Source/Editor/UnrealEd/Private/Kismet2/BlueprintEditorUtils.cpp
@@ -7424,10 +7424,26 @@ static void ConformInterfaceByName(UBlueprint* Blueprint, FBPInterfaceDescriptio
 		for (int32 GraphIndex = 0; GraphIndex < CurrentInterfaceDesc.Graphs.Num(); GraphIndex++)
 		{
 			// If we can't find the function associated with the graph, delete it
-			const UEdGraph* CurrentGraph = CurrentInterfaceDesc.Graphs[GraphIndex];
+			UEdGraph* CurrentGraph = CurrentInterfaceDesc.Graphs[GraphIndex];
 
 			if (!CurrentGraph || !FindUField<UFunction>(CurrentInterfaceDesc.Interface, CurrentGraph->GetFName()))
 			{
+				if(CurrentGraph)
+				{
+					CurrentGraph->GetSchema()->HandleGraphBeingDeleted(*CurrentGraph);
+
+					// rename to free up the graph's name.. which may be needed by an inherited function
+					CurrentGraph->Rename(
+						nullptr,
+						GetTransientPackage(),
+						(Blueprint->bIsRegeneratingOnLoad ? REN_ForceNoResetLoaders : 0) | REN_DoNotDirty | REN_DontCreateRedirectors);
+					// orphan the graph, we need to mark as garbage so that it's tab cleans up (FTabPayload_UObject)
+					// removing from root, standalone, and public is defensive to make sure it is not saved:
+					CurrentGraph->ClearFlags(RF_Standalone | RF_Public);
+					CurrentGraph->RemoveFromRoot();
+					CurrentGraph->MarkAsGarbage();
+				}
+				
 				CurrentInterfaceDesc.Graphs.RemoveAt(GraphIndex, 1);
 				GraphIndex--;
 			}
