[Core] changes to make lazyloadimports performant
changed lazy load back to looking at meta data to decide which types are lazyloaded
removed callback that handles lazyload can add back in the future if needed
only imports that marked as lazy are skipped. verifyimports still resolves other imports
skipping all imports causes performance problems because PackageLoadInternal is called many more times then necessary

#p4v-preflight-copy 21567032

diff --git a/Engine/Source/Editor/UnrealEd/Classes/CookOnTheSide/CookOnTheFlyServer.h b/Engine/Source/Editor/UnrealEd/Classes/CookOnTheSide/CookOnTheFlyServer.h
--- a/Engine/Source/Editor/UnrealEd/Classes/CookOnTheSide/CookOnTheFlyServer.h
+++ b/Engine/Source/Editor/UnrealEd/Classes/CookOnTheSide/CookOnTheFlyServer.h
@@ -1267,10 +1267,7 @@ private:
 	/** Allocate a new FCookSavePackageContext and ICookedPackageWriter for the given platform. */
 	UE::Cook::FCookSavePackageContext* CreateSaveContext(const ITargetPlatform* TargetPlatform);
 
-	void ConditionalInstallImportBehaviorCallback();
-	void ConditionalUninstallImportBehaviorCallback();
-	static void PropertyImportBehaviorCallback(const FObjectImport& Import, const FLinkerLoad& LinkerLoad, UE::LinkerLoad::EImportBehavior& OutBehavior);
-	
+
 	static UCookOnTheFlyServer* ActiveCOTFS;
 	uint32		StatLoadedPackageCount = 0;
 	uint32		StatSavedPackageCount = 0;
diff --git a/Engine/Source/Editor/UnrealEd/Private/CookOnTheFlyServer.cpp b/Engine/Source/Editor/UnrealEd/Private/CookOnTheFlyServer.cpp
--- a/Engine/Source/Editor/UnrealEd/Private/CookOnTheFlyServer.cpp
+++ b/Engine/Source/Editor/UnrealEd/Private/CookOnTheFlyServer.cpp
@@ -5166,7 +5166,6 @@ void UCookOnTheFlyServer::PostGarbageCollectImpl(TArray<UObject*>& LocalGCKeepOb
 void UCookOnTheFlyServer::BeginDestroy()
 {
 	ShutdownCookOnTheFly();
-	ConditionalUninstallImportBehaviorCallback();
 
 	Super::BeginDestroy();
 }
@@ -6908,8 +6907,6 @@ void UCookOnTheFlyServer::OnRequestClusterCompleted(const UE::Cook::FRequestClus
 			}
 			PackageDatas->SetLogDiscoveredPackages(true); // Turn this on for the rest of the cook after the initial cluster
 		}
-
-		ConditionalInstallImportBehaviorCallback();
 	}
 }
 
@@ -8399,8 +8396,6 @@ void UCookOnTheFlyServer::ShutdownCookSession()
 		PackageData->DestroyGeneratorPackage();
 	}
 
-	ConditionalUninstallImportBehaviorCallback();
-
 	if (IsCookByTheBookMode())
 	{
 		UnregisterCookByTheBookDelegates();
@@ -10699,86 +10694,6 @@ const UE::Cook::FCookSavePackageContext* UCookOnTheFlyServer::FindSaveContext(co
 	return nullptr;
 }
 
-void UCookOnTheFlyServer::ConditionalInstallImportBehaviorCallback()
-{
-	if (!bImportBehaviorCallbackInstalled && !IsCookingInEditor() && FEditorDomain::Get())
-	{
-		check(ActiveCOTFS == nullptr);
-		ActiveCOTFS = this;
-		UE::LinkerLoad::SetPropertyImportBehaviorCallback(&PropertyImportBehaviorCallback);
-		bImportBehaviorCallbackInstalled = true;
-	}
-}
-
-void UCookOnTheFlyServer::ConditionalUninstallImportBehaviorCallback()
-{
-	if (bImportBehaviorCallbackInstalled)
-	{
-		check(ActiveCOTFS == this);
-		UE::LinkerLoad::SetPropertyImportBehaviorCallback((UE::LinkerLoad::PropertyImportBehaviorFunction*)nullptr);
-		ActiveCOTFS = nullptr;
-		bImportBehaviorCallbackInstalled = false;
-	}
-}
-
-void UCookOnTheFlyServer::PropertyImportBehaviorCallback(const FObjectImport& Import, const FLinkerLoad& LinkerLoad, UE::LinkerLoad::EImportBehavior& OutBehavior)
-{
-	auto GetImportPackageName = [&LinkerLoad](const FObjectImport& Import)
-	{
-		const FObjectImport* CurrentResource = &Import;
-		for (int32 NumCycles = 0; NumCycles < LinkerLoad.ImportMap.Num(); ++NumCycles)
-		{
-			// If the import has a package name set, then that's the import package name,
-			if (CurrentResource->HasPackageName())
-			{
-				return CurrentResource->GetPackageName();
-			}
-			// If our outer is null, then we have a package
-			else if (CurrentResource->OuterIndex.IsNull())
-			{
-				return CurrentResource->ObjectName;
-			}
-			if (!CurrentResource->OuterIndex.IsImport())
-			{
-				return FName();
-			}
-			CurrentResource = &LinkerLoad.Imp(CurrentResource->OuterIndex);
-		}
-		return FName();
-	};
-
-	if (ActiveCOTFS != nullptr)
-	{
-		FName ImportPackageName = GetImportPackageName(Import);
-		if (!ImportPackageName.IsNone())
-		{
-			UE::Cook::FPackageData* PackageData = ActiveCOTFS->PackageDatas->FindPackageDataByPackageName(ImportPackageName);
-
-			if (PackageData && PackageData->IsInProgress())
-			{
-				OutBehavior = UE::LinkerLoad::EImportBehavior::Eager;
-			}
-			else
-			{
-				FNameBuilder ClassPathBuilder(Import.ClassPackage);
-				ClassPathBuilder.AppendChar(TEXT('.'));
-				ClassPathBuilder << Import.ClassName;
-				FName ClassPath(ClassPathBuilder.ToView());
-				UE::EditorDomain::FClassDigestMap& ClassDigests = UE::EditorDomain::GetClassDigests();
-				FReadScopeLock ClassDigestsScopeLock(ClassDigests.Lock);
-				UE::EditorDomain::FClassDigestData* ExistingData = ClassDigests.Map.Find(ClassPath);
-				if (!ExistingData || !ExistingData->bTargetIterativeEnabled)
-				{
-					OutBehavior = UE::LinkerLoad::EImportBehavior::Eager;
-				}
-				else
-				{
-					OutBehavior = UE::LinkerLoad::EImportBehavior::LazyOnDemand;
-				}
-			}
-		}
-	}
-}
 
 void UCookOnTheFlyServer::GenerateLocalizationReferences(TConstArrayView<FString> CookCultures)
 {
diff --git a/Engine/Source/Runtime/CoreUObject/Internal/UObject/LinkerLoadImportBehavior.h b/Engine/Source/Runtime/CoreUObject/Internal/UObject/LinkerLoadImportBehavior.h
--- a/Engine/Source/Runtime/CoreUObject/Internal/UObject/LinkerLoadImportBehavior.h
+++ b/Engine/Source/Runtime/CoreUObject/Internal/UObject/LinkerLoadImportBehavior.h
@@ -19,11 +19,8 @@ enum class EImportBehavior : uint8
 	LazyOnDemand,
 };
 
-using PropertyImportBehaviorFunction = void(const FObjectImport& Import, const FLinkerLoad& LinkerLoad, EImportBehavior& OutBehavior);
 
-UE_API void SetPropertyImportBehaviorCallback(PropertyImportBehaviorFunction* Function);
-
-void GetPropertyImportLoadBehavior(const FObjectImport& Import, const FLinkerLoad& LinkerLoad, EImportBehavior& OutBehavior);
+EImportBehavior GetPropertyImportLoadBehavior(const FObjectImport& Import, const FLinkerLoad& LinkerLoad);
 }
 
 #undef UE_API
diff --git a/Engine/Source/Runtime/CoreUObject/Private/Serialization/AsyncLoading.cpp b/Engine/Source/Runtime/CoreUObject/Private/Serialization/AsyncLoading.cpp
--- a/Engine/Source/Runtime/CoreUObject/Private/Serialization/AsyncLoading.cpp
+++ b/Engine/Source/Runtime/CoreUObject/Private/Serialization/AsyncLoading.cpp
@@ -26,6 +26,7 @@
 #include "UObject/CoreRedirects.h"
 #include "UObject/SoftObjectPath.h"
 #include "UObject/LinkerLoad.h"
+#include "UObject/LinkerLoadImportBehavior.h"
 #include "UObject/PackageResourceManager.h"
 #include "Serialization/DeferredMessageLog.h"
 #include "UObject/UObjectThreadContext.h"
@@ -6469,12 +6470,6 @@ EAsyncPackageState::Type FAsyncPackage::LoadImports()
 {
 	TRACE_CPUPROFILER_EVENT_SCOPE(FAsyncPackage::LoadImports);
 	SCOPE_CYCLE_COUNTER(STAT_FAsyncPackage_LoadImports);
-
-	if (FLinkerLoad::IsImportLazyLoadEnabled())
-	{
-		return EAsyncPackageState::Complete;
-	}
-
 	LastObjectWorkWasPerformedOn	= LinkerRoot;
 	LastTypeOfWorkPerformed			= TEXT("loading imports");
 
@@ -6504,6 +6499,17 @@ EAsyncPackageState::Type FAsyncPackage::LoadImports()
 			continue;
 		}
 
+#if UE_WITH_OBJECT_HANDLE_LATE_RESOLVE
+		if (FLinkerLoad::IsImportLazyLoadEnabled())
+		{
+			using namespace UE::LinkerLoad;
+			if (GetPropertyImportLoadBehavior(*Import, *Linker) == EImportBehavior::LazyOnDemand)
+			{
+				continue;
+			}
+		}
+#endif
+
 		// Our import package name is the import name itself if import is a package, else the specified package name
 		const FLinkerInstancingContext& InstancingContext = Linker->GetInstancingContext();
 		FName ImportToLoad = !Import->HasPackageName() ? Import->ObjectName : Import->GetPackageName();
@@ -6621,18 +6627,26 @@ EAsyncPackageState::Type FAsyncPackage::CreateImports()
 {
 	SCOPED_LOADTIMER(CreateImportsTime);
 	SCOPE_CYCLE_COUNTER(STAT_FAsyncPackage_CreateImports);
-
-	if (FLinkerLoad::IsImportLazyLoadEnabled())
-	{
-		return EAsyncPackageState::Complete;
-	}
-
 	// GC can't run in here
 	FGCScopeGuard GCGuard;
 
 	// Create imports.
 	while( ImportIndex < Linker->ImportMap.Num() && !IsTimeLimitExceeded() )
 	{
+
+#if UE_WITH_OBJECT_HANDLE_LATE_RESOLVE
+		if (FLinkerLoad::IsImportLazyLoadEnabled())
+		{
+			FObjectImport& Import = Linker->ImportMap[ImportIndex];
+			using namespace UE::LinkerLoad;
+			if (GetPropertyImportLoadBehavior(Import, *Linker) == EImportBehavior::LazyOnDemand)
+			{
+				ImportIndex++;
+				continue;
+			}
+		}
+#endif
+
  		UObject* Object	= Linker->CreateImport( ImportIndex++ );
 		LastObjectWorkWasPerformedOn	= Object;
 		LastTypeOfWorkPerformed			= TEXT("creating imports for");
diff --git a/Engine/Source/Runtime/CoreUObject/Private/UObject/LinkerLoad.cpp b/Engine/Source/Runtime/CoreUObject/Private/UObject/LinkerLoad.cpp
--- a/Engine/Source/Runtime/CoreUObject/Private/UObject/LinkerLoad.cpp
+++ b/Engine/Source/Runtime/CoreUObject/Private/UObject/LinkerLoad.cpp
@@ -5549,15 +5549,12 @@ FArchive& FLinkerLoad::operator<<( FObjectPtr& ObjectPtr )
 	if (Index.IsImport())
 	{
 		using namespace UE::LinkerLoad;
-		EImportBehavior LoadBehavior = EImportBehavior::Eager;
-		GetPropertyImportLoadBehavior(Imp(Index), *this, LoadBehavior);
-		if (LoadBehavior == EImportBehavior::LazyOnDemand)
+		if (GetPropertyImportLoadBehavior(Imp(Index), *this) == EImportBehavior::LazyOnDemand)
 		{
 			const FObjectImport& Import = Imp(Index);
 			FObjectPathId ObjectPath;
 			FName PackageName = FObjectPathId::MakeImportPathIdAndPackageName(Import, *this, ObjectPath);
 			FObjectRef ImportRef{PackageName, Import.ClassPackage, Import.ClassName, ObjectPath};
-
 			ObjectPtr = FObjectPtr(ImportRef);
 
 			return *this;
diff --git a/Engine/Source/Runtime/CoreUObject/Private/UObject/LinkerLoadImportBehavior.cpp b/Engine/Source/Runtime/CoreUObject/Private/UObject/LinkerLoadImportBehavior.cpp
--- a/Engine/Source/Runtime/CoreUObject/Private/UObject/LinkerLoadImportBehavior.cpp
+++ b/Engine/Source/Runtime/CoreUObject/Private/UObject/LinkerLoadImportBehavior.cpp
@@ -4,38 +4,72 @@
 #include "UObject/LinkerLoad.h"
 #include "UObject/ObjectResource.h"
 #include "UObject/Package.h"
+#include "Misc/CommandLine.h"
 
 #if UE_WITH_OBJECT_HANDLE_LATE_RESOLVE
 
 namespace UE::LinkerLoad
 {
-PropertyImportBehaviorFunction* PropertyImportBehaviorCallback = nullptr;
 
-void SetPropertyImportBehaviorCallback(PropertyImportBehaviorFunction* Function)
+/// @brief Finds LoadBehavior meta data recursively
+/// @return Eager by default in not found
+static EImportBehavior FindLoadBehavior(const UClass& Class)
 {
-	void* PreviousValue = (void*)FPlatformAtomics::InterlockedExchangePtr((void**)&PropertyImportBehaviorCallback, (void*)Function);
-	checkf((Function == nullptr) || (PreviousValue == nullptr), TEXT("LinkerLoad property import behavior callback set more than once.  This is not supported."))
+	static const FName Name_LoadBehavior("LoadBehavior");
+	if (const FString* LoadBehaviorMeta = Class.FindMetaData(Name_LoadBehavior))
+	{
+		if (*LoadBehaviorMeta == "LazyOnDemand")
+		{
+			return EImportBehavior::LazyOnDemand;
+		}
+		return EImportBehavior::Eager;
+	}
+	else
+	{
+		const UClass* Super = Class.GetSuperClass();
+		if (Super != nullptr)
+		{
+			return FindLoadBehavior(*Super);
+		}
+		return EImportBehavior::Eager;
+	}
 }
 
-void GetPropertyImportLoadBehavior(const FObjectImport& Import, const FLinkerLoad& LinkerLoad, EImportBehavior& OutBehavior)
+EImportBehavior GetPropertyImportLoadBehavior(const FObjectImport& Import, const FLinkerLoad& LinkerLoad)
 {
-	OutBehavior = EImportBehavior::Eager;
-
 	if (Import.bImportSearchedFor)
 	{
 		// If it was something that's been searched for, we've already attempted a resolve, might as well use it
-		return;
+		return EImportBehavior::Eager;
 	}
 
 	if (!LinkerLoad.IsImportLazyLoadEnabled() || !LinkerLoad.IsAllowingLazyLoading() || (LinkerLoad.LinkerRoot && LinkerLoad.LinkerRoot->HasAnyPackageFlags(PKG_PlayInEditor)))
 	{
-		return;
+		return EImportBehavior::Eager;
+	}
+
+	// Attempt to get the meta from the referenced class.  This only looks in already loaded classes.  May need to resolve the class in the future.
+	static const bool bDefaultLoadBehaviorTest = FParse::Param(FCommandLine::Get(), TEXT("DefaultLoadBehaviorTest"));
+	if (bDefaultLoadBehaviorTest)
+	{
+		return EImportBehavior::LazyOnDemand;
+	}
+
+	//Packages can't have meta data because they 
+	if (Import.ClassName == TEXT("Package"))
+	{
+		return EImportBehavior::LazyOnDemand;
 	}
 
-	if (PropertyImportBehaviorCallback)
+	TRACE_CPUPROFILER_EVENT_SCOPE(LinkerLoader::GetPropertyImportLoadBehavior);
+
+	UObject* ClassPackage = FindObjectFast<UPackage>(nullptr, Import.ClassPackage);
+	const UClass* FindClass = ClassPackage ? FindObjectFast<const UClass>(ClassPackage, Import.ClassName) : nullptr;
+	if (FindClass != nullptr)
 	{
-		PropertyImportBehaviorCallback(Import, LinkerLoad, OutBehavior);
+		return FindLoadBehavior(*FindClass);
 	}
+	return EImportBehavior::Eager;
 }
 }
 
