Check if world exists before early-ing out due to world switching.

diff --git a/Engine/Source/Editor/UnrealEd/Private/EditorViewportClient.cpp b/Engine/Source/Editor/UnrealEd/Private/EditorViewportClient.cpp
--- a/Engine/Source/Editor/UnrealEd/Private/EditorViewportClient.cpp
+++ b/Engine/Source/Editor/UnrealEd/Private/EditorViewportClient.cpp
@@ -3903,7 +3903,7 @@ void FEditorViewportClient::Draw(FViewport* InViewport, FCanvas* Canvas)
 	}
 
 	// Early out if we are changing maps in editor as there is no reason to render the scene and it may not even be valid (For unsaved maps)
-	if (World->IsPreparingMapChange())
+	if (World && World->IsPreparingMapChange())
 	{
 		return;
 	}
