World Partition: create a loading region that covers the whole world when saving a new partitioned world from a template that doesn't include default loading regions.

#rnx

diff --git a/Engine/Source/Editor/UnrealEd/Private/FileHelpers.cpp b/Engine/Source/Editor/UnrealEd/Private/FileHelpers.cpp
--- a/Engine/Source/Editor/UnrealEd/Private/FileHelpers.cpp
+++ b/Engine/Source/Editor/UnrealEd/Private/FileHelpers.cpp
@@ -79,6 +79,7 @@
 #include "WorldPartition/WorldPartition.h"
 #include "WorldPartition/WorldPartitionEditorPerProjectUserSettings.h"
 #include "WorldPartition/HLOD/HLODLayer.h"
+#include "WorldPartition/LoaderAdapter/LoaderAdapterShape.h"
 #include "PackageSourceControlHelper.h"
 #include "ActorFolder.h"
 #include "InterchangeManager.h"
@@ -896,15 +897,25 @@ static bool SaveWorld(UWorld* World,
 			
 			if (bPackageNeedsRename || bNewlyCreated || !bNewPackageExists)
 			{
-				// Force rescan to make sure assets are found on map open or world partition initialize
+				// Force rescan to make sure assets are found on map open or world partition initialize`
 				AssetRegistry.ScanPathsSynchronous( ULevel::GetExternalObjectsPaths(NewPackageName) , true);
 			}
 
 			if (RenamedWorldPartition)
 			{
-				// Save Snapshot of loaded Editor regions
-				GetMutableDefault<UWorldPartitionEditorPerProjectUserSettings>()->SetEditorLoadedRegions(SaveWorld, LoadedEditorRegions);
-				RenamedWorldPartition->LoadLastLoadedRegions(LoadedEditorRegions);
+				if (LoadedEditorRegions.Num())
+				{
+					// Save Snapshot of loaded Editor regions
+					GetMutableDefault<UWorldPartitionEditorPerProjectUserSettings>()->SetEditorLoadedRegions(SaveWorld, LoadedEditorRegions);
+					RenamedWorldPartition->LoadLastLoadedRegions(LoadedEditorRegions);
+				}
+				else if (bIsTempPackage)
+				{
+					const FBox WorldBounds = RenamedWorldPartition->GetRuntimeWorldBounds();
+					UWorldPartitionEditorLoaderAdapter* EditorLoaderAdapter = RenamedWorldPartition->CreateEditorLoaderAdapter<FLoaderAdapterShape>(World, WorldBounds, TEXT("Loaded Region"));
+					EditorLoaderAdapter->GetLoaderAdapter()->SetUserCreated(true);
+					EditorLoaderAdapter->GetLoaderAdapter()->Load();
+				}
 			}
 		}
 
