Fixing localization label names in BaseIteratePackagesCommandlet.cpp

diff --git a/Engine/Source/Editor/UnrealEd/Private/Commandlets/BaseIteratePackagesCommandlet.cpp b/Engine/Source/Editor/UnrealEd/Private/Commandlets/BaseIteratePackagesCommandlet.cpp
--- a/Engine/Source/Editor/UnrealEd/Private/Commandlets/BaseIteratePackagesCommandlet.cpp
+++ b/Engine/Source/Editor/UnrealEd/Private/Commandlets/BaseIteratePackagesCommandlet.cpp
@@ -1153,7 +1153,7 @@ int32 UBaseIteratePackagesCommandlet::Main( const FString& Params )
 FText UBaseIteratePackagesCommandlet::GetChangelistDescription() const
 {
 	FTextBuilder ChangelistDescription;
-	ChangelistDescription.AppendLine(NSLOCTEXT("ContentCmdlets", "ChangelistDescription", "Resave Packages"));
+	ChangelistDescription.AppendLine(NSLOCTEXT("ContentCmdlets", "ChangelistDescriptionResavePackages", "Resave Packages"));
 	return ChangelistDescription.ToText();
 }
 
