Fixed crash occurring when the user changes a TextureSample or TextureObject node connected to the Material outputs to Object/Sample respectively.

diff --git a/Engine/Source/Editor/UnrealEd/Private/MaterialGraph.cpp b/Engine/Source/Editor/UnrealEd/Private/MaterialGraph.cpp
--- a/Engine/Source/Editor/UnrealEd/Private/MaterialGraph.cpp
+++ b/Engine/Source/Editor/UnrealEd/Private/MaterialGraph.cpp
@@ -638,31 +638,34 @@ void UMaterialGraph::LinkMaterialExpressionsFromGraph()
 							// Wire up non-execution input pins
 							// Implicitly generated property inputs are not returned by GetInputs(), so check index is within valid range.
 							FExpressionInput* ExpressionInput = ExpressionInputs.IsValidIndex(Pin->SourceIndex) ? ExpressionInputs[Pin->SourceIndex] : nullptr;
-							if (Pin->LinkedTo.Num() > 0)
+							if (ExpressionInput)
 							{
-								UEdGraphPin* ConnectedPin = Pin->LinkedTo[0];
-								UMaterialGraphNode* ConnectedNode = CastChecked<UMaterialGraphNode>(ConnectedPin->GetOwningNode());
+								if (Pin->LinkedTo.Num() > 0)
+								{
+									UEdGraphPin* ConnectedPin = Pin->LinkedTo[0];
+									UMaterialGraphNode* ConnectedNode = Cast<UMaterialGraphNode>(ConnectedPin->GetOwningNode());
+
+									if (ConnectedNode && !ConnectedNode->MaterialExpression->IsExpressionConnected(ExpressionInput, ConnectedPin->SourceIndex))
+									{
+										if (!bModifiedExpression)
+										{
+											bModifiedExpression = true;
+											Expression->Modify();
+										}
 
-								if (ExpressionInput && !ConnectedNode->MaterialExpression->IsExpressionConnected(ExpressionInput, ConnectedPin->SourceIndex))
+										ConnectedNode->MaterialExpression->Modify();
+										ExpressionInput->Connect(ConnectedPin->SourceIndex, ConnectedNode->MaterialExpression);
+									}
+								}
+								else if (ExpressionInput->Expression)
 								{
 									if (!bModifiedExpression)
 									{
 										bModifiedExpression = true;
 										Expression->Modify();
 									}
-
-									ConnectedNode->MaterialExpression->Modify();
-									ExpressionInput->Connect(ConnectedPin->SourceIndex, ConnectedNode->MaterialExpression);
-								}
-							}
-							else if (ExpressionInput && ExpressionInput->Expression)
-							{
-								if (!bModifiedExpression)
-								{
-									bModifiedExpression = true;
-									Expression->Modify();
+									ExpressionInput->Expression = NULL;
 								}
-								ExpressionInput->Expression = NULL;
 							}
 						}
 						else if (Pin->Direction == EGPD_Output && Pin->PinType.PinCategory == UMaterialGraphSchema::PC_Exec)
