Fixed some checks for Strata and lightmass.

diff --git a/Engine/Source/Editor/UnrealEd/Private/Lightmass/Lightmass.cpp b/Engine/Source/Editor/UnrealEd/Private/Lightmass/Lightmass.cpp
--- a/Engine/Source/Editor/UnrealEd/Private/Lightmass/Lightmass.cpp
+++ b/Engine/Source/Editor/UnrealEd/Private/Lightmass/Lightmass.cpp
@@ -54,6 +54,8 @@
 
 extern FSwarmDebugOptions GSwarmDebugOptions;
 
+bool Lightmass_IsStrataEnabled();
+
 DEFINE_LOG_CATEGORY_STATIC(LogLightmassSolver, Warning, All);
 /**
  * If false (default behavior), Lightmass is launched automatically when a lighting build starts.
@@ -1344,6 +1346,13 @@ void FLightmassExporter::GetMaterialHash(const UMaterialInterface* Material, FSH
 			LastGuid = MaterialGuid;
 		}
 	}
+
+	if (Lightmass_IsStrataEnabled())
+	{
+		uint32 LightmassStrataVersion = 0XB6A0D99F; // This can be change when the code/logic for converting material to strata for lightmap has changed.
+		HashState.Update((const uint8*)&LightmassStrataVersion, sizeof(LightmassStrataVersion));
+	}
+
 	HashState.Final();
 	HashState.GetHash(&OutHash.Hash[0]);
 }
diff --git a/Engine/Source/Editor/UnrealEd/Private/Lightmass/LightmassRender.cpp b/Engine/Source/Editor/UnrealEd/Private/Lightmass/LightmassRender.cpp
--- a/Engine/Source/Editor/UnrealEd/Private/Lightmass/LightmassRender.cpp
+++ b/Engine/Source/Editor/UnrealEd/Private/Lightmass/LightmassRender.cpp
@@ -28,7 +28,7 @@
 
 DEFINE_LOG_CATEGORY_STATIC(LogLightmassRender, Error, All);
 
-bool Lightmass_IsStrataEnabled()
+extern bool Lightmass_IsStrataEnabled()
 {
 	static const auto CVar = IConsoleManager::Get().FindTConsoleVariableDataInt(TEXT("r.Strata"));
 	return CVar && CVar->GetValueOnAnyThread() > 0;
@@ -1094,6 +1094,15 @@ bool FLightmassMaterialRenderer::GenerateMaterialData(
 	// Set the blend mode
 	static_assert(EBlendMode::BLEND_MAX == (EBlendMode)Lightmass::BLEND_MAX, "Debug type sizes must match.");
 	OutMaterialData.BlendMode = (Lightmass::EBlendMode)((int32)BlendMode);
+	if (Lightmass_IsStrataEnabled())
+	{
+		OutMaterialData.StrataBlendMode = (Lightmass::EStrataBlendMode)((int32)InMaterial.GetStrataBlendMode());
+	}
+	else
+	{
+		OutMaterialData.StrataBlendMode = Lightmass::SBM_MAX;
+	}
+
 	// Set the two-sided flag
 	OutMaterialData.bTwoSided = (uint32)InMaterial.IsTwoSided();
 	OutMaterialData.OpacityMaskClipValue = InMaterial.GetOpacityMaskClipValue();
diff --git a/Engine/Source/Programs/UnrealLightmass/Private/ImportExport/Material.cpp b/Engine/Source/Programs/UnrealLightmass/Private/ImportExport/Material.cpp
--- a/Engine/Source/Programs/UnrealLightmass/Private/ImportExport/Material.cpp
+++ b/Engine/Source/Programs/UnrealLightmass/Private/ImportExport/Material.cpp
@@ -55,8 +55,16 @@ namespace Lightmass
 		}
 		else
 		{
-			// Materials with a translucent blend mode should always import transmission
-			check(BlendMode != BLEND_Translucent && BlendMode != BLEND_Additive && BlendMode != BLEND_Modulate && BlendMode != BLEND_AlphaComposite && BlendMode != BLEND_AlphaHoldout);
+			if (StrataBlendMode != SBM_MAX)
+			{
+				// Strata is enabled: only opaque materials do not have Transmission data.
+				check(StrataBlendMode == SBM_Opaque);
+			}
+			else
+			{
+				// Materials with a translucent blend mode should always import transmission
+				check(BlendMode != BLEND_Translucent && BlendMode != BLEND_Additive && BlendMode != BLEND_Modulate && BlendMode != BLEND_AlphaComposite && BlendMode != BLEND_AlphaHoldout);
+			}
 		}
 
 		// Normal
diff --git a/Engine/Source/Programs/UnrealLightmass/Public/MaterialExport.h b/Engine/Source/Programs/UnrealLightmass/Public/MaterialExport.h
--- a/Engine/Source/Programs/UnrealLightmass/Public/MaterialExport.h
+++ b/Engine/Source/Programs/UnrealLightmass/Public/MaterialExport.h
@@ -53,11 +53,23 @@ enum EBlendMode
 	BLEND_AlphaHoldout      =6,
 	BLEND_MAX               =7,
 };
+enum EStrataBlendMode
+{
+	SBM_Opaque								= 0,
+	SBM_Masked								= 1,
+	SBM_TranslucentGreyTransmittance		= 2,
+	SBM_TranslucentColoredTransmittance		= 3,
+	SBM_ColoredTransmittanceOnly			= 4,
+	SBM_AlphaHoldout						= 5,
+	SBM_MAX									= 6,
+};
 
 struct FMaterialData
 {
 	/** The BLEND mode of the material */
 	EBlendMode BlendMode;
+	/** The Strata BLEND mode of the material */
+	EStrataBlendMode StrataBlendMode;
 	/** Whether the material is two-sided or not */
 	uint32 bTwoSided:1;
 	/** Whether the material should cast shadows as masked even though it has a translucent blend mode. */
@@ -78,8 +90,9 @@ struct FMaterialData
 	int32 TransmissionSize;
 	int32 NormalSize;
 
-	FMaterialData() : 
-		  EmissiveBoost(1.0f)
+	FMaterialData() :
+		  StrataBlendMode(SBM_MAX)
+		, EmissiveBoost(1.0f)
 		, DiffuseBoost(1.0f)
 		, EmissiveSize(0)
 		, DiffuseSize(0)
