Fixed an object path issue when capturing a thumbnail.

It used the full class name when setting the captured thumbnail while the code that fetch the thumbnails use the short type name.

diff --git a/Engine/Source/Editor/UnrealEd/Private/ObjectTools.cpp b/Engine/Source/Editor/UnrealEd/Private/ObjectTools.cpp
--- a/Engine/Source/Editor/UnrealEd/Private/ObjectTools.cpp
+++ b/Engine/Source/Editor/UnrealEd/Private/ObjectTools.cpp
@@ -4978,20 +4978,19 @@ namespace ThumbnailTools
 				DestPackage->SetThumbnailMap(MakeUnique<FThumbnailMap>());
 			}
 
-			// @todo thumbnails: Backwards compat
-			FName ObjectFullNameFName( *ObjectFullName );
-			const FObjectThumbnail* CachedThumbnail = DestPackage->GetThumbnailMap().Find( ObjectFullNameFName );
+			FName InObjectShortClassFullName( *UClass::ConvertFullNameToShortTypeFullName( ObjectFullName ) );
+			const FObjectThumbnail* CachedThumbnail = DestPackage->GetThumbnailMap().Find( InObjectShortClassFullName );
 			if ( Thumbnail != nullptr )
 			{
 				// Cache the thumbnail (possibly replacing an existing thumb!)
-				Result = &DestPackage->AccessThumbnailMap().Add( ObjectFullNameFName, *Thumbnail );
+				Result = &DestPackage->AccessThumbnailMap().Add( InObjectShortClassFullName, *Thumbnail );
 			}
 			//only let thumbnails loaded from disk to be removed.
 			//When capturing thumbnails from the content browser, it will only exist in memory until it is saved out to a package.
 			//Don't let the recycling purge them
 			else if ((CachedThumbnail != nullptr) && (CachedThumbnail->IsLoadedFromDisk()))
 			{
-				DestPackage->AccessThumbnailMap().Remove( ObjectFullNameFName );
+				DestPackage->AccessThumbnailMap().Remove( InObjectShortClassFullName );
 			}
 
 		}
