Added missing null checks.

The rest of the code in FEditorViewportClient::Draw has explicit null checks for the world, but there were a couple of places where it was missing, so we could crash when Slate did cursor queries during level load.

#rnx

diff --git a/Engine/Source/Editor/UnrealEd/Private/EditorViewportClient.cpp b/Engine/Source/Editor/UnrealEd/Private/EditorViewportClient.cpp
--- a/Engine/Source/Editor/UnrealEd/Private/EditorViewportClient.cpp
+++ b/Engine/Source/Editor/UnrealEd/Private/EditorViewportClient.cpp
@@ -4077,12 +4077,12 @@ void FEditorViewportClient::Draw(FViewport* InViewport, FCanvas* Canvas)
 
 	// Remove temporary debug lines.
 	// Possibly a hack. Lines may get added without the scene being rendered etc.
-	if (World->LineBatcher != NULL && (World->LineBatcher->BatchedLines.Num() || World->LineBatcher->BatchedPoints.Num() || World->LineBatcher->BatchedMeshes.Num() ) )
+	if (World && World->LineBatcher != NULL && (World->LineBatcher->BatchedLines.Num() || World->LineBatcher->BatchedPoints.Num() || World->LineBatcher->BatchedMeshes.Num() ) )
 	{
 		World->LineBatcher->Flush();
 	}
 
-	if (World->ForegroundLineBatcher != NULL && (World->ForegroundLineBatcher->BatchedLines.Num() || World->ForegroundLineBatcher->BatchedPoints.Num() || World->ForegroundLineBatcher->BatchedMeshes.Num() ) )
+	if (World && World->ForegroundLineBatcher != NULL && (World->ForegroundLineBatcher->BatchedLines.Num() || World->ForegroundLineBatcher->BatchedPoints.Num() || World->ForegroundLineBatcher->BatchedMeshes.Num() ) )
 	{
 		World->ForegroundLineBatcher->Flush();
 	}
diff --git a/Engine/Source/Runtime/Engine/Private/UnrealEngine.cpp b/Engine/Source/Runtime/Engine/Private/UnrealEngine.cpp
--- a/Engine/Source/Runtime/Engine/Private/UnrealEngine.cpp
+++ b/Engine/Source/Runtime/Engine/Private/UnrealEngine.cpp
@@ -11837,8 +11837,8 @@ void DrawStatsHUD( UWorld* World, FViewport* Viewport, FCanvas* Canvas, UCanvas*
 	DECLARE_SCOPE_CYCLE_COUNTER( TEXT( "DrawStatsHUD" ), STAT_DrawStatsHUD, STATGROUP_StatSystem );
 	CSV_SCOPED_TIMING_STAT_EXCLUSIVE(DebugHUD);
 
-	// We cannot draw without a canvas
-	if (Canvas == NULL)
+	// We cannot draw without a world or a canvas.
+	if (World == NULL || Canvas == NULL)
 	{
 		return;
 	}
