Cooker and AssetRegistry minor fixes - headers, comments, and whitespace.
#rnx

diff --git a/Engine/Source/Editor/UnrealEd/Classes/CookOnTheSide/CookOnTheFlyServer.h b/Engine/Source/Editor/UnrealEd/Classes/CookOnTheSide/CookOnTheFlyServer.h
--- a/Engine/Source/Editor/UnrealEd/Classes/CookOnTheSide/CookOnTheFlyServer.h
+++ b/Engine/Source/Editor/UnrealEd/Classes/CookOnTheSide/CookOnTheFlyServer.h
@@ -2,8 +2,6 @@
 
 #pragma once
 
-#include "CoreMinimal.h"
-
 #include "Containers/ArrayView.h"
 #include "Containers/RingBuffer.h"
 #include "CookOnTheSide/CookLog.h"
@@ -297,16 +295,24 @@ private:
 	int32 MaxConcurrentShaderJobs;
 	/** Min number of free UObject indices before the cooker should partial gc */
 	int32 MinFreeUObjectIndicesBeforeGC;
-	/** The maximum number of packages that should be preloaded at once. Once this is full, packages in LoadPrepare will remain unpreloaded in LoadPrepare until the existing preloaded packages exit {LoadPrepare,LoadReady} state. */
+	/**
+	 * The maximum number of packages that should be preloaded at once. Once this is full, packages in LoadPrepare will
+	 * remain unpreloaded in LoadPrepare until the existing preloaded packages exit {LoadPrepare,LoadReady} state.
+	 */
 	uint32 MaxPreloadAllocated;
 	/**
-	 * A knob to tune performance - How many packages should be present in the SaveQueue before we start processing the SaveQueue. If number of files is less, we will look for other work to do, and save packages only if all other work is done.
-	 * This allows us to have enough population in the SaveQueue to get benefit from the asynchronous work done on packages in the SaveQueue.
+	 * A knob to tune performance - How many packages should be present in the SaveQueue before we start processing the
+	 * SaveQueue. If number is less, we will find other work to do and save packages only if all other work is done.
+	 * This allows us to have enough population in the SaveQueue to get benefit from the asynchronous work done on
+	 * packages in the SaveQueue.
 	 */
 	uint32 DesiredSaveQueueLength;
 	/**
-	 * A knob to tune performance - How many packages should be present in the LoadPrepare+LoadReady queues before we start processing the LoadQueue. If number of files is less, we will look for other work to do, and load packages only if all other work is done.
-	 * This allows us to have enough population in the LoadPrepareQueue to get benefit from the asynchronous work done on packages in the LoadPrepareQueue.
+	 * A knob to tune performance - How many packages should be present in the LoadPrepare+LoadReady queues before we
+	 * start processing the LoadQueue. If number is less, we will find other work to do, and load packages only if all
+	 * other work is done.
+	 * This allows us to have enough population in the LoadPrepareQueue to get benefit from the asynchronous work done
+	 * on packages in the LoadPrepareQueue.
 	 */
 	uint32 DesiredLoadQueueLength;
 	/** A knob to tune performance - how many packages to pull off in each call to PumpRequests. */
@@ -319,11 +325,16 @@ private:
 	FString SandboxFileOutputDirectory;
 	TUniquePtr<FAsyncIODelete> AsyncIODelete; // Helper for deleting the old cook directory asynchronously
 	bool bIsSavingPackage = false; // used to stop recursive mark package dirty functions
-	/** Set to true during CookOnTheFly if a plugin is calling RequestPackage and we should therefore not make assumptions about when platforms are done cooking */
+	/**
+	 * Set to true during CookOnTheFly if a plugin is calling RequestPackage and we should therefore not make
+	 * assumptions about when platforms are done cooking
+	 */
 	bool bCookOnTheFlyExternalRequests = false;
 
-	TMap<FName, int32> MaxAsyncCacheForType; // max number of objects of a specific type which are allowed to async cache at once
-	mutable TMap<FName, int32> CurrentAsyncCacheForType; // max number of objects of a specific type which are allowed to async cache at once
+	/** max number of objects of a specific type which are allowed to async cache at once */
+	TMap<FName, int32> MaxAsyncCacheForType;
+	/** max number of objects of a specific type which are allowed to async cache at once */
+	mutable TMap<FName, int32> CurrentAsyncCacheForType;
 
 	/** List of additional plugin directories to remap into the sandbox as needed */
 	TArray<TSharedRef<IPlugin> > PluginsToRemap;
@@ -432,7 +443,10 @@ private:
 	bool TryCreateRequestCluster(UE::Cook::FPackageData& PackageData);
 	/** Inspect the next package in the RequestQueue and push it on to its next state. Report the number of PackageDatas that were pushed to load. */
 	void PumpRequests(UE::Cook::FTickStackData& StackData, int32& OutNumPushed);
-	/** Assign the requests found in PumpRequests; either pushing them to ReadyRequests if SingleProcess Cook or to CookWorkers if MultiProcess. */
+	/**
+	 * Assign the requests found in PumpRequests; either pushing them to ReadyRequests if SingleProcess Cook or
+	 * to CookWorkers if MultiProcess.
+	 */
 	void AssignRequests(TArrayView<UE::Cook::FPackageData*> Requests, UE::Cook::FRequestQueue& RequestQueue);
 	/**
 	 * Multiprocess cook: Notify the CookDirector that a package it assigned to a worker was demoted
@@ -440,18 +454,25 @@ private:
 	 */
 	void NotifyRemovedFromWorker(UE::Cook::FPackageData& PackageData);
 
-	/** Load all packages in the LoadQueue until it's time to break. Report the number of loads that were pushed to save. */
+	/** Load packages in the LoadQueue until it's time to break. Report the number of loads that were pushed to save. */
 	void PumpLoads(UE::Cook::FTickStackData& StackData, uint32 DesiredQueueLength, int32& OutNumPushed, bool& bOutBusy);
 	/** Move packages from LoadPrepare's entry queue into the PreloadingQueue until we run out of Preload slots. */
 	void PumpPreloadStarts();
-	/** Move preload-completed packages from LoadPrepare into state LoadReady until we find the first one that is not yet finished preloading. */
+	/** Move preload-completed packages from LoadPrepare->LoadReady until we find one that is not finished preloading. */
 	void PumpPreloadCompletes();
-	/** Load the given PackageData that was in the load queue and send it on to its next state. Report the number of PackageDatas that were pushed to save (0 or 1) */
+	/**
+	 * Load the given PackageData that was in the load queue and send it on to its next state.
+	 * Report the number of PackageDatas that were pushed to save (0 or 1)
+	 */
 	void LoadPackageInQueue(UE::Cook::FPackageData& PackageData, uint32& ResultFlags, int32& OutNumPushed);
 	/** Mark that the given PackageData failed to load and return it to idle. */
-	void RejectPackageToLoad(UE::Cook::FPackageData& PackageData, const TCHAR* ReasonText, UE::Cook::ESuppressCookReason Reason);
+	void RejectPackageToLoad(UE::Cook::FPackageData& PackageData, const TCHAR* ReasonText,
+		UE::Cook::ESuppressCookReason Reason);
 
-	/** Try to save all packages in the SaveQueue until it's time to break. Report the number of requests that were completed (either skipped or successfully saved or failed to save) */
+	/**
+	 * Try to save all packages in the SaveQueue until it's time to break.
+	 * Report the number of requests that were completed (either skipped or successfully saved or failed to save)
+	 */
 	void PumpSaves(UE::Cook::FTickStackData& StackData, uint32 DesiredQueueLength, int32& OutNumPushed, bool& bOutBusy);
 	/**
 	 * Inspect the given package and queue it for saving if necessary.
@@ -471,18 +492,18 @@ private:
 	void QueueDiscoveredPackageData(UE::Cook::FPackageData& PackageData, UE::Cook::FInstigator&& Instigator,
 		bool bLoadReady = false);
 
-	/** Function called whenever an inprogress package is cancelled and returned to idle. Notifies the CookDirector when on a CookWorker. */
+	/** Called when a package is cancelled and returned to idle. Notifies CookDirector when on a CookWorker. */
 	void DemoteToIdle(UE::Cook::FPackageData& PackageData, UE::Cook::ESendFlags SendFlags, UE::Cook::ESuppressCookReason Reason);
-	/** Function called whenever an inprogress package completes its save and returns to idle. Notifies the CookDirector when on a CookWorker. */
+	/** Called when a package completes its save and returns to idle. Notifies the CookDirector when on a CookWorker. */
 	void PromoteToSaveComplete(UE::Cook::FPackageData& PackageData, UE::Cook::ESendFlags SendFlags);
 
-	/** Check whether the package filter has change, and if so call QueueDiscoveredPackage again on each existing package. */
+	/** If the package filter has changed, call QueueDiscoveredPackage again on each existing package. */
 	void UpdatePackageFilter();
 
 	/**
 	 * Remove all request data about the given platform from any data in the CookOnTheFlyServer.
-	 * Called in response to a platform being removed from the list of session platforms e.g. because it has not been recently used by CookOnTheFly.
-	 * Does not modify Cooked platforms.
+	 * Called when a platform is removed from the list of session platforms e.g. because it has not been recently
+	 * used by CookOnTheFly. Does not modify Cooked platforms.
 	 */
 	void OnRemoveSessionPlatform(const ITargetPlatform* TargetPlatform);
 
@@ -522,14 +543,16 @@ public:
 		TArray<FString> NeverCookDirectories;
 		TArray<FString> CookCultures;
 		TArray<FString> IniMapSections;
-		TArray<FString> CookPackages; // list of packages we should cook, used to specify specific packages to cook
+		/** list of packages we should cook, used to specify specific packages to cook */
+		TArray<FString> CookPackages;
 		ECookByTheBookOptions CookOptions = ECookByTheBookOptions::None;
 		FString DLCName;
 		FString CreateReleaseVersion;
 		FString BasedOnReleaseVersion;
 		bool bGenerateStreamingInstallManifests = false;
 		bool bGenerateDependenciesForMaps = false;
-		bool bErrorOnEngineContentUse = false; // this is a flag for dlc, will cause the cooker to error if the dlc references engine content
+		/** this is a flag for dlc, will cause the cooker to error if the dlc references engine content */
+		bool bErrorOnEngineContentUse = false;
 	};
 
 	struct FCookOnTheFlyStartupOptions
@@ -538,64 +561,51 @@ public:
 		bool bBindAnyPort = false;
 		/** Whether to save the cooked output to the Zen storage server. */
 		bool bZenStore = false;
-		/** Whether the network file server should use a platform-specific communication protocol instead of TCP (used when bZenStore == false) */
+		/**
+		 * Whether the network file server should use a platform-specific communication protocol instead of TCP
+		 * (used when bZenStore == false)
+		 */
 		bool bPlatformProtocol = false;
 		/** Target platforms */
 		TArray<ITargetPlatform*> TargetPlatforms;
 	};
 
-
 	virtual ~UCookOnTheFlyServer();
 
-	/**
-	* FTickableEditorObject interface used by cook on the side
-	*/
+	// FTickableEditorObject interface used by cook on the side
 	TStatId GetStatId() const override;
 	void Tick(float DeltaTime) override;
 	bool IsTickable() const override;
 	ECookMode::Type GetCookMode() const { return CurrentCookMode; }
 	ECookInitializationFlags GetCookFlags() const { return CookFlags; }
 
-
-	/**
-	* FExec interface used in the editor
-	*/
+	// FExec interface used in the editor
 	virtual bool Exec(class UWorld* InWorld, const TCHAR* Cmd, FOutputDevice& Ar) override;
 
-	/**
-	* ICookInfo interface
-	*/
+	// ICookInfo interface
 	virtual UE::Cook::FInstigator GetInstigator(FName PackageName) override;
 	virtual TArray<UE::Cook::FInstigator> GetInstigatorChain(FName PackageName) override;
 
-	/**
-	 * Dumps cooking stats to the log
-	 *  run from the exec command "Cook stats"
-	 */
+	/** Dumps cooking stats to the log. Run from the exec command "Cook stats". */
 	void DumpStats();
 
-	/**
-	* Initialize the CookServer so that either CookOnTheFly can be called or Cook on the side can be started and ticked
-	*/
-	void Initialize( ECookMode::Type DesiredCookMode, ECookInitializationFlags InCookInitializationFlags, const FString& OutputDirectoryOverride = FString() );
+	/** Initialize *this so that either CookOnTheFly or CookByTheBook can be started and ticked */
+	void Initialize( ECookMode::Type DesiredCookMode, ECookInitializationFlags InCookInitializationFlags,
+		const FString& OutputDirectoryOverride = FString() );
 
 	/**
-	* Initialize cook on the fly server
-	*
-	* @param InCookOnTheFlyOptions Cook on the fly startup options
-	*
-	* @return true on success, false otherwise.
-	*/
+	 * Initialize cook on the fly server
+	 *
+	 * @param InCookOnTheFlyOptions Cook on the fly startup options
+	 *
+	 * @return true on success, false otherwise.
+	 */
 	bool StartCookOnTheFly(FCookOnTheFlyStartupOptions InCookOnTheFlyOptions); 
 
-	/**
-	* Broadcast our the fileserver presence on the network
-	*/
+	/** Broadcast the fileserver's presence on the network */
 	bool BroadcastFileserverPresence( const FGuid &InstanceId );
 	
-	/** 
-	* Shutdown cook on the fly server
-	*/
+	/** Shutdown *this from running in CookOnTheFly mode */
 	void ShutdownCookOnTheFly();
 
 	/**
@@ -604,14 +614,10 @@ public:
 	*/
 	void StartCookByTheBook( const FCookByTheBookStartupOptions& CookByTheBookStartupOptions );
 
-	/**
-	* Queue a cook by the book cancel (you might want to do this instead of calling cancel directly so that you don't have to be in the game thread when canceling
-	*/
+	/** Queue a cook by the book cancel (can be called from any thread) */
 	void QueueCancelCookByTheBook();
 
-	/**
-	* Cancel the currently running cook by the book (needs to be called from the game thread)
-	*/
+	/** Cancel the currently running cook by the book (needs to be called from the game thread) */
 	void CancelCookByTheBook();
 
 	/** Connect to the CookDirector host, Initialize this UCOTFS, and start the CookWorker session */
@@ -620,14 +626,14 @@ public:
 	/** Terminate the CookWorker session */
 	void ShutdownCookAsCookWorker();
 
-	/** Is the local CookOnTheFlyServer in a cook session in any CookMode? Used to restrict operations when cooking and reduce cputime when not cooking. */
+	/**
+	 * Is the local CookOnTheFlyServer in a cook session in any CookMode?
+	 * Used to restrict operations when cooking and reduce cputime when not cooking.
+	 */
 	bool IsInSession() const;
 	/** Is the local CookOnTheFlyServer in a cook session in CookByTheBook Mode? */
 	bool IsCookByTheBookRunning() const;
 
-	/**
-	* Get any packages which are in memory, these were probably required to be loaded because of the current package we are cooking, so we should probably cook them also
-	*/
 	UE_DEPRECATED(4.26, "Unsolicited packages are now added directly to the savequeue and are not marked as unsolicited")
 	TArray<UPackage*> GetUnsolicitedPackages(const TArray<const ITargetPlatform*>& TargetPlatforms) const;
 
@@ -642,28 +648,30 @@ public:
 	uint32 TickCookWorker();
 
 	/**
-	 * Execute CookByTheBook as a load of all packages followed by a save of all packages. This mode is sometimes faster than normal CBTB.
+	 * Execute CookByTheBook as a load of all packages followed by a save of all packages.
+	 * This mode is sometimes faster than normal CBTB.
 	 * Should be called at the beginning of the cook after StartCookByTheBook, and only if IsFullLoadAndSave.
 	 */
 	uint32 CookFullLoadAndSave();
 
-	/** Is the Local CookOnTheFLyServer using FullLoadAndSave? Used in SingleProcess cooks only, always false on CookWorkers or MPCook Directors. */
+	/**
+	 * Is the Local CookOnTheFlyServer using FullLoadAndSave?
+	 * FullLoadAndSave is used in SingleProcess cooks only, always false on CookWorkers or MPCook Directors.
+	 */
 	bool IsFullLoadAndSave() const;
 
-	/**
-	* Clear all the previously cooked data all cook requests from now on will be considered recook requests
-	*/
+	/** Clear all the previously cooked data all cook requests from now on will be considered recook requests */
 	void ClearAllCookedData();
 
-	/** Demote all PackageDatas in any queue back to Idle, and eliminate any pending requests. Used when canceling a cook. */
+	/** Demote PackageDatas in any queue back to Idle, and eliminate pending requests. Used when canceling a cook. */
 	void CancelAllQueues();
 
 
 	/**
-	* Clear any cached cooked platform data for a platform
-	*  call ClearCachedCookedPlatformData on all UObjects
-	* @param TargetPlatform platform to clear all the cached data for
-	*/
+	 * Clear any cached cooked platform data for a platform and call ClearCachedCookedPlatformData on all UObjects.
+	 * 
+	 * @param TargetPlatform platform to clear all the cached data for
+	 */
 	void ClearCachedCookedPlatformDataForPlatform(const ITargetPlatform* TargetPlatform);
 
 	UE_DEPRECATED(4.25, "Use version that takes const ITargetPlatform* instead")
@@ -671,10 +679,10 @@ public:
 
 
 	/**
-	* Clear all the previously cooked data for the platform passed in 
-	* 
-	* @param TargetPlatform the platform to clear the cooked packages for
-	*/
+	 * Clear all the previously cooked data for the platform passed in 
+	 * 
+	 * @param TargetPlatform the platform to clear the cooked packages for
+	 */
 	void ClearPlatformCookedData(const ITargetPlatform* TargetPlatform);
 
 	UE_DEPRECATED(4.25, "Use version that takes const ITargetPlatform* instead")
@@ -689,11 +697,11 @@ public:
 	void ResetCook(TConstArrayView<TPair<const ITargetPlatform*, bool>> TargetPlatforms);
 
 	/**
-	* Recompile any global shader changes 
-	* if any are detected then clear the cooked platform data so that they can be rebuilt
-	* 
-	* @return return true if shaders were recompiled
-	*/
+	 * Recompile any global shader changes 
+	 * if any are detected then clear the cooked platform data so that they can be rebuilt
+	 * 
+	 * @return return true if shaders were recompiled
+	 */
 	bool RecompileChangedShaders(const TArray<const ITargetPlatform*>& TargetPlatforms);
 
 	UE_DEPRECATED(4.25, "Use version that takes const ITargetPlatform* instead")
@@ -701,23 +709,24 @@ public:
 
 
 	/**
-	* Force stop whatever pending cook requests are going on and clear all the cooked data
-	* Note cook on the side / cook on the fly clients may not be able to recover from this if they are waiting on a cook request to complete
-	*/
+	 * Force stop whatever pending cook requests are going on and clear all the cooked data
+	 * Note CookOnTheFly clients may not be able to recover from this if they are waiting on a cook request to complete
+	 */
 	void StopAndClearCookedData();
 
 	void TickRequestManager();
 
-	/** 
-	* Process any shader recompile requests
-	*/
+	/** Process any shader recompile requests */
 	UE_DEPRECATED(5.1, "Ticking RecompileShaderRequests is now handled by the Tick function")
 	void TickRecompileShaderRequests() {}
 
 	UE_DEPRECATED(5.1, "Ticking RecompileShaderRequests is now handled by the Tick function")
 	bool HasRecompileShaderRequests() const;
 
-	/** Return whether the tick needs to take any action for the current session. If not, the session is done. Used for external managers of the cooker to know when to tick it. */
+	/**
+	 * Return whether the tick needs to take any action for the current session. If not, the session is done.
+	 * Used for external managers of the cooker to know when to tick it.
+	 */
 	bool HasRemainingWork() const;
 	void WaitForRequests(int TimeoutMs);
 
@@ -726,26 +735,44 @@ public:
 	/** Is the local CookOnTheFlyServer running in the editor? */
 	bool IsCookingInEditor() const;
 
-	/** Is the local CookOnTheFlyServer initialized to run in real time mode (where it needs to respect the timeslice), e.g. it is running in Editor? */
+	/** Is the local CookOnTheFlyServer initialized to run in real time mode (respects the timeslice), (e.g. in Editor)? */
 	bool IsRealtimeMode() const;
 
-	/** Is the local CookOnTheFlyServer initialized to run CookByTheBook (find all packages needed by the game and cook them)? */
+	/** Is the local CookOnTheFlyServer initialized to run CookByTheBook (find all required packages and cook them)? */
 	bool IsCookByTheBookMode() const;
-	/** Is the CookDirector (local CookOnTheFlyServer for SPCook or MPDirector, or the remote MPDirector for a CookWorker) initialized to run CookByTheBook? */
+	/**
+	 * Is the CookDirector (local CookOnTheFlyServer for SPCook or MPDirector, or the remote MPDirector for a
+	 * CookWorker) initialized to run CookByTheBook?
+	 */
 	bool IsDirectorCookByTheBook() const;
 
-	/** Is the local CookOnTheFlyServer in a mode that uses ShaderCodeLibraries rather than storing shaders in the Packages that use them? */
+	/**
+	 * Is the local CookOnTheFlyServer in a mode that uses ShaderCodeLibraries rather than storing shaders in
+	 * the Packages that use them?
+	 */
 	bool IsUsingShaderCodeLibrary() const;
 
-	/** Is the local CookOnTheFlyServer using ZenStore (cooked packages are stored in Zen's Cache storage rather than as loose files on disk)? */
+	/**
+	 * Is the local CookOnTheFlyServer using ZenStore (cooked packages are stored in Zen's Cache storage rather
+	 * than as loose files on disk)?
+	 */
 	bool IsUsingZenStore() const;
 
-	/** Is the local CookOnTheFlyServer initialized to run CookOnTheFly (accept connections from game executables, cook what they ask for)? */
+	/**
+	 * Is the local CookOnTheFlyServer initialized to run CookOnTheFly (accept connections from game executables,
+	 * cook what they ask for)?
+	 */
 	bool IsCookOnTheFlyMode() const;
-	/** Is the CookDirector (local CookOnTheFlyServer for SPCook or MPDirector, or the remote MPDirector for a CookWorker) initialized to run CookOnTheFly? */
+	/**
+	 * Is the CookDirector (local CookOnTheFlyServer for SPCook or MPDirector, or the remote MPDirector for a
+	 * CookWorker) initialized to run CookOnTheFly?
+	 */
 	bool IsDirectorCookOnTheFly() const;
 
-	/** Is the local CookOnTheFlyServer initialized to run as a CookWorker (connect to a Director cooker cooking CBTB or COTF, cook what it assigns to us)? */
+	/**
+	 * Is the local CookOnTheFlyServer initialized to run as a CookWorker (connect to a Director cooker cooking
+	 * CBTB or COTF, cook what it assigns to us)?
+	 */
 	bool IsCookWorkerMode() const;
 
 	virtual void BeginDestroy() override;
@@ -773,24 +800,26 @@ public:
 	void EvaluateGarbageCollectionResults(int32 NumObjectsBeforeGC, const FPlatformMemoryStats& MemStatsBeforeGC);
 
 	/**
-	* RequestPackage to be cooked
-	*
-	* @param StandardFileName FileName of the package in standard format as returned by FPaths::MakeStandardFilename
-	* @param TargetPlatforms The TargetPlatforms we want this package cooked for
-	* @param bForceFrontOfQueue should we put this package in the front of the cook queue (next to be processed) or at the end
-	*/
-	bool RequestPackage(const FName& StandardFileName, const TArrayView<const ITargetPlatform* const>& TargetPlatforms, const bool bForceFrontOfQueue);
+	 * RequestPackage to be cooked
+	 *
+	 * @param StandardFileName FileName of the package in standard format as returned by FPaths::MakeStandardFilename
+	 * @param TargetPlatforms The TargetPlatforms we want this package cooked for
+	 * @param bForceFrontOfQueue should package go to front of the cook queue (next to be processed) or the end
+	 */
+	bool RequestPackage(const FName& StandardFileName, const TArrayView<const ITargetPlatform* const>& TargetPlatforms,
+		const bool bForceFrontOfQueue);
 
 	UE_DEPRECATED(4.25, "Use Version that takes TArray<const ITargetPlatform*> instead")
-	bool RequestPackage(const FName& StandardFileName, const TArrayView<const FName>& TargetPlatformNames, const bool bForceFrontOfQueue);
+	bool RequestPackage(const FName& StandardFileName, const TArrayView<const FName>& TargetPlatformNames,
+		const bool bForceFrontOfQueue);
 
 	/**
-	* RequestPackage to be cooked
-	* This function can only be called while the cooker is in cook by the book mode
-	*
-	* @param StandardPackageFName name of the package in standard format as returned by FPaths::MakeStandardFilename
-	* @param bForceFrontOfQueue should we put this package in the front of the cook queue (next to be processed) or at the end
-	*/
+	 * RequestPackage to be cooked
+	 * This function can only be called while the cooker is in cook by the book mode
+	 *
+	 * @param StandardPackageFName name of the package in standard format as returned by FPaths::MakeStandardFilename
+	 * @param bForceFrontOfQueue should package go to front of the cook queue (next to be processed) or the end
+	 */
 	bool RequestPackage(const FName& StandardPackageFName, const bool bForceFrontOfQueue);
 
 
@@ -809,31 +838,35 @@ public:
 	*/
 	void MarkPackageDirtyForCooker( UPackage *Package, bool bAllowInSession = false );
 
-	/** Helper function for MarkPackageDirtyForCooker. Executes the MarkPackageDirtyForCooker operations that are only safe to execute from the scheduler's designated point for handling external requests. */
+	/**
+	 * Helper function for MarkPackageDirtyForCooker. Executes the MarkPackageDirtyForCooker operations that are
+	 * only safe to execute from the scheduler's designated point for handling external requests.
+	 */
 	void MarkPackageDirtyForCookerFromSchedulerThread(const FName& PackageName);
 
 	/**
-	* MaybeMarkPackageAsAlreadyLoaded
-	* Mark the package as already loaded if we have already cooked the package for all requested target platforms
-	* this hints to the objects on load that we don't need to load all our bulk data
-	* 
-	* @param Package to mark as not requiring reload
-	*/
+	 * MaybeMarkPackageAsAlreadyLoaded
+	 * Mark the package as already loaded if we have already cooked the package for all requested target platforms
+	 * this hints to the objects on load that we don't need to load all our bulk data
+	 * 
+	 * @param Package to mark as not requiring reload
+	 */
 	void MaybeMarkPackageAsAlreadyLoaded(UPackage* Package);
 
-	/**
-	* Callbacks from UObject globals
-	*/
+	// Callbacks from UObject globals
 	void PreGarbageCollect();
 	void CookerAddReferencedObjects(FReferenceCollector& Ar);
 	void PostGarbageCollect();
 
 private:
 
-	/** Is the local CookOnTheFlyServer initialized to run in CookOnTheFly AND using the legacy scheduling method that predates COTF2? */
+	/**
+	 * Is the local CookOnTheFlyServer initialized to run in CookOnTheFly AND using the legacy scheduling method
+	 * that predates COTF2?
+	 */
 	bool IsUsingLegacyCookOnTheFlyScheduling() const;
 
-	/** Update accumulators of editor data and consume editor changes since the last cook when starting cooks in the editor. */
+	/** Update accumulators of editor data and consume editor changes after last cook when starting cooks in the editor. */
 	void BeginCookEditorSystems();
 	void BeginCookPackageWriters(FBeginCookContext& BeginContext);
 	void InitializeSession();
@@ -858,27 +891,27 @@ private:
 	/** Start the collection of EDL diagnostics for the cook, if applicable. */
 	void BeginCookEDLCookInfo(FBeginCookContext& BeginContext);
 
-	/**
-	* Collect all the files which need to be cooked for a cook by the book session
-	*/
+	/** Collect all the files which need to be cooked for a cook by the book session */
 	void CollectFilesToCook(TArray<FName>& FilesInPath, TMap<FName, UE::Cook::FInstigator>& Instigators,
 		const TArray<FString>& CookMaps, const TArray<FString>& CookDirectories, 
 		const TArray<FString>& IniMapSections, ECookByTheBookOptions FilesToCookFlags,
 		const TArrayView<const ITargetPlatform* const>& TargetPlatforms,
 		const TMap<FName, TArray<FName>>& GameDefaultObjects);
 
-	/**
-	* Gets all game default objects for all platforms.
-	*/
-	static void GetGameDefaultObjects(const TArray<ITargetPlatform*>& TargetPlatforms, TMap<FName, TArray<FName>>& GameDefaultObjectsOut);
+	/** Gets all game default objects for all platforms. */
+	static void GetGameDefaultObjects(const TArray<ITargetPlatform*>& TargetPlatforms, TMap<FName,
+		TArray<FName>>& GameDefaultObjectsOut);
 
-	/* Collect filespackages that should not be cooked from ini settings and commandline. Does not include checking UAssetManager, which has to be queried later */
-	TArray<FName> GetNeverCookPackageFileNames(TArrayView<const FString> ExtraNeverCookDirectories = TArrayView<const FString>()) const; // const because it is not always called and should avoid side effects
+	/*
+	 * Collect filespackages that should not be cooked from ini settings and commandline.
+	 * Does not include checking UAssetManager, which has to be queried later
+	 * This function is const because it is not always called and should avoid side effects
+	 */
+	TArray<FName> GetNeverCookPackageFileNames(TArrayView<const FString> ExtraNeverCookDirectories
+		= TArrayView<const FString>()) const;
 
 
-	/**
-	* AddFileToCook add file to cook list 
-	*/
+	/** AddFileToCook add file to cook list */
 	void AddFileToCook( TArray<FName>& InOutFilesToCook, TMap<FName, UE::Cook::FInstigator>& InOutInstigators,
 		const FString &InFilename, const UE::Cook::FInstigator& Instigator) const;
 
@@ -887,37 +920,31 @@ private:
 
 	/** Invokes the necessary FShaderCodeLibrary functions to start cooking the shader code library. */
 	void BeginCookStartShaderCodeLibrary(FBeginCookContext& BeginContext);
-	/** Finishes async operations from ShaderCodeLibraryBeginCookStart, saves initial data to disk, and opens the library for the rest of the cook */
-	void BeginCookFinishShaderCodeLibrary(FBeginCookContext& BeginContext);
-    
 	/**
-	 * Opens Global shader library. Global shaderlib isn't split into chunks nor associated with the assets, so it a special case
+	 * Finishes async operations from ShaderCodeLibraryBeginCookStart, saves initial data to disk,
+	 * and opens the library for the rest of the cook
 	 */
+	void BeginCookFinishShaderCodeLibrary(FBeginCookContext& BeginContext);
+    
+	/** Opens Global shader library. Global shaderlib is a special case - no chunks and not connected to assets. */
 	void OpenGlobalShaderLibrary();
 
-	/**
-	 * Saves Global shader library. Global shaderlib isn't split into chunks nor associated with the assets, so it a special case
-	 */
+	/** Saves Global shader library. Global shaderlib is a special case - no chunks and not connected to assets. */
 	void SaveAndCloseGlobalShaderLibrary();
 
-	/**
-	 * Invokes the necessary FShaderCodeLibrary functions to open a named code library.
-	 */
+	/** Invokes the necessary FShaderCodeLibrary functions to open a named code library. */
 	void OpenShaderLibrary(FString const& Name);
     
-	/**
-	 * Invokes the necessary FShaderCodeLibrary functions to save and close a named code library.
-	 */
+	/** Invokes the necessary FShaderCodeLibrary functions to save and close a named code library. */
 	void SaveShaderLibrary(const ITargetPlatform* TargetPlatform, FString const& Name);
 
 	/**
-	* Calls the ShaderPipelineCacheToolsCommandlet to build a upipelinecache file from the stable pipeline cache (.spc, in the past textual .stablepc.csv) file, if any
-	*/
+	 * Calls the ShaderPipelineCacheToolsCommandlet to build a upipelinecache file from the stable pipeline cache
+	 * (.spc, in the past textual .stablepc.csv) file, if any
+	 */
 	void CreatePipelineCache(const ITargetPlatform* TargetPlatform, const FString& LibraryName);
 
-	/**
-	* Invokes the necessary FShaderCodeLibrary functions to clean out all the temporary files.
-	*/
+	/** Invokes the necessary FShaderCodeLibrary functions to clean out all the temporary files. */
 	void CleanShaderCodeLibraries();
 
 	void RegisterShaderChunkDataGenerator();
@@ -932,14 +959,14 @@ private:
 	void PrintFinishStats();
 
 	/**
-	* Get all the packages which are listed in asset registry passed in.  
-	*
-	* @param AssetRegistryPath path of the assetregistry.bin file to read
-	* @param bVerifyPackagesExist whether or not we should verify the packages exist on disk.
-	* @param bSkipUncookedPackages whether or not we should skip over packages that are in the AR but were not cooked/saved
-	* @param OutPackageDatas out list of packagename/filenames for packages contained in the asset registry file
-	* @return true if successfully read false otherwise
-	*/
+	 * Get all the packages which are listed in asset registry passed in.  
+	 *
+	 * @param AssetRegistryPath path of the assetregistry.bin file to read
+	 * @param bVerifyPackagesExist whether to verify the packages exist on disk.
+	 * @param bSkipUncookedPackages whether to skip over packages that are in the AR but were not cooked/saved
+	 * @param OutPackageDatas out list of packagename/filenames for packages contained in the asset registry file
+	 * @return true if successfully read false otherwise
+	 */
 	bool GetAllPackageFilenamesFromAssetRegistry( const FString& AssetRegistryPath, bool bVerifyPackagesExist,
 		 bool bSkipUncookedPackages, TArray<UE::Cook::FConstructPackageData>& OutPackageDatas) const;
 
@@ -960,7 +987,8 @@ private:
 	FBeginCookContext CreateBeginCookOnTheFlyContext(const FCookOnTheFlyStartupOptions& Options);
 	/** Construct the on-stack data for StartCook functions, based on the arguments when adding a COTF platform */
 	FBeginCookContext CreateAddPlatformContext(ITargetPlatform* TargetPlatform);
-	void GetCookOnTheFlyUnsolicitedFiles(const ITargetPlatform* TargetPlatform, const FString& PlatformName, TArray<FString>& UnsolicitedFiles, const FString& Filename, bool bIsCookable);
+	void GetCookOnTheFlyUnsolicitedFiles(const ITargetPlatform* TargetPlatform, const FString& PlatformName,
+		TArray<FString>& UnsolicitedFiles, const FString& Filename, bool bIsCookable);
 
 	//////////////////////////////////////////////////////////////////////////
 	// CookWorker specific functions
@@ -972,29 +1000,34 @@ private:
 	//////////////////////////////////////////////////////////////////////////
 	// general functions
 
-	/** Perform any special processing for freshly loaded packages 
-	 */
-	void ProcessUnsolicitedPackages(TArray<FName>* OutDiscoveredPackageNames = nullptr, TMap<FName, UE::Cook::FInstigator>* OutInstigators = nullptr);
+	/** Perform any special processing for freshly loaded packages */
+	void ProcessUnsolicitedPackages(TArray<FName>* OutDiscoveredPackageNames = nullptr,
+		TMap<FName, UE::Cook::FInstigator>* OutInstigators = nullptr);
 
 	/**
 	 * Loads a package and prepares it for cooking
 	 *  this is the same as a normal load but also ensures that the sublevels are loaded if they are streaming sublevels
 	 *
 	 * @param BuildFilename long package name of the package to load 
-	 * @param OutPackage UPackage of the package loaded, non-null on success, may be non-null on failure if the UPackage existed but had another failure
+	 * @param OutPackage UPackage of the package loaded, non-null on success, may be non-null on failure if
+	 *        the UPackage existed but had another failure
 	 * @param ReportingPackageData PackageData which caused the load, for e.g. generated packages
 	 * @return Whether LoadPackage was completely successful and the package can be cooked
 	 */
 	bool LoadPackageForCooking(UE::Cook::FPackageData& PackageData, UPackage*& OutPackage,
 		UE::Cook::FPackageData* ReportingPackageData = nullptr);
 
-	/** Read information about the previous cook from disk and set whether the current cook is iterative based on previous cook, config, and commandline. */
+	/**
+	 * Read information about the previous cook from disk and set whether the current cook is iterative based on
+	 * previous cook, config, and commandline.
+	 */
 	void LoadBeginCookIterativeFlags(FBeginCookContext& BeginContext);
-	void LoadBeginCookIterativeFlagsLocal(FBeginCookContext& BeginContext) const; // const because it is not always called and should avoid sideeffects
+	/** const because it is not always calledand should avoid sideeffects */
+	void LoadBeginCookIterativeFlagsLocal(FBeginCookContext& BeginContext) const;
 	/** Initialize the sandbox for a new cook session */
 	void BeginCookSandbox(FBeginCookContext& BeginContext);
 
-	/** Set parameters that rely on config settings from startup and don't depend on information from StartCook* functions. */
+	/** Set parameters that rely on config settings from startup and don't depend on StartCook* functions. */
 	void LoadInitializeConfigSettings(const FString& InOutputDirectoryOverride);
 	void SetInitializeConfigSettings(UE::Cook::FInitializeConfigSettings&& Settings);
 
@@ -1003,54 +1036,49 @@ private:
 	void SetBeginCookConfigSettings(FBeginCookContext& BeginContext, UE::Cook::FBeginCookConfigSettings&& Settings);
 	void SetNeverCookPackageConfigSettings(FBeginCookContext& BeginContext, UE::Cook::FBeginCookConfigSettings& Settings);
 
-	/**
-	* Finalize the package store
-	*/
+	/** Finalize the package store */
 	void FinalizePackageStore();
 
-	/**
-	* Empties SavePackageContexts and deletes the contents
-	*/
+	/** Empties SavePackageContexts and deletes the contents */
 	void ClearPackageStoreContexts();
 
-	/**
-	* Initialize shaders for the specified platforms when running cook on the fly.
-	*/
+	/** Initialize shaders for the specified platforms when running cook on the fly. */
 	void InitializeShadersForCookOnTheFly(const TArrayView<ITargetPlatform* const>& NewTargetPlatforms);
 
 	/**
-	* Some content plugins does not support all target platforms.
-	* Build up a map of unsupported packages per platform that can be checked before saving.
-	*/
+	 * Some content plugins does not support all target platforms.
+	 * Build up a map of unsupported packages per platform that can be checked before saving.
+	 * const because it is not always called and should avoid side effects
+	 */
 	void DiscoverPlatformSpecificNeverCookPackages(
 		const TArrayView<const ITargetPlatform* const>& TargetPlatforms, const TArray<FString>& UBTPlatformStrings,
-		UE::Cook::FBeginCookConfigSettings& Settings) const; // const because it is not always called and should avoid side effects
+		UE::Cook::FBeginCookConfigSettings& Settings) const;
 
 	/**
-	* GetDependentPackages
-	* get package dependencies according to the asset registry
-	* 
-	* @param Packages List of packages to use as the root set for dependency checking
-	* @param Found return value, all objects which package is dependent on
-	*/
+	 * GetDependentPackages
+	 * get package dependencies according to the asset registry
+	 * 
+	 * @param Packages List of packages to use as the root set for dependency checking
+	 * @param Found return value, all objects which package is dependent on
+	 */
 	void GetDependentPackages( const TSet<UPackage*>& Packages, TSet<FName>& Found);
 
 	/**
-	* GetDependentPackages
-	* get package dependencies according to the asset registry
-	*
-	* @param Root set of packages to use when looking for dependencies
-	* @param FoundPackages list of packages which were found
-	*/
+	 * GetDependentPackages
+	 * get package dependencies according to the asset registry
+	 *
+	 * @param Root set of packages to use when looking for dependencies
+	 * @param FoundPackages list of packages which were found
+	 */
 	void GetDependentPackages(const TSet<FName>& RootPackages, TSet<FName>& FoundPackages);
 
 	/**
-	* ContainsWorld
-	* use the asset registry to determine if a Package contains a UWorld or ULevel object
-	* 
-	* @param PackageName to return if it contains the a UWorld object or a ULevel
-	* @return true if the Package contains a UWorld or ULevel false otherwise
-	*/
+	 * ContainsWorld
+	 * use the asset registry to determine if a Package contains a UWorld or ULevel object
+	 * 
+	 * @param PackageName to return if it contains the a UWorld object or a ULevel
+	 * @return true if the Package contains a UWorld or ULevel false otherwise
+	 */
 	bool ContainsMap(const FName& PackageName) const;
 
 
@@ -1073,24 +1101,29 @@ private:
 	 * @param bIsPrecaching true if called for precaching
 	 * @return false if time slice was reached, true if all objects have had BeginCacheForCookedPlatformData called
 	 */
-	UE::Cook::EPollStatus PrepareSave(UE::Cook::FPackageData& PackageData, UE::Cook::FCookerTimer& Timer, bool bPrecaching);
-	UE::Cook::EPollStatus PrepareSaveInternal(UE::Cook::FPackageData& PackageData, UE::Cook::FCookerTimer& Timer, bool bPrecaching);
-	/** Call BeginCacheForCookedPlatformData on all of the objects in a PackageData or a GeneratorPackage's current round. */
+	UE::Cook::EPollStatus PrepareSave(UE::Cook::FPackageData& PackageData, UE::Cook::FCookerTimer& Timer,
+		bool bPrecaching);
+	UE::Cook::EPollStatus PrepareSaveInternal(UE::Cook::FPackageData& PackageData, UE::Cook::FCookerTimer& Timer,
+		bool bPrecaching);
+	/** Call BeginCacheForCookedPlatformData on all objects in a PackageData or a GeneratorPackage's current round. */
 	UE::Cook::EPollStatus CallBeginCacheOnObjects(UE::Cook::FPackageData& PackageData, UPackage* Package,
 		TArray<FWeakObjectPtr>& Objects, int32& NextIndex, UE::Cook::FCookerTimer& Timer);
 
 	/**
 	 * Frees all the memory used to call BeginCacheForCookedPlatformData on all the objects in PackageData.
-	 * If the calls were incomplete because the PackageData's save was cancelled, handles canceling them and leaving any required CancelManagers in GetPendingCookedPlatformDatas
+	 * If the calls were incomplete because the PackageData's save was cancelled, handles canceling them and
+	 * leaving any required CancelManagers in GetPendingCookedPlatformDatas
 	 * 
-	 * @param bCompletedSave If false, data that can not be efficiently recomputed will be preserved to try the save again. If true, all data will be wiped.
+	 * @param bCompletedSave If false, data that can not be efficiently recomputed will be preserved to try
+	 *        the save again. If true, all data will be wiped.
 	 * @param ReleaseSaveReason Why the save data is being released, allows specifying how much to tear down
 	 */
 	void ReleaseCookedPlatformData(UE::Cook::FPackageData& PackageData, UE::Cook::EReleaseSaveReason ReleaseSaveReason);
 
 	/**
 	 * Poll the GetPendingCookedPlatformDatas and release their resources when they are complete.
-	 * This is done inside of PumpSaveQueue, but is also required after a cancelled cook so that references to the pending objects will be eventually dropped.
+	 * This is done inside of PumpSaveQueue, but is also required after a cancelled cook so that references to
+	 * the pending objects will be eventually dropped.
 	 */
 	void TickCancels();
 
@@ -1100,7 +1133,8 @@ private:
 	* @param IniVersionStrings return list of the important current ini version strings
 	* @return false if function fails (should assume all platforms are out of date)
 	*/
-	bool GetCurrentIniVersionStrings( const ITargetPlatform* TargetPlatform, UE::Cook::FIniSettingContainer& IniVersionStrings ) const;
+	bool GetCurrentIniVersionStrings( const ITargetPlatform* TargetPlatform,
+		UE::Cook::FIniSettingContainer& IniVersionStrings ) const;
 
 	/**
 	* GetCookedIniVersionStrings gets the ini version strings used in previous cook for specified target platform
@@ -1108,39 +1142,45 @@ private:
 	* @param IniVersionStrings return list of the previous cooks ini version strings
 	* @return false if function fails to find the ini version strings
 	*/
-	bool GetCookedIniVersionStrings( const ITargetPlatform* TargetPlatform, UE::Cook::FIniSettingContainer& IniVersionStrings, TMap<FString, FString>& AdditionalStrings ) const;
+	bool GetCookedIniVersionStrings( const ITargetPlatform* TargetPlatform,
+		UE::Cook::FIniSettingContainer& IniVersionStrings, TMap<FString, FString>& AdditionalStrings ) const;
 
-	/** Test the CurrentCookSettings against the previous cooksettings to decide if we have to wipe the previous cook even when running iteratively. */
-	bool ArePreviousCookSettingsCompatible(const TMap<FName, FString>& CurrentCookSettings, const ITargetPlatform* TargetPlatform) const;
+	/**
+	 * Test the CurrentCookSettings against the previous cooksettings to decide if we have to wipe the
+	 * previous cook even when running iteratively.
+	 */
+	bool ArePreviousCookSettingsCompatible(const TMap<FName, FString>& CurrentCookSettings,
+		const ITargetPlatform* TargetPlatform) const;
 	/** Save the CurrentCookSettings into the output directory. */
 	void SaveCookSettings(const TMap<FName, FString>& CurrentCookSettings, const ITargetPlatform* TargetPlatform);
-	/** Populate a map suitable for saving as an ini with the current value of the Cook settings that need to be tested for compatibility. */
+	/**
+	 * Populate a map suitable for saving as an ini with the current value of the Cook settings that need to be
+	 * tested for compatibility.
+	 */
 	TMap<FName, FString> CalculateCookSettingStrings() const;
 	/** Get the path to the CookSettings metadata file for the given platform and current output directory. */
 	FString GetCookSettingsFileName(const ITargetPlatform* TargetPlatform) const;
 
 	/**
-	* Convert a path to a full sandbox path
-	* is effected by the cooking dlc settings
-	* This function should be used instead of calling the FSandbox Sandbox->ConvertToSandboxPath functions
-	*/
+	 * Convert a path to a full sandbox path. Is effected by the cooking dlc settings.
+	 * This function should be used instead of calling the FSandbox Sandbox->ConvertToSandboxPath functions.
+	 */
 	FString ConvertToFullSandboxPath( const FString &FileName, bool bForWrite = false ) const;
 	FString ConvertToFullSandboxPath( const FString &FileName, bool bForWrite, const FString& PlatformName ) const;
 
 	/**
-	* GetSandboxAssetRegistryFilename
-	* 
-	* return full path of the asset registry in the sandbox
-	*/
+	 * GetSandboxAssetRegistryFilename
+	 * 
+	 * @return full path of the asset registry in the sandbox
+	 */
 	const FString GetSandboxAssetRegistryFilename();
 
 	const FString GetCookedAssetRegistryFilename(const FString& PlatformName);
 
 	/**
-	* Get the sandbox root directory for that platform
-	* is effected by the CookingDlc settings
-	* This should be used instead of calling the Sandbox function
-	*/
+	 * Get the sandbox root directory for that platform. Is effected by the CookingDlc settings.
+	 * This should be used instead of calling the Sandbox function.
+	 */
 	FString GetSandboxDirectory( const FString& PlatformName ) const;
 
 	/* Create the delete-old-cooked-directory helper.*/
@@ -1152,43 +1192,37 @@ private:
 	/** Is the local CookOnTheFlyServer cooking a Project+Engine+EmbeddedPlugin Patch based on a previous Release? */
 	bool IsCookingAgainstFixedBase() const;
 
-	/**
-	* Returns whether or not we should populate the Asset Registry using the main game content
-	*/
+	/** Returns whether or not we should populate the Asset Registry using the main game content */
 	bool ShouldPopulateFullAssetRegistry() const;
 
 	/**
-	* GetBaseDirectoryForDLC
-	* 
-	* @return return the path to the DLC
-	*/
+	 * GetBaseDirectoryForDLC
+	 * 
+	 * @return return the path to the DLC
+	 */
 	FString GetBaseDirectoryForDLC() const;
 
 	FString GetContentDirectoryForDLC() const;
 
-	/** Is the local CookOnTheFlyServer cooking a Project+Engine+EmbeddedPlugin Release that can be used as a base for future DLC or Patches? */
-	bool IsCreatingReleaseVersion();
-
 	/**
-	* Loads the cooked ini version settings maps into the Ini settings cache
-	* 
-	* @param TargetPlatforms to look for ini settings for
-	*/
-	//bool CacheIniVersionStringsMap( const ITargetPlatform* TargetPlatform ) const;
+	 * Is the local CookOnTheFlyServer cooking a Project+Engine+EmbeddedPlugin Release that can be used as a
+	 * base for future DLC or Patches?
+	 */
+	bool IsCreatingReleaseVersion();
 
 	/**
-	* Checks if important ini settings have changed since last cook for each target platform 
-	* 
-	* @param TargetPlatforms to check if out of date
-	* @param OutOfDateTargetPlatforms return list of out of date target platforms which should be cleaned
-	*/
+	 * Checks if important ini settings have changed since last cook for each target platform 
+	 * 
+	 * @param TargetPlatforms to check if out of date
+	 * @param OutOfDateTargetPlatforms return list of out of date target platforms which should be cleaned
+	 */
 	bool IniSettingsOutOfDate( const ITargetPlatform* TargetPlatform ) const;
 
 	/**
-	* Saves ini settings which are in the memory cache to the hard drive in ini files
-	*
-	* @param TargetPlatforms to save
-	*/
+	 * Saves ini settings which are in the memory cache to the hard drive in ini files
+	 *
+	 * @param TargetPlatforms to save
+	 */
 	bool SaveCurrentIniSettings( const ITargetPlatform* TargetPlatform ) const;
 
 
@@ -1204,11 +1238,12 @@ private:
 	friend class UE::Cook::FSaveCookedPackageContext;
 
 	/**
-	*  Save the global shader map
-	*  
-	*  @param	Platforms		List of platforms to make global shader maps for
-	*/
-	void SaveGlobalShaderMapFiles(const TArrayView<const ITargetPlatform* const>& Platforms, ODSCRecompileCommand RecompileCommand);
+	 * Save the global shader map
+	 *
+	 * @param Platforms List of platforms to make global shader maps for
+	 */
+	void SaveGlobalShaderMapFiles(const TArrayView<const ITargetPlatform* const>& Platforms,
+		ODSCRecompileCommand RecompileCommand);
 
 
 	/** Create sandbox file in directory using current settings supplied */
@@ -1257,7 +1292,8 @@ private:
 	/** Try creating (or finding from earlier creation) the generated package for later population */
 	UPackage* TryCreateGeneratedPackage(UE::Cook::FGeneratorPackage& Generator, UE::Cook::FCookGenerationInfo& GeneratedInfo);
 	/** Try calling the splitter's populate to create the package */
-	UE::Cook::EPollStatus TryPopulateGeneratedPackage(UE::Cook::FGeneratorPackage& Generator, UE::Cook::FCookGenerationInfo& GeneratedInfo);
+	UE::Cook::EPollStatus TryPopulateGeneratedPackage(UE::Cook::FGeneratorPackage& Generator,
+		UE::Cook::FCookGenerationInfo& GeneratedInfo);
 
 	ICookedPackageWriter& FindOrCreatePackageWriter(const ITargetPlatform* TargetPlatform);
 	const ICookedPackageWriter* FindPackageWriter(const ITargetPlatform* TargetPlatform) const;
@@ -1274,22 +1310,28 @@ private:
 	UE::Cook::FStatHistoryInt NumObjectsHistory;
 	UE::Cook::FStatHistoryInt VirtualMemoryHistory;
 
-	/** This is set to true when the decision about which packages we need to cook changes because e.g. a platform was added to the sessionplatforms. */
+	/** True when the which packages we need to cook changes because e.g. a platform was added to the sessionplatforms. */
 	bool bPackageFilterDirty = false;
-	/** Set to true when PumpLoads has detected it is blocked on async work and CookOnTheFlyServer should do work elsewhere. */
+	/** True when PumpLoads has detected it is blocked on async work and CookOnTheFlyServer should do work elsewhere. */
 	bool bLoadBusy = false;
-	/** Set to true when PumpSaves has detected it is blocked on async work and CookOnTheFlyServer should do work elsewhere. */
+	/** True when PumpSaves has detected it is blocked on async work and CookOnTheFlyServer should do work elsewhere. */
 	bool bSaveBusy = false;
-	/** If preloading is enabled, we call TryPreload until it returns true before sending the package to LoadReady, otherwise we skip TryPreload and it goes immediately. */
+	/**
+	 * If preloading is enabled, we call TryPreload until it returns true before sending the package to LoadReady,
+	 * otherwise we skip TryPreload and it goes immediately.
+	 */
 	bool bPreloadingEnabled = false;
-	/** If enabled, we load/save TargetDomainKey hashes and use those to test whether packages have already been cooked in hybrid-iterative builds. */
+	/**
+	 * If enabled, we load/save TargetDomainKey hashes and use those to test whether packages have already been
+	 * cooked in hybrid-iterative builds.
+	 */
 	bool bHybridIterativeEnabled = true;
 	/** Test mode for the debug of hybrid iterative dependencies. */
 	bool bHybridIterativeDebug = false;
 	bool bFirstCookInThisProcessInitialized = false;
 	bool bFirstCookInThisProcess = true;
 	bool bImportBehaviorCallbackInstalled = false;
-	/** True if a session is started for any mode. If the session is started, then we have at least one TargetPlatform specified. */
+	/** True if a session is started for any mode. If started, then we have at least one TargetPlatform specified. */
 	bool bSessionRunning = false;
 	/** Whether we're using ZenStore for storage of cook results. If false, we are using LooseCookedPackageWriter. */
 	bool bZenStore = false;
@@ -1306,7 +1348,8 @@ private:
 	/** Tracking for the ticking of tickable cook objects */
 	double LastCookableObjectTickTime = 0.;
 
-	// These helper structs are all TUniquePtr rather than inline members so that we can keep their headers private.  See class header comments for their purpose.
+	// These structs are TUniquePtr rather than inline members so we can keep their headers private.
+	// See class header comments for their purpose.
 	TUniquePtr<UE::Cook::FPackageTracker> PackageTracker;
 	TUniquePtr<UE::Cook::FPackageDatas> PackageDatas;
 	TUniquePtr<UE::Cook::IWorkerRequests> WorkerRequests;
@@ -1315,7 +1358,10 @@ private:
 	TUniquePtr<UE::Cook::FCookWorkerClient> CookWorkerClient;
 
 	TArray<UE::Cook::FCookSavePackageContext*> SavePackageContexts;
-	/** Objects that were collected during the single-threaded PreGarbageCollect callback and that should be reported as referenced in CookerAddReferencedObjects. */
+	/**
+	 * Objects that were collected during the single-threaded PreGarbageCollect callback and that should be reported
+	 * as referenced in CookerAddReferencedObjects.
+	 */
 	TArray<UObject*> GCKeepObjects;
 	UE::Cook::FPackageData* SavingPackageData = nullptr;
 	/** Helper struct for running cooking in diagnostic modes */
diff --git a/Engine/Source/Editor/UnrealEd/Private/Commandlets/CookCommandlet.cpp b/Engine/Source/Editor/UnrealEd/Private/Commandlets/CookCommandlet.cpp
--- a/Engine/Source/Editor/UnrealEd/Private/Commandlets/CookCommandlet.cpp
+++ b/Engine/Source/Editor/UnrealEd/Private/Commandlets/CookCommandlet.cpp
@@ -50,357 +50,359 @@
 DEFINE_LOG_CATEGORY_STATIC(LogCookCommandlet, Log, All);
 
 #if ENABLE_COOK_STATS
-#include "ProfilingDebugging/ScopedTimers.h"
+#include "AnalyticsET.h"
 #include "AnalyticsEventAttribute.h"
 #include "IAnalyticsProviderET.h"
-#include "AnalyticsET.h"
-#include "Virtualization/VirtualizationSystem.h"
+#include "ProfilingDebugging/ScopedTimers.h"
 #include "StudioAnalytics.h"
+#include "Virtualization/VirtualizationSystem.h"
 
 namespace DetailedCookStats
 {
-	FString CookProject;
-	FString CookCultures;
-	FString CookLabel;
-	FString TargetPlatforms;
-	double CookWallTimeSec = 0.0;
-	double StartupWallTimeSec = 0.0;
-	double CookByTheBookTimeSec = 0.0;
-	double StartCookByTheBookTimeSec = 0.0;
-	extern double TickCookOnTheSideTimeSec;
-	extern double TickCookOnTheSideLoadPackagesTimeSec;
-	extern double TickCookOnTheSideResolveRedirectorsTimeSec;
-	extern double TickCookOnTheSideSaveCookedPackageTimeSec;
-	extern double TickCookOnTheSidePrepareSaveTimeSec;
-	extern double GameCookModificationDelegateTimeSec;
-	double TickLoopGCTimeSec = 0.0;
-	double TickLoopRecompileShaderRequestsTimeSec = 0.0;
-	double TickLoopShaderProcessAsyncResultsTimeSec = 0.0;
-	double TickLoopProcessDeferredCommandsTimeSec = 0.0;
-	double TickLoopTickCommandletStatsTimeSec = 0.0;
-	double TickLoopFlushRenderingCommandsTimeSec = 0.0;
-	bool IsCookAll = false;
-	bool IsCookOnTheFly = false;
-	bool IsIterativeCook = false;
-	bool IsFastCook = false;
-	bool IsUnversioned = false;
-
-	FCookStatsManager::FAutoRegisterCallback RegisterCookStats([](FCookStatsManager::AddStatFuncRef AddStat)
-	{
-		const FString StatName(TEXT("Cook.Profile"));
-		#define ADD_COOK_STAT_FLT(Path, Name) AddStat(StatName, FCookStatsManager::CreateKeyValueArray(TEXT("Path"), TEXT(Path), TEXT(#Name), Name))
-		ADD_COOK_STAT_FLT(" 0", CookWallTimeSec);
-		ADD_COOK_STAT_FLT(" 0. 0", StartupWallTimeSec);
-		ADD_COOK_STAT_FLT(" 0. 1", CookByTheBookTimeSec);
-		ADD_COOK_STAT_FLT(" 0. 1. 0", StartCookByTheBookTimeSec);
-		ADD_COOK_STAT_FLT(" 0. 1. 0. 0", GameCookModificationDelegateTimeSec);
-		ADD_COOK_STAT_FLT(" 0. 1. 1", TickCookOnTheSideTimeSec);
-		ADD_COOK_STAT_FLT(" 0. 1. 1. 0", TickCookOnTheSideLoadPackagesTimeSec);
-		ADD_COOK_STAT_FLT(" 0. 1. 1. 1", TickCookOnTheSideSaveCookedPackageTimeSec);
-		ADD_COOK_STAT_FLT(" 0. 1. 1. 1. 0", TickCookOnTheSideResolveRedirectorsTimeSec);
-		ADD_COOK_STAT_FLT(" 0. 1. 1. 2", TickCookOnTheSidePrepareSaveTimeSec);
-		ADD_COOK_STAT_FLT(" 0. 1. 2", TickLoopGCTimeSec);
-		ADD_COOK_STAT_FLT(" 0. 1. 3", TickLoopRecompileShaderRequestsTimeSec);
-		ADD_COOK_STAT_FLT(" 0. 1. 4", TickLoopShaderProcessAsyncResultsTimeSec);
-		ADD_COOK_STAT_FLT(" 0. 1. 5", TickLoopProcessDeferredCommandsTimeSec);
-		ADD_COOK_STAT_FLT(" 0. 1. 6", TickLoopTickCommandletStatsTimeSec);
-		ADD_COOK_STAT_FLT(" 0. 1. 7", TickLoopFlushRenderingCommandsTimeSec);
-		ADD_COOK_STAT_FLT(" 0. 1. 8", TargetPlatforms);
-		ADD_COOK_STAT_FLT(" 0. 1. 9", CookProject);
-		ADD_COOK_STAT_FLT(" 0. 1. 10", CookCultures);
-		ADD_COOK_STAT_FLT(" 0. 1. 11", IsCookAll);
-		ADD_COOK_STAT_FLT(" 0. 1. 12", IsCookOnTheFly);
-		ADD_COOK_STAT_FLT(" 0. 1. 13", IsIterativeCook);
-		ADD_COOK_STAT_FLT(" 0. 1. 14", IsUnversioned);
-		ADD_COOK_STAT_FLT(" 0. 1. 15", CookLabel);
-		ADD_COOK_STAT_FLT(" 0. 1. 16", IsFastCook);
+
+FString CookProject;
+FString CookCultures;
+FString CookLabel;
+FString TargetPlatforms;
+double CookWallTimeSec = 0.0;
+double StartupWallTimeSec = 0.0;
+double CookByTheBookTimeSec = 0.0;
+double StartCookByTheBookTimeSec = 0.0;
+extern double TickCookOnTheSideTimeSec;
+extern double TickCookOnTheSideLoadPackagesTimeSec;
+extern double TickCookOnTheSideResolveRedirectorsTimeSec;
+extern double TickCookOnTheSideSaveCookedPackageTimeSec;
+extern double TickCookOnTheSidePrepareSaveTimeSec;
+extern double GameCookModificationDelegateTimeSec;
+double TickLoopGCTimeSec = 0.0;
+double TickLoopRecompileShaderRequestsTimeSec = 0.0;
+double TickLoopShaderProcessAsyncResultsTimeSec = 0.0;
+double TickLoopProcessDeferredCommandsTimeSec = 0.0;
+double TickLoopTickCommandletStatsTimeSec = 0.0;
+double TickLoopFlushRenderingCommandsTimeSec = 0.0;
+bool IsCookAll = false;
+bool IsCookOnTheFly = false;
+bool IsIterativeCook = false;
+bool IsFastCook = false;
+bool IsUnversioned = false;
+
+FCookStatsManager::FAutoRegisterCallback RegisterCookStats([](FCookStatsManager::AddStatFuncRef AddStat)
+{
+	const FString StatName(TEXT("Cook.Profile"));
+	#define ADD_COOK_STAT_FLT(Path, Name) AddStat(StatName, FCookStatsManager::CreateKeyValueArray(TEXT("Path"), TEXT(Path), TEXT(#Name), Name))
+	ADD_COOK_STAT_FLT(" 0", CookWallTimeSec);
+	ADD_COOK_STAT_FLT(" 0. 0", StartupWallTimeSec);
+	ADD_COOK_STAT_FLT(" 0. 1", CookByTheBookTimeSec);
+	ADD_COOK_STAT_FLT(" 0. 1. 0", StartCookByTheBookTimeSec);
+	ADD_COOK_STAT_FLT(" 0. 1. 0. 0", GameCookModificationDelegateTimeSec);
+	ADD_COOK_STAT_FLT(" 0. 1. 1", TickCookOnTheSideTimeSec);
+	ADD_COOK_STAT_FLT(" 0. 1. 1. 0", TickCookOnTheSideLoadPackagesTimeSec);
+	ADD_COOK_STAT_FLT(" 0. 1. 1. 1", TickCookOnTheSideSaveCookedPackageTimeSec);
+	ADD_COOK_STAT_FLT(" 0. 1. 1. 1. 0", TickCookOnTheSideResolveRedirectorsTimeSec);
+	ADD_COOK_STAT_FLT(" 0. 1. 1. 2", TickCookOnTheSidePrepareSaveTimeSec);
+	ADD_COOK_STAT_FLT(" 0. 1. 2", TickLoopGCTimeSec);
+	ADD_COOK_STAT_FLT(" 0. 1. 3", TickLoopRecompileShaderRequestsTimeSec);
+	ADD_COOK_STAT_FLT(" 0. 1. 4", TickLoopShaderProcessAsyncResultsTimeSec);
+	ADD_COOK_STAT_FLT(" 0. 1. 5", TickLoopProcessDeferredCommandsTimeSec);
+	ADD_COOK_STAT_FLT(" 0. 1. 6", TickLoopTickCommandletStatsTimeSec);
+	ADD_COOK_STAT_FLT(" 0. 1. 7", TickLoopFlushRenderingCommandsTimeSec);
+	ADD_COOK_STAT_FLT(" 0. 1. 8", TargetPlatforms);
+	ADD_COOK_STAT_FLT(" 0. 1. 9", CookProject);
+	ADD_COOK_STAT_FLT(" 0. 1. 10", CookCultures);
+	ADD_COOK_STAT_FLT(" 0. 1. 11", IsCookAll);
+	ADD_COOK_STAT_FLT(" 0. 1. 12", IsCookOnTheFly);
+	ADD_COOK_STAT_FLT(" 0. 1. 13", IsIterativeCook);
+	ADD_COOK_STAT_FLT(" 0. 1. 14", IsUnversioned);
+	ADD_COOK_STAT_FLT(" 0. 1. 15", CookLabel);
+	ADD_COOK_STAT_FLT(" 0. 1. 16", IsFastCook);
 		
-		#undef ADD_COOK_STAT_FLT
-	});
+	#undef ADD_COOK_STAT_FLT
+});
 
-	static void LogCookStats(const FString& CookCmdLine)
+static void LogCookStats(const FString& CookCmdLine)
+{
+	if (FStudioAnalytics::IsAvailable())
 	{
-		if (FStudioAnalytics::IsAvailable())
-		{
 
-			// convert filtered stats directly to an analytics event
-			TArray<FAnalyticsEventAttribute> StatAttrs;
+		// convert filtered stats directly to an analytics event
+		TArray<FAnalyticsEventAttribute> StatAttrs;
 
-			// Sends each cook stat to the studio analytics system.
-			auto SendCookStatsToAnalytics = [&StatAttrs](const FString& StatName, const TArray<FCookStatsManager::StringKeyValue>& StatAttributes)
+		// Sends each cook stat to the studio analytics system.
+		auto SendCookStatsToAnalytics = [&StatAttrs](const FString& StatName, const TArray<FCookStatsManager::StringKeyValue>& StatAttributes)
+		{
+			for (const auto& Attr : StatAttributes)
 			{
-				for (const auto& Attr : StatAttributes)
-				{
-					FString FormattedAttrName = StatName + "." + Attr.Key;
+				FString FormattedAttrName = StatName + "." + Attr.Key;
 
-					StatAttrs.Emplace(FormattedAttrName, Attr.Value);
-				}
-			};
+				StatAttrs.Emplace(FormattedAttrName, Attr.Value);
+			}
+		};
 
-			// Now actually grab the stats 
-			FCookStatsManager::LogCookStats(SendCookStatsToAnalytics);
+		// Now actually grab the stats 
+		FCookStatsManager::LogCookStats(SendCookStatsToAnalytics);
 
-			// Record them all under cooking event
-			FStudioAnalytics::GetProvider().RecordEvent(TEXT("Core.Cooking"), StatAttrs);
+		// Record them all under cooking event
+		FStudioAnalytics::GetProvider().RecordEvent(TEXT("Core.Cooking"), StatAttrs);
 
-			FStudioAnalytics::GetProvider().BlockUntilFlushed(60.0f);
-		}
+		FStudioAnalytics::GetProvider().BlockUntilFlushed(60.0f);
+	}
 
-		bool bSendCookAnalytics = false;
-		GConfig->GetBool(TEXT("CookAnalytics"), TEXT("SendAnalytics"), bSendCookAnalytics, GEngineIni);
+	bool bSendCookAnalytics = false;
+	GConfig->GetBool(TEXT("CookAnalytics"), TEXT("SendAnalytics"), bSendCookAnalytics, GEngineIni);
 
-		if (GIsBuildMachine || FParse::Param(FCommandLine::Get(), TEXT("SendCookAnalytics")) || bSendCookAnalytics)
+	if (GIsBuildMachine || FParse::Param(FCommandLine::Get(), TEXT("SendCookAnalytics")) || bSendCookAnalytics)
+	{
+		FString APIServerET;
+		if (GConfig->GetString(TEXT("CookAnalytics"), TEXT("APIServer"), APIServerET, GEngineIni))
 		{
-			FString APIServerET;
-			if (GConfig->GetString(TEXT("CookAnalytics"), TEXT("APIServer"), APIServerET, GEngineIni))
-			{
-				FString AppId(TEXT("Cook"));
-				bool bUseLegacyCookProtocol = !GConfig->GetString(TEXT("CookAnalytics"), TEXT("AppId"), AppId, GEngineIni);
+			FString AppId(TEXT("Cook"));
+			bool bUseLegacyCookProtocol = !GConfig->GetString(TEXT("CookAnalytics"), TEXT("AppId"), AppId, GEngineIni);
 
-				// Optionally create an analytics provider to send stats to for central collection.
-				TSharedPtr<IAnalyticsProviderET> CookAnalytics = FAnalyticsET::Get().CreateAnalyticsProvider(FAnalyticsET::Config(AppId, APIServerET, FString(), bUseLegacyCookProtocol));
-				if (CookAnalytics.IsValid())
+			// Optionally create an analytics provider to send stats to for central collection.
+			TSharedPtr<IAnalyticsProviderET> CookAnalytics = FAnalyticsET::Get().CreateAnalyticsProvider(FAnalyticsET::Config(AppId, APIServerET, FString(), bUseLegacyCookProtocol));
+			if (CookAnalytics.IsValid())
+			{
+				CookAnalytics->SetUserID(FString::Printf(TEXT("%s\\%s"), FPlatformProcess::ComputerName(), FPlatformProcess::UserName(false)));
+				CookAnalytics->StartSession(MakeAnalyticsEventAttributeArray(
+					TEXT("Project"), CookProject,
+					TEXT("CmdLine"), CookCmdLine,
+					TEXT("IsBuildMachine"), GIsBuildMachine,
+					TEXT("TargetPlatforms"), TargetPlatforms
+				));
+
+				TArray<FString> CookStatsToSend;
+				const bool bFilterStats = GConfig->GetArray(TEXT("CookAnalytics"), TEXT("CookStats"), CookStatsToSend, GEngineIni) > 0;
+				// Sends each cook stat to the analytics provider.
+				auto SendCookStatsToAnalytics = [CookAnalytics, &CookStatsToSend, bFilterStats](const FString& StatName, const TArray<FCookStatsManager::StringKeyValue>& StatAttributes)
 				{
-					CookAnalytics->SetUserID(FString::Printf(TEXT("%s\\%s"), FPlatformProcess::ComputerName(), FPlatformProcess::UserName(false)));
-					CookAnalytics->StartSession(MakeAnalyticsEventAttributeArray(
-						TEXT("Project"), CookProject,
-						TEXT("CmdLine"), CookCmdLine,
-						TEXT("IsBuildMachine"), GIsBuildMachine,
-						TEXT("TargetPlatforms"), TargetPlatforms
-					));
-
-					TArray<FString> CookStatsToSend;
-					const bool bFilterStats = GConfig->GetArray(TEXT("CookAnalytics"), TEXT("CookStats"), CookStatsToSend, GEngineIni) > 0;
-					// Sends each cook stat to the analytics provider.
-					auto SendCookStatsToAnalytics = [CookAnalytics, &CookStatsToSend, bFilterStats](const FString& StatName, const TArray<FCookStatsManager::StringKeyValue>& StatAttributes)
+					if (!bFilterStats || CookStatsToSend.Contains(StatName))
 					{
-						if (!bFilterStats || CookStatsToSend.Contains(StatName))
-						{
-							// convert filtered stats directly to an analytics event
-							TArray<FAnalyticsEventAttribute> StatAttrs;
-							StatAttrs.Reset(StatAttributes.Num());
-							for (const auto& Attr : StatAttributes)
-							{
-								StatAttrs.Emplace(Attr.Key, Attr.Value);
-							}
-							CookAnalytics->RecordEvent(StatName, StatAttrs);
-						}
-						else
+						// convert filtered stats directly to an analytics event
+						TArray<FAnalyticsEventAttribute> StatAttrs;
+						StatAttrs.Reset(StatAttributes.Num());
+						for (const auto& Attr : StatAttributes)
 						{
-							UE_LOG(LogCookCommandlet, Verbose, TEXT("[%s] not present in analytics CookStats filter"), *StatName);
+							StatAttrs.Emplace(Attr.Key, Attr.Value);
 						}
-					};
-					FCookStatsManager::LogCookStats(SendCookStatsToAnalytics);
-				}
+						CookAnalytics->RecordEvent(StatName, StatAttrs);
+					}
+					else
+					{
+						UE_LOG(LogCookCommandlet, Verbose, TEXT("[%s] not present in analytics CookStats filter"), *StatName);
+					}
+				};
+				FCookStatsManager::LogCookStats(SendCookStatsToAnalytics);
 			}
 		}
+	}
 
-		/** Used for custom logging of DDC Resource usage stats. */
-		struct FDDCResourceUsageStat
+	/** Used for custom logging of DDC Resource usage stats. */
+	struct FDDCResourceUsageStat
+	{
+	public:
+		FDDCResourceUsageStat(FString InAssetType, double InTotalTimeSec, bool bIsGameThreadTime, double InSizeMB, int64 InAssetsBuilt) : AssetType(MoveTemp(InAssetType)), TotalTimeSec(InTotalTimeSec), GameThreadTimeSec(bIsGameThreadTime ? InTotalTimeSec : 0.0), SizeMB(InSizeMB), AssetsBuilt(InAssetsBuilt) {}
+		void Accumulate(const FDDCResourceUsageStat& OtherStat)
 		{
-		public:
-			FDDCResourceUsageStat(FString InAssetType, double InTotalTimeSec, bool bIsGameThreadTime, double InSizeMB, int64 InAssetsBuilt) : AssetType(MoveTemp(InAssetType)), TotalTimeSec(InTotalTimeSec), GameThreadTimeSec(bIsGameThreadTime ? InTotalTimeSec : 0.0), SizeMB(InSizeMB), AssetsBuilt(InAssetsBuilt) {}
-			void Accumulate(const FDDCResourceUsageStat& OtherStat)
-			{
-				TotalTimeSec += OtherStat.TotalTimeSec;
-				GameThreadTimeSec += OtherStat.GameThreadTimeSec;
-				SizeMB += OtherStat.SizeMB;
-				AssetsBuilt += OtherStat.AssetsBuilt;
-			}
-			FString AssetType;
-			double TotalTimeSec;
-			double GameThreadTimeSec;
-			double SizeMB;
-			int64 AssetsBuilt;
-		};
+			TotalTimeSec += OtherStat.TotalTimeSec;
+			GameThreadTimeSec += OtherStat.GameThreadTimeSec;
+			SizeMB += OtherStat.SizeMB;
+			AssetsBuilt += OtherStat.AssetsBuilt;
+		}
+		FString AssetType;
+		double TotalTimeSec;
+		double GameThreadTimeSec;
+		double SizeMB;
+		int64 AssetsBuilt;
+	};
 
-		/** Used for custom TSet comparison of DDC Resource usage stats. */
-		struct FDDCResourceUsageStatKeyFuncs : BaseKeyFuncs<FDDCResourceUsageStat, FString, false>
-		{
-			static const FString& GetSetKey(const FDDCResourceUsageStat& Element) { return Element.AssetType; }
-			static bool Matches(const FString& A, const FString& B) { return A == B; }
-			static uint32 GetKeyHash(const FString& Key) { return GetTypeHash(Key); }
-		};
+	/** Used for custom TSet comparison of DDC Resource usage stats. */
+	struct FDDCResourceUsageStatKeyFuncs : BaseKeyFuncs<FDDCResourceUsageStat, FString, false>
+	{
+		static const FString& GetSetKey(const FDDCResourceUsageStat& Element) { return Element.AssetType; }
+		static bool Matches(const FString& A, const FString& B) { return A == B; }
+		static uint32 GetKeyHash(const FString& Key) { return GetTypeHash(Key); }
+	};
 
-		/** Used to store profile data for custom logging. */
-		struct FCookProfileData
-		{
-		public:
-			FCookProfileData(FString InPath, FString InKey, FString InValue) : Path(MoveTemp(InPath)), Key(MoveTemp(InKey)), Value(MoveTemp(InValue)) {}
-			FString Path;
-			FString Key;
-			FString Value;
-		};
+	/** Used to store profile data for custom logging. */
+	struct FCookProfileData
+	{
+	public:
+		FCookProfileData(FString InPath, FString InKey, FString InValue) : Path(MoveTemp(InPath)), Key(MoveTemp(InKey)), Value(MoveTemp(InValue)) {}
+		FString Path;
+		FString Key;
+		FString Value;
+	};
 
-		// instead of printing the usage stats generically, we capture them so we can log a subset of them in an easy-to-read way.
-		TSet<FDDCResourceUsageStat, FDDCResourceUsageStatKeyFuncs> DDCResourceUsageStats;
-		TArray<FCookStatsManager::StringKeyValue> DDCSummaryStats;
-		TArray<FCookProfileData> CookProfileData;
-		TArray<FString> StatCategories;
-		TMap<FString, TArray<FCookStatsManager::StringKeyValue>> StatsInCategories;
+	// instead of printing the usage stats generically, we capture them so we can log a subset of them in an easy-to-read way.
+	TSet<FDDCResourceUsageStat, FDDCResourceUsageStatKeyFuncs> DDCResourceUsageStats;
+	TArray<FCookStatsManager::StringKeyValue> DDCSummaryStats;
+	TArray<FCookProfileData> CookProfileData;
+	TArray<FString> StatCategories;
+	TMap<FString, TArray<FCookStatsManager::StringKeyValue>> StatsInCategories;
 
-		/** this functor will take a collected cooker stat and log it out using some custom formatting based on known stats that are collected.. */
-		auto LogStatsFunc = [&DDCResourceUsageStats, &DDCSummaryStats, &CookProfileData, &StatCategories, &StatsInCategories]
-		(const FString& StatName, const TArray<FCookStatsManager::StringKeyValue>& StatAttributes)
-		{
-			// Some stats will use custom formatting to make a visibly pleasing summary.
-			bool bStatUsedCustomFormatting = false;
+	/** this functor will take a collected cooker stat and log it out using some custom formatting based on known stats that are collected.. */
+	auto LogStatsFunc = [&DDCResourceUsageStats, &DDCSummaryStats, &CookProfileData, &StatCategories, &StatsInCategories]
+	(const FString& StatName, const TArray<FCookStatsManager::StringKeyValue>& StatAttributes)
+	{
+		// Some stats will use custom formatting to make a visibly pleasing summary.
+		bool bStatUsedCustomFormatting = false;
 
-			if (StatName == TEXT("DDC.Usage"))
-			{
-				// Don't even log this detailed DDC data. It's mostly only consumable by ingestion into pivot tools.
-				bStatUsedCustomFormatting = true;
-			}
-			else if (StatName.EndsWith(TEXT(".Usage"), ESearchCase::IgnoreCase))
+		if (StatName == TEXT("DDC.Usage"))
+		{
+			// Don't even log this detailed DDC data. It's mostly only consumable by ingestion into pivot tools.
+			bStatUsedCustomFormatting = true;
+		}
+		else if (StatName.EndsWith(TEXT(".Usage"), ESearchCase::IgnoreCase))
+		{
+			// Anything that ends in .Usage is assumed to be an instance of FCookStats.FDDCResourceUsageStats. We'll log that using custom formatting.
+			FString AssetType = StatName;
+			AssetType.RemoveFromEnd(TEXT(".Usage"), ESearchCase::IgnoreCase);
+			// See if the asset has a subtype (found via the "Node" parameter")
+			const FCookStatsManager::StringKeyValue* AssetSubType = StatAttributes.FindByPredicate([](const FCookStatsManager::StringKeyValue& Item) { return Item.Key == TEXT("Node"); });
+			if (AssetSubType && AssetSubType->Value.Len() > 0)
 			{
-				// Anything that ends in .Usage is assumed to be an instance of FCookStats.FDDCResourceUsageStats. We'll log that using custom formatting.
-				FString AssetType = StatName;
-				AssetType.RemoveFromEnd(TEXT(".Usage"), ESearchCase::IgnoreCase);
-				// See if the asset has a subtype (found via the "Node" parameter")
-				const FCookStatsManager::StringKeyValue* AssetSubType = StatAttributes.FindByPredicate([](const FCookStatsManager::StringKeyValue& Item) { return Item.Key == TEXT("Node"); });
-				if (AssetSubType && AssetSubType->Value.Len() > 0)
-				{
-					AssetType += FString::Printf(TEXT(" (%s)"), *AssetSubType->Value);
-				}
-				// Pull the Time and Size attributes and AddOrAccumulate them into the set of stats. Ugly string/container manipulation code courtesy of UE/C++.
-				const FCookStatsManager::StringKeyValue* AssetTimeSecAttr = StatAttributes.FindByPredicate([](const FCookStatsManager::StringKeyValue& Item) { return Item.Key == TEXT("TimeSec"); });
-				double AssetTimeSec = 0.0;
-				if (AssetTimeSecAttr)
-				{
-					LexFromString(AssetTimeSec, *AssetTimeSecAttr->Value);
-				}
-				const FCookStatsManager::StringKeyValue* AssetSizeMBAttr = StatAttributes.FindByPredicate([](const FCookStatsManager::StringKeyValue& Item) { return Item.Key == TEXT("MB"); });
-				double AssetSizeMB = 0.0;
-				if (AssetSizeMBAttr)
-				{
-					LexFromString(AssetSizeMB, *AssetSizeMBAttr->Value);
-				}
-				const FCookStatsManager::StringKeyValue* ThreadNameAttr = StatAttributes.FindByPredicate([](const FCookStatsManager::StringKeyValue& Item) { return Item.Key == TEXT("ThreadName"); });
-				bool bIsGameThreadTime = ThreadNameAttr != nullptr && ThreadNameAttr->Value == TEXT("GameThread");
-
-				const FCookStatsManager::StringKeyValue* HitOrMissAttr = StatAttributes.FindByPredicate([](const FCookStatsManager::StringKeyValue& Item) { return Item.Key == TEXT("HitOrMiss"); });
-				bool bWasMiss = HitOrMissAttr != nullptr && HitOrMissAttr->Value == TEXT("Miss");
-				int64 AssetsBuilt = 0;
-				if (bWasMiss)
-				{
-					const FCookStatsManager::StringKeyValue* CountAttr = StatAttributes.FindByPredicate([](const FCookStatsManager::StringKeyValue& Item) { return Item.Key == TEXT("Count"); });
-					if (CountAttr)
-					{
-						LexFromString(AssetsBuilt, *CountAttr->Value);
-					}
-				}
-
-
-				FDDCResourceUsageStat Stat(AssetType, AssetTimeSec, bIsGameThreadTime, AssetSizeMB, AssetsBuilt);
-				FDDCResourceUsageStat* ExistingStat = DDCResourceUsageStats.Find(Stat.AssetType);
-				if (ExistingStat)
-				{
-					ExistingStat->Accumulate(Stat);
-				}
-				else
-				{
-					DDCResourceUsageStats.Add(Stat);
-				}
-				bStatUsedCustomFormatting = true;
+				AssetType += FString::Printf(TEXT(" (%s)"), *AssetSubType->Value);
 			}
-			else if (StatName == TEXT("DDC.Summary"))
+			// Pull the Time and Size attributes and AddOrAccumulate them into the set of stats. Ugly string/container manipulation code courtesy of UE/C++.
+			const FCookStatsManager::StringKeyValue* AssetTimeSecAttr = StatAttributes.FindByPredicate([](const FCookStatsManager::StringKeyValue& Item) { return Item.Key == TEXT("TimeSec"); });
+			double AssetTimeSec = 0.0;
+			if (AssetTimeSecAttr)
 			{
-				DDCSummaryStats.Append(StatAttributes);
-				bStatUsedCustomFormatting = true;
+				LexFromString(AssetTimeSec, *AssetTimeSecAttr->Value);
 			}
-			else if (StatName == TEXT("Cook.Profile"))
+			const FCookStatsManager::StringKeyValue* AssetSizeMBAttr = StatAttributes.FindByPredicate([](const FCookStatsManager::StringKeyValue& Item) { return Item.Key == TEXT("MB"); });
+			double AssetSizeMB = 0.0;
+			if (AssetSizeMBAttr)
 			{
-				if (StatAttributes.Num() >= 2)
-				{
-					CookProfileData.Emplace(StatAttributes[0].Value, StatAttributes[1].Key, StatAttributes[1].Value);
-				}
-				bStatUsedCustomFormatting = true;
+				LexFromString(AssetSizeMB, *AssetSizeMBAttr->Value);
 			}
+			const FCookStatsManager::StringKeyValue* ThreadNameAttr = StatAttributes.FindByPredicate([](const FCookStatsManager::StringKeyValue& Item) { return Item.Key == TEXT("ThreadName"); });
+			bool bIsGameThreadTime = ThreadNameAttr != nullptr && ThreadNameAttr->Value == TEXT("GameThread");
 
-			// if a stat doesn't use custom formatting, just spit out the raw info.
-			if (!bStatUsedCustomFormatting)
+			const FCookStatsManager::StringKeyValue* HitOrMissAttr = StatAttributes.FindByPredicate([](const FCookStatsManager::StringKeyValue& Item) { return Item.Key == TEXT("HitOrMiss"); });
+			bool bWasMiss = HitOrMissAttr != nullptr && HitOrMissAttr->Value == TEXT("Miss");
+			int64 AssetsBuilt = 0;
+			if (bWasMiss)
 			{
-				TArray<FCookStatsManager::StringKeyValue>& StatsInCategory = StatsInCategories.FindOrAdd(StatName);
-				if (StatsInCategory.Num() == 0)
+				const FCookStatsManager::StringKeyValue* CountAttr = StatAttributes.FindByPredicate([](const FCookStatsManager::StringKeyValue& Item) { return Item.Key == TEXT("Count"); });
+				if (CountAttr)
 				{
-					StatCategories.Add(StatName);
+					LexFromString(AssetsBuilt, *CountAttr->Value);
 				}
-				StatsInCategory.Append(StatAttributes);
 			}
-		};
 
-		FCookStatsManager::LogCookStats(LogStatsFunc);
 
-		UE_LOG(LogCookCommandlet, Display, TEXT("Misc Cook Stats"));
-		UE_LOG(LogCookCommandlet, Display, TEXT("==============="));
-		for (FString& StatCategory : StatCategories)
-		{
-			UE_LOG(LogCookCommandlet, Display, TEXT("%s"), *StatCategory);
-			TArray<FCookStatsManager::StringKeyValue>& StatsInCategory = StatsInCategories.FindOrAdd(StatCategory);
-
-			// log each key/value pair, with the equal signs lined up.
-			for (const FCookStatsManager::StringKeyValue& StatKeyValue : StatsInCategory)
+			FDDCResourceUsageStat Stat(AssetType, AssetTimeSec, bIsGameThreadTime, AssetSizeMB, AssetsBuilt);
+			FDDCResourceUsageStat* ExistingStat = DDCResourceUsageStats.Find(Stat.AssetType);
+			if (ExistingStat)
 			{
-				UE_LOG(LogCookCommandlet, Display, TEXT("    %s=%s"), *StatKeyValue.Key, *StatKeyValue.Value);
+				ExistingStat->Accumulate(Stat);
 			}
+			else
+			{
+				DDCResourceUsageStats.Add(Stat);
+			}
+			bStatUsedCustomFormatting = true;
 		}
-
-		// DDC Usage stats are custom formatted, and the above code just accumulated them into a TSet. Now log it with our special formatting for readability.
-		if (CookProfileData.Num() > 0)
+		else if (StatName == TEXT("DDC.Summary"))
+		{
+			DDCSummaryStats.Append(StatAttributes);
+			bStatUsedCustomFormatting = true;
+		}
+		else if (StatName == TEXT("Cook.Profile"))
 		{
-			UE_LOG(LogCookCommandlet, Display, TEXT(""));
-			UE_LOG(LogCookCommandlet, Display, TEXT("Cook Profile"));
-			UE_LOG(LogCookCommandlet, Display, TEXT("============"));
-			for (const auto& ProfileEntry : CookProfileData)
+			if (StatAttributes.Num() >= 2)
 			{
-				UE_LOG(LogCookCommandlet, Display, TEXT("%s.%s=%s"), *ProfileEntry.Path, *ProfileEntry.Key, *ProfileEntry.Value);
+				CookProfileData.Emplace(StatAttributes[0].Value, StatAttributes[1].Key, StatAttributes[1].Value);
 			}
+			bStatUsedCustomFormatting = true;
 		}
-		if (DDCSummaryStats.Num() > 0)
+
+		// if a stat doesn't use custom formatting, just spit out the raw info.
+		if (!bStatUsedCustomFormatting)
 		{
-			UE_LOG(LogCookCommandlet, Display, TEXT(""));
-			UE_LOG(LogCookCommandlet, Display, TEXT("DDC Summary Stats"));
-			UE_LOG(LogCookCommandlet, Display, TEXT("================="));
-			for (const auto& Attr : DDCSummaryStats)
+			TArray<FCookStatsManager::StringKeyValue>& StatsInCategory = StatsInCategories.FindOrAdd(StatName);
+			if (StatsInCategory.Num() == 0)
 			{
-				UE_LOG(LogCookCommandlet, Display, TEXT("%-16s=%10s"), *Attr.Key, *Attr.Value);
+				StatCategories.Add(StatName);
 			}
+			StatsInCategory.Append(StatAttributes);
 		}
+	};
 
-		DumpDerivedDataBuildRemoteExecutorStats();
+	FCookStatsManager::LogCookStats(LogStatsFunc);
 
-		if (DDCResourceUsageStats.Num() > 0)
+	UE_LOG(LogCookCommandlet, Display, TEXT("Misc Cook Stats"));
+	UE_LOG(LogCookCommandlet, Display, TEXT("==============="));
+	for (FString& StatCategory : StatCategories)
+	{
+		UE_LOG(LogCookCommandlet, Display, TEXT("%s"), *StatCategory);
+		TArray<FCookStatsManager::StringKeyValue>& StatsInCategory = StatsInCategories.FindOrAdd(StatCategory);
+
+		// log each key/value pair, with the equal signs lined up.
+		for (const FCookStatsManager::StringKeyValue& StatKeyValue : StatsInCategory)
 		{
-			// sort the list
-			TArray<FDDCResourceUsageStat> SortedDDCResourceUsageStats;
-			SortedDDCResourceUsageStats.Empty(DDCResourceUsageStats.Num());
-			for (const FDDCResourceUsageStat& Stat : DDCResourceUsageStats)
-			{
-				SortedDDCResourceUsageStats.Emplace(Stat);
-			}
-			SortedDDCResourceUsageStats.Sort([](const FDDCResourceUsageStat& LHS, const FDDCResourceUsageStat& RHS)
-				{
-					return LHS.TotalTimeSec > RHS.TotalTimeSec;
-				});
-
-			UE_LOG(LogCookCommandlet, Display, TEXT(""));
-			UE_LOG(LogCookCommandlet, Display, TEXT("DDC Resource Stats"));
-			UE_LOG(LogCookCommandlet, Display, TEXT("======================================================================================================="));
-			UE_LOG(LogCookCommandlet, Display, TEXT("Asset Type                          Total Time (Sec)  GameThread Time (Sec)  Assets Built  MB Processed"));
-			UE_LOG(LogCookCommandlet, Display, TEXT("----------------------------------  ----------------  ---------------------  ------------  ------------"));
-			for (const FDDCResourceUsageStat& Stat : SortedDDCResourceUsageStats)
-			{
-				UE_LOG(LogCookCommandlet, Display, TEXT("%-34s  %16.2f  %21.2f  %12d  %12.2f"), *Stat.AssetType, Stat.TotalTimeSec, Stat.GameThreadTimeSec, Stat.AssetsBuilt, Stat.SizeMB);
-			}
+			UE_LOG(LogCookCommandlet, Display, TEXT("    %s=%s"), *StatKeyValue.Key, *StatKeyValue.Value);
 		}
+	}
 
-		DumpBuildDependencyTrackerStats();
+	// DDC Usage stats are custom formatted, and the above code just accumulated them into a TSet. Now log it with our special formatting for readability.
+	if (CookProfileData.Num() > 0)
+	{
+		UE_LOG(LogCookCommandlet, Display, TEXT(""));
+		UE_LOG(LogCookCommandlet, Display, TEXT("Cook Profile"));
+		UE_LOG(LogCookCommandlet, Display, TEXT("============"));
+		for (const auto& ProfileEntry : CookProfileData)
+		{
+			UE_LOG(LogCookCommandlet, Display, TEXT("%s.%s=%s"), *ProfileEntry.Path, *ProfileEntry.Key, *ProfileEntry.Value);
+		}
+	}
+	if (DDCSummaryStats.Num() > 0)
+	{
+		UE_LOG(LogCookCommandlet, Display, TEXT(""));
+		UE_LOG(LogCookCommandlet, Display, TEXT("DDC Summary Stats"));
+		UE_LOG(LogCookCommandlet, Display, TEXT("================="));
+		for (const auto& Attr : DDCSummaryStats)
+		{
+			UE_LOG(LogCookCommandlet, Display, TEXT("%-16s=%10s"), *Attr.Key, *Attr.Value);
+		}
+	}
 
-		if (UE::Virtualization::IVirtualizationSystem::IsInitialized())
+	DumpDerivedDataBuildRemoteExecutorStats();
+
+	if (DDCResourceUsageStats.Num() > 0)
+	{
+		// sort the list
+		TArray<FDDCResourceUsageStat> SortedDDCResourceUsageStats;
+		SortedDDCResourceUsageStats.Empty(DDCResourceUsageStats.Num());
+		for (const FDDCResourceUsageStat& Stat : DDCResourceUsageStats)
+		{
+			SortedDDCResourceUsageStats.Emplace(Stat);
+		}
+		SortedDDCResourceUsageStats.Sort([](const FDDCResourceUsageStat& LHS, const FDDCResourceUsageStat& RHS)
+			{
+				return LHS.TotalTimeSec > RHS.TotalTimeSec;
+			});
+
+		UE_LOG(LogCookCommandlet, Display, TEXT(""));
+		UE_LOG(LogCookCommandlet, Display, TEXT("DDC Resource Stats"));
+		UE_LOG(LogCookCommandlet, Display, TEXT("======================================================================================================="));
+		UE_LOG(LogCookCommandlet, Display, TEXT("Asset Type                          Total Time (Sec)  GameThread Time (Sec)  Assets Built  MB Processed"));
+		UE_LOG(LogCookCommandlet, Display, TEXT("----------------------------------  ----------------  ---------------------  ------------  ------------"));
+		for (const FDDCResourceUsageStat& Stat : SortedDDCResourceUsageStats)
 		{
-			UE::Virtualization::IVirtualizationSystem::Get().DumpStats();
+			UE_LOG(LogCookCommandlet, Display, TEXT("%-34s  %16.2f  %21.2f  %12d  %12.2f"), *Stat.AssetType, Stat.TotalTimeSec, Stat.GameThreadTimeSec, Stat.AssetsBuilt, Stat.SizeMB);
 		}
 	}
+
+	DumpBuildDependencyTrackerStats();
+
+	if (UE::Virtualization::IVirtualizationSystem::IsInitialized())
+	{
+		UE::Virtualization::IVirtualizationSystem::Get().DumpStats();
+	}
+}
+
 }
 #endif
 
diff --git a/Engine/Source/Editor/UnrealEd/Private/Cooker/AsyncIODelete.cpp b/Engine/Source/Editor/UnrealEd/Private/Cooker/AsyncIODelete.cpp
--- a/Engine/Source/Editor/UnrealEd/Private/Cooker/AsyncIODelete.cpp
+++ b/Engine/Source/Editor/UnrealEd/Private/Cooker/AsyncIODelete.cpp
@@ -1,7 +1,6 @@
 // Copyright Epic Games, Inc. All Rights Reserved.
 
 #include "AsyncIODelete.h"
-#include "CoreMinimal.h"
 
 #include "Async/Async.h"
 #include "Containers/UnrealString.h"
diff --git a/Engine/Source/Editor/UnrealEd/Private/Cooker/CookPackageSplitter.cpp b/Engine/Source/Editor/UnrealEd/Private/Cooker/CookPackageSplitter.cpp
--- a/Engine/Source/Editor/UnrealEd/Private/Cooker/CookPackageSplitter.cpp
+++ b/Engine/Source/Editor/UnrealEd/Private/Cooker/CookPackageSplitter.cpp
@@ -4,11 +4,7 @@
 
 #if WITH_EDITOR
 
-namespace UE
-{
-namespace Cook
-{
-namespace Private
+namespace UE::Cook::Private
 {
 
 static TLinkedList<FRegisteredCookPackageSplitter*>* GRegisteredCookPackageSplitterList = nullptr;
@@ -37,8 +33,6 @@ void FRegisteredCookPackageSplitter::ForEach(TFunctionRef<void(FRegisteredCookPa
 	}
 }
 
-}
-}
 }
 
 #endif
\ No newline at end of file
diff --git a/Engine/Source/Editor/UnrealEd/Private/Cooker/CookPlatformManager.cpp b/Engine/Source/Editor/UnrealEd/Private/Cooker/CookPlatformManager.cpp
--- a/Engine/Source/Editor/UnrealEd/Private/Cooker/CookPlatformManager.cpp
+++ b/Engine/Source/Editor/UnrealEd/Private/Cooker/CookPlatformManager.cpp
@@ -17,369 +17,376 @@
 #include "Misc/CoreMisc.h"
 #include "Misc/ScopeRWLock.h"
 
-namespace UE
-{
-namespace Cook
+namespace UE::Cook
 {
 
-	class FRWScopeLockConditional
+class FRWScopeLockConditional
+{
+public:
+	FRWScopeLockConditional(FRWLock& InLockObject, FRWScopeLockType InLockType, bool bInNeedsLock)
+		:LockObject(InLockObject)
+		, LockType(InLockType)
+		, bNeedsLock(bInNeedsLock)
 	{
-	public:
-		FRWScopeLockConditional(FRWLock& InLockObject, FRWScopeLockType InLockType, bool bInNeedsLock)
-			:LockObject(InLockObject)
-			, LockType(InLockType)
-			, bNeedsLock(bInNeedsLock)
+		if (bNeedsLock)
 		{
-			if (bNeedsLock)
+			if (LockType != SLT_ReadOnly)
 			{
-				if (LockType != SLT_ReadOnly)
-				{
-					LockObject.WriteLock();
-				}
-				else
-				{
-					LockObject.ReadLock();
-				}
+				LockObject.WriteLock();
+			}
+			else
+			{
+				LockObject.ReadLock();
 			}
 		}
-		~FRWScopeLockConditional()
+	}
+	~FRWScopeLockConditional()
+	{
+		if (bNeedsLock)
 		{
-			if (bNeedsLock)
+			if (LockType != SLT_ReadOnly)
 			{
-				if (LockType != SLT_ReadOnly)
-				{
-					LockObject.WriteUnlock();
-				}
-				else
-				{
-					LockObject.ReadUnlock();
-				}
+				LockObject.WriteUnlock();
+			}
+			else
+			{
+				LockObject.ReadUnlock();
 			}
 		}
+	}
 
-	private:
-		UE_NONCOPYABLE(FRWScopeLockConditional);
+private:
+	UE_NONCOPYABLE(FRWScopeLockConditional);
 
-		FRWLock& LockObject;
-		FRWScopeLockType LockType;
-		bool bNeedsLock;
-	};
+	FRWLock& LockObject;
+	FRWScopeLockType LockType;
+	bool bNeedsLock;
+};
 
-	FPlatformData::FPlatformData()
-		: LastReferenceTime(0.0)
-		, ReferenceCount(0)
-	{
-	}
+FPlatformData::FPlatformData()
+	: LastReferenceTime(0.0)
+	, ReferenceCount(0)
+{
+}
 
-	uint32 FPlatformManager::IsInPlatformsLockTLSSlot = 0;
+uint32 FPlatformManager::IsInPlatformsLockTLSSlot = 0;
 
-	void FPlatformManager::InitializeTls()
-	{
-		IsInPlatformsLockTLSSlot = FPlatformTLS::AllocTlsSlot();
-	}
+void FPlatformManager::InitializeTls()
+{
+	IsInPlatformsLockTLSSlot = FPlatformTLS::AllocTlsSlot();
+}
 
-	bool FPlatformManager::IsInPlatformsLock()
-	{
-		return FPlatformTLS::GetTlsValue(IsInPlatformsLockTLSSlot) != 0;
-	}
+bool FPlatformManager::IsInPlatformsLock()
+{
+	return FPlatformTLS::GetTlsValue(IsInPlatformsLockTLSSlot) != 0;
+}
 
-	void FPlatformManager::SetIsInPlatformsLock(bool bValue)
-	{
-		FPlatformTLS::SetTlsValue(IsInPlatformsLockTLSSlot, bValue ? (void*)0x1 : (void*)0x0);
-	}
+void FPlatformManager::SetIsInPlatformsLock(bool bValue)
+{
+	FPlatformTLS::SetTlsValue(IsInPlatformsLockTLSSlot, bValue ? (void*)0x1 : (void*)0x0);
+}
 
-	FPlatformManager::~FPlatformManager()
+FPlatformManager::~FPlatformManager()
+{
 	{
+		FRWScopeLock PlatformDatasScopeLock(PlatformDatasLock, FRWScopeLockType::SLT_Write);
+		for (TPair<const ITargetPlatform*, FPlatformData*>& kvpair : PlatformDatas)
 		{
-			FRWScopeLock PlatformDatasScopeLock(PlatformDatasLock, FRWScopeLockType::SLT_Write);
-			for (TPair<const ITargetPlatform*, FPlatformData*>& kvpair : PlatformDatas)
-			{
-				FPlatformData* Value = kvpair.Value;
-				delete Value;
-			}
-			PlatformDatas.Empty();
-			PlatformDatasByName.Empty();
+			FPlatformData* Value = kvpair.Value;
+			delete Value;
 		}
+		PlatformDatas.Empty();
+		PlatformDatasByName.Empty();
 	}
+}
 
-	const TArray<FPlatformId>& FPlatformManager::GetSessionPlatforms() const
-	{
-		checkf(IsSchedulerThread(), TEXT("Access to SessionPlatforms on non-scheduler thread that persists outside of function scope is not yet implemented."));
-		checkf(bHasSelectedSessionPlatforms, TEXT("Calling GetSessionPlatforms or (any of the top level cook functions that call it) without first calling SelectSessionPlatforms is invalid"));
-		return SessionPlatforms;
-	}
+const TArray<FPlatformId>& FPlatformManager::GetSessionPlatforms() const
+{
+	checkf(IsSchedulerThread(), TEXT("Access to SessionPlatforms on non-scheduler thread that persists outside of function scope is not yet implemented."));
+	checkf(bHasSelectedSessionPlatforms, TEXT("Calling GetSessionPlatforms or (any of the top level cook functions that call it) without first calling SelectSessionPlatforms is invalid"));
+	return SessionPlatforms;
+}
 
-	int32 FPlatformManager::GetNumSessionPlatforms() const
-	{
-		checkf(IsSchedulerThread() || IsInPlatformsLock(), TEXT("Access to SessionPlatforms is only legal on non-scheduler threads when inside a ReadLockPlatforms scope."));
-		return SessionPlatforms.Num();
-	}
+int32 FPlatformManager::GetNumSessionPlatforms() const
+{
+	checkf(IsSchedulerThread() || IsInPlatformsLock(), TEXT("Access to SessionPlatforms is only legal on non-scheduler threads when inside a ReadLockPlatforms scope."));
+	return SessionPlatforms.Num();
+}
 
-	bool FPlatformManager::HasSessionPlatform(FPlatformId TargetPlatform) const
+bool FPlatformManager::HasSessionPlatform(FPlatformId TargetPlatform) const
+{
+	const bool bIsSchedulerThread = IsSchedulerThread();
+	checkf(bIsSchedulerThread || IsInPlatformsLock(), TEXT("Looking up platforms by PlatformId is only legal on non-scheduler threads when inside a ReadLockPlatforms scope."));
+	FRWScopeLockConditional ScopeLock(SessionLock, FRWScopeLockType::SLT_ReadOnly, !bIsSchedulerThread);
+	return SessionPlatforms.Contains(TargetPlatform);
+}
+
+void FPlatformManager::SelectSessionPlatforms(UCookOnTheFlyServer& COTFS,
+	const TArrayView<FPlatformId const>& TargetPlatforms)
+{
+	checkf(IsSchedulerThread(), TEXT("Writing to SessionPlatforms is only allowed from the scheduler thread."));
+	for (FPlatformId TargetPlatform : TargetPlatforms)
 	{
-		const bool bIsSchedulerThread = IsSchedulerThread();
-		checkf(bIsSchedulerThread || IsInPlatformsLock(), TEXT("Looking up platforms by PlatformId is only legal on non-scheduler threads when inside a ReadLockPlatforms scope."));
-		FRWScopeLockConditional ScopeLock(SessionLock, FRWScopeLockType::SLT_ReadOnly, !bIsSchedulerThread);
-		return SessionPlatforms.Contains(TargetPlatform);
+		CreatePlatformData(TargetPlatform);
 	}
 
-	void FPlatformManager::SelectSessionPlatforms(UCookOnTheFlyServer& COTFS, const TArrayView<FPlatformId const>& TargetPlatforms)
-	{
-		checkf(IsSchedulerThread(), TEXT("Writing to SessionPlatforms is only allowed from the scheduler thread."));
-		for (FPlatformId TargetPlatform : TargetPlatforms)
-		{
-			CreatePlatformData(TargetPlatform);
-		}
+	FRWScopeLock ScopeLock(SessionLock, FRWScopeLockType::SLT_Write);
+	SessionPlatforms.Empty(TargetPlatforms.Num());
+	SessionPlatforms.Append(TargetPlatforms.GetData(), TargetPlatforms.Num());
+	COTFS.bSessionRunning = true;
+	bHasSelectedSessionPlatforms = true;
+}
 
-		FRWScopeLock ScopeLock(SessionLock, FRWScopeLockType::SLT_Write);
-		SessionPlatforms.Empty(TargetPlatforms.Num());
-		SessionPlatforms.Append(TargetPlatforms.GetData(), TargetPlatforms.Num());
-		COTFS.bSessionRunning = true;
-		bHasSelectedSessionPlatforms = true;
-	}
+void FPlatformManager::ClearSessionPlatforms(UCookOnTheFlyServer& COTFS)
+{
+	checkf(IsSchedulerThread(), TEXT("Writing to SessionPlatforms is only allowed from the scheduler thread."));
+	FRWScopeLock ScopeLock(SessionLock, FRWScopeLockType::SLT_Write);
+	SessionPlatforms.Empty();
+	COTFS.bSessionRunning = false;
+	bHasSelectedSessionPlatforms = false;
+}
+
+void FPlatformManager::AddSessionPlatform(UCookOnTheFlyServer& COTFS, FPlatformId TargetPlatform)
+{
+	checkf(IsSchedulerThread(), TEXT("Writing to SessionPlatforms is only allowed from the scheduler thread."));
+	CreatePlatformData(TargetPlatform);
 
-	void FPlatformManager::ClearSessionPlatforms(UCookOnTheFlyServer& COTFS)
+	FRWScopeLock ScopeLock(SessionLock, FRWScopeLockType::SLT_Write);
+	if (SessionPlatforms.Contains(TargetPlatform))
 	{
-		checkf(IsSchedulerThread(), TEXT("Writing to SessionPlatforms is only allowed from the scheduler thread."));
-		FRWScopeLock ScopeLock(SessionLock, FRWScopeLockType::SLT_Write);
-		SessionPlatforms.Empty();
-		COTFS.bSessionRunning = false;
-		bHasSelectedSessionPlatforms = false;
+		return;
 	}
+	SessionPlatforms.Add(TargetPlatform);
+	COTFS.bSessionRunning = true;
+	bHasSelectedSessionPlatforms = true;
+}
 
-	void FPlatformManager::AddSessionPlatform(UCookOnTheFlyServer& COTFS, FPlatformId TargetPlatform)
-	{
-		checkf(IsSchedulerThread(), TEXT("Writing to SessionPlatforms is only allowed from the scheduler thread."));
-		CreatePlatformData(TargetPlatform);
+FPlatformData* FPlatformManager::GetPlatformData(FPlatformId Platform)
+{
+	const bool bIsSchedulerThread = IsSchedulerThread();
+	checkf(bIsSchedulerThread || IsInPlatformsLock(), TEXT("Reading FPlatformData on non-scheduler threads is only allowed when inside a ReadLockPlatforms scope."));
+	FPlatformData** Existing = PlatformDatas.Find(Platform);
+	return Existing ? *Existing : nullptr;
+}
 
-		FRWScopeLock ScopeLock(SessionLock, FRWScopeLockType::SLT_Write);
-		if (SessionPlatforms.Contains(TargetPlatform))
-		{
-			return;
-		}
-		SessionPlatforms.Add(TargetPlatform);
-		COTFS.bSessionRunning = true;
-		bHasSelectedSessionPlatforms = true;
-	}
+FPlatformData* FPlatformManager::GetPlatformDataByName(FName PlatformName)
+{
+	const bool bIsSchedulerThread = IsSchedulerThread();
+	checkf(bIsSchedulerThread || IsInPlatformsLock(), TEXT("Reading FPlatformData on non-scheduler threads is only allowed when inside a ReadLockPlatforms scope."));
+	FPlatformData** Existing = PlatformDatasByName.Find(PlatformName);
+	return Existing ? *Existing : nullptr;
+}
 
-	FPlatformData* FPlatformManager::GetPlatformData(FPlatformId Platform)
+FPlatformData& FPlatformManager::CreatePlatformData(const ITargetPlatform* Platform)
+{
+	check(Platform != nullptr);
+	FPlatformData** ExistingPlatformData = PlatformDatas.Find(Platform);
+	if (ExistingPlatformData)
 	{
-		const bool bIsSchedulerThread = IsSchedulerThread();
-		checkf(bIsSchedulerThread || IsInPlatformsLock(), TEXT("Reading FPlatformData on non-scheduler threads is only allowed when inside a ReadLockPlatforms scope."));
-		FPlatformData** Existing = PlatformDatas.Find(Platform);
-		return Existing ? *Existing : nullptr;
+		return **ExistingPlatformData;
 	}
 
-	FPlatformData* FPlatformManager::GetPlatformDataByName(FName PlatformName)
+	checkf(IsSchedulerThread(), TEXT("Writing to FPlatformData is only allowed from the scheduler thread."));
+	FName PlatformName(Platform->PlatformName());
+	checkf(!PlatformName.IsNone(), TEXT("Invalid ITargetPlatform with an empty name"));
 	{
-		const bool bIsSchedulerThread = IsSchedulerThread();
-		checkf(bIsSchedulerThread || IsInPlatformsLock(), TEXT("Reading FPlatformData on non-scheduler threads is only allowed when inside a ReadLockPlatforms scope."));
-		FPlatformData** Existing = PlatformDatasByName.Find(PlatformName);
-		return Existing ? *Existing : nullptr;
+		FRWScopeLock PlatformDatasScopeLock(PlatformDatasLock, FRWScopeLockType::SLT_Write);
+
+		FPlatformData*& PlatformData = PlatformDatas.Add(Platform);
+		PlatformData = new FPlatformData;
+		PlatformDatasByName.Add(PlatformName, PlatformData);
+
+		// We could get the non-const ITargetPlatform* from the global PlatformTargetModule,
+		// so this const cast is just a performance shortcut rather than a contract break
+		ITargetPlatform* NonConstPlatform = const_cast<ITargetPlatform*>(Platform);
+		PlatformData->TargetPlatform = NonConstPlatform;
+		PlatformData->PlatformName = PlatformName;
+		return *PlatformData;
 	}
+}
 
-	FPlatformData& FPlatformManager::CreatePlatformData(const ITargetPlatform* Platform)
+bool FPlatformManager::IsPlatformInitialized(FPlatformId Platform) const
+{
+	checkf(IsSchedulerThread() || IsInPlatformsLock(), TEXT("Looking up platforms by PlatformId is only legal on non-scheduler threads when inside a ReadLockPlatforms scope."));
+	const FPlatformData*const* PlatformData = PlatformDatas.Find(Platform);
+	if (!PlatformData)
 	{
-		check(Platform != nullptr);
-		FPlatformData** ExistingPlatformData = PlatformDatas.Find(Platform);
-		if (ExistingPlatformData)
-		{
-			return **ExistingPlatformData;
-		}
+		return false;
+	}
+	return (*PlatformData)->bIsSandboxInitialized;
+}
 
-		checkf(IsSchedulerThread(), TEXT("Writing to FPlatformData is only allowed from the scheduler thread."));
-		FName PlatformName(Platform->PlatformName());
-		checkf(!PlatformName.IsNone(), TEXT("Invalid ITargetPlatform with an empty name"));
-		{
-			FRWScopeLock PlatformDatasScopeLock(PlatformDatasLock, FRWScopeLockType::SLT_Write);
+void FPlatformManager::SetArePlatformsPrepopulated(bool bValue)
+{
+	checkf(IsSchedulerThread(), TEXT("Get/SetArePlatformsPrepopulated is only allowed from the scheduler thread."));
+	bArePlatformsPrepopulated = bValue;
+}
 
-			FPlatformData*& PlatformData = PlatformDatas.Add(Platform);
-			PlatformData = new FPlatformData;
-			PlatformDatasByName.Add(PlatformName, PlatformData);
+bool FPlatformManager::GetArePlatformsPrepopulated() const
+{
+	checkf(IsSchedulerThread(), TEXT("Get/SetArePlatformsPrepopulated is only allowed from the scheduler thread."));
+	return bArePlatformsPrepopulated;
+}
 
-			// We could get the non-const ITargetPlatform* from the global PlatformTargetModule, so this const cast is just a performance shortcut rather than a contract break
-			ITargetPlatform* NonConstPlatform = const_cast<ITargetPlatform*>(Platform);
-			PlatformData->TargetPlatform = NonConstPlatform;
-			PlatformData->PlatformName = PlatformName;
-			return *PlatformData;
-		}
-	}
+void FPlatformManager::PruneUnreferencedSessionPlatforms(UCookOnTheFlyServer& CookOnTheFlyServer)
+{
+	checkf(IsSchedulerThread(), TEXT("Writing to SessionPlatforms is only allowed from the scheduler thread."));
+	const double SecondsToLive = 5.0 * 60;
 
-	bool FPlatformManager::IsPlatformInitialized(FPlatformId Platform) const
-	{
-		checkf(IsSchedulerThread() || IsInPlatformsLock(), TEXT("Looking up platforms by PlatformId is only legal on non-scheduler threads when inside a ReadLockPlatforms scope."));
-		const FPlatformData*const* PlatformData = PlatformDatas.Find(Platform);
-		if (!PlatformData)
-		{
-			return false;
-		}
-		return (*PlatformData)->bIsSandboxInitialized;
-	}
+	// OldestKeepTime is constructed to less than -SecondsToLive, so we can robustly detect not-yet-initialized
+	double OldestKeepTime = -1.0e10;
+	TArray<const ITargetPlatform*, TInlineAllocator<1>> RemovePlatforms;
 
-	void FPlatformManager::SetArePlatformsPrepopulated(bool bValue)
+	for (TPair<const ITargetPlatform*, FPlatformData*>& kvpair : PlatformDatas)
 	{
-		checkf(IsSchedulerThread(), TEXT("Get/SetArePlatformsPrepopulated is only allowed from the scheduler thread."));
-		bArePlatformsPrepopulated = bValue;
-	}
+		FPlatformData* PlatformData = kvpair.Value;
+		if (PlatformData->LastReferenceTime > 0. && PlatformData->ReferenceCount == 0)
+		{
+			// We have a platform that we need to check for pruning.
+			// Initialize the OldestKeepTime so we can check whether the platform has aged out.
+			if (OldestKeepTime < -SecondsToLive)
+			{
+				const double CurrentTimeSeconds = FPlatformTime::Seconds();
+				OldestKeepTime = CurrentTimeSeconds - SecondsToLive;
+			}
 
-	bool FPlatformManager::GetArePlatformsPrepopulated() const
-	{
-		checkf(IsSchedulerThread(), TEXT("Get/SetArePlatformsPrepopulated is only allowed from the scheduler thread."));
-		return bArePlatformsPrepopulated;
+			// Note that this loop is outside of the critical section, for performance.
+			// If we find any candidates for pruning we have to check them again once inside the critical section.
+			if (kvpair.Value->LastReferenceTime < OldestKeepTime)
+			{
+				RemovePlatforms.Add(kvpair.Key);
+			}
+		}
 	}
 
-	void FPlatformManager::PruneUnreferencedSessionPlatforms(UCookOnTheFlyServer& CookOnTheFlyServer)
+	if (RemovePlatforms.Num() > 0)
 	{
-		checkf(IsSchedulerThread(), TEXT("Writing to SessionPlatforms is only allowed from the scheduler thread."));
-		const double SecondsToLive = 5.0 * 60;
+		FRWScopeLock Lock(SessionLock, FRWScopeLockType::SLT_Write);
 
-		double OldestKeepTime = -1.0e10; // Constructed to something smaller than 0 - SecondsToLive, so we can robustly detect not-yet-initialized
-		TArray<const ITargetPlatform*, TInlineAllocator<1>> RemovePlatforms;
-
-		for (TPair<const ITargetPlatform*, FPlatformData*>& kvpair : PlatformDatas)
+		for (const ITargetPlatform* TargetPlatform : RemovePlatforms)
 		{
-			FPlatformData* PlatformData = kvpair.Value;
-			if (PlatformData->LastReferenceTime > 0. && PlatformData->ReferenceCount == 0)
+			FPlatformData* PlatformData = *PlatformDatas.Find(TargetPlatform);
+			if (PlatformData->LastReferenceTime > 0. && PlatformData->ReferenceCount == 0
+				&& PlatformData->LastReferenceTime < OldestKeepTime)
 			{
-				// We have a platform that we need to check for pruning.  Initialize the OldestKeepTime so we can check whether the platform has aged out.
-				if (OldestKeepTime < -SecondsToLive)
-				{
-					const double CurrentTimeSeconds = FPlatformTime::Seconds();
-					OldestKeepTime = CurrentTimeSeconds - SecondsToLive;
-				}
-
-				// Note that this loop is outside of the critical section, for performance.
-				// If we find any candidates for pruning we have to check them again once inside the critical section.
-				if (kvpair.Value->LastReferenceTime < OldestKeepTime)
-				{
-					RemovePlatforms.Add(kvpair.Key);
-				}
-			}
-		}
+				// Mark that the platform no longer needs to be inspected for pruning because we have removed it
+				// from CookOnTheFly's SessionPlatforms
+				PlatformData->LastReferenceTime = 0.;
 
-		if (RemovePlatforms.Num() > 0)
-		{
-			FRWScopeLock Lock(SessionLock, FRWScopeLockType::SLT_Write);
+				// Remove the SessionPlatform
+				CookOnTheFlyServer.OnRemoveSessionPlatform(TargetPlatform);
 
-			for (const ITargetPlatform* TargetPlatform : RemovePlatforms)
-			{
-				FPlatformData* PlatformData = *PlatformDatas.Find(TargetPlatform);
-				if (PlatformData->LastReferenceTime > 0. && PlatformData->ReferenceCount == 0 && PlatformData->LastReferenceTime < OldestKeepTime)
+				SessionPlatforms.Remove(TargetPlatform);
+				if (SessionPlatforms.Num() == 0)
 				{
-					// Mark that the platform no longer needs to be inspected for pruning because we have removed it from CookOnTheFly's SessionPlatforms
-					PlatformData->LastReferenceTime = 0.;
-
-					// Remove the SessionPlatform
-					CookOnTheFlyServer.OnRemoveSessionPlatform(TargetPlatform);
-
-					SessionPlatforms.Remove(TargetPlatform);
-					if (SessionPlatforms.Num() == 0)
-					{
-						CookOnTheFlyServer.bSessionRunning = false;
-						bHasSelectedSessionPlatforms = false;
-					}
+					CookOnTheFlyServer.bSessionRunning = false;
+					bHasSelectedSessionPlatforms = false;
 				}
 			}
 		}
 	}
+}
 
-	void FPlatformManager::AddRefCookOnTheFlyPlatform(FName PlatformName, UCookOnTheFlyServer& CookOnTheFlyServer)
-	{
-		checkf(IsSchedulerThread() || IsInPlatformsLock(), TEXT("AddRefCookOnTheFlyPlatform is only legal on non-scheduler threads when inside a ReadLockPlatforms scope."));
+void FPlatformManager::AddRefCookOnTheFlyPlatform(FName PlatformName, UCookOnTheFlyServer& CookOnTheFlyServer)
+{
+	checkf(IsSchedulerThread() || IsInPlatformsLock(), TEXT("AddRefCookOnTheFlyPlatform is only legal on non-scheduler threads when inside a ReadLockPlatforms scope."));
 
-		FPlatformData* PlatformData = GetPlatformDataByName(PlatformName);
-		checkf(PlatformData != nullptr, TEXT("Unrecognized Platform %s"), *PlatformName.ToString());
-		++PlatformData->ReferenceCount;
+	FPlatformData* PlatformData = GetPlatformDataByName(PlatformName);
+	checkf(PlatformData != nullptr, TEXT("Unrecognized Platform %s"), *PlatformName.ToString());
+	++PlatformData->ReferenceCount;
 
-		if (!HasSessionPlatform(PlatformData->TargetPlatform))
-		{
-			CookOnTheFlyServer.WorkerRequests->AddCookOnTheFlyCallback([PlatformName, LocalCookOnTheFlyServer = &CookOnTheFlyServer]()
+	if (!HasSessionPlatform(PlatformData->TargetPlatform))
+	{
+		CookOnTheFlyServer.WorkerRequests->AddCookOnTheFlyCallback([PlatformName,
+			LocalCookOnTheFlyServer = &CookOnTheFlyServer]()
+			{
+				ITargetPlatform* TargetPlatform = GetTargetPlatformManager()->FindTargetPlatform(PlatformName.ToString());
+				if (TargetPlatform)
 				{
-					ITargetPlatform* TargetPlatform = GetTargetPlatformManager()->FindTargetPlatform(PlatformName.ToString());
-					if (TargetPlatform)
-					{
-						LocalCookOnTheFlyServer->StartCookOnTheFlySessionFromGameThread(TargetPlatform);
-					}
-				});
-		}
+					LocalCookOnTheFlyServer->StartCookOnTheFlySessionFromGameThread(TargetPlatform);
+				}
+			});
 	}
+}
 
-	void FPlatformManager::ReleaseCookOnTheFlyPlatform(FName PlatformName)
-	{
-		checkf(IsSchedulerThread() || IsInPlatformsLock(), TEXT("ReleaseCookOnTheFlyPlatform is only legal on non-scheduler threads when inside a ReadLockPlatforms scope."));
+void FPlatformManager::ReleaseCookOnTheFlyPlatform(FName PlatformName)
+{
+	checkf(IsSchedulerThread() || IsInPlatformsLock(), TEXT("ReleaseCookOnTheFlyPlatform is only legal on non-scheduler threads when inside a ReadLockPlatforms scope."));
 		
-		FPlatformData* PlatformData = GetPlatformDataByName(PlatformName);
-		checkf(PlatformData != nullptr, TEXT("Unrecognized Platform %s"), *PlatformName.ToString());
-		check(PlatformData->ReferenceCount > 0);
-		--PlatformData->ReferenceCount;
-		PlatformData->LastReferenceTime = FPlatformTime::Seconds();
-	}
+	FPlatformData* PlatformData = GetPlatformDataByName(PlatformName);
+	checkf(PlatformData != nullptr, TEXT("Unrecognized Platform %s"), *PlatformName.ToString());
+	check(PlatformData->ReferenceCount > 0);
+	--PlatformData->ReferenceCount;
+	PlatformData->LastReferenceTime = FPlatformTime::Seconds();
+}
 
-	TMap<ITargetPlatform*, ITargetPlatform*> FPlatformManager::RemapTargetPlatforms()
-	{
-		checkf(IsSchedulerThread(), TEXT("Writing to PlatformDatas is only allowed from the scheduler thread."));
-		TMap<ITargetPlatform*, ITargetPlatform*> Remap;
+TMap<ITargetPlatform*, ITargetPlatform*> FPlatformManager::RemapTargetPlatforms()
+{
+	checkf(IsSchedulerThread(), TEXT("Writing to PlatformDatas is only allowed from the scheduler thread."));
+	TMap<ITargetPlatform*, ITargetPlatform*> Remap;
 
-		ITargetPlatformManagerModule* PlatformManager = GetTargetPlatformManager();
-		TFastPointerMap<const ITargetPlatform*, FPlatformData*> NewPlatformDatas;
+	ITargetPlatformManagerModule* PlatformManager = GetTargetPlatformManager();
+	TFastPointerMap<const ITargetPlatform*, FPlatformData*> NewPlatformDatas;
 
+	{
+		FRWScopeLock PlatformDatasScopeLock(PlatformDatasLock, FRWScopeLockType::SLT_Write);
+		for (TPair<const ITargetPlatform*, FPlatformData*>& kvpair : PlatformDatas)
 		{
-			FRWScopeLock PlatformDatasScopeLock(PlatformDatasLock, FRWScopeLockType::SLT_Write);
-			for (TPair<const ITargetPlatform*, FPlatformData*>& kvpair : PlatformDatas)
-			{
-				ITargetPlatform* Old = const_cast<ITargetPlatform*>(kvpair.Key);
-				FPlatformData& Data = *kvpair.Value;
-				FName PlatformName = Data.PlatformName;
+			ITargetPlatform* Old = const_cast<ITargetPlatform*>(kvpair.Key);
+			FPlatformData& Data = *kvpair.Value;
+			FName PlatformName = Data.PlatformName;
 
-				ITargetPlatform* New = PlatformManager->FindTargetPlatform(PlatformName.ToString());
-				checkf(New, TEXT("TargetPlatform %s has been removed from the list of TargetPlatforms from ITargetPlatformManagerModule after cooking has started; this case is not handled."), *PlatformName.ToString());
-				Data.TargetPlatform = New;
+			ITargetPlatform* New = PlatformManager->FindTargetPlatform(PlatformName.ToString());
+			checkf(New, TEXT("TargetPlatform %s has been removed from the list of TargetPlatforms from ITargetPlatformManagerModule after cooking has started; this case is not handled."),
+				*PlatformName.ToString());
+			Data.TargetPlatform = New;
 
-				Remap.Add(Old, New);
-				NewPlatformDatas.Add(New, &Data);
-			}
-			// Note that PlatformDatasByName is unchanged, since it maps PlatformName -> PlatformData*, and neither of those have changed. The values inside the FPlatformData may have changed, though, so we need to hold the lock during this function.
-			Swap(PlatformDatas, NewPlatformDatas);
-
-			{
-				FRWScopeLock SessionScopeLock(SessionLock, FRWScopeLockType::SLT_Write);
-				RemapArrayElements(SessionPlatforms, Remap);
-			}
+			Remap.Add(Old, New);
+			NewPlatformDatas.Add(New, &Data);
 		}
+		// Note that PlatformDatasByName is unchanged, since it maps PlatformName -> PlatformData*,
+		// and neither of those have changed. The values inside the FPlatformData may have changed, though,
+		// so we need to hold the lock during this function.
+		Swap(PlatformDatas, NewPlatformDatas);
 
-		return Remap;
+		{
+			FRWScopeLock SessionScopeLock(SessionLock, FRWScopeLockType::SLT_Write);
+			RemapArrayElements(SessionPlatforms, Remap);
+		}
 	}
 
-	FPlatformManager::FReadScopeLock::FReadScopeLock(FPlatformManager& InPlatformManager)
-		:PlatformManager(InPlatformManager)
-	{
-		bAttached = true;
-		PlatformManager.PlatformDatasLock.ReadLock();
-		check(!IsInPlatformsLock());
-		SetIsInPlatformsLock(true);
-	}
+	return Remap;
+}
 
-	FPlatformManager::FReadScopeLock::FReadScopeLock(FReadScopeLock&& Other)
-		:PlatformManager(Other.PlatformManager)
-	{
-		bAttached = Other.bAttached;
-		Other.bAttached = false;
-	}
+FPlatformManager::FReadScopeLock::FReadScopeLock(FPlatformManager& InPlatformManager)
+	:PlatformManager(InPlatformManager)
+{
+	bAttached = true;
+	PlatformManager.PlatformDatasLock.ReadLock();
+	check(!IsInPlatformsLock());
+	SetIsInPlatformsLock(true);
+}
 
-	FPlatformManager::FReadScopeLock::~FReadScopeLock()
-	{
-		if (bAttached)
-		{
-			SetIsInPlatformsLock(false);
-			PlatformManager.PlatformDatasLock.ReadUnlock();
-			bAttached = false;
-		}
-	}
+FPlatformManager::FReadScopeLock::FReadScopeLock(FReadScopeLock&& Other)
+	:PlatformManager(Other.PlatformManager)
+{
+	bAttached = Other.bAttached;
+	Other.bAttached = false;
+}
 
-	FPlatformManager::FReadScopeLock FPlatformManager::ReadLockPlatforms()
+FPlatformManager::FReadScopeLock::~FReadScopeLock()
+{
+	if (bAttached)
 	{
-		return FReadScopeLock(*this);
+		SetIsInPlatformsLock(false);
+		PlatformManager.PlatformDatasLock.ReadUnlock();
+		bAttached = false;
 	}
+}
 
+FPlatformManager::FReadScopeLock FPlatformManager::ReadLockPlatforms()
+{
+	return FReadScopeLock(*this);
 }
+
 }
diff --git a/Engine/Source/Editor/UnrealEd/Private/Cooker/CookPlatformManager.h b/Engine/Source/Editor/UnrealEd/Private/Cooker/CookPlatformManager.h
--- a/Engine/Source/Editor/UnrealEd/Private/Cooker/CookPlatformManager.h
+++ b/Engine/Source/Editor/UnrealEd/Private/Cooker/CookPlatformManager.h
@@ -26,8 +26,10 @@ namespace Cook
 	/*
 	 * Struct to hold data about each platform we have encountered in the cooker.
 	 * Fields in this struct persist across multiple CookByTheBook sessions.
-	 * All fields on this struct are readable only from either the scheduler thread or within a ReadLock on non-scheduler thread.
-	 * Some fields on this struct are SingleThreadWrite fields - written only from the scheduler thread. These are unsynchronized other than SingleThreadWrite readlocks.
+	 * All fields on this struct are readable only from either the scheduler thread or within a ReadLock on
+	 * non-scheduler thread.
+	 * Some fields on this struct are SingleThreadWrite fields - written only from the scheduler thread.
+	 * These are unsynchronized other than SingleThreadWrite readlocks.
 	 * The remaining fields are read/write from any thread, and are TAtomics.
 	 */
 	struct FPlatformData
@@ -39,15 +41,22 @@ namespace Cook
 
 		/**
 		 * The TargetPlatform returned from TargetPlatformManagerModule.
-		 * This may change to a different pointer if Invalidate is called on the TargetPlatformManagerModule, so it is only valid to dereference on the ScheduleThread since invalidation occurs on the ScheduleThread.
+		 * This may change to a different pointer if Invalidate is called on the TargetPlatformManagerModule, so it is
+		 * only valid to dereference on the ScheduleThread since invalidation occurs on the ScheduleThread.
 		 */
 		ITargetPlatform* TargetPlatform;
 
-		/* Pointer to the platform-specific RegistryGenerator for this platform.  If already constructed we can take a faster refresh path on future sessions. */
+		/*
+		* Pointer to the platform-specific RegistryGenerator for this platform.  If already constructed we can take a
+		* faster refresh path on future sessions.
+		*/
 		TUniquePtr<FAssetRegistryGenerator> RegistryGenerator;
 		TUniquePtr<UE::Cook::IAssetRegistryReporter> RegistryReporter;
 
-		/* Whether BeginCookSandbox has been called for this platform.  If we have already initialized the sandbox we can take a faster refresh path on future sessions. */
+		/*
+		* Whether BeginCookSandbox has been called for this platform.  If we have already initialized the sandbox we
+		* can take a faster refresh path on future sessions.
+		*/
 		bool bIsSandboxInitialized = false;
 
 		/* Whether BeginCookSandbox specified a full, non-iterative build for this platform. */
@@ -67,7 +76,8 @@ namespace Cook
 
 	/*
 	 * Information about the platforms (a) known and (b) active for the current cook session in the UCookOnTheFlyServer.
-	 * This structure is SingleThreadWrite-ThreadSafe.  All read accesses from threads other than the SchedulerThread enter read locks (or require read locks held). 
+	 * This structure is SingleThreadWrite-ThreadSafe.  All read accesses from threads other than the SchedulerThread
+	 * enter read locks (or require read locks held). 
 	 * Write accesses enter write locks, and assert if not on the SchedulerThread.
 	*/
 	struct FPlatformManager
@@ -75,17 +85,20 @@ namespace Cook
 	public:
 		~FPlatformManager();
 
-		/** Returns the set of TargetPlatforms that is active for the CurrentCookByTheBook session or CookOnTheFly request. */
+		/** Returns the TargetPlatforms that are active for the CurrentCookByTheBook session or CookOnTheFly request. */
 		const TArray<const ITargetPlatform*>& GetSessionPlatforms() const;
 		int32 GetNumSessionPlatforms() const;
 
 		/** Return whether the given platform is already in the list of platforms for the current session. */
 		bool HasSessionPlatform(FPlatformId TargetPlatform) const;
 
-		/** Specify the set of TargetPlatforms to use for the currently-initializing CookByTheBook session or CookOnTheFly request. */
+		/** Specify the TargetPlatforms to use for the currently CookByTheBook session or CookOnTheFly request. */
 		void SelectSessionPlatforms(UCookOnTheFlyServer& COTFS, const TArrayView<FPlatformId const>& TargetPlatforms);
 
-		/** Mark that the list of TargetPlatforms for the session has no longer been set; it will be an error to try to read them until SelectSessionPlatforms is called. */
+		/**
+		 * Mark that the TargetPlatforms are no longer set; reading them without first calling
+		 * SelectSessionPlatforms is an error.
+		 */
 		void ClearSessionPlatforms(UCookOnTheFlyServer& COTFS);
 
 		/** Add @param TargetPlatform to the session platforms if not already present. */
@@ -97,40 +110,55 @@ namespace Cook
 		 */
 		FPlatformData* GetPlatformData(FPlatformId Platform);
 
-		/** Get The PlatformData for the given Platform, looked up by name since NetworkThread does not have FPlatformId to start out with. */
+		/**
+		 * Get The PlatformData for the given Platform, looked up by name since NetworkThread does not have
+		 * FPlatformId to start out with.
+		 */
 		FPlatformData* GetPlatformDataByName(FName PlatformName);
 
 		/** Create if not already created the necessary platform-specific data for the given platform. */
 		FPlatformData& CreatePlatformData(const ITargetPlatform* Platform);
 
-		/** Return whether platform-specific setup steps have been executed for the given platform in the current UCookOnTheFlyServer. */
+		/** Return whether platform-specific setup steps have run in the current UCookOnTheFlyServer. */
 		bool IsPlatformInitialized(FPlatformId Platform) const;
 
-		/** Get/Set for the Prepopulated flag. If platforms are prepopulated, then we need to repopulate the list whenever targetplatforms are invalidated. */
+		/**
+		 * Get/Set the Prepopulated flag. If platforms are prepopulated, then we need to repopulate the list
+		 * whenever targetplatforms are invalidated.
+		 */
 		void SetArePlatformsPrepopulated(bool bValue);
 		bool GetArePlatformsPrepopulated() const;
 
 		/**
-		 * Platforms requested in CookOnTheFly requests are added to the list of SessionPlatforms, and some packages (e.g. unsolicited packages) are cooked against all session packages.
-		 * To have good performance in the case where multiple targetplatforms are being requested over time from the same CookOnTheFly server, we prune platforms from the list of
-		 * active session platforms if they haven't been requested in a while.
+		 * Platforms requested in CookOnTheFly requests are added to the list of SessionPlatforms, and some packages
+		 * (e.g. unsolicited packages) are cooked against all session packages.
+		 * To have good performance in the case where multiple targetplatforms are being requested over time from the
+		 * same CookOnTheFly server, we prune platforms from the list of active session platforms if they haven't been
+		 * requested in a while.
 		 */
 		void PruneUnreferencedSessionPlatforms(UCookOnTheFlyServer& CookOnTheFlyServer);
 
-		/** Increment the counter indicating the current platform is requested in an active CookOnTheFly request.  Add it to the SessionPlatforms if not already present. */
+		/**
+		 * Increment the counter indicating the current platform is requested in an active CookOnTheFly request.
+		 * Add it to the SessionPlatforms if not already present.
+		 */
 		void AddRefCookOnTheFlyPlatform(FName PlatformName, UCookOnTheFlyServer& CookOnTheFlyServer);
 
 		/** Decrement the counter indicating the current platform is being used in a CookOnTheFly request. */
 		void ReleaseCookOnTheFlyPlatform(FName TargetPlatform);
 
 		/**
-		 * Called from CookOnTheFlyServer when TargetPlatformManagerModule reports that ITargetPlatform* have been invalided.
-		 * Constructs a map from the old ITargetPlatform* to the new ITargetPlatform* using the PlatformName that we cached from the old ITargetPlatform* to identify the matching new one.
-		 * Modifies all internal ITargetPlatform* values and returns the map to use for other stored copies of ITargetPlatform*.
+		 * Called from CookOnTheFlyServer when TargetPlatformManagerModule reports that ITargetPlatform* have been
+		 * invalided. Constructs a map from the old ITargetPlatform* to the new ITargetPlatform* using the PlatformName
+		 * that we cached from the old ITargetPlatform* to identify the matching new one. Modifies all internal
+		 * ITargetPlatform* values and returns the map to use for other stored copies of ITargetPlatform*.
 		 */
 		TMap<ITargetPlatform*, ITargetPlatform*> RemapTargetPlatforms();
 
-		/* When using FPlatformManager from a thread other than the scheduler thread (e.g. HandleNetworkFileServer functions), callers must lock the list of PlatformDatas using ReadLockPlatforms */
+		/* 
+		 * When using FPlatformManager from a thread other than the scheduler thread (e.g. HandleNetworkFileServer
+		 * functions), callers must lock the list of PlatformDatas using ReadLockPlatforms.
+		 */
 		struct FReadScopeLock
 		{
 			~FReadScopeLock();
@@ -154,31 +182,40 @@ namespace Cook
 		static void SetIsInPlatformsLock(bool bValue);
 		static uint32 IsInPlatformsLockTLSSlot;
 
-		/** A collection of initialization flags and other data we maintain for each platform we have encountered in any session. */
+		/** A collection of flags and other data we maintain for each platform we have encountered in any session. */
 		TFastPointerMap<FPlatformId, FPlatformData*> PlatformDatas;
 		TMap<FName, FPlatformData*> PlatformDatasByName;
 
 		/*
-		 * RWLock used to guard PlatformDatas, PlatformDatasByName, and the validity of FPlatformIds. These fields are writable only on the SchedulerThread, so SchedulerThreadread operations can skip taking this lock.
-		 * To prevent deadlocks, if PlatformDatasLock and SessionLock are held at the same time in a block of code, then PlatformDatasLock must be entered before entering SessionLock.
+		 * RWLock used to guard PlatformDatas, PlatformDatasByName, and the validity of FPlatformIds. These fields are
+		 * writable only on the SchedulerThread, so SchedulerThreadread operations can skip taking this lock.
+		 * To prevent deadlocks, if PlatformDatasLock and SessionLock are held at the same time in a block of code,
+		 * then PlatformDatasLock must be entered before entering SessionLock.
 		 */
 		mutable FRWLock PlatformDatasLock;
 
 		/**
-		 * A collection of Platforms that are active for the current CookByTheBook session or CookOnTheFly request.  Used so we can refer to "all platforms" without having to store a list on every FileRequest.
-		 * Writing to the list of active session platforms requires a CriticalSection, because it is read (under critical section) on NetworkFileServer threads.
+		 * A collection of Platforms that are active for the current CookByTheBook session or CookOnTheFly request.
+		 * Used so we can refer to "all platforms" without having to store a list on every FileRequest.
+		 * Writing to the list of active session platforms requires a CriticalSection, because it is read (under
+		 * critical section) on NetworkFileServer threads.
 		 */
 		TArray<FPlatformId> SessionPlatforms;
 
 		/**
-		 * RWLock used to guard SessionPlatforms. SessionPlatforms can only be written from the SchedulerThread, so SchedulerThread read operations can skip taking this lock.
-		 * To prevent deadlocks, if PlatformDatasLock and SessionLock are held at the same time in a block of code, then PlatformDatasLock must be entered before entering SessionLock.
+		 * RWLock used to guard SessionPlatforms. SessionPlatforms can only be written from the SchedulerThread, so
+		 * SchedulerThread read operations can skip taking this lock.
+		 * To prevent deadlocks, if PlatformDatasLock and SessionLock are held at the same time in a block of code,
+		 * then PlatformDatasLock must be entered before entering SessionLock.
 		 */
 		mutable FRWLock SessionLock;
 		
 		bool bArePlatformsPrepopulated = false;
 
-		/** Duplicate of COTFS.bSessionRunning. Used for assertions on other threads. It is invalid to attempt to cook if session platforms have not been selected. */
+		/**
+		 * Duplicate of COTFS.bSessionRunning. Used for assertions on other threads. It is invalid to attempt to cook
+		 * if session platforms have not been selected.
+		 */
 		bool bHasSelectedSessionPlatforms = false;
 	};
 
@@ -186,7 +223,8 @@ namespace Cook
 }
 
 template <typename Value>
-void RemapMapKeys(TFastPointerMap<const ITargetPlatform*, Value>& Map, const TMap<ITargetPlatform*, ITargetPlatform*>& Remap)
+void RemapMapKeys(TFastPointerMap<const ITargetPlatform*, Value>& Map, const TMap<ITargetPlatform*,
+	ITargetPlatform*>& Remap)
 {
 	TFastPointerMap<const ITargetPlatform*, Value> NewMap;
 	NewMap.Reserve(Map.Num());
@@ -197,7 +235,8 @@ void RemapMapKeys(TFastPointerMap<const ITargetPlatform*, Value>& Map, const TMa
 	Swap(Map, NewMap);
 }
 
-inline void RemapArrayElements(TArray<const ITargetPlatform*>& Array, const TMap<ITargetPlatform*, ITargetPlatform*>& Remap)
+inline void RemapArrayElements(TArray<const ITargetPlatform*>& Array, const TMap<ITargetPlatform*,
+	ITargetPlatform*>& Remap)
 {
 	for (const ITargetPlatform*& Old : Array)
 	{
diff --git a/Engine/Source/Editor/UnrealEd/Private/Cooker/CookRequests.cpp b/Engine/Source/Editor/UnrealEd/Private/Cooker/CookRequests.cpp
--- a/Engine/Source/Editor/UnrealEd/Private/Cooker/CookRequests.cpp
+++ b/Engine/Source/Editor/UnrealEd/Private/Cooker/CookRequests.cpp
@@ -16,7 +16,6 @@
 #include "Templates/UnrealTemplate.h"
 #include "Trace/Detail/Channel.h"
 #include "UObject/UnrealNames.h"
-#include "atomic"
 
 namespace UE::Cook
 {
diff --git a/Engine/Source/Editor/UnrealEd/Private/Cooker/IoStoreCookOnTheFlyRequestManager.cpp b/Engine/Source/Editor/UnrealEd/Private/Cooker/IoStoreCookOnTheFlyRequestManager.cpp
--- a/Engine/Source/Editor/UnrealEd/Private/Cooker/IoStoreCookOnTheFlyRequestManager.cpp
+++ b/Engine/Source/Editor/UnrealEd/Private/Cooker/IoStoreCookOnTheFlyRequestManager.cpp
@@ -1,28 +1,29 @@
 // Copyright Epic Games, Inc. All Rights Reserved.
 
 #include "IoStoreCookOnTheFlyRequestManager.h"
-#include "CookOnTheFlyServerInterface.h"
+
+#include "AssetRegistry/IAssetRegistry.h"
+#include "Async/Async.h"
+#include "Cooker/ExternalCookOnTheFlyServer.h"
 #include "CookOnTheFly.h"
 #include "CookOnTheFlyMessages.h"
-#include "Modules/ModuleManager.h"
-#include "PackageStoreWriter.h"
+#include "CookOnTheFlyNetServer.h"
+#include "CookOnTheFlyServerInterface.h"
+#include "Engine/Engine.h"
 #include "HAL/PlatformTime.h"
-#include "Misc/PackageName.h"
-#include "UObject/SavePackage.h"
-#include "ProfilingDebugging/CountersTrace.h"
-#include "AssetRegistry/IAssetRegistry.h"
-#include "Templates/Function.h"
 #include "IPAddress.h"
-#include "SocketSubsystem.h"
-#include "Sockets.h"
-#include "Async/Async.h"
-#include "Serialization/BufferArchive.h"
-#include "NetworkMessage.h"
-#include "Engine/Engine.h"
 #include "MessageEndpoint.h"
 #include "MessageEndpointBuilder.h"
-#include "Cooker/ExternalCookOnTheFlyServer.h"
-#include "CookOnTheFlyNetServer.h"
+#include "Misc/PackageName.h"
+#include "Modules/ModuleManager.h"
+#include "NetworkMessage.h"
+#include "PackageStoreWriter.h"
+#include "ProfilingDebugging/CountersTrace.h"
+#include "Serialization/BufferArchive.h"
+#include "Sockets.h"
+#include "SocketSubsystem.h"
+#include "Templates/Function.h"
+#include "UObject/SavePackage.h"
 
 class FIoStoreCookOnTheFlyRequestManager final
 	: public UE::Cook::ICookOnTheFlyRequestManager
diff --git a/Engine/Source/Editor/UnrealEd/Private/Cooker/IoStoreCookOnTheFlyRequestManager.h b/Engine/Source/Editor/UnrealEd/Private/Cooker/IoStoreCookOnTheFlyRequestManager.h
--- a/Engine/Source/Editor/UnrealEd/Private/Cooker/IoStoreCookOnTheFlyRequestManager.h
+++ b/Engine/Source/Editor/UnrealEd/Private/Cooker/IoStoreCookOnTheFlyRequestManager.h
@@ -2,17 +2,19 @@
 
 #pragma once
 
-#include "CoreMinimal.h"
+#include "Templates/SharedPointer.h"
+#include "Templates/UniquePtr.h"
 
 class IAssetRegistry;
 
-namespace UE { namespace Cook
+namespace UE::Cook
 {
 
 class ICookOnTheFlyServer;
 class ICookOnTheFlyRequestManager;
 class ICookOnTheFlyNetworkServer;
 
-TUniquePtr<ICookOnTheFlyRequestManager> MakeIoStoreCookOnTheFlyRequestManager(ICookOnTheFlyServer& CookOnTheFlyServer, const IAssetRegistry* AssetRegistry, TSharedRef<ICookOnTheFlyNetworkServer> ConnectionServer);
+TUniquePtr<ICookOnTheFlyRequestManager> MakeIoStoreCookOnTheFlyRequestManager(ICookOnTheFlyServer& CookOnTheFlyServer,
+	const IAssetRegistry* AssetRegistry, TSharedRef<ICookOnTheFlyNetworkServer> ConnectionServer);
 
-}} // namespace UE::Cook
+} // namespace UE::Cook
diff --git a/Engine/Source/Editor/UnrealEd/Private/Cooker/LooseCookedPackageWriter.h b/Engine/Source/Editor/UnrealEd/Private/Cooker/LooseCookedPackageWriter.h
--- a/Engine/Source/Editor/UnrealEd/Private/Cooker/LooseCookedPackageWriter.h
+++ b/Engine/Source/Editor/UnrealEd/Private/Cooker/LooseCookedPackageWriter.h
@@ -4,24 +4,24 @@
 
 #include "Containers/Array.h"
 #include "Containers/ArrayView.h"
+#include "Containers/Map.h"
 #include "Containers/StringFwd.h"
 #include "Containers/UnrealString.h"
 #include "HAL/CriticalSection.h"
 #include "HAL/Platform.h"
 #include "IO/IoDispatcher.h"
+#include "Memory/CompositeBuffer.h"
 #include "Memory/SharedBuffer.h"
 #include "Misc/DateTime.h"
+#include "Misc/PackagePath.h"
+#include "PackageStoreManifest.h"
 #include "Serialization/CompactBinary.h"
 #include "Serialization/FileRegions.h"
 #include "Serialization/PackageWriter.h"
-#include "UObject/NameTypes.h"
-#include "Containers/Map.h"
-#include "Memory/CompositeBuffer.h"
-#include "Misc/PackagePath.h"
-#include "PackageStoreManifest.h"
 #include "Serialization/PackageWriterToSharedBuffer.h"
 #include "Templates/SharedPointer.h"
 #include "Templates/UniquePtr.h"
+#include "UObject/NameTypes.h"
 
 class FAsyncIODelete;
 class FAssetRegistryState;
diff --git a/Engine/Source/Editor/UnrealEd/Private/Cooker/NetworkFileCookOnTheFlyRequestManager.h b/Engine/Source/Editor/UnrealEd/Private/Cooker/NetworkFileCookOnTheFlyRequestManager.h
--- a/Engine/Source/Editor/UnrealEd/Private/Cooker/NetworkFileCookOnTheFlyRequestManager.h
+++ b/Engine/Source/Editor/UnrealEd/Private/Cooker/NetworkFileCookOnTheFlyRequestManager.h
@@ -2,18 +2,17 @@
 
 #pragma once
 
-
-
 #include "Templates/SharedPointer.h"
 #include "Templates/UniquePtr.h"
 
-namespace UE { namespace Cook
+namespace UE::Cook
 {
 
 class ICookOnTheFlyNetworkServer;
 class ICookOnTheFlyRequestManager;
 class ICookOnTheFlyServer;
 
-TUniquePtr<ICookOnTheFlyRequestManager> MakeNetworkFileCookOnTheFlyRequestManager(ICookOnTheFlyServer& CookOnTheFlyServer, TSharedRef<ICookOnTheFlyNetworkServer> NetworkServer);
+TUniquePtr<ICookOnTheFlyRequestManager> MakeNetworkFileCookOnTheFlyRequestManager(ICookOnTheFlyServer& CookOnTheFlyServer,
+	TSharedRef<ICookOnTheFlyNetworkServer> NetworkServer);
 
-}} // namespace UE::Cook
+} // namespace UE::Cook
diff --git a/Engine/Source/Runtime/AssetRegistry/Private/AssetDataGatherer.h b/Engine/Source/Runtime/AssetRegistry/Private/AssetDataGatherer.h
--- a/Engine/Source/Runtime/AssetRegistry/Private/AssetDataGatherer.h
+++ b/Engine/Source/Runtime/AssetRegistry/Private/AssetDataGatherer.h
@@ -2,7 +2,6 @@
 
 #pragma once
 
-#include <atomic>
 #include "Containers/Array.h"
 #include "Containers/ArrayView.h"
 #include "Containers/Map.h"
@@ -20,8 +19,10 @@
 #include "Misc/Optional.h"
 #include "PackageDependencyData.h"
 #include "PackageReader.h"
-#include "UObject/NameTypes.h"
 #include "Templates/UniquePtr.h"
+#include "UObject/NameTypes.h"
+
+#include <atomic>
 
 class FArchive;
 struct FAssetData;
diff --git a/Engine/Source/Runtime/AssetRegistry/Private/AssetRegistryArchive.h b/Engine/Source/Runtime/AssetRegistry/Private/AssetRegistryArchive.h
--- a/Engine/Source/Runtime/AssetRegistry/Private/AssetRegistryArchive.h
+++ b/Engine/Source/Runtime/AssetRegistry/Private/AssetRegistryArchive.h
@@ -2,7 +2,6 @@
 
 #pragma once
 
-#include "CoreMinimal.h"
 #include "AssetRegistry/AssetDataTagMapSerializationDetails.h"
 #include "Serialization/LargeMemoryWriter.h"
 #include "UObject/NameBatchSerialization.h"
diff --git a/Engine/Source/Runtime/AssetRegistry/Private/AssetRegistryConsoleCommands.h b/Engine/Source/Runtime/AssetRegistry/Private/AssetRegistryConsoleCommands.h
--- a/Engine/Source/Runtime/AssetRegistry/Private/AssetRegistryConsoleCommands.h
+++ b/Engine/Source/Runtime/AssetRegistry/Private/AssetRegistryConsoleCommands.h
@@ -2,13 +2,12 @@
 
 #pragma once
 
-#include "CoreMinimal.h"
+#include "AssetRegistry/AssetData.h"
+#include "AssetRegistryPrivate.h"
 #include "HAL/IConsoleManager.h"
 #include "HAL/FileManager.h"
 #include "Misc/PackageName.h"
-#include "AssetRegistry/AssetData.h"
 #include "Misc/Paths.h"
-#include "AssetRegistryPrivate.h"
 
 #define LOCTEXT_NAMESPACE "AssetRegistry"
 
diff --git a/Engine/Source/Runtime/AssetRegistry/Private/AssetRegistryPrivate.h b/Engine/Source/Runtime/AssetRegistry/Private/AssetRegistryPrivate.h
--- a/Engine/Source/Runtime/AssetRegistry/Private/AssetRegistryPrivate.h
+++ b/Engine/Source/Runtime/AssetRegistry/Private/AssetRegistryPrivate.h
@@ -2,6 +2,6 @@
 
 #pragma once
 
-#include "CoreMinimal.h"
+#include "Logging/LogMacros.h"
 
 DECLARE_LOG_CATEGORY_EXTERN(LogAssetRegistry, Log, All);
diff --git a/Engine/Source/Runtime/AssetRegistry/Private/AssetRegistryState.cpp b/Engine/Source/Runtime/AssetRegistry/Private/AssetRegistryState.cpp
--- a/Engine/Source/Runtime/AssetRegistry/Private/AssetRegistryState.cpp
+++ b/Engine/Source/Runtime/AssetRegistry/Private/AssetRegistryState.cpp
@@ -3,28 +3,28 @@
 #include "AssetRegistry/AssetRegistryState.h"
 
 #include "Algo/Sort.h"
+#include "AssetRegistry/ARFilter.h"
+#include "AssetRegistryArchive.h"
+#include "AssetRegistryImpl.h"
+#include "AssetRegistryPrivate.h"
+#include "Blueprint/BlueprintSupport.h"
+#include "DependsNode.h"
+#include "GenericPlatform/GenericPlatformChunkInstall.h"
 #include "Misc/CommandLine.h"
+#include "Misc/ConfigCacheIni.h"
 #include "Misc/FileHelper.h"
 #include "Misc/Paths.h"
 #include "Misc/PathViews.h"
+#include "NameTableArchive.h"
+#include "PackageReader.h"
 #include "Serialization/ArrayReader.h"
 #include "Serialization/LargeMemoryReader.h"
 #include "Serialization/MemoryWriter.h"
-#include "Misc/ConfigCacheIni.h"
-#include "UObject/UObjectHash.h"
-#include "UObject/UObjectIterator.h"
 #include "UObject/MetaData.h"
-#include "AssetRegistryArchive.h"
-#include "AssetRegistryImpl.h"
-#include "AssetRegistryPrivate.h"
-#include "AssetRegistry/ARFilter.h"
-#include "DependsNode.h"
-#include "PackageReader.h"
-#include "NameTableArchive.h"
-#include "GenericPlatform/GenericPlatformChunkInstall.h"
-#include "Blueprint/BlueprintSupport.h"
 #include "UObject/NameBatchSerialization.h"
 #include "UObject/PrimaryAssetId.h"
+#include "UObject/UObjectHash.h"
+#include "UObject/UObjectIterator.h"
 
 FAssetRegistryState& FAssetRegistryState::operator=(FAssetRegistryState&& Rhs)
 {
diff --git a/Engine/Source/Runtime/AssetRegistry/Private/DependsNode.h b/Engine/Source/Runtime/AssetRegistry/Private/DependsNode.h
--- a/Engine/Source/Runtime/AssetRegistry/Private/DependsNode.h
+++ b/Engine/Source/Runtime/AssetRegistry/Private/DependsNode.h
@@ -2,12 +2,10 @@
 
 #pragma once
 
-#include "CoreMinimal.h"
-
 #include "Algo/BinarySearch.h"
 #include "AssetRegistry/AssetData.h"
-#include "Containers/BitArray.h"
 #include "AssetRegistry/IAssetRegistry.h"
+#include "Containers/BitArray.h"
 #include "Misc/AssetRegistryInterface.h"
 #include "PropertyCombinationSet.h"
 
diff --git a/Engine/Source/Runtime/AssetRegistry/Private/DiskCachedAssetData.h b/Engine/Source/Runtime/AssetRegistry/Private/DiskCachedAssetData.h
--- a/Engine/Source/Runtime/AssetRegistry/Private/DiskCachedAssetData.h
+++ b/Engine/Source/Runtime/AssetRegistry/Private/DiskCachedAssetData.h
@@ -2,7 +2,6 @@
 
 #pragma once
 
-#include "CoreMinimal.h"
 #include "AssetRegistry/AssetData.h"
 #include "PackageDependencyData.h"
 
diff --git a/Engine/Source/Runtime/AssetRegistry/Private/NameTableArchive.cpp b/Engine/Source/Runtime/AssetRegistry/Private/NameTableArchive.cpp
--- a/Engine/Source/Runtime/AssetRegistry/Private/NameTableArchive.cpp
+++ b/Engine/Source/Runtime/AssetRegistry/Private/NameTableArchive.cpp
@@ -1,8 +1,10 @@
 // Copyright Epic Games, Inc. All Rights Reserved.
 
 #include "NameTableArchive.h"
-#include "AssetRegistryPrivate.h"
+
 #include "AssetRegistry/AssetRegistryState.h"
+#include "AssetRegistryPrivate.h"
+#include "UObject/NameTypes.h"
 
 
 class FNameTableErrorArchive : public FArchive
diff --git a/Engine/Source/Runtime/AssetRegistry/Private/NameTableArchive.h b/Engine/Source/Runtime/AssetRegistry/Private/NameTableArchive.h
--- a/Engine/Source/Runtime/AssetRegistry/Private/NameTableArchive.h
+++ b/Engine/Source/Runtime/AssetRegistry/Private/NameTableArchive.h
@@ -2,9 +2,11 @@
 
 #pragma once
 
-#include "CoreMinimal.h"
+#include "Containers/Array.h"
+#include "Serialization/Archive.h"
 
 struct FAssetData;
+struct FNameEntryId;
 
 /** 
  * Legacy name table reader replaced by FAssetRegistryReader.
diff --git a/Engine/Source/Runtime/AssetRegistry/Private/PackageDependencyData.h b/Engine/Source/Runtime/AssetRegistry/Private/PackageDependencyData.h
--- a/Engine/Source/Runtime/AssetRegistry/Private/PackageDependencyData.h
+++ b/Engine/Source/Runtime/AssetRegistry/Private/PackageDependencyData.h
@@ -2,8 +2,6 @@
 
 #pragma once
 
-#include "CoreMinimal.h"
-
 #include "AssetRegistry/AssetData.h"
 #include "AssetRegistry/IAssetRegistry.h"
 #include "Containers/ArrayView.h"
diff --git a/Engine/Source/Runtime/AssetRegistry/Private/PackageReader.h b/Engine/Source/Runtime/AssetRegistry/Private/PackageReader.h
--- a/Engine/Source/Runtime/AssetRegistry/Private/PackageReader.h
+++ b/Engine/Source/Runtime/AssetRegistry/Private/PackageReader.h
@@ -2,9 +2,8 @@
 
 #pragma once
 
-#include "CoreMinimal.h"
-#include "Serialization/ArchiveUObject.h"
 #include "Misc/EnumClassFlags.h"
+#include "Serialization/ArchiveUObject.h"
 #include "UObject/ObjectResource.h"
 #include "UObject/PackageFileSummary.h"
 
diff --git a/Engine/Source/Runtime/AssetRegistry/Private/PathTree.cpp b/Engine/Source/Runtime/AssetRegistry/Private/PathTree.cpp
--- a/Engine/Source/Runtime/AssetRegistry/Private/PathTree.cpp
+++ b/Engine/Source/Runtime/AssetRegistry/Private/PathTree.cpp
@@ -1,7 +1,8 @@
 // Copyright Epic Games, Inc. All Rights Reserved.
 
 #include "AssetRegistry/PathTree.h"
-#include <UObject/NameTypes.h>
+
+#include "UObject/NameTypes.h"
 
 bool FPathTree::CachePath(FName Path, TFunctionRef<void(FName)> OnPathAdded)
 {
