DumpGPU:  fix so editor preview windows are dumped, useful for debugging graphical issues with such views.  There is logic which skips rendering of preview windows based on the time interval since they were last ticked, and running DumpGPU delays timing.  The fix is to force a render if a GPU dump was started on the current frame.  I had also attempted to invalidate the viewport to force it to refresh, but this didn't work because the command processing happens later in the Engine Tick than the viewport's Tick, so the invalidate doesn't happen until the next frame.

#rnx

diff --git a/Engine/Source/Editor/UnrealEd/Private/SEditorViewport.cpp b/Engine/Source/Editor/UnrealEd/Private/SEditorViewport.cpp
--- a/Engine/Source/Editor/UnrealEd/Private/SEditorViewport.cpp
+++ b/Engine/Source/Editor/UnrealEd/Private/SEditorViewport.cpp
@@ -581,11 +581,21 @@ bool SEditorViewport::IsRealtime() const
 
 bool SEditorViewport::IsVisible() const
 {
+#if WITH_DUMPGPU
+	extern ENGINE_API uint32 GDumpGPU_FrameNumber;
+#endif
+
 	const float VisibilityTimeThreshold = .25f;
-	// The viewport is visible if we don't have a parent layout (likely a floating window) or this viewport is visible in the parent layout
+	// The viewport is visible if we don't have a parent layout (likely a floating window) or this viewport is visible in the parent layout.
+	// Also, always render the viewport if DumpGPU is active, regardless of tick time threshold -- otherwise these don't show up due to lag
+	// caused by the GPU dump being triggered.
 	return 
 		LastTickTime == 0.0	||	// Never been ticked
-		FPlatformTime::Seconds() - LastTickTime <= VisibilityTimeThreshold;	// Ticked recently
+		FPlatformTime::Seconds() - LastTickTime <= VisibilityTimeThreshold	// Ticked recently
+#if WITH_DUMPGPU
+		|| GDumpGPU_FrameNumber == GFrameNumber	// GPU dump in progress
+#endif		
+		;
 }
 
 void SEditorViewport::OnScreenCapture()
diff --git a/Engine/Source/Runtime/Engine/Private/UnrealEngine.cpp b/Engine/Source/Runtime/Engine/Private/UnrealEngine.cpp
--- a/Engine/Source/Runtime/Engine/Private/UnrealEngine.cpp
+++ b/Engine/Source/Runtime/Engine/Private/UnrealEngine.cpp
@@ -5721,6 +5721,8 @@ bool UEngine::HandleProfileGPUCommand( const TCHAR* Cmd, FOutputDevice& Ar )
 #endif // WITH_PROFILEGPU
 
 #if WITH_DUMPGPU
+ENGINE_API uint32 GDumpGPU_FrameNumber = INDEX_NONE;
+
 bool UEngine::HandleDumpGPUCommand(const TCHAR* Cmd, FOutputDevice& Ar)
 {
 	if (!FApp::CanEverRender())
@@ -5729,6 +5731,8 @@ bool UEngine::HandleDumpGPUCommand(const TCHAR* Cmd, FOutputDevice& Ar)
 	}
 
 	check(IsInGameThread());
+	GDumpGPU_FrameNumber = GFrameNumber;
+
 	TArray<FString> Args;
 	FString ResourceDumpDirectory = FRDGBuilder::BeginResourceDump(Cmd);
 
