Remove a target Blueprint from cached dependent lists prior to reloading its outer package.

Also remove a target Blueprint from cached dependency/dependent lists prior to unloading its outer package.

diff --git a/Engine/Source/Editor/UnrealEd/Private/PackageTools.cpp b/Engine/Source/Editor/UnrealEd/Private/PackageTools.cpp
--- a/Engine/Source/Editor/UnrealEd/Private/PackageTools.cpp
+++ b/Engine/Source/Editor/UnrealEd/Private/PackageTools.cpp
@@ -483,6 +483,28 @@ UPackageTools::UPackageTools(const FObjectInitializer& ObjectInitializer)
 					if (UBlueprint* BP = Cast<UBlueprint>(Obj))
 					{
 						BP->ClearEditorReferences();
+
+						// Remove from cached dependent lists.
+						for (const TWeakObjectPtr<UBlueprint> Dependency : BP->CachedDependencies)
+						{
+							if (UBlueprint* ResolvedDependency = Dependency.Get())
+							{
+								ResolvedDependency->CachedDependents.Remove(BP);
+							}
+						}
+
+						BP->CachedDependencies.Reset();
+
+						// Remove from cached dependency lists.
+						for (const TWeakObjectPtr<UBlueprint> Dependent : BP->CachedDependents)
+						{
+							if (UBlueprint* ResolvedDependent = Dependent.Get())
+							{
+								ResolvedDependent->CachedDependencies.Remove(BP);
+							}
+						}
+
+						BP->CachedDependents.Reset();
 					}
 					if (UWorld* World = Cast<UWorld>(Obj))
 					{
@@ -987,6 +1009,18 @@ UPackageTools::UPackageTools(const FObjectInitializer& ObjectInitializer)
 				if (UBlueprint* BP = Cast<UBlueprint>(InObject))
 				{
 					BP->ClearEditorReferences();
+
+					// Remove from cached dependent lists; this will be repopulated on reload, but we
+					// don't wish to consider every one of our dependencies as a potential referencer,
+					// and this set will be serialized by the archiver that's used to find those. Any
+					// references will instead be collected from other fields that reference the asset.
+					for (const TWeakObjectPtr<UBlueprint> Dependency : BP->CachedDependencies)
+					{
+						if (UBlueprint* ResolvedDependency = Dependency.Get())
+						{
+							ResolvedDependency->CachedDependents.Remove(BP);
+						}
+					}
 				}
 				if (UWorld* World = Cast<UWorld>(InObject))
 				{
