- Added new optional world partition builder UWorldPartitionLandscapeSplineMeshesBuilder to extract Landscape SplineMeshComponents and ControlPointMeshComponents and distributes them into ALandscapeSplineMeshesActors.
- When builder is used, ALandscape actor is marked with a flag. This flag is then used to automatically mark hidden in game (for PIE) and strip from cooked builds the source components.
- Added new ALandscapeSplineMeshesActor (partitioned actor) that serves as a container for extracted landscape spline meshes and control point meshes.

diff --git a/Engine/Source/Editor/LevelEditor/Private/LevelEditor.cpp b/Engine/Source/Editor/LevelEditor/Private/LevelEditor.cpp
--- a/Engine/Source/Editor/LevelEditor/Private/LevelEditor.cpp
+++ b/Engine/Source/Editor/LevelEditor/Private/LevelEditor.cpp
@@ -1487,6 +1487,9 @@ void FLevelEditorModule::BindGlobalLevelEditorCommands()
 	ActionList.MapAction(Commands.BuildMinimap,
 		FExecuteAction::CreateStatic(&FLevelEditorActionCallbacks::BuildMinimap_Execute));
 
+	ActionList.MapAction(Commands.BuildLandscapeSplineMeshes,
+		FExecuteAction::CreateStatic(&FLevelEditorActionCallbacks::BuildLandscapeSplineMeshes_Execute));
+
 	ActionList.MapAction(Commands.BuildTextureStreamingOnly,
 		FExecuteAction::CreateStatic(&FLevelEditorActionCallbacks::BuildTextureStreamingOnly_Execute));
 
diff --git a/Engine/Source/Editor/LevelEditor/Private/LevelEditorActions.cpp b/Engine/Source/Editor/LevelEditor/Private/LevelEditorActions.cpp
--- a/Engine/Source/Editor/LevelEditor/Private/LevelEditorActions.cpp
+++ b/Engine/Source/Editor/LevelEditor/Private/LevelEditorActions.cpp
@@ -832,6 +832,12 @@ void FLevelEditorActionCallbacks::BuildMinimap_Execute()
 	FEditorBuildUtils::EditorBuild(GetWorld(), FBuildOptions::BuildMinimap);
 }
 
+void FLevelEditorActionCallbacks::BuildLandscapeSplineMeshes_Execute()
+{
+	// Build Landscape Spline Meshes
+	FEditorBuildUtils::EditorBuild(GetWorld(), FBuildOptions::BuildLandscapeSplineMeshes);
+}
+
 void FLevelEditorActionCallbacks::BuildTextureStreamingOnly_Execute()
 {
 	FEditorBuildUtils::EditorBuildTextureStreaming(GetWorld());
@@ -3442,6 +3448,7 @@ void FLevelEditorCommands::RegisterCommands()
 	UI_COMMAND( BuildPathsOnly, "Build Paths", "Only builds paths (all levels.)", EUserInterfaceActionType::Button, FInputChord() );
 	UI_COMMAND( BuildHLODs, "Build HLODs", "Builds all HLODs for the current world", EUserInterfaceActionType::Button, FInputChord() );
 	UI_COMMAND( BuildMinimap, "Build Minimap", "Builds the minimap for the current world", EUserInterfaceActionType::Button, FInputChord() );
+	UI_COMMAND( BuildLandscapeSplineMeshes, "Build Landscape Spline Meshes", "Builds landscape spline meshes for the current world", EUserInterfaceActionType::Button, FInputChord());
 	UI_COMMAND( BuildTextureStreamingOnly, "Build Texture Streaming", "Build texture streaming data", EUserInterfaceActionType::Button, FInputChord() );
 	UI_COMMAND( BuildVirtualTextureOnly, "Build Virtual Textures", "Build runtime virtual texture low mips streaming data", EUserInterfaceActionType::Button, FInputChord());
 	UI_COMMAND( BuildAllLandscape, "Build All Landscape", "Build All Landscape Data(Grass, Baked Textures)", EUserInterfaceActionType::Button, FInputChord());
diff --git a/Engine/Source/Editor/LevelEditor/Private/LevelEditorMenu.cpp b/Engine/Source/Editor/LevelEditor/Private/LevelEditorMenu.cpp
--- a/Engine/Source/Editor/LevelEditor/Private/LevelEditorMenu.cpp
+++ b/Engine/Source/Editor/LevelEditor/Private/LevelEditorMenu.cpp
@@ -611,6 +611,7 @@ void FLevelEditorMenu::RegisterBuildMenu()
 				FToolMenuSection& Section = InMenu->AddSection("LevelEditorWorldPartition", LOCTEXT("WorldPartitionHeading", "World Partition"));
 				Section.AddMenuEntry(FLevelEditorCommands::Get().BuildHLODs);
 				Section.AddMenuEntry(FLevelEditorCommands::Get().BuildMinimap);
+				Section.AddMenuEntry(FLevelEditorCommands::Get().BuildLandscapeSplineMeshes);
 			}
 		}));
 	}
diff --git a/Engine/Source/Editor/LevelEditor/Public/LevelEditorActions.h b/Engine/Source/Editor/LevelEditor/Public/LevelEditorActions.h
--- a/Engine/Source/Editor/LevelEditor/Public/LevelEditorActions.h
+++ b/Engine/Source/Editor/LevelEditor/Public/LevelEditorActions.h
@@ -79,6 +79,7 @@ public:
 	TSharedPtr< FUICommandInfo > BuildPathsOnly;
 	TSharedPtr< FUICommandInfo > BuildHLODs;
 	TSharedPtr< FUICommandInfo > BuildMinimap;
+	TSharedPtr< FUICommandInfo > BuildLandscapeSplineMeshes;
 	TSharedPtr< FUICommandInfo > BuildTextureStreamingOnly;
 	TSharedPtr< FUICommandInfo > BuildVirtualTextureOnly;
 	TSharedPtr< FUICommandInfo > BuildAllLandscape;
@@ -752,6 +753,7 @@ public:
 	static void BuildPathsOnly_Execute();
 	static void BuildHLODs_Execute();
 	static void BuildMinimap_Execute();
+	static void BuildLandscapeSplineMeshes_Execute();
 	static void BuildTextureStreamingOnly_Execute();
 	static void BuildVirtualTextureOnly_Execute();
 	static void BuildAllLandscape_Execute();
diff --git a/Engine/Source/Editor/UnrealEd/Private/EditorBuildUtils.cpp b/Engine/Source/Editor/UnrealEd/Private/EditorBuildUtils.cpp
--- a/Engine/Source/Editor/UnrealEd/Private/EditorBuildUtils.cpp
+++ b/Engine/Source/Editor/UnrealEd/Private/EditorBuildUtils.cpp
@@ -74,6 +74,7 @@ const FName FBuildOptions::BuildAllSubmit(TEXT("BuildAllSubmit"));
 const FName FBuildOptions::BuildAllOnlySelectedPaths(TEXT("BuildAllOnlySelectedPaths"));
 const FName FBuildOptions::BuildHierarchicalLOD(TEXT("BuildHierarchicalLOD"));
 const FName FBuildOptions::BuildMinimap(TEXT("BuildMinimap"));
+const FName FBuildOptions::BuildLandscapeSplineMeshes(TEXT("BuildLandscapeSplineMeshes"));
 const FName FBuildOptions::BuildTextureStreaming(TEXT("BuildTextureStreaming"));
 const FName FBuildOptions::BuildVirtualTexture(TEXT("BuildVirtualTexture"));
 const FName FBuildOptions::BuildAllLandscape(TEXT("BuildAllLandscape"));
@@ -317,6 +318,10 @@ bool FEditorBuildUtils::EditorBuild( UWorld* InWorld, FName Id, const bool bAllo
 	{
 		BuildType = SBuildProgressWidget::BUILDTYPE_Minimap;
 	}
+	else if (Id == FBuildOptions::BuildLandscapeSplineMeshes)
+	{
+		BuildType = SBuildProgressWidget::BUILDTYPE_LandscapeSplineMeshes;
+	}
 	else if (Id == FBuildOptions::BuildTextureStreaming)
 	{
 		BuildType = SBuildProgressWidget::BUILDTYPE_TextureStreaming;
@@ -469,7 +474,7 @@ bool FEditorBuildUtils::EditorBuild( UWorld* InWorld, FName Id, const bool bAllo
 		bDoBuild = InWorld->IsPartitionedWorld();
 		if ( bDoBuild )
 		{
-			GEditor->ResetTransaction( NSLOCTEXT("UnrealEd", "BuildHLODMeshes", "Building Hierarchical LOD Meshes") );
+			GEditor->ResetTransaction( NSLOCTEXT("UnrealEd", "BuildMinimap", "Building Minimap") );
 
 			// We can't set the busy cursor for all windows, because lighting
 			// needs a cursor for the lighting options dialog.
@@ -481,6 +486,23 @@ bool FEditorBuildUtils::EditorBuild( UWorld* InWorld, FName Id, const bool bAllo
 			TriggerMinimapBuilder(InWorld, Id);
 		}
 	}
+	else if (Id == FBuildOptions::BuildLandscapeSplineMeshes)
+	{
+		bDoBuild = InWorld->IsPartitionedWorld();
+		if ( bDoBuild )
+		{
+			GEditor->ResetTransaction( NSLOCTEXT("UnrealEd", "BuildLandscapeSplineMeshes", "Building Landscape Spline Meshes") );
+
+			// We can't set the busy cursor for all windows, because lighting
+			// needs a cursor for the lighting options dialog.
+			const FScopedBusyCursor BusyCursor;
+
+			bShouldMapCheck = false;
+			bDirtyPersistentLevel = false;
+
+			TriggerLandscapeSplineMeshesBuilder(InWorld);
+		}
+	}
 	else if (Id == FBuildOptions::BuildAllLandscape)
 	{
 		bDoBuild = GEditor->WarnAboutHiddenLevels(InWorld, false);
@@ -1204,6 +1226,11 @@ void FBuildAllHandler::ProcessBuild(const TWeakPtr<SBuildProgressWidget>& BuildP
 			BuildProgressWidget.Pin()->SetBuildType(SBuildProgressWidget::BUILDTYPE_Minimap);
 			FEditorBuildUtils::TriggerMinimapBuilder(CurrentWorld, CurrentBuildId);
 		}
+		else if (StepId == FBuildOptions::BuildLandscapeSplineMeshes)
+		{
+			BuildProgressWidget.Pin()->SetBuildType(SBuildProgressWidget::BUILDTYPE_LandscapeSplineMeshes);
+			FEditorBuildUtils::TriggerLandscapeSplineMeshesBuilder(CurrentWorld);
+		}
 		else if (StepId == FBuildOptions::BuildTextureStreaming)
 		{
 			BuildProgressWidget.Pin()->SetBuildType(SBuildProgressWidget::BUILDTYPE_TextureStreaming);
@@ -1290,7 +1317,7 @@ void FEditorBuildUtils::TriggerHierarchicalLODBuilder(UWorld* InWorld, FName Id)
 	{
 		IWorldPartitionEditorModule& WorldPartitionEditorModule = FModuleManager::LoadModuleChecked<IWorldPartitionEditorModule>("WorldPartitionEditor");
 		TSubclassOf<UWorldPartitionBuilder> WorldPartitionHLODsBuilder = FindObjectChecked<UClass>(nullptr, TEXT("/Script/UnrealEd.WorldPartitionHLODsBuilder"), true);
-		WorldPartitionEditorModule.RunBuilder(WorldPartitionHLODsBuilder, InWorld->GetPackage()->GetName());
+		WorldPartitionEditorModule.RunBuilder(WorldPartitionHLODsBuilder, InWorld);
 	}
 	else
 	{
@@ -1305,7 +1332,17 @@ void FEditorBuildUtils::TriggerMinimapBuilder(UWorld* InWorld, FName Id)
 	{
 		IWorldPartitionEditorModule& WorldPartitionEditorModule = FModuleManager::LoadModuleChecked<IWorldPartitionEditorModule>("WorldPartitionEditor");
 		TSubclassOf<UWorldPartitionBuilder> WorldPartitionMiniMapBuilder = FindObjectChecked<UClass>(nullptr, TEXT("/Script/UnrealEd.WorldPartitionMiniMapBuilder"), true);
-		WorldPartitionEditorModule.RunBuilder(WorldPartitionMiniMapBuilder, InWorld->GetPackage()->GetName());
+		WorldPartitionEditorModule.RunBuilder(WorldPartitionMiniMapBuilder, InWorld);
+	}
+}
+
+void FEditorBuildUtils::TriggerLandscapeSplineMeshesBuilder(UWorld* InWorld)
+{
+	if (InWorld->IsPartitionedWorld())
+	{
+		IWorldPartitionEditorModule& WorldPartitionEditorModule = FModuleManager::LoadModuleChecked<IWorldPartitionEditorModule>("WorldPartitionEditor");
+		TSubclassOf<UWorldPartitionBuilder> WorldPartitionLandscapeSplineMeshesBuilder = FindObjectChecked<UClass>(nullptr, TEXT("/Script/UnrealEd.WorldPartitionLandscapeSplineMeshesBuilder"), true);
+		WorldPartitionEditorModule.RunBuilder(WorldPartitionLandscapeSplineMeshesBuilder, InWorld);
 	}
 }
 
diff --git a/Engine/Source/Editor/UnrealEd/Private/WorldPartition/WorldPartitionFoliageBuilder.cpp b/Engine/Source/Editor/UnrealEd/Private/WorldPartition/WorldPartitionFoliageBuilder.cpp
--- a/Engine/Source/Editor/UnrealEd/Private/WorldPartition/WorldPartitionFoliageBuilder.cpp
+++ b/Engine/Source/Editor/UnrealEd/Private/WorldPartition/WorldPartitionFoliageBuilder.cpp
@@ -79,7 +79,7 @@ bool UWorldPartitionFoliageBuilder::RunInternal(UWorld* World, const FCellInfo&
 	// - Unload the Existing IFA and all the new IFAs
 	// - Repeat for next Existing IFA
 	FWorldPartitionHelpers::FForEachActorWithLoadingParams ForEachActorWithLoadingParams;
-	ForEachActorWithLoadingParams.ActorClass = AInstancedFoliageActor::StaticClass();
+	ForEachActorWithLoadingParams.ActorClasses = { AInstancedFoliageActor::StaticClass() };
 
 	FWorldPartitionHelpers::ForEachActorWithLoading(WorldPartition, [this, World, WorldPartition, ActorPartitionSubsystem, &SCCHelper, &NewActors, &NumInstances, &NumInstancesProcessed, &ExistingActors](const FWorldPartitionActorDesc* ActorDesc)
 	{
diff --git a/Engine/Source/Editor/UnrealEd/Private/WorldPartition/WorldPartitionLandscapeSplineMeshesBuilder.cpp b/Engine/Source/Editor/UnrealEd/Private/WorldPartition/WorldPartitionLandscapeSplineMeshesBuilder.cpp
new file mode 100644
--- /dev/null
+++ b/Engine/Source/Editor/UnrealEd/Private/WorldPartition/WorldPartitionLandscapeSplineMeshesBuilder.cpp
@@ -0,0 +1,538 @@
+// Copyright Epic Games, Inc. All Rights Reserved.
+
+#include "WorldPartition/WorldPartitionLandscapeSplineMeshesBuilder.h"
+
+#include "PackageSourceControlHelper.h"
+#include "SourceControlHelpers.h"
+#include "HAL/PlatformFileManager.h"
+#include "UObject/SavePackage.h"
+#include "UObject/ScriptInterface.h"
+#include "UObject/GCObjectScopeGuard.h"
+#include "Engine/World.h"
+#include "Engine/StaticMesh.h"
+#include "ActorFolder.h"
+#include "GameFramework/Actor.h"
+#include "GameFramework/WorldSettings.h"
+#include "Components/StaticMeshComponent.h"
+#include "Components/SplineMeshComponent.h"
+#include "ControlPointMeshComponent.h"
+#include "LandscapeSplinesComponent.h"
+#include "LandscapeSplineControlPoint.h"
+#include "LandscapeSplineMeshesActor.h"
+#include "LandscapeSplineActor.h"
+#include "LandscapeInfo.h"
+#include "Landscape.h"
+#include "WorldPartition/WorldPartition.h"
+#include "WorldPartition/WorldPartitionActorDesc.h"
+#include "WorldPartition/WorldPartitionHelpers.h"
+#include "ActorPartition/ActorPartitionSubsystem.h"
+#include "Rendering/ColorVertexBuffer.h"
+
+DEFINE_LOG_CATEGORY_STATIC(LogWorldPartitionLandscapeSplineMeshBuilder, All, All);
+
+static const FName CreatedFromBuilderTag(TEXT("CreatedFromWorldPartitionLandscapeSplineMeshBuilder"));
+
+UStaticMesh* UWorldPartitionLandscapeSplineMeshesBuilder::SplineEditorMesh = nullptr;
+
+UWorldPartitionLandscapeSplineMeshesBuilder::UWorldPartitionLandscapeSplineMeshesBuilder(const FObjectInitializer& ObjectInitializer)
+	: Super(ObjectInitializer)
+	, NewGridSize(0)
+{
+	FParse::Value(FCommandLine::Get(), TEXT("NewGridSize="), NewGridSize);
+
+	static ConstructorHelpers::FObjectFinder<UStaticMesh> ConstructorStatics(TEXT("/Engine/EditorLandscapeResources/SplineEditorMesh"));
+	SplineEditorMesh = ConstructorStatics.Object;
+}
+
+FName UWorldPartitionLandscapeSplineMeshesBuilder::GetSplineCollisionProfileName(const UStaticMeshComponent* InMeshComponent)
+{
+	FName CollisionProfileName = InMeshComponent->GetCollisionProfileName();
+
+	// Get collision profile name
+	if (ILandscapeSplineInterface* LSSplineInterface = Cast<ILandscapeSplineInterface>(InMeshComponent->GetOwner()))
+	{
+		UObject* ComponentOwner = LSSplineInterface->GetSplinesComponent()->GetOwnerForMeshComponent(InMeshComponent);
+		if (ComponentOwner)
+		{
+			if (ULandscapeSplineControlPoint* ControlPoint = Cast<ULandscapeSplineControlPoint>(ComponentOwner))
+			{
+				CollisionProfileName = ControlPoint->GetCollisionProfileName();
+			}
+			else if (ULandscapeSplineSegment* SplineSegment = Cast<ULandscapeSplineSegment>(ComponentOwner))
+			{
+				CollisionProfileName = SplineSegment->GetCollisionProfileName();
+			}
+		}
+	}
+
+	return CollisionProfileName;
+}
+
+int32 UWorldPartitionLandscapeSplineMeshesBuilder::HashStaticMeshComponent(const UStaticMeshComponent* InComponent)
+{
+	// Here we use ToCompactString as it internally clamps to 2 decimal digits, or else we get error precision and hash will differ)
+	uint32 HashValue = GetTypeHash(InComponent->GetComponentLocation().ToCompactString());
+	HashValue = HashCombine(HashValue, GetTypeHash(InComponent->GetComponentScale().ToCompactString()));
+	HashValue = HashCombine(HashValue, GetTypeHash(InComponent->GetComponentRotation().GetNormalized().ToCompactString()));
+	HashValue = HashCombine(HashValue, GetTypeHash(GetSplineCollisionProfileName(InComponent)));
+	HashValue = HashCombine(HashValue, GetTypeHash(InComponent->CastShadow));
+	HashValue = HashCombine(HashValue, GetTypeHash(InComponent->GetStaticMesh()));
+	HashValue = HashCombine(HashValue, GetTypeHash(InComponent->LDMaxDrawDistance));
+	HashValue = HashCombine(HashValue, GetTypeHash(InComponent->TranslucencySortPriority));
+
+	// Hash VT settings
+	for (const URuntimeVirtualTexture* RVT : InComponent->RuntimeVirtualTextures)
+	{
+		HashValue = HashCombine(HashValue, GetTypeHash(RVT));
+	}
+	HashValue = HashCombine(HashValue, GetTypeHash(InComponent->VirtualTextureRenderPassType));
+	HashValue = HashCombine(HashValue, GetTypeHash(InComponent->VirtualTextureLodBias));
+	HashValue = HashCombine(HashValue, GetTypeHash(InComponent->VirtualTextureCullMips));
+
+	// Hash vertex colors
+	int32 NumLODs = InComponent->GetStaticMesh()->GetNumLODs();
+	for (int32 CurrentLOD = 0; CurrentLOD != NumLODs; CurrentLOD++)
+	{
+		FStaticMeshLODResources& SourceLODModel = InComponent->GetStaticMesh()->GetRenderData()->LODResources[CurrentLOD];
+		if (InComponent->LODData.IsValidIndex(CurrentLOD))
+		{
+			const FStaticMeshComponentLODInfo& SourceLODInfo = InComponent->LODData[CurrentLOD];
+			if (SourceLODInfo.OverrideVertexColors != NULL)
+			{
+				// Copy vertex colors from source to target.
+				FColorVertexBuffer* SourceColorBuffer = SourceLODInfo.OverrideVertexColors;
+
+				for (uint32 ColorVertexIndex = 0; ColorVertexIndex < SourceColorBuffer->GetNumVertices(); ColorVertexIndex++)
+				{
+					HashValue = HashCombine(HashValue, GetTypeHash(SourceColorBuffer->VertexColor(ColorVertexIndex)));
+				}
+			}
+		}
+	}
+
+	// Hash used materials
+	for (int32 i = 0; i < InComponent->GetNumMaterials(); ++i)
+	{
+		HashValue = HashCombine(HashValue, GetTypeHash(InComponent->GetMaterial(i)->GetFName()));
+	}
+
+	// Hash spline mesh component params
+	if (const USplineMeshComponent* SplineMeshComponent = Cast<USplineMeshComponent>(InComponent))
+	{
+		HashValue = HashCombine(HashValue, GetTypeHash(SplineMeshComponent->GetForwardAxis()));
+		HashValue = HashCombine(HashValue, GetTypeHash(SplineMeshComponent->GetStartPosition()));
+		HashValue = HashCombine(HashValue, GetTypeHash(SplineMeshComponent->GetStartTangent()));
+		HashValue = HashCombine(HashValue, GetTypeHash(SplineMeshComponent->GetEndPosition()));
+		HashValue = HashCombine(HashValue, GetTypeHash(SplineMeshComponent->GetEndTangent()));
+		HashValue = HashCombine(HashValue, GetTypeHash(SplineMeshComponent->GetStartOffset()));
+		HashValue = HashCombine(HashValue, GetTypeHash(SplineMeshComponent->GetEndOffset()));
+		HashValue = HashCombine(HashValue, GetTypeHash(SplineMeshComponent->GetStartScale()));
+		HashValue = HashCombine(HashValue, GetTypeHash(SplineMeshComponent->GetEndScale()));
+		HashValue = HashCombine(HashValue, GetTypeHash(SplineMeshComponent->GetStartRoll()));
+		HashValue = HashCombine(HashValue, GetTypeHash(SplineMeshComponent->GetEndRoll()));
+		HashValue = HashCombine(HashValue, GetTypeHash(SplineMeshComponent->GetSplineUpDir()));
+
+		HashValue = HashCombine(HashValue, GetTypeHash(SplineMeshComponent->VirtualTextureMainPassMaxDrawDistance));
+	}
+	// Hash control point mesh component params
+	else if (const UControlPointMeshComponent* ControlPointMeshComponent = Cast<UControlPointMeshComponent>(InComponent))
+	{
+		HashValue = HashCombine(HashValue, GetTypeHash(ControlPointMeshComponent->VirtualTextureMainPassMaxDrawDistance));
+	}
+
+	// Hash the owning actor's tags
+	if (const AActor* Owner = InComponent->GetOwner())
+	{
+		for (const FName& ActorTag : Owner->Tags)
+		{
+			if (!ActorTag.IsEqual(CreatedFromBuilderTag))
+			{
+				HashValue = HashCombine(HashValue, GetTypeHash(ActorTag));
+			}
+		}
+	}
+
+	return HashValue;
+}
+
+void UWorldPartitionLandscapeSplineMeshesBuilder::CloneStaticMeshComponent(const UStaticMeshComponent* InSrcMeshComponent, UStaticMeshComponent* DstMeshComponent)
+{
+	check(InSrcMeshComponent);
+	check(DstMeshComponent);
+
+	// Clone a USplineMeshComponent
+	if (const USplineMeshComponent* SplineMeshComponent = Cast<USplineMeshComponent>(InSrcMeshComponent))
+	{
+		USplineMeshComponent* DstSplineMeshComponent = CastChecked<USplineMeshComponent>(DstMeshComponent);
+		DstSplineMeshComponent->bAllowSplineEditingPerInstance = true;
+		// Copy spline settings
+		DstSplineMeshComponent->SetForwardAxis(SplineMeshComponent->GetForwardAxis(), false);
+		DstSplineMeshComponent->SetStartAndEnd(SplineMeshComponent->GetStartPosition(), SplineMeshComponent->GetStartTangent(), SplineMeshComponent->GetEndPosition(), SplineMeshComponent->GetEndTangent(), false);
+		DstSplineMeshComponent->SetStartOffset(SplineMeshComponent->GetStartOffset(), false);
+		DstSplineMeshComponent->SetEndOffset(SplineMeshComponent->GetEndOffset(), false);
+		DstSplineMeshComponent->SetStartScale(SplineMeshComponent->GetStartScale(), false);
+		DstSplineMeshComponent->SetEndScale(SplineMeshComponent->GetEndScale(), false);
+		DstSplineMeshComponent->SetStartRoll(SplineMeshComponent->GetStartRoll(), false);
+		DstSplineMeshComponent->SetEndRoll(SplineMeshComponent->GetEndRoll(), false);
+		DstSplineMeshComponent->SetSplineUpDir(SplineMeshComponent->GetSplineUpDir(), false);
+		DstSplineMeshComponent->VirtualTextureMainPassMaxDrawDistance = SplineMeshComponent->VirtualTextureMainPassMaxDrawDistance;
+	}
+	// Clone a UControlPointMeshComponent
+	else if (const UControlPointMeshComponent* ControlPointMeshComponent = Cast<UControlPointMeshComponent>(InSrcMeshComponent))
+	{
+		UControlPointMeshComponent* DstControlPointMeshComponent = CastChecked<UControlPointMeshComponent>(DstMeshComponent);
+		DstControlPointMeshComponent->VirtualTextureMainPassMaxDrawDistance = ControlPointMeshComponent->VirtualTextureMainPassMaxDrawDistance;
+	}
+
+	// Setup component transform
+	DstMeshComponent->SetWorldTransform(InSrcMeshComponent->GetComponentTransform());
+
+	// Copy collision profile
+	DstMeshComponent->SetCollisionProfileName(GetSplineCollisionProfileName(InSrcMeshComponent));
+
+	// Shadow
+	DstMeshComponent->SetCastShadow(InSrcMeshComponent->CastShadow);
+
+	// Copy mesh
+	DstMeshComponent->SetStaticMesh(InSrcMeshComponent->GetStaticMesh());
+	if (DstMeshComponent->IsA<USplineMeshComponent>())
+	{
+		CastChecked<USplineMeshComponent>(DstMeshComponent)->UpdateMesh();
+	}
+
+	// Max draw distance
+	DstMeshComponent->LDMaxDrawDistance = InSrcMeshComponent->LDMaxDrawDistance;
+
+	// Translucency sort priority
+	DstMeshComponent->TranslucencySortPriority = InSrcMeshComponent->TranslucencySortPriority;
+
+	// Copy VT settings
+	DstMeshComponent->RuntimeVirtualTextures = InSrcMeshComponent->RuntimeVirtualTextures;
+	DstMeshComponent->VirtualTextureRenderPassType = InSrcMeshComponent->VirtualTextureRenderPassType;
+	DstMeshComponent->VirtualTextureLodBias = InSrcMeshComponent->VirtualTextureLodBias;
+	DstMeshComponent->VirtualTextureCullMips = InSrcMeshComponent->VirtualTextureCullMips;
+
+	// Copy vertex colors
+	DstMeshComponent->CopyInstanceVertexColorsIfCompatible(InSrcMeshComponent);
+
+	// Copy assigned materials
+	for (int32 i = 0; i < InSrcMeshComponent->GetNumMaterials(); ++i)
+	{
+		DstMeshComponent->SetMaterial(i, InSrcMeshComponent->GetMaterial(i));
+	}
+
+	// Copy the owning actor's tags
+	if (const AActor* Owner = InSrcMeshComponent->GetOwner())
+	{
+		DstMeshComponent->ComponentTags.Append(Owner->Tags);
+	}
+}
+
+void UWorldPartitionLandscapeSplineMeshesBuilder::CloneStaticMeshComponentInActor(ALandscapeSplineMeshesActor* InActor, const UStaticMeshComponent* InMeshComponent)
+{
+	check(InActor);
+	check(InMeshComponent->IsA<USplineMeshComponent>() || InMeshComponent->IsA<UControlPointMeshComponent>());
+	UStaticMeshComponent* NewMeshComponent = nullptr;
+	if (InMeshComponent->IsA<USplineMeshComponent>())
+	{
+		NewMeshComponent = CastChecked<USplineMeshComponent>(InActor->CreateStaticMeshComponent(USplineMeshComponent::StaticClass()));
+	}
+	else
+	{
+		NewMeshComponent = CastChecked<UControlPointMeshComponent>(InActor->CreateStaticMeshComponent(UControlPointMeshComponent::StaticClass()));
+	}
+	CloneStaticMeshComponent(InMeshComponent, NewMeshComponent);
+}
+
+ALandscapeSplineMeshesActor* UWorldPartitionLandscapeSplineMeshesBuilder::GetOrCreatePartitionActorForComponent(UWorld* InWorld, const UStaticMeshComponent* InMeshComponent, const FGuid& InLandscapeGuid)
+{
+	check(InWorld);
+	check(InWorld->IsPartitionedWorld());
+	check(InMeshComponent);
+	check(InMeshComponent->IsA<USplineMeshComponent>() || InMeshComponent->IsA<UControlPointMeshComponent>());
+	UActorPartitionSubsystem* ActorPartitionSubsystem = InWorld->GetSubsystem<UActorPartitionSubsystem>();
+	check(ActorPartitionSubsystem);
+
+	// Create or find the placement partition actor
+	auto OnActorCreated = [&InLandscapeGuid](APartitionActor* CreatedPartitionActor)
+	{
+		if (ALandscapeSplineMeshesActor* LandscapeSplineMeshesActor = CastChecked<ALandscapeSplineMeshesActor>(CreatedPartitionActor))
+		{
+			LandscapeSplineMeshesActor->SetGridGuid(InLandscapeGuid);
+		}
+	};
+
+	FActorPartitionGetParams Params(
+		ALandscapeSplineMeshesActor::StaticClass(),
+		true,
+		InWorld->PersistentLevel,
+		FVector(FVector2D(InMeshComponent->GetComponentLocation()), 0),
+		0,
+		InLandscapeGuid,
+		true,
+		OnActorCreated);
+
+	return CastChecked<ALandscapeSplineMeshesActor>(ActorPartitionSubsystem->GetActor(Params));
+}
+
+void UWorldPartitionLandscapeSplineMeshesBuilder::FilterStaticMeshComponents(TArray<UStaticMeshComponent*>& InOutComponents)
+{
+	check(SplineEditorMesh);
+
+	// Remove invalid components from the list
+	InOutComponents.RemoveAll([](UStaticMeshComponent* Component)
+	{
+		if (!Component || !Component->GetStaticMesh())
+		{
+			return true;
+		}
+
+		// Remove editor splines
+		if (Component->GetStaticMesh() == SplineEditorMesh)
+		{
+			return true;
+		}
+
+		// We only care for SplineMeshComponents & ControlPointMeshComponents
+		if (!Component->IsA<USplineMeshComponent>() && !Component->IsA<UControlPointMeshComponent>())
+		{
+			return true;
+		}
+
+		return false;
+	});
+}
+
+bool UWorldPartitionLandscapeSplineMeshesBuilder::RunInternal(UWorld* InWorld, const FCellInfo& InCellInfo, FPackageSourceControlHelper& PackageHelper)
+{			
+	UWorldPartition* WorldPartition = InWorld->GetWorldPartition();
+	if (!WorldPartition)
+	{
+		UE_LOG(LogWorldPartitionLandscapeSplineMeshBuilder, Error, TEXT("World is not partitioned."));
+		return false;
+	}
+
+	TSet<UPackage*> PackagesToSave;
+
+	auto AddObjectToSaveIfDirty = [&PackagesToSave](UObject* Object)
+	{
+		UPackage* Package = Object->GetPackage();
+		if (Package->IsDirty())
+		{
+			PackagesToSave.Add(Package);
+		}
+	};
+
+	// Update world settings if necessary
+	if (NewGridSize != 0 && (InWorld->GetWorldSettings()->LandscapeSplineMeshesGridSize != NewGridSize))
+	{
+		UE_LOG(LogWorldPartitionLandscapeSplineMeshBuilder, Display, TEXT("Changing Landscape Spline Meshes Grid Size from %d to %d"), InWorld->GetWorldSettings()->LandscapeSplineMeshesGridSize, NewGridSize);
+		InWorld->GetWorldSettings()->Modify();
+		InWorld->GetWorldSettings()->LandscapeSplineMeshesGridSize = NewGridSize;
+		AddObjectToSaveIfDirty(InWorld->GetWorldSettings());
+	}
+		
+	FWorldPartitionHelpers::FForEachActorWithLoadingParams ForEachActorWithLoadingParams;
+	ForEachActorWithLoadingParams.ActorClasses = { ALandscapeSplineMeshesActor::StaticClass(), ALandscapeSplineActor::StaticClass() };
+
+	TArray<FWorldPartitionReference> ActorReferences;
+	TSet<ALandscapeSplineMeshesActor*> PreviousGeneratedActors;
+
+	// Preload all ALandscapeSplineMeshesActor and ALandscapeSplineActor
+	FWorldPartitionHelpers::ForEachActorWithLoading(WorldPartition, [&ActorReferences, &PreviousGeneratedActors, WorldPartition](const FWorldPartitionActorDesc* ActorDesc)
+	{
+		ActorReferences.Emplace(WorldPartition, ActorDesc->GetGuid());
+		AActor* Actor = ActorDesc->GetActor();
+		if (IsValid(Actor))
+		{
+			if (ALandscapeSplineMeshesActor* LandscapeSplineMeshesActor = Cast<ALandscapeSplineMeshesActor>(Actor))
+			{
+				if (LandscapeSplineMeshesActor->Tags.Contains(CreatedFromBuilderTag))
+				{
+					PreviousGeneratedActors.Add(LandscapeSplineMeshesActor);
+				}
+			}
+		}
+		return true;
+	}, ForEachActorWithLoadingParams);
+
+	TMap<ALandscapeSplineMeshesActor*, TArray<UStaticMeshComponent*>> GeneratedActorComponents;
+	TSet<ALandscapeSplineMeshesActor*> ActorsToValidate;
+	TSet<ALandscapeSplineMeshesActor*> ActorsToProcess;
+	for (TActorIterator<ALandscape> It(InWorld); It; ++It)
+	{
+		ALandscape* Landscape = *It;
+		ULandscapeInfo* LandscapeInfo = IsValid(Landscape) ? Landscape->GetLandscapeInfo() : nullptr;
+		if (!LandscapeInfo)
+		{
+			continue;
+		}
+
+		bool bLandscapeHasSplineActors = false;
+		for (TScriptInterface<ILandscapeSplineInterface> SplineOwner : LandscapeInfo->GetSplineActors())
+		{
+			ALandscapeSplineActor* SplineActor = Cast<ALandscapeSplineActor>(SplineOwner.GetObject());
+			if (!SplineActor)
+			{
+				continue;
+			}
+
+			TArray<UStaticMeshComponent*> Components;
+			SplineActor->GetComponents(Components);
+
+			// Remove invalid components from the list
+			FilterStaticMeshComponents(Components);
+			if (Components.IsEmpty())
+			{
+				continue;
+			}
+
+			// For components sources (USplineMeshComponent & UControlPointMeshComponent) to properly return IsEditorOnly, 
+			// we need to ensure that old LandscapeSplineActors have their LandscapeActor pointer properly setup
+			SplineActor->SetLandscapeActor(Landscape);
+			AddObjectToSaveIfDirty(SplineActor);
+			bLandscapeHasSplineActors = true;
+
+			// Build a list of new actors and existing actors
+			for (UStaticMeshComponent* MeshComponent : Components)
+			{
+				ALandscapeSplineMeshesActor* Actor = GetOrCreatePartitionActorForComponent(InWorld, MeshComponent, LandscapeInfo->LandscapeGuid);
+				if (ensure(Actor))
+				{
+					GeneratedActorComponents.FindOrAdd(Actor).Add(MeshComponent);
+					if (PreviousGeneratedActors.Remove(Actor))
+					{
+						ActorsToValidate.Add(Actor);
+					}
+					else if (!ActorsToValidate.Contains(Actor))
+					{
+						ActorsToProcess.Add(Actor);
+					}
+				}
+			}
+		}
+
+		if (bLandscapeHasSplineActors)
+		{
+			Landscape->SetUseGeneratedLandscapeSplineMeshesActors(true);
+			AddObjectToSaveIfDirty(Landscape);
+		}
+	}
+
+	// Find existing actors that needs to be updated
+	for (ALandscapeSplineMeshesActor* ExistingActor : ActorsToValidate)
+	{
+		TArray<UStaticMeshComponent*> NewComponents = GeneratedActorComponents.FindChecked(ExistingActor);
+
+		// Compute hash for each component that would be part of this actor
+		TSet<uint32> NewComponentHashes;
+		for (UStaticMeshComponent* Component : NewComponents)
+		{
+			NewComponentHashes.Add(HashStaticMeshComponent(Component));
+		}
+
+		const int32 NewComponentCount = NewComponents.Num();
+		if (ALandscapeSplineMeshesActor* LandscapeSplineMeshesActor = Cast<ALandscapeSplineMeshesActor>(ExistingActor))
+		{
+			bool bDiffers = (LandscapeSplineMeshesActor->GetStaticMeshComponents().Num() != NewComponentCount);
+			if (!bDiffers)
+			{
+				for (const UStaticMeshComponent* Component : LandscapeSplineMeshesActor->GetStaticMeshComponents())
+				{
+					uint32 HashValue = HashStaticMeshComponent(Component);
+					if (NewComponentHashes.Remove(HashValue) != 1)
+					{
+						bDiffers = true;
+						break;
+					}
+				}
+			}
+			if (bDiffers)
+			{
+				LandscapeSplineMeshesActor->ClearStaticMeshComponents();
+				ActorsToProcess.Add(ExistingActor);
+			}
+		}
+	}
+
+	static const FName FolderPath(TEXT("Spline Mesh Actors"));
+
+	int32 Progress = 0;
+	UE_LOG(LogWorldPartitionLandscapeSplineMeshBuilder, Display, TEXT("Processing Landscape Splines for world %s"), *InWorld->GetName());
+	for (ALandscapeSplineMeshesActor* Actor : ActorsToProcess)
+	{
+		UE_LOG(LogWorldPartitionLandscapeSplineMeshBuilder, Display, TEXT("Processing Splines for %s [%d/%d]"), *Actor->GetName(), ++Progress, ActorsToProcess.Num());
+
+		// Add tag to track actors created from this tool
+		Actor->Tags.AddUnique(CreatedFromBuilderTag);
+		// Regroup in one folder
+		Actor->SetFolderPath(FolderPath);
+
+		TArray<UStaticMeshComponent*> SourceComponents = GeneratedActorComponents.FindChecked(Actor);
+		for (UStaticMeshComponent* Component : SourceComponents)
+		{
+			CloneStaticMeshComponentInActor(Actor, Component);
+		}
+
+		check(Actor->GetPackage()->IsDirty());
+		AddObjectToSaveIfDirty(Actor);
+	}
+
+	if (UActorFolder* ActorFolder = InWorld->PersistentLevel->GetActorFolder(FolderPath))
+	{
+		AddObjectToSaveIfDirty(ActorFolder);
+	}
+
+	// Destroy actors that are no longer needed
+	TArray<FString> PackagesToDelete;
+	for (ALandscapeSplineMeshesActor* Actor : PreviousGeneratedActors)
+	{
+		if (UPackage* Package = Actor->GetExternalPackage())
+		{
+			PackagesToDelete.Add(Package->GetName());
+		}
+		InWorld->DestroyActor(Actor);
+	}
+
+	if (IsRunningCommandlet())
+	{
+		if (!UWorldPartitionBuilder::SavePackages(PackagesToSave.Array(), PackageHelper))
+		{
+			return false;
+		}
+
+		if (!UWorldPartitionBuilder::DeletePackages(PackagesToDelete, PackageHelper))
+		{
+			return false;
+		}
+	}
+
+	return true;
+}
+
+// Version of builder
+bool UWorldPartitionLandscapeSplineMeshesBuilder::RunOnInitializedWorld(UWorld* InWorld)
+{
+	if (!InWorld || !InWorld->bIsWorldInitialized)
+	{
+		UE_LOG(LogWorldPartitionLandscapeSplineMeshBuilder, Error, TEXT("World must be initialized."));
+		return false;
+	}
+
+	UWorldPartitionLandscapeSplineMeshesBuilder* Builder = NewObject<UWorldPartitionLandscapeSplineMeshesBuilder>(GetTransientPackage(), UWorldPartitionLandscapeSplineMeshesBuilder::StaticClass());
+	if (!Builder)
+	{
+		UE_LOG(LogWorldPartitionLandscapeSplineMeshBuilder, Error, TEXT("Failed to create WorldPartitionLandscapeSplineMeshesBuilder."));
+		return false;
+	}
+	
+	FPackageSourceControlHelper SCCHelper;
+	bool bResult = false;
+
+	{
+		FGCObjectScopeGuard BuilderGuard(Builder);
+		bResult = Builder->Run(InWorld, SCCHelper);
+	}
+
+	return bResult;
+}
\ No newline at end of file
diff --git a/Engine/Source/Editor/UnrealEd/Private/WorldPartition/WorldPartitionResaveActorsBuilder.cpp b/Engine/Source/Editor/UnrealEd/Private/WorldPartition/WorldPartitionResaveActorsBuilder.cpp
--- a/Engine/Source/Editor/UnrealEd/Private/WorldPartition/WorldPartitionResaveActorsBuilder.cpp
+++ b/Engine/Source/Editor/UnrealEd/Private/WorldPartition/WorldPartitionResaveActorsBuilder.cpp
@@ -280,7 +280,7 @@ bool UWorldPartitionResaveActorsBuilder::RunInternal(UWorld* World, const FCellI
 
 		FWorldPartitionHelpers::FForEachActorWithLoadingParams ForEachActorWithLoadingParams;
 
-		ForEachActorWithLoadingParams.ActorClass = ActorClass;
+		ForEachActorWithLoadingParams.ActorClasses = { ActorClass };
 
 		ForEachActorWithLoadingParams.FilterActorDesc = [this](const FWorldPartitionActorDesc* ActorDesc) -> bool
 		{
diff --git a/Engine/Source/Editor/UnrealEd/Public/Dialogs/SBuildProgress.h b/Engine/Source/Editor/UnrealEd/Public/Dialogs/SBuildProgress.h
--- a/Engine/Source/Editor/UnrealEd/Public/Dialogs/SBuildProgress.h
+++ b/Engine/Source/Editor/UnrealEd/Public/Dialogs/SBuildProgress.h
@@ -44,6 +44,8 @@ public:
 		BUILDTYPE_HLODs,
 		/** Minimap is being built */
 		BUILDTYPE_Minimap,
+		/** Landscape Spline Meshes are being built */
+		BUILDTYPE_LandscapeSplineMeshes,
 		/** Texture streaming data is being built */
 		BUILDTYPE_TextureStreaming, 
 		/** Virtual texture data is being built */
diff --git a/Engine/Source/Editor/UnrealEd/Public/EditorBuildUtils.h b/Engine/Source/Editor/UnrealEd/Public/EditorBuildUtils.h
--- a/Engine/Source/Editor/UnrealEd/Public/EditorBuildUtils.h
+++ b/Engine/Source/Editor/UnrealEd/Public/EditorBuildUtils.h
@@ -34,6 +34,8 @@ struct FBuildOptions
 	UNREALED_API static const FName BuildHierarchicalLOD;
 	/** Build minimap */
 	UNREALED_API static const FName BuildMinimap;
+	/** Build landscape spline meshes */
+	UNREALED_API static const FName BuildLandscapeSplineMeshes;
 	/** Build texture streaming */
 	UNREALED_API static const FName BuildTextureStreaming;
 	/** Build virtual textures */
@@ -279,6 +281,13 @@ private:
 	 */
 	static void TriggerMinimapBuilder(UWorld* InWorld, FName Id);
 
+	/**
+	 * Trigger LandscapeSplineMeshes builder to (re)generate landscape spline meshes actors
+	 *
+	 * @param	InWorld			WorldContext
+	 */
+	static void TriggerLandscapeSplineMeshesBuilder(UWorld* InWorld);
+
 	/** Intentionally hide constructors, etc. to prevent instantiation */
 	FEditorBuildUtils();
 	~FEditorBuildUtils();
diff --git a/Engine/Source/Editor/UnrealEd/Public/WorldPartition/WorldPartitionLandscapeSplineMeshesBuilder.h b/Engine/Source/Editor/UnrealEd/Public/WorldPartition/WorldPartitionLandscapeSplineMeshesBuilder.h
new file mode 100644
--- /dev/null
+++ b/Engine/Source/Editor/UnrealEd/Public/WorldPartition/WorldPartitionLandscapeSplineMeshesBuilder.h
@@ -0,0 +1,41 @@
+// Copyright Epic Games, Inc. All Rights Reserved.
+
+#pragma once
+
+#include "WorldPartition/WorldPartitionBuilder.h"
+#include "WorldPartitionLandscapeSplineMeshesBuilder.generated.h"
+
+class UWorld;
+class ALandscapeSplineMeshesActor;
+class UStaticMeshComponent;
+class UPackage;
+class ILandscapeSplineInterface;
+
+// Example Command Line: ProjectName MapName -run=WorldPartitionBuilderCommandlet -SCCProvider=Perforce -Builder=WorldPartitionLandscapeSplineMeshesBuilder (<Optional> -NewGridSize=Value)
+UCLASS()
+class UNREALED_API UWorldPartitionLandscapeSplineMeshesBuilder : public UWorldPartitionBuilder
+{
+	GENERATED_UCLASS_BODY()
+
+public:
+	// UWorldPartitionBuilder interface begin
+	virtual bool RequiresCommandletRendering() const override { return false; }
+	virtual ELoadingMode GetLoadingMode() const override { return ELoadingMode::Custom; }
+	virtual bool RunInternal(UWorld* World, const FCellInfo& InCellInfo, FPackageSourceControlHelper& PackageHelper) override;
+	// UWorldPartitionBuilder interface end
+
+	// Helper function that creates a builder and runs it on an already loaded/initialized world
+	static bool RunOnInitializedWorld(UWorld* World);
+
+private:
+	void FilterStaticMeshComponents(TArray<class UStaticMeshComponent*>& InOutComponents);
+	void CloneStaticMeshComponentInActor(ALandscapeSplineMeshesActor* InActor, const UStaticMeshComponent* InMeshComponent);
+	void CloneStaticMeshComponent(const UStaticMeshComponent* InSrcMeshComponent, UStaticMeshComponent* DstMeshComponent);
+	int32 HashStaticMeshComponent(const UStaticMeshComponent* InComponent);
+	FName GetSplineCollisionProfileName(const UStaticMeshComponent* InMeshComponent);
+	ALandscapeSplineMeshesActor* GetOrCreatePartitionActorForComponent(class UWorld* InWorld, const UStaticMeshComponent* InMeshComponent, const FGuid& InLandscapeGuid);
+
+	uint32 NewGridSize;
+	
+	static class UStaticMesh* SplineEditorMesh;
+};
\ No newline at end of file
diff --git a/Engine/Source/Editor/WorldPartitionEditor/Private/WorldPartitionEditorModule.cpp b/Engine/Source/Editor/WorldPartitionEditor/Private/WorldPartitionEditorModule.cpp
--- a/Engine/Source/Editor/WorldPartitionEditor/Private/WorldPartitionEditorModule.cpp
+++ b/Engine/Source/Editor/WorldPartitionEditor/Private/WorldPartitionEditorModule.cpp
@@ -4,6 +4,7 @@
 #include "WorldPartition/WorldPartition.h"
 #include "WorldPartition/WorldPartitionHLODsBuilder.h"
 #include "WorldPartition/WorldPartitionMiniMapBuilder.h"
+#include "WorldPartition/WorldPartitionLandscapeSplineMeshesBuilder.h"
 #include "WorldPartition/WorldPartitionActorLoaderInterface.h"
 #include "WorldPartition/HLOD/HLODLayerAssetTypeActions.h"
 #include "WorldPartition/SWorldPartitionEditor.h"
@@ -451,18 +452,23 @@ bool FWorldPartitionEditorModule::ConvertMap(const FString& InLongPackageName)
 	return false;
 }
 
-bool FWorldPartitionEditorModule::RunBuilder(TSubclassOf<UWorldPartitionBuilder> InWorldPartitionBuilder, const FString& InLongPackageName)
+bool FWorldPartitionEditorModule::RunBuilder(TSubclassOf<UWorldPartitionBuilder> InWorldPartitionBuilder, UWorld* InWorld)
 {
 	// Ideally this should be improved to automatically register all builders & present their options in a consistent way...
 
 	if (InWorldPartitionBuilder == UWorldPartitionHLODsBuilder::StaticClass())
 	{
-		return BuildHLODs(InLongPackageName);
+		return BuildHLODs(InWorld->GetPackage()->GetName());
 	}
 	
 	if (InWorldPartitionBuilder == UWorldPartitionMiniMapBuilder::StaticClass())
 	{
-		return BuildMinimap(InLongPackageName);
+		return BuildMinimap(InWorld->GetPackage()->GetName());
+	}
+
+	if (InWorldPartitionBuilder == UWorldPartitionLandscapeSplineMeshesBuilder::StaticClass())
+	{
+		return BuildLandscapeSplineMeshes(InWorld);
 	}
 
 	return false;
@@ -558,6 +564,16 @@ bool FWorldPartitionEditorModule::BuildMinimap(const FString& InMapToProcess)
 	return bSuccess;
 }
 
+bool FWorldPartitionEditorModule::BuildLandscapeSplineMeshes(UWorld* InWorld)
+{
+	if (!UWorldPartitionLandscapeSplineMeshesBuilder::RunOnInitializedWorld(InWorld))
+	{
+		FMessageDialog::Open(EAppMsgType::Ok, LOCTEXT("LandscapeSplineMeshesBuildFailed", "Landscape Spline Meshes build failed! See log for details."));
+		return false;
+	}
+	return true;
+}
+
 void FWorldPartitionEditorModule::OnMapChanged(uint32 MapFlags)
 {
 	if (MapFlags == MapChangeEventFlags::NewMap)
diff --git a/Engine/Source/Editor/WorldPartitionEditor/Public/WorldPartitionEditorModule.h b/Engine/Source/Editor/WorldPartitionEditor/Public/WorldPartitionEditorModule.h
--- a/Engine/Source/Editor/WorldPartitionEditor/Public/WorldPartitionEditorModule.h
+++ b/Engine/Source/Editor/WorldPartitionEditor/Public/WorldPartitionEditorModule.h
@@ -54,7 +54,7 @@ public:
 	/**
 	 * Run a world partition builder for the current map. Will display a dialog to specify options for the generation.
 	 */
-	virtual bool RunBuilder(TSubclassOf<UWorldPartitionBuilder> WorldPartitionBuilder, const FString& InLongPackageName) override;
+	virtual bool RunBuilder(TSubclassOf<UWorldPartitionBuilder> WorldPartitionBuilder, UWorld* InWorld) override;
 
 	/** Return the world added event. */
 	virtual FWorldPartitionCreated& OnWorldPartitionCreated() override { return WorldPartitionCreatedEvent; }
@@ -80,6 +80,7 @@ private:
 
 	bool BuildHLODs(const FString& InMapToProcess);
 	bool BuildMinimap(const FString& InMapToProcess);
+	bool BuildLandscapeSplineMeshes(UWorld* InWorld);
 
 private:
 	void RunCommandletAsExternalProcess(const FString& InCommandletArgs, const FText& InOperationDescription, int32& OutResult, bool& bOutCancelled, FString& OutCommandletOutput);
diff --git a/Engine/Source/Runtime/Engine/Classes/Components/SplineMeshComponent.h b/Engine/Source/Runtime/Engine/Classes/Components/SplineMeshComponent.h
--- a/Engine/Source/Runtime/Engine/Classes/Components/SplineMeshComponent.h
+++ b/Engine/Source/Runtime/Engine/Classes/Components/SplineMeshComponent.h
@@ -156,6 +156,7 @@ class ENGINE_API USplineMeshComponent : public UStaticMeshComponent, public IInt
 	//Begin UObject Interface
 	virtual void Serialize(FArchive& Ar) override;
 #if WITH_EDITOR
+	virtual bool IsEditorOnly() const override;
 	virtual bool Modify(bool bAlwaysMarkDirty = true) override;
 	virtual void PostEditChangeProperty(FPropertyChangedEvent& PropertyChangedEvent) override;
 #endif
diff --git a/Engine/Source/Runtime/Engine/Classes/GameFramework/WorldSettings.h b/Engine/Source/Runtime/Engine/Classes/GameFramework/WorldSettings.h
--- a/Engine/Source/Runtime/Engine/Classes/GameFramework/WorldSettings.h
+++ b/Engine/Source/Runtime/Engine/Classes/GameFramework/WorldSettings.h
@@ -601,6 +601,9 @@ public:
 	UPROPERTY(EditAnywhere, Category = Foliage)
 	uint32 InstancedFoliageGridSize;
 
+	UPROPERTY(EditAnywhere, Category = Landscape)
+	uint32 LandscapeSplineMeshesGridSize;
+
 	/** Size of the grid for navigation data chunk actors */
 	UPROPERTY(EditAnywhere, Category = Navigation)
 	uint32 NavigationDataChunkGridSize;
diff --git a/Engine/Source/Runtime/Engine/Private/Components/SplineMeshComponent.cpp b/Engine/Source/Runtime/Engine/Private/Components/SplineMeshComponent.cpp
--- a/Engine/Source/Runtime/Engine/Private/Components/SplineMeshComponent.cpp
+++ b/Engine/Source/Runtime/Engine/Private/Components/SplineMeshComponent.cpp
@@ -21,6 +21,7 @@
 #if WITH_EDITOR
 #include "IHierarchicalLODUtilities.h"
 #include "HierarchicalLODUtilitiesModule.h"
+#include "LandscapeSplineActor.h"
 #endif // WITH_EDITOR
 
 int32 GNoRecreateSplineMeshProxy = 1;
@@ -547,6 +548,23 @@ void USplineMeshComponent::Serialize(FArchive& Ar)
 }
 
 #if WITH_EDITOR
+bool USplineMeshComponent::IsEditorOnly() const
+{
+	if (Super::IsEditorOnly())
+	{
+		return true;
+	}
+
+	// If Landscape uses generated LandscapeSplineMeshesActors, SplineMeshComponents is removed from cooked build  
+	ALandscapeSplineActor* SplineActor = Cast<ALandscapeSplineActor>(GetOwner());
+	if (SplineActor && SplineActor->HasGeneratedLandscapeSplineMeshesActors())
+	{
+		return true;
+	}
+
+	return false;
+}
+
 bool USplineMeshComponent::Modify(bool bAlwaysMarkDirty)
 {
 	bool bSavedToTransactionBuffer = Super::Modify(bAlwaysMarkDirty);
diff --git a/Engine/Source/Runtime/Engine/Private/WorldPartition/WorldPartitionHelpers.cpp b/Engine/Source/Runtime/Engine/Private/WorldPartition/WorldPartitionHelpers.cpp
--- a/Engine/Source/Runtime/Engine/Private/WorldPartition/WorldPartitionHelpers.cpp
+++ b/Engine/Source/Runtime/Engine/Private/WorldPartition/WorldPartitionHelpers.cpp
@@ -4,6 +4,7 @@
 
 #include "WorldPartition/WorldPartition.h"
 #include "WorldPartition/WorldPartitionEditorHash.h"
+#include "Algo/AnyOf.h"
 
 #include "Commandlets/Commandlet.h"
 
@@ -92,13 +93,13 @@ namespace WorldPartitionHelpers
 FWorldPartitionHelpers::FForEachActorWithLoadingParams::FForEachActorWithLoadingParams()
 	: bGCPerActor(false)
 	, bKeepReferences(false)
-	, ActorClass(AActor::StaticClass())
+	, ActorClasses({ AActor::StaticClass() })
 {}
 
 void FWorldPartitionHelpers::ForEachActorWithLoading(UWorldPartition* WorldPartition, TSubclassOf<AActor> ActorClass, TFunctionRef<bool(const FWorldPartitionActorDesc*)> Func, TFunctionRef<void()> OnReleasingActorReferences, bool bGCPerActor)
 {
 	FForEachActorWithLoadingParams Params;
-	Params.ActorClass = ActorClass;
+	Params.ActorClasses = { ActorClass };
 	Params.OnPreGarbageCollect = [&OnReleasingActorReferences]() { OnReleasingActorReferences(); };
 	ForEachActorWithLoading(WorldPartition, Func, Params);
 }
@@ -134,13 +135,14 @@ void FWorldPartitionHelpers::ForEachActorWithLoading(UWorldPartition* WorldParti
 
 	for (FActorDescContainerCollection::TConstIterator<> ActorDescIterator(WorldPartition); ActorDescIterator; ++ActorDescIterator)
 	{
-		if (IsActorDescClassCompatibleWith(*ActorDescIterator, Params.ActorClass))
+		const FWorldPartitionActorDesc* ActorDesc = *ActorDescIterator;
+		if (Algo::AnyOf(Params.ActorClasses, [ActorDesc](UClass* ActorClass) { return IsActorDescClassCompatibleWith(ActorDesc, ActorClass); }))
 		{
-			if (!Params.FilterActorDesc || Params.FilterActorDesc(*ActorDescIterator))
+			if (!Params.FilterActorDesc || Params.FilterActorDesc(ActorDesc))
 			{
-				WorldPartitionHelpers::LoadReferences(WorldPartition, ActorDescIterator->GetGuid(), Result.ActorReferences);
+				WorldPartitionHelpers::LoadReferences(WorldPartition, ActorDesc->GetGuid(), Result.ActorReferences);
 
-				FWorldPartitionReference ActorReference(WorldPartition, ActorDescIterator->GetGuid());
+				FWorldPartitionReference ActorReference(WorldPartition, ActorDesc->GetGuid());
 				if (!Func(ActorReference.Get()))
 				{
 					break;
diff --git a/Engine/Source/Runtime/Engine/Private/WorldSettings.cpp b/Engine/Source/Runtime/Engine/Private/WorldSettings.cpp
--- a/Engine/Source/Runtime/Engine/Private/WorldSettings.cpp
+++ b/Engine/Source/Runtime/Engine/Private/WorldSettings.cpp
@@ -93,7 +93,7 @@ PRAGMA_ENABLE_DEPRECATION_WARNINGS
 	KillZDamageType = ConstructorStatics.DmgType_Environmental_Object.Object;
 
 #if WITH_EDITORONLY_DATA
-	InstancedFoliageGridSize = DefaultPlacementGridSize = 25600;
+	InstancedFoliageGridSize = DefaultPlacementGridSize = LandscapeSplineMeshesGridSize = 25600;
 	NavigationDataChunkGridSize = 102400;
 	NavigationDataBuilderLoadingCellSize = 102400 * 4;
 	bIncludeGridSizeInNameForFoliageActors = false;
@@ -681,7 +681,8 @@ bool AWorldSettings::CanEditChange(const FProperty* InProperty) const
 		}
 		else if (PropertyName == GET_MEMBER_NAME_STRING_CHECKED(AWorldSettings, bEnableWorldComposition ) ||
 				 PropertyName == GET_MEMBER_NAME_STRING_CHECKED(AWorldSettings, bForceNoPrecomputedLighting ) ||
-				 PropertyName == GET_MEMBER_NAME_STRING_CHECKED(AWorldSettings, bPrecomputeVisibility))
+				 PropertyName == GET_MEMBER_NAME_STRING_CHECKED(AWorldSettings, bPrecomputeVisibility) ||
+				 PropertyName == GET_MEMBER_NAME_STRING_CHECKED(AWorldSettings, LandscapeSplineMeshesGridSize))
 		{
 			return !IsPartitionedWorld();
 		}
diff --git a/Engine/Source/Runtime/Engine/Public/WorldPartition/IWorldPartitionEditorModule.h b/Engine/Source/Runtime/Engine/Public/WorldPartition/IWorldPartitionEditorModule.h
--- a/Engine/Source/Runtime/Engine/Public/WorldPartition/IWorldPartitionEditorModule.h
+++ b/Engine/Source/Runtime/Engine/Public/WorldPartition/IWorldPartitionEditorModule.h
@@ -17,7 +17,10 @@ public:
 
 	virtual bool ConvertMap(const FString& InLongPackageName) = 0;
 
-	virtual bool RunBuilder(TSubclassOf<UWorldPartitionBuilder> BuilderClass, const FString& InLongPackageName) = 0;
+	UE_DEPRECATED(5.1, "Use RunBuilder with UWorld* instead.")
+	virtual bool RunBuilder(TSubclassOf<UWorldPartitionBuilder> BuilderClass, const FString& InLongPackageName) { return false; }
+
+	virtual bool RunBuilder(TSubclassOf<UWorldPartitionBuilder> BuilderClass, UWorld* InWorld) = 0;
 
 	virtual int32 GetPlacementGridSize() const = 0;
 	virtual int32 GetInstancedFoliageGridSize() const = 0;
diff --git a/Engine/Source/Runtime/Engine/Public/WorldPartition/WorldPartitionHelpers.h b/Engine/Source/Runtime/Engine/Public/WorldPartition/WorldPartitionHelpers.h
--- a/Engine/Source/Runtime/Engine/Public/WorldPartition/WorldPartitionHelpers.h
+++ b/Engine/Source/Runtime/Engine/Public/WorldPartition/WorldPartitionHelpers.h
@@ -42,8 +42,8 @@ public:
 		/* Prevent clearing of actor references. */
 		bool bKeepReferences;
 
-		/* The class used to filter actors loading. */
-		TSubclassOf<AActor> ActorClass;
+		/* The classes used to filter actors loading. */
+		TArray<TSubclassOf<AActor>> ActorClasses;
 
 		/* Custom filter function used to filter actors loading. */
 		TUniqueFunction<bool(const FWorldPartitionActorDesc*)> FilterActorDesc;
diff --git a/Engine/Source/Runtime/Landscape/Classes/ControlPointMeshComponent.h b/Engine/Source/Runtime/Landscape/Classes/ControlPointMeshComponent.h
--- a/Engine/Source/Runtime/Landscape/Classes/ControlPointMeshComponent.h
+++ b/Engine/Source/Runtime/Landscape/Classes/ControlPointMeshComponent.h
@@ -30,6 +30,7 @@ class UControlPointMeshComponent : public UStaticMeshComponent
 	{
 		return Super::ShouldRenderSelected() || bSelected;
 	}
+	virtual bool IsEditorOnly() const override;
 #endif
 	virtual float GetVirtualTextureMainPassMaxDrawDistance() const override { return VirtualTextureMainPassMaxDrawDistance; }
 	//End UPrimitiveComponent Interface
diff --git a/Engine/Source/Runtime/Landscape/Classes/Landscape.h b/Engine/Source/Runtime/Landscape/Classes/Landscape.h
--- a/Engine/Source/Runtime/Landscape/Classes/Landscape.h
+++ b/Engine/Source/Runtime/Landscape/Classes/Landscape.h
@@ -353,6 +353,9 @@ public:
 	void UpdateProxyLayersWeightmapUsage();
 	void ValidateProxyLayersWeightmapUsage() const;
 
+	LANDSCAPE_API void SetUseGeneratedLandscapeSplineMeshesActors(bool bInEnabled);
+	LANDSCAPE_API bool GetUseGeneratedLandscapeSplineMeshesActors() const;
+
 protected:
 	FName GenerateUniqueLayerName(FName InName = NAME_None) const;
 
@@ -459,6 +462,13 @@ public:
 	UPROPERTY(EditAnywhere, Category=Landscape)
 	bool bCanHaveLayersContent = false;
 
+	/*
+	 * If true, WorldPartitionLandscapeSplineMeshesBuilder is responsible of generating partitioned actors of type ALandscapeSplineMeshesActor that will contain all landscape spline/controlpoints static meshes. 
+	 * Source components will be editor only and hidden in game for PIE.
+	 */
+	UPROPERTY()
+	bool bUseGeneratedLandscapeSplineMeshesActors = false;
+
 	DECLARE_EVENT(ALandscape, FLandscapeBlueprintBrushChangedDelegate);
 	FLandscapeBlueprintBrushChangedDelegate& OnBlueprintBrushChangedDelegate() { return LandscapeBlueprintBrushChangedDelegate; }
 
diff --git a/Engine/Source/Runtime/Landscape/Classes/LandscapeSplineActor.h b/Engine/Source/Runtime/Landscape/Classes/LandscapeSplineActor.h
--- a/Engine/Source/Runtime/Landscape/Classes/LandscapeSplineActor.h
+++ b/Engine/Source/Runtime/Landscape/Classes/LandscapeSplineActor.h
@@ -8,6 +8,8 @@
 #include "ILandscapeSplineInterface.h"
 #include "LandscapeSplineActor.generated.h"
 
+class ALandscape;
+
 UCLASS(MinimalAPI, NotBlueprintable, NotPlaceable, hidecategories = (Display, Attachment, Physics, Debug, Lighting, Input))
 class ALandscapeSplineActor : public AActor, public ILandscapeSplineInterface
 {
@@ -17,15 +19,22 @@ protected:
 	/** Guid for LandscapeInfo **/
 	UPROPERTY()
 	FGuid LandscapeGuid;
+
+#if WITH_EDITORONLY_DATA
+	/** Landscape **/
+	UPROPERTY()
+	TObjectPtr<ALandscape> LandscapeActor;
+#endif
 	
 public:
 	virtual FGuid GetLandscapeGuid() const override { return LandscapeGuid; }
-	void SetLandscapeGuid(const FGuid& Guid) { LandscapeGuid = Guid; }
 	virtual ULandscapeSplinesComponent* GetSplinesComponent() const override;
 	virtual FTransform LandscapeActorToWorld() const override;
 	virtual ULandscapeInfo* GetLandscapeInfo() const override;
 		
 #if WITH_EDITOR
+	LANDSCAPE_API bool HasGeneratedLandscapeSplineMeshesActors() const;
+	void SetLandscapeGuid(const FGuid& InGuid);
 	virtual bool CanChangeIsSpatiallyLoadedFlag() const override { return true; }
 	virtual void GetActorDescProperties(FPropertyPairsMap& PropertyPairsMap) const override;
 
@@ -42,5 +51,9 @@ public:
 	// Interface existes for backward compatibility. Should already be created since its this actor's root component.
 	virtual void CreateSplineComponent() override { check(false); }
 	virtual void CreateSplineComponent(const FVector& Scale3D) override { check(false); }
+
+private:
+	LANDSCAPE_API void SetLandscapeActor(ALandscape* InLandscapeActor);
+	friend class UWorldPartitionLandscapeSplineMeshesBuilder;
 #endif
 };
diff --git a/Engine/Source/Runtime/Landscape/Classes/LandscapeSplineMeshesActor.h b/Engine/Source/Runtime/Landscape/Classes/LandscapeSplineMeshesActor.h
new file mode 100644
--- /dev/null
+++ b/Engine/Source/Runtime/Landscape/Classes/LandscapeSplineMeshesActor.h
@@ -0,0 +1,53 @@
+// Copyright Epic Games, Inc. All Rights Reserved.
+
+#pragma once
+
+#include "ActorPartition/PartitionActor.h"
+#include "LandscapeSplineMeshesActor.generated.h"
+
+//
+// This class is only intended to be used by UWorldPartitionLandscapeSplineMeshesBuilder which extracts and clones
+// landscape spline and control point static meshes into partitioned actors
+//
+// This serves as an optimization as these actors will be streamed at runtime
+//
+UCLASS(NotPlaceable, NotBlueprintable, hideCategories = (Input), showCategories = ("Input|MouseInput", "Input|TouchInput"), ConversionRoot, ComponentWrapperClass, meta = (ChildCanTick))
+class LANDSCAPE_API ALandscapeSplineMeshesActor : public APartitionActor
+{
+	GENERATED_UCLASS_BODY()
+
+private:
+	UPROPERTY(Category = PartitionSplineMeshActor, VisibleAnywhere, BlueprintReadOnly, meta = (ExposeFunctionCategories = "Mesh,Rendering,Physics,Components|StaticMesh,Components|SplineMesh", AllowPrivateAccess = "true"))
+	TArray<TObjectPtr<class UStaticMeshComponent>> StaticMeshComponents;
+
+public:
+
+#if WITH_EDITOR
+	class UStaticMeshComponent* CreateStaticMeshComponent(const TSubclassOf<class UStaticMeshComponent>& InComponentClass);
+	void ClearStaticMeshComponents();
+
+	//~ Begin AActor Interface
+	virtual void CheckForErrors() override;
+	virtual bool GetReferencedContentObjects(TArray<UObject*>& Objects) const override;
+	//~ End AActor Interface
+
+	//~ Begin APartitionActor Interface
+	virtual uint32 GetDefaultGridSize(UWorld* InWorld) const override;
+	virtual FGuid GetGridGuid() const override { return GridGuid; }
+	//~ End APartitionActor Interface
+
+	void SetGridGuid(const FGuid& InGuid) { GridGuid = InGuid; }
+#endif // WITH_EDITOR
+
+	const TArray<class UStaticMeshComponent*>& GetStaticMeshComponents() const;
+
+protected:
+	//~ Begin UObject Interface.
+	virtual FString GetDetailedInfoInternal() const override;
+	//~ End UObject Interface.
+
+private:
+#if WITH_EDITORONLY_DATA
+	FGuid GridGuid;
+#endif
+};
\ No newline at end of file
diff --git a/Engine/Source/Runtime/Landscape/Private/Landscape.cpp b/Engine/Source/Runtime/Landscape/Private/Landscape.cpp
--- a/Engine/Source/Runtime/Landscape/Private/Landscape.cpp
+++ b/Engine/Source/Runtime/Landscape/Private/Landscape.cpp
@@ -2755,6 +2755,20 @@ bool ALandscape::IsNaniteEnabled() const
 	return bEnableNanite;
 }
 
+void ALandscape::SetUseGeneratedLandscapeSplineMeshesActors(bool bInEnabled)
+{
+	if (bUseGeneratedLandscapeSplineMeshesActors != bInEnabled)
+	{
+		Modify();
+		bUseGeneratedLandscapeSplineMeshesActors = bInEnabled;
+	}
+}
+
+bool ALandscape::GetUseGeneratedLandscapeSplineMeshesActors() const
+{
+	return bUseGeneratedLandscapeSplineMeshesActors;
+}
+
 void ALandscapeProxy::OnFeatureLevelChanged(ERHIFeatureLevel::Type NewFeatureLevel)
 {
 	FlushGrassComponents();
diff --git a/Engine/Source/Runtime/Landscape/Private/LandscapeSplineActor.cpp b/Engine/Source/Runtime/Landscape/Private/LandscapeSplineActor.cpp
--- a/Engine/Source/Runtime/Landscape/Private/LandscapeSplineActor.cpp
+++ b/Engine/Source/Runtime/Landscape/Private/LandscapeSplineActor.cpp
@@ -8,6 +8,9 @@
 
 #if WITH_EDITOR
 #include "WorldPartition/WorldPartitionActorDesc.h"
+#include "Components/StaticMeshComponent.h"
+#include "Components/SplineMeshComponent.h"
+#include "ControlPointMeshComponent.h"
 #endif
 
 ALandscapeSplineActor::ALandscapeSplineActor(const FObjectInitializer& ObjectInitializer)
@@ -49,6 +52,25 @@ void ALandscapeSplineActor::GetSharedProperties(ULandscapeInfo* InLandscapeInfo)
 {
 	Modify();
 	LandscapeGuid = InLandscapeInfo->LandscapeGuid;
+	LandscapeActor = InLandscapeInfo->LandscapeActor.Get();
+}
+
+void ALandscapeSplineActor::SetLandscapeGuid(const FGuid& InGuid)
+{ 
+	LandscapeGuid = InGuid;
+	check(!LandscapeActor || (LandscapeGuid == LandscapeActor->GetLandscapeGuid()));
+}
+
+// This function is only allowed to be called to fixup LandscapeActor pointer
+void ALandscapeSplineActor::SetLandscapeActor(ALandscape* InLandscapeActor)
+{
+	if (LandscapeActor != InLandscapeActor)
+	{
+		Modify();
+		check(LandscapeActor == nullptr);
+		check(LandscapeGuid == InLandscapeActor->GetLandscapeGuid());
+		LandscapeActor = InLandscapeActor;
+	}
 }
 
 void ALandscapeSplineActor::Destroyed()
@@ -69,11 +91,24 @@ void ALandscapeSplineActor::PostRegisterAllComponents()
 	Super::PostRegisterAllComponents();
 	if (!IsPendingKillPending())
 	{
+		UWorld* World = GetWorld();
 		if (LandscapeGuid.IsValid())
 		{
-			ULandscapeInfo* LandscapeInfo = ULandscapeInfo::FindOrCreate(GetWorld(), LandscapeGuid);
+			ULandscapeInfo* LandscapeInfo = ULandscapeInfo::FindOrCreate(World, LandscapeGuid);
 			LandscapeInfo->RegisterSplineActor(this);
 		}
+
+		// If Landscape uses generated LandscapeSplineMeshesActors, ensure SplineMeshComponents & ControlPointMeshComponents are hidden in PIE
+		if (World->IsGameWorld() && HasGeneratedLandscapeSplineMeshesActors())
+		{
+			ForEachComponent<UStaticMeshComponent>(true, [](UStaticMeshComponent* Component)
+			{
+				if (Component->IsA<USplineMeshComponent>() || Component->IsA<UControlPointMeshComponent>())
+				{
+					Component->SetHiddenInGame(true);
+				}
+			});
+		}
 	}
 }
 
@@ -110,4 +145,9 @@ AActor* ALandscapeSplineActor::GetSceneOutlinerParent() const
 	return nullptr;
 }
 
-#endif
\ No newline at end of file
+bool ALandscapeSplineActor::HasGeneratedLandscapeSplineMeshesActors() const
+{
+	return IsValid(LandscapeActor) && LandscapeActor->GetUseGeneratedLandscapeSplineMeshesActors();
+}
+
+#endif
diff --git a/Engine/Source/Runtime/Landscape/Private/LandscapeSplineMeshesActor.cpp b/Engine/Source/Runtime/Landscape/Private/LandscapeSplineMeshesActor.cpp
new file mode 100644
--- /dev/null
+++ b/Engine/Source/Runtime/Landscape/Private/LandscapeSplineMeshesActor.cpp
@@ -0,0 +1,119 @@
+// Copyright Epic Games, Inc. All Rights Reserved.
+
+#include "LandscapeSplineMeshesActor.h"
+#include "Components/StaticMeshComponent.h"
+#include "ControlPointMeshComponent.h"
+#include "Components/SplineMeshComponent.h"
+#include "Engine/StaticMesh.h"
+#include "Misc/UObjectToken.h"
+#include "Misc/MapErrors.h"
+#include "GameFramework/WorldSettings.h"
+
+#define LOCTEXT_NAMESPACE "LandscapeSplineMeshesActor"
+
+// ----------------------------------------------------------------
+// ALandscapeSplineMeshesActor Implementation
+// ----------------------------------------------------------------
+
+ALandscapeSplineMeshesActor::ALandscapeSplineMeshesActor(const FObjectInitializer& ObjectInitializer)
+	: Super(ObjectInitializer)
+{
+	SetCanBeDamaged(false);
+}
+
+const TArray<UStaticMeshComponent*>& ALandscapeSplineMeshesActor::GetStaticMeshComponents() const
+{
+	return StaticMeshComponents;
+}
+
+#if WITH_EDITOR
+uint32 ALandscapeSplineMeshesActor::GetDefaultGridSize(UWorld* InWorld) const
+{
+	return InWorld->GetWorldSettings()->LandscapeSplineMeshesGridSize;
+}
+
+UStaticMeshComponent* ALandscapeSplineMeshesActor::CreateStaticMeshComponent(const TSubclassOf<UStaticMeshComponent>& InComponentClass)
+{
+	check(InComponentClass->IsChildOf<USplineMeshComponent>() || InComponentClass->IsChildOf<UControlPointMeshComponent>());
+	Modify();
+
+	UStaticMeshComponent* Component = NewObject<UStaticMeshComponent>(this, InComponentClass.Get(), NAME_None, RF_Transactional);
+	AddInstanceComponent(Component);
+	Component->SetupAttachment(RootComponent);
+	if (GetRootComponent()->IsRegistered())
+	{
+		Component->RegisterComponent();
+	}
+	Component->SetWorldTransform(RootComponent->GetComponentTransform());
+	Component->SetCollisionProfileName(UCollisionProfile::BlockAll_ProfileName);
+	Component->Mobility = EComponentMobility::Static;
+	Component->SetGenerateOverlapEvents(false);
+	StaticMeshComponents.Add(Component);
+	return Component;
+}
+
+void ALandscapeSplineMeshesActor::ClearStaticMeshComponents()
+{
+	Modify();
+
+	for (UStaticMeshComponent* Component : StaticMeshComponents)
+	{
+		if (Component)
+		{
+			Component->Modify();
+			Component->DestroyComponent();
+		}
+	}
+
+	StaticMeshComponents.Empty();
+}
+
+void ALandscapeSplineMeshesActor::CheckForErrors()
+{
+	Super::CheckForErrors();
+
+	FMessageLog MapCheck("MapCheck");
+
+	for (UStaticMeshComponent* Component : StaticMeshComponents)
+	{
+		if (Component && !Component->GetStaticMesh())
+		{
+			MapCheck.Warning()
+				->AddToken(FUObjectToken::Create(this))
+				->AddToken(FTextToken::Create(LOCTEXT("MapCheck_Message_SplineMeshesActorNull", "LandscapeSplineMeshesActor has a component with NULL StaticMesh")))
+				->AddToken(FMapErrorToken::Create(FMapErrors::StaticMeshNull));
+		}
+	}
+}
+
+bool ALandscapeSplineMeshesActor::GetReferencedContentObjects(TArray<UObject*>& Objects) const
+{
+	Super::GetReferencedContentObjects(Objects);
+
+	for (UStaticMeshComponent* Component : StaticMeshComponents)
+	{
+		if (Component && Component->GetStaticMesh())
+		{
+			Objects.Add(Component->GetStaticMesh());
+		}
+	}
+	return true;
+}
+
+#endif
+
+FString ALandscapeSplineMeshesActor::GetDetailedInfoInternal() const
+{
+	TStringBuilder<1024> Builder;
+	for (UStaticMeshComponent* Component : StaticMeshComponents)
+	{
+		if (Component && Component->GetStaticMesh())
+		{
+			Builder += Component->GetDetailedInfoInternal();
+			Builder += TEXT("\n");
+		}
+	}
+	return Builder.ToString();
+}
+
+#undef LOCTEXT_NAMESPACE
\ No newline at end of file
diff --git a/Engine/Source/Runtime/Landscape/Private/LandscapeSplines.cpp b/Engine/Source/Runtime/Landscape/Private/LandscapeSplines.cpp
--- a/Engine/Source/Runtime/Landscape/Private/LandscapeSplines.cpp
+++ b/Engine/Source/Runtime/Landscape/Private/LandscapeSplines.cpp
@@ -43,6 +43,7 @@
 #include "Landscape.h"
 #include "LandscapeLayerInfoObject.h"
 #include "LandscapeInfoMap.h"
+#include "LandscapeSplineActor.h"
 #endif
 
 IMPLEMENT_HIT_PROXY(HLandscapeSplineProxy, HHitProxy);
@@ -1535,6 +1536,25 @@ UControlPointMeshComponent::UControlPointMeshComponent(const FObjectInitializer&
 #endif
 }
 
+#if WITH_EDITOR
+bool UControlPointMeshComponent::IsEditorOnly() const
+{
+	if (Super::IsEditorOnly())
+	{
+		return true;
+	}
+
+	// If Landscape uses generated LandscapeSplineMeshesActors, ControlPointMeshComponents is removed from cooked build  
+	ALandscapeSplineActor* SplineActor = Cast<ALandscapeSplineActor>(GetOwner());
+	if (SplineActor && SplineActor->HasGeneratedLandscapeSplineMeshesActors())
+	{
+		return true;
+	}
+
+	return false;
+}
+#endif
+
 //////////////////////////////////////////////////////////////////////////
 // SPLINE CONTROL POINT
 
