Fixed build break
#rnx

diff --git a/Engine/Source/Editor/UnrealEd/Private/WorldPartition/WorldPartitionLandscapeSplineMeshesBuilder.cpp b/Engine/Source/Editor/UnrealEd/Private/WorldPartition/WorldPartitionLandscapeSplineMeshesBuilder.cpp
--- a/Engine/Source/Editor/UnrealEd/Private/WorldPartition/WorldPartitionLandscapeSplineMeshesBuilder.cpp
+++ b/Engine/Source/Editor/UnrealEd/Private/WorldPartition/WorldPartitionLandscapeSplineMeshesBuilder.cpp
@@ -27,6 +27,7 @@
 #include "WorldPartition/WorldPartitionHelpers.h"
 #include "ActorPartition/ActorPartitionSubsystem.h"
 #include "Rendering/ColorVertexBuffer.h"
+#include "EngineUtils.h"
 
 DEFINE_LOG_CATEGORY_STATIC(LogWorldPartitionLandscapeSplineMeshBuilder, All, All);
 
