Allow texture dropping in viewport to create a MIC instead of a new Material

diff --git a/Engine/Source/Editor/UnrealEd/Classes/Settings/LevelEditorViewportSettings.h b/Engine/Source/Editor/UnrealEd/Classes/Settings/LevelEditorViewportSettings.h
--- a/Engine/Source/Editor/UnrealEd/Classes/Settings/LevelEditorViewportSettings.h
+++ b/Engine/Source/Editor/UnrealEd/Classes/Settings/LevelEditorViewportSettings.h
@@ -71,6 +71,16 @@ enum class EScrollGestureDirection : uint8
 	Natural				UMETA(DisplayName = "Natural"),
 };
 
+UENUM()
+enum class EMaterialKind : uint8
+{
+	Unknown,
+	Base,
+	Normal,
+	Specular,
+	Emissive,
+};
+
 /**
  * Implements the Level Editor's per-instance view port settings.
  */
@@ -542,6 +552,14 @@ public:
 	UPROPERTY(config)
 	FVector2D LastInViewportMenuLocation;
 
+	/** When dropping a texture in the viewport, create an instance of this material instead of creating a new material. Populate MaterialParamsForDroppedTextures to specify the parameter names. */
+	UPROPERTY(EditAnywhere, config, Category = Behavior)
+	TSoftObjectPtr<UMaterialInterface> MaterialForDroppedTextures;
+
+	/** When dropping a texture in the viewport, determines which material parameter to assign for each found texture type. Only relevant if MaterialForDroppedTextures is assigned. */
+	UPROPERTY(EditAnywhere, config, Category = Behavior)
+	TMap<EMaterialKind, FName> MaterialParamsForDroppedTextures;
+
 private:
 
 	// Per-instance viewport settings.
diff --git a/Engine/Source/Editor/UnrealEd/Private/LevelEditorViewport.cpp b/Engine/Source/Editor/UnrealEd/Private/LevelEditorViewport.cpp
--- a/Engine/Source/Editor/UnrealEd/Private/LevelEditorViewport.cpp
+++ b/Engine/Source/Editor/UnrealEd/Private/LevelEditorViewport.cpp
@@ -11,6 +11,7 @@
 #include "Components/MeshComponent.h"
 #include "Components/SkeletalMeshComponent.h"
 #include "Materials/Material.h"
+#include "Materials/MaterialInstanceConstant.h"
 #include "Animation/AnimSequenceBase.h"
 #include "CanvasItem.h"
 #include "Engine/BrushBuilder.h"
@@ -22,6 +23,7 @@
 #include "Animation/AnimBlueprint.h"
 #include "Exporters/ExportTextContainer.h"
 #include "Factories/MaterialFactoryNew.h"
+#include "Factories/MaterialInstanceConstantFactoryNew.h"
 #include "Editor/GroupActor.h"
 #include "Components/DecalComponent.h"
 #include "Components/DirectionalLightComponent.h"
@@ -304,19 +306,9 @@ TArray<AActor*> FLevelEditorViewportClient::TryPlacingActorFromObject( ULevel* I
 	return PlacedActors;
 }
 
-namespace EMaterialKind
-{
-	enum Type
-	{
-		Unknown = 0,
-		Base,
-		Normal,
-		Specular,
-		Emissive,
-	};
-}
 
-static FString GetSharedTextureNameAndKind( FString TextureName, EMaterialKind::Type& Kind)
+
+static FString GetSharedTextureNameAndKind( FString TextureName, EMaterialKind& Kind)
 {
 	// Try and strip the suffix from the texture name, if we're successful it must be of that type.
 	bool hasBaseSuffix = TextureName.RemoveFromEnd( "_D" ) || TextureName.RemoveFromEnd( "_Diff" ) || TextureName.RemoveFromEnd( "_Diffuse" ) || TextureName.RemoveFromEnd( "_Detail" ) || TextureName.RemoveFromEnd( "_Base" );
@@ -372,7 +364,7 @@ static UTexture* GetTextureWithNameVariations( const FString& BasePackageName, c
 	return nullptr;
 }
 
-static bool TryAndCreateMaterialInput( UMaterial* UnrealMaterial, EMaterialKind::Type TextureKind, UTexture* UnrealTexture, FExpressionInput& MaterialInput, int X, int Y )
+static bool TryAndCreateMaterialInput( UMaterial* UnrealMaterial, EMaterialKind TextureKind, UTexture* UnrealTexture, FExpressionInput& MaterialInput, int X, int Y )
 {
 	// Ignore null textures.
 	if ( UnrealTexture == nullptr )
@@ -514,16 +506,21 @@ static bool TryAndCreateMaterialInput( UMaterial* UnrealMaterial, EMaterialKind:
 	return true;
 }
 
-UObject* FLevelEditorViewportClient::GetOrCreateMaterialFromTexture( UTexture* UnrealTexture )
+UObject* FLevelEditorViewportClient::GetOrCreateMaterialFromTexture(UTexture* UnrealTexture)
 {
+	// Check if a base material and corresponding params are set in the Settings to determine whether to create a Material or MIC
+	const ULevelEditorViewportSettings* ViewportSettings = GetDefault<ULevelEditorViewportSettings>();
+	UMaterialInterface* BaseMaterial = ViewportSettings->MaterialForDroppedTextures.LoadSynchronous();
+	const bool bCreateMaterialInstance = BaseMaterial && ViewportSettings->MaterialParamsForDroppedTextures.Num() > 0;
+
 	FString TextureShortName = FPackageName::GetShortName( UnrealTexture->GetOutermost()->GetName() );
 
 	// See if we can figure out what kind of material it is, based on a suffix, like _S for Specular, _D for Base/Detail/Diffuse.
 	// if it can determine which type of texture it was, it will return the base name of the texture minus the suffix.
-	EMaterialKind::Type MaterialKind;
+	EMaterialKind MaterialKind;
 	TextureShortName = GetSharedTextureNameAndKind( TextureShortName, MaterialKind );
 	
-	FString MaterialFullName = TextureShortName + "_Mat";
+	const FString MaterialFullName = bCreateMaterialInstance ? (TEXT("MI_") + TextureShortName) : (TextureShortName + TEXT("_Mat"));
 	FString NewPackageName = FPackageName::GetLongPackagePath( UnrealTexture->GetOutermost()->GetName() ) + TEXT( "/" ) + MaterialFullName;
 	NewPackageName = UPackageTools::SanitizePackageName( NewPackageName );
 	UPackage* Package = CreatePackage( *NewPackageName );
@@ -535,96 +532,164 @@ UObject* FLevelEditorViewportClient::GetOrCreateMaterialFromTexture( UTexture* U
 	TArray<FAssetData> OutAssetData;
 	if ( AssetRegistryModule.Get().GetAssetsByPackageName( *NewPackageName, OutAssetData ) && OutAssetData.Num() > 0 )
 	{
-		// TODO Check if is material?
-		return OutAssetData[0].GetAsset();
+		UObject* FoundAsset = OutAssetData[0].GetAsset();
+		if (FoundAsset->IsA(UMaterialInterface::StaticClass()))
+		{
+			return FoundAsset;
+		}
+		UE_LOG(LogEditorViewport, Warning, TEXT("Failed to create material %s from texture because a non-material asset already exists with that name"), *NewPackageName);
+		return nullptr;
 	}
 
-	// create an unreal material asset
-	UMaterialFactoryNew* MaterialFactory = NewObject<UMaterialFactoryNew>();
+	// Variations for Base Maps.
+	TArray<FString> BaseSuffixes;
+	BaseSuffixes.Add("_D");
+	BaseSuffixes.Add("_Diff");
+	BaseSuffixes.Add("_Diffuse");
+	BaseSuffixes.Add("_Detail");
+	BaseSuffixes.Add("_Base");
 
-	UMaterial* UnrealMaterial = (UMaterial*)MaterialFactory->FactoryCreateNew(
-		UMaterial::StaticClass(), Package, *MaterialFullName, RF_Standalone | RF_Public, NULL, GWarn );
+	// Variations for Normal Maps.
+	TArray<FString> NormalSuffixes;
+	NormalSuffixes.Add("_N");
+	NormalSuffixes.Add("_Norm");
+	NormalSuffixes.Add("_Normal");
 
-	if (UnrealMaterial == nullptr)
-	{
-		return nullptr;
-	}
+	// Variations for Specular Maps.
+	TArray<FString> SpecularSuffixes;
+	SpecularSuffixes.Add("_S");
+	SpecularSuffixes.Add("_Spec");
+	SpecularSuffixes.Add("_Specular");
 
-	UMaterialEditorOnlyData* UnrealMaterialEditorOnly = UnrealMaterial->GetEditorOnlyData();
+	// Variations for Emissive Maps.
+	TArray<FString> EmissiveSuffixes;
+	EmissiveSuffixes.Add("_E");
+	EmissiveSuffixes.Add("_Emissive");
 
-	const int HSpace = -300;
+	// The asset path for the base texture, we need this to try and append different suffixes to to find other textures we can use.
+	const FString BaseTexturePackage = FPackageName::GetLongPackagePath(UnrealTexture->GetOutermost()->GetName()) + TEXT("/") + TextureShortName;
 
-	// If we were able to figure out the material kind, we need to try and build a complex material
-	// involving multiple textures.  If not, just try and connect what we found to the base map.
-	if ( MaterialKind == EMaterialKind::Unknown )
+	UMaterialInterface* CreatedMaterialInterface = nullptr;
+
+	if (bCreateMaterialInstance)
 	{
-		TryAndCreateMaterialInput( UnrealMaterial, EMaterialKind::Base, UnrealTexture, UnrealMaterialEditorOnly->BaseColor, HSpace, 0 );
+		UMaterialInstanceConstantFactoryNew* Factory = NewObject<UMaterialInstanceConstantFactoryNew>();
+		Factory->InitialParent = BaseMaterial;
+		UMaterialInstanceConstant* CreatedMIC = Cast<UMaterialInstanceConstant>(Factory->FactoryCreateNew(UMaterialInstanceConstant::StaticClass(), Package, *MaterialFullName, RF_Standalone | RF_Public, NULL, GWarn));
+		if (!CreatedMIC)
+		{
+			return nullptr;
+		}
+		CreatedMaterialInterface = CreatedMIC;
+
+		if (MaterialKind == EMaterialKind::Unknown)
+		{
+			// If the texture type cannot be determined, treat it as a base texture and don't look for any matching textures
+			if (const FName* UnknownParam = ViewportSettings->MaterialParamsForDroppedTextures.Find(EMaterialKind::Unknown))
+			{
+				CreatedMIC->SetTextureParameterValueEditorOnly(*UnknownParam, UnrealTexture);
+			}
+			else if (const FName* BaseParam = ViewportSettings->MaterialParamsForDroppedTextures.Find(EMaterialKind::Base))
+			{
+				CreatedMIC->SetTextureParameterValueEditorOnly(*BaseParam, UnrealTexture);
+			}
+			else
+			{
+				UE_LOG(LogEditorViewport, Warning, TEXT("Dropped texture not assigned to material instance %s because no material parameter name was defined for Unknown or Base in LevelEditorViewportSettings"), *MaterialFullName);
+			}
+		}
+		else
+		{
+			auto AssignParam = [ViewportSettings, CreatedMIC, BaseTexturePackage](EMaterialKind ParamKind, TArray<FString>& Suffixes)->bool
+			{
+				if (const FName* Param = ViewportSettings->MaterialParamsForDroppedTextures.Find(ParamKind))
+				{
+					if (UTexture* Texture = GetTextureWithNameVariations(BaseTexturePackage, Suffixes))
+					{
+						CreatedMIC->SetTextureParameterValueEditorOnly(*Param, Texture);
+						return true;
+					}
+				}
+				return false;
+			};
+
+			// The passed-in texture should be found and assigned in one of the below functions, however it's possible this fails because
+			// the user hasn't set up their Settings properly. If another type of texture is found but not defined in the settings, then
+			// those can just silently fail - only warn the user if the specific texture they chose failed to assign.
+			bool bAssignedInputTexture = false;
+			bAssignedInputTexture |= AssignParam(EMaterialKind::Base, BaseSuffixes) && MaterialKind == EMaterialKind::Base;
+			bAssignedInputTexture |= AssignParam(EMaterialKind::Normal, NormalSuffixes) && MaterialKind == EMaterialKind::Normal;
+			bAssignedInputTexture |= AssignParam(EMaterialKind::Specular, SpecularSuffixes) && MaterialKind == EMaterialKind::Specular;
+			bAssignedInputTexture |= AssignParam(EMaterialKind::Emissive, EmissiveSuffixes) && MaterialKind == EMaterialKind::Emissive;
+			if (!bAssignedInputTexture)
+			{
+				UE_LOG(LogEditorViewport, Warning, TEXT("Dropped texture not assigned to material instance %s because no material parameter name was defined in LevelEditorViewportSettings"), *MaterialFullName);
+			}
+		}
 	}
 	else
 	{
-		// Variations for Base Maps.
-		TArray<FString> BaseSuffixes;
-		BaseSuffixes.Add( "_D" );
-		BaseSuffixes.Add( "_Diff" );
-		BaseSuffixes.Add( "_Diffuse" );
-		BaseSuffixes.Add( "_Detail" );
-		BaseSuffixes.Add( "_Base" );
-
-		// Variations for Normal Maps.
-		TArray<FString> NormalSuffixes;
-		NormalSuffixes.Add( "_N" );
-		NormalSuffixes.Add( "_Norm" );
-		NormalSuffixes.Add( "_Normal" );
-
-		// Variations for Specular Maps.
-		TArray<FString> SpecularSuffixes;
-		SpecularSuffixes.Add( "_S" );
-		SpecularSuffixes.Add( "_Spec" );
-		SpecularSuffixes.Add( "_Specular" );
-
-		// Variations for Emissive Maps.
-		TArray<FString> EmissiveSuffixes;
-		EmissiveSuffixes.Add( "_E" );
-		EmissiveSuffixes.Add( "_Emissive" );
-
-		// The asset path for the base texture, we need this to try and append different suffixes to to find
-		// other textures we can use.
-		FString BaseTexturePackage = FPackageName::GetLongPackagePath( UnrealTexture->GetOutermost()->GetName() ) + TEXT( "/" ) + TextureShortName;
-
-		// Try and find different variations
-		UTexture* BaseTexture = GetTextureWithNameVariations( BaseTexturePackage, BaseSuffixes );
-		UTexture* NormalTexture = GetTextureWithNameVariations( BaseTexturePackage, NormalSuffixes );
-		UTexture* SpecularTexture = GetTextureWithNameVariations( BaseTexturePackage, SpecularSuffixes );
-		UTexture* EmissiveTexture = GetTextureWithNameVariations( BaseTexturePackage, EmissiveSuffixes );
-
-		// Connect and layout any textures we find into their respective inputs in the material.
-		const int VSpace = 170;
-		TryAndCreateMaterialInput( UnrealMaterial, EMaterialKind::Base, BaseTexture, UnrealMaterialEditorOnly->BaseColor, HSpace, VSpace * -1 );
-		TryAndCreateMaterialInput( UnrealMaterial, EMaterialKind::Specular, SpecularTexture, UnrealMaterialEditorOnly->Specular, HSpace, VSpace * 0 );
-		TryAndCreateMaterialInput( UnrealMaterial, EMaterialKind::Emissive, EmissiveTexture, UnrealMaterialEditorOnly->EmissiveColor, HSpace, VSpace * 1 );
-		TryAndCreateMaterialInput( UnrealMaterial, EMaterialKind::Normal, NormalTexture, UnrealMaterialEditorOnly->Normal, HSpace, VSpace * 2 );
-	}
-
-	UnrealMaterial->PreEditChange(nullptr);
-	UnrealMaterial->PostEditChange();
+		// create an unreal material asset
+		UMaterialFactoryNew* MaterialFactory = NewObject<UMaterialFactoryNew>();
+
+		UMaterial* UnrealMaterial = (UMaterial*)MaterialFactory->FactoryCreateNew(
+			UMaterial::StaticClass(), Package, *MaterialFullName, RF_Standalone | RF_Public, NULL, GWarn);
+
+		if (UnrealMaterial == nullptr)
+		{
+			return nullptr;
+		}
+		CreatedMaterialInterface = UnrealMaterial;
+
+		UMaterialEditorOnlyData* UnrealMaterialEditorOnly = UnrealMaterial->GetEditorOnlyData();
+
+		const int HSpace = -300;
+
+		// If we were able to figure out the material kind, we need to try and build a complex material
+		// involving multiple textures.  If not, just try and connect what we found to the base map.
+		if (MaterialKind == EMaterialKind::Unknown)
+		{
+			TryAndCreateMaterialInput(UnrealMaterial, EMaterialKind::Base, UnrealTexture, UnrealMaterialEditorOnly->BaseColor, HSpace, 0);
+		}
+		else
+		{
+			// Try and find different variations
+			UTexture* BaseTexture = GetTextureWithNameVariations(BaseTexturePackage, BaseSuffixes);
+			UTexture* NormalTexture = GetTextureWithNameVariations(BaseTexturePackage, NormalSuffixes);
+			UTexture* SpecularTexture = GetTextureWithNameVariations(BaseTexturePackage, SpecularSuffixes);
+			UTexture* EmissiveTexture = GetTextureWithNameVariations(BaseTexturePackage, EmissiveSuffixes);
+
+			// Connect and layout any textures we find into their respective inputs in the material.
+			const int VSpace = 170;
+			TryAndCreateMaterialInput(UnrealMaterial, EMaterialKind::Base, BaseTexture, UnrealMaterialEditorOnly->BaseColor, HSpace, VSpace * -1);
+			TryAndCreateMaterialInput(UnrealMaterial, EMaterialKind::Specular, SpecularTexture, UnrealMaterialEditorOnly->Specular, HSpace, VSpace * 0);
+			TryAndCreateMaterialInput(UnrealMaterial, EMaterialKind::Emissive, EmissiveTexture, UnrealMaterialEditorOnly->EmissiveColor, HSpace, VSpace * 1);
+			TryAndCreateMaterialInput(UnrealMaterial, EMaterialKind::Normal, NormalTexture, UnrealMaterialEditorOnly->Normal, HSpace, VSpace * 2);
+		}
+	}
+
+	CreatedMaterialInterface->PreEditChange(nullptr);
+	CreatedMaterialInterface->PostEditChange();
 
 	// Notify the asset registry
-	FAssetRegistryModule::AssetCreated( UnrealMaterial );
+	FAssetRegistryModule::AssetCreated(CreatedMaterialInterface);
 
 	// Set the dirty flag so this package will get saved later
 	Package->SetDirtyFlag( true );
 
-	UnrealMaterial->ForceRecompileForRendering();
+	CreatedMaterialInterface->ForceRecompileForRendering();
+
+	const FString MaterialTypeCreated = bCreateMaterialInstance ? TEXT("MaterialInstance") : TEXT("Material");
 
 	// Warn users that a new material has been created
-	FNotificationInfo Info( FText::Format( LOCTEXT( "DropTextureMaterialCreated", "Material '{0}' Created" ), FText::FromString(MaterialFullName)) );
+	FNotificationInfo Info( FText::Format( LOCTEXT( "DropTextureMaterialCreated", "{0} '{1}' Created" ), FText::FromString(MaterialTypeCreated), FText::FromString(MaterialFullName)));
 	Info.ExpireDuration = 4.0f;
 	Info.bUseLargeFont = true;
 	Info.bUseSuccessFailIcons = false;
 	Info.Image = FAppStyle::GetBrush( "ClassThumbnail.Material" );
 	FSlateNotificationManager::Get().AddNotification( Info );
 
-	return UnrealMaterial;
+	return CreatedMaterialInterface;
 }
 
 /**
