add forward declare for UMaterialInterface

diff --git a/Engine/Source/Editor/UnrealEd/Classes/Settings/LevelEditorViewportSettings.h b/Engine/Source/Editor/UnrealEd/Classes/Settings/LevelEditorViewportSettings.h
--- a/Engine/Source/Editor/UnrealEd/Classes/Settings/LevelEditorViewportSettings.h
+++ b/Engine/Source/Editor/UnrealEd/Classes/Settings/LevelEditorViewportSettings.h
@@ -554,7 +554,7 @@ public:
 
 	/** When dropping a texture in the viewport, create an instance of this material instead of creating a new material. Populate MaterialParamsForDroppedTextures to specify the parameter names. */
 	UPROPERTY(EditAnywhere, config, Category = Behavior)
-	TSoftObjectPtr<UMaterialInterface> MaterialForDroppedTextures;
+	TSoftObjectPtr<class UMaterialInterface> MaterialForDroppedTextures;
 
 	/** When dropping a texture in the viewport, determines which material parameter to assign for each found texture type. Only relevant if MaterialForDroppedTextures is assigned. */
 	UPROPERTY(EditAnywhere, config, Category = Behavior)
