Add a renderer project setting to allow virtual textures to be used in opacity masks.
Default to off which is the current hard coded behavior.
The reason for disabling is that shadow passes read but don't do virtual texture feedback.

diff --git a/Engine/Plugins/Enterprise/DatasmithImporter/Source/DatasmithImporter/Private/DatasmithImporter.cpp b/Engine/Plugins/Enterprise/DatasmithImporter/Source/DatasmithImporter/Private/DatasmithImporter.cpp
--- a/Engine/Plugins/Enterprise/DatasmithImporter/Source/DatasmithImporter/Private/DatasmithImporter.cpp
+++ b/Engine/Plugins/Enterprise/DatasmithImporter/Source/DatasmithImporter/Private/DatasmithImporter.cpp
@@ -43,6 +43,7 @@
 #include "Editor/UnrealEdEngine.h"
 #include "EditorLevelUtils.h"
 #include "Engine/Engine.h"
+#include "Engine/RendererSettings.h"
 #include "Engine/SkinnedAsset.h"
 #include "Engine/StaticMesh.h"
 #include "Engine/StaticMeshActor.h"
@@ -893,21 +894,22 @@ UMaterialInterface* FDatasmithImporter::ImportMaterial( FDatasmithImportContext&
 		return nullptr;
 	}
 
-#if MATERIAL_OPACITYMASK_DOESNT_SUPPORT_VIRTUALTEXTURE
-	TArray<UTexture*> OutOpacityMaskTextures;
-	if (ImportedMaterial->GetTexturesInPropertyChain(MP_OpacityMask, OutOpacityMaskTextures, nullptr, nullptr))
+	if (GetDefault<URendererSettings>()->bEnableVirtualTextureOpacityMask == false)
 	{
-		for (UTexture* CurrentTexture : OutOpacityMaskTextures)
+		//Virtual textures are not supported in the OpacityMask slot, convert any textures back to a regular texture.
+		TArray<UTexture*> OutOpacityMaskTextures;
+		if (ImportedMaterial->GetTexturesInPropertyChain(MP_OpacityMask, OutOpacityMaskTextures, nullptr, nullptr))
 		{
-			UTexture2D* Texture2D = Cast<UTexture2D>(CurrentTexture);
-			if (Texture2D && Texture2D->VirtualTextureStreaming)
+			for (UTexture* CurrentTexture : OutOpacityMaskTextures)
 			{
-				//Virtual textures are not supported yet in the OpacityMask slot, convert the texture back to a regular texture.
-				ImportContext.AssetsContext.VirtualTexturesToConvert.Add(Texture2D);
+				UTexture2D* Texture2D = Cast<UTexture2D>(CurrentTexture);
+				if (Texture2D && Texture2D->VirtualTextureStreaming)
+				{
+					ImportContext.AssetsContext.VirtualTexturesToConvert.Add(Texture2D);
+				}
 			}
 		}
 	}
-#endif
 
 	UDatasmithAssetImportData* AssetImportData = Cast< UDatasmithAssetImportData >(ImportedMaterial->AssetImportData);
 
diff --git a/Engine/Plugins/Enterprise/MDLImporter/Source/MDLImporter/Private/MDLImporter.cpp b/Engine/Plugins/Enterprise/MDLImporter/Source/MDLImporter/Private/MDLImporter.cpp
--- a/Engine/Plugins/Enterprise/MDLImporter/Source/MDLImporter/Private/MDLImporter.cpp
+++ b/Engine/Plugins/Enterprise/MDLImporter/Source/MDLImporter/Private/MDLImporter.cpp
@@ -17,6 +17,7 @@
 
 #include "AssetToolsModule.h"
 #include "Engine/Texture2D.h"
+#include "Engine/RendererSettings.h"
 #include "Factories/TextureFactory.h"
 #include "HAL/FileManager.h"
 #include "HAL/PlatformFileManager.h"
@@ -269,7 +270,13 @@ bool FMDLImporter::DistillMaterials(const TMap<FString, UMaterial*>& MaterialsMa
 
 void FMDLImporter::ConvertUnsuportedVirtualTextures() const 
 {
-#if MATERIAL_OPACITYMASK_DOESNT_SUPPORT_VIRTUALTEXTURE
+	//This function only converts virtual textures connected to opacity mask.
+	//Early out if the project supports this.
+	if (GetDefault<URendererSettings>()->bEnableVirtualTextureOpacityMask)
+	{
+		return;
+	}
+
 	const TArray<UMaterialInterface*>& CreatedMaterials = MaterialFactory->GetCreatedMaterials();
 	TArray<UTexture2D*> VirtualTexturesToConvert;
 	TArray<UMaterial*> MaterialsToRefreshAfterVirtualTextureConversion;
@@ -310,7 +317,6 @@ void FMDLImporter::ConvertUnsuportedVirtualTextures() const
 
 	IAssetTools& AssetTools = FModuleManager::LoadModuleChecked<FAssetToolsModule>("AssetTools").Get();
 	AssetTools.ConvertVirtualTextures(VirtualTexturesToConvert, true, &MaterialsToRefreshAfterVirtualTextureConversion);
-#endif
 }
 
 bool FMDLImporter::ImportMaterials(UObject* ParentPackage, EObjectFlags Flags, Mdl::FMaterialCollection& Materials,
diff --git a/Engine/Source/Editor/UnrealEd/Private/Fbx/FbxMaterialImport.cpp b/Engine/Source/Editor/UnrealEd/Private/Fbx/FbxMaterialImport.cpp
--- a/Engine/Source/Editor/UnrealEd/Private/Fbx/FbxMaterialImport.cpp
+++ b/Engine/Source/Editor/UnrealEd/Private/Fbx/FbxMaterialImport.cpp
@@ -28,6 +28,7 @@
 #include "AssetRegistry/ARFilter.h"
 #include "Factories/MaterialImportHelpers.h"
 #include "MaterialEditingLibrary.h"
+#include "Engine/RendererSettings.h"
 
 DEFINE_LOG_CATEGORY_STATIC(LogFbxMaterialImport, Log, All);
 
@@ -288,6 +289,8 @@ bool UnFbx::FFbxImporter::CreateAndLinkExpressionForMaterialProperty(
 			int32 TextureCount = FbxProperty.GetSrcObjectCount<FbxTexture>();
 			if (TextureCount>0)
 			{
+				const bool bEnableVirtualTextureOpacityMask = GetDefault<URendererSettings>()->bEnableVirtualTextureOpacityMask;
+
 				for(int32 TextureIndex =0; TextureIndex<TextureCount; ++TextureIndex)
 				{
 					FbxFileTexture* FbxTexture = FbxProperty.GetSrcObject<FbxFileTexture>(TextureIndex);
@@ -301,9 +304,9 @@ bool UnFbx::FFbxImporter::CreateAndLinkExpressionForMaterialProperty(
 						float ScaleV = FbxTexture->GetScaleV();
 						bool bIsVirtualTexture = UnrealTexture->VirtualTextureStreaming;
 
-#if MATERIAL_OPACITYMASK_DOESNT_SUPPORT_VIRTUALTEXTURE
-						if (bIsVirtualTexture && MaterialProperty == FbxSurfaceMaterial::sTransparencyFactor)
+						if (bIsVirtualTexture && MaterialProperty == FbxSurfaceMaterial::sTransparencyFactor && !bEnableVirtualTextureOpacityMask)
 						{
+							//Virtual textures are not supported in the OpacityMask slot, convert any textures back to a regular texture.
 							//TODO, add a tracking of the materials created during the import so we can refresh here them if the TextureFactory actually found a existing asset instead of creating a new one.
 							if (UTexture2D* UnrealTexture2D = Cast<UTexture2D>(UnrealTexture))
 							{
@@ -318,9 +321,8 @@ bool UnFbx::FFbxImporter::CreateAndLinkExpressionForMaterialProperty(
 							bIsVirtualTexture = false;
 
 							FString TextureName = UnrealTexture->GetName();
-							UE_LOG(LogFbxMaterialImport, Warning, TEXT("Texture %s could not be imported as a Virtual Texture as they are not supported in OpacityMask property (material %s)."), *TextureName, UTF8_TO_TCHAR(FbxMaterial.GetName()));
+							UE_LOG(LogFbxMaterialImport, Warning, TEXT("Texture %s could not be imported as a Virtual Texture because the project renderer settings disable virtual textures in OpacityMask property (material %s)."), *TextureName, UTF8_TO_TCHAR(FbxMaterial.GetName()));
 						}
-#endif
 
 						// and link it to the material 
 						UMaterialExpressionTextureSample* UnrealTextureExpression = NewObject<UMaterialExpressionTextureSample>(UnrealMaterial);
diff --git a/Engine/Source/Runtime/Engine/Classes/Engine/RendererSettings.h b/Engine/Source/Runtime/Engine/Classes/Engine/RendererSettings.h
--- a/Engine/Source/Runtime/Engine/Classes/Engine/RendererSettings.h
+++ b/Engine/Source/Runtime/Engine/Classes/Engine/RendererSettings.h
@@ -412,6 +412,12 @@ class ENGINE_API URendererSettings : public UDeveloperSettings
 		ConfigRestartRequired = true))
 	uint32 bVirtualTextureAnisotropicFiltering : 1;
 
+	UPROPERTY(config, EditAnywhere, Category = VirtualTextures, meta = (
+		EditCondition = "bVirtualTextures",
+		DisplayName = "Enable virtual textures for Opacity Mask",
+		ToolTip = "Relax restriction on virtual textures contributing to Opacity Mask. In some edge cases this can lead to low resolution shadow edges."))
+		uint32 bEnableVirtualTextureOpacityMask : 1;
+
 	UPROPERTY(config, EditAnywhere, Category = VirtualTextures, meta = (
 		EditCondition = "bVirtualTextures",
 		ConsoleVariable = "r.VT.TileSize", DisplayName = "Tile size",
diff --git a/Engine/Source/Runtime/Engine/Private/Materials/HLSLMaterialTranslator.cpp b/Engine/Source/Runtime/Engine/Private/Materials/HLSLMaterialTranslator.cpp
--- a/Engine/Source/Runtime/Engine/Private/Materials/HLSLMaterialTranslator.cpp
+++ b/Engine/Source/Runtime/Engine/Private/Materials/HLSLMaterialTranslator.cpp
@@ -6,6 +6,7 @@
 #include "VT/VirtualTextureScalability.h"
 #include "Materials/MaterialExpressionCustom.h"
 #include "MaterialRestrictiveMode.h"
+#include "Engine/RendererSettings.h"
 
 bool IsExpressionClassPermitted(const UClass* const Class)
 {
@@ -684,13 +685,10 @@ void FHLSLMaterialTranslator::ValidateVtPropertyLimits()
 			VisitExpressionsForProperty(PropertyToValidate, Visitor);
 			if (Visitor.FoundVirtualTexture)
 			{
-				// virtual texture connected to an invalid property, report the correct error
-#if MATERIAL_OPACITYMASK_DOESNT_SUPPORT_VIRTUALTEXTURE
-				if (PropertyToValidate == MP_OpacityMask)
+				if (PropertyToValidate == MP_OpacityMask && GetDefault<URendererSettings>()->bEnableVirtualTextureOpacityMask == false)
 				{
-					Errorf(TEXT("Sampling a virtual texture is currently not supported when connected to the Opacity Mask material attribute."));
+					Errorf(TEXT("Sampling a virtual texture for the Opacity Mask material attribute is disabled in the Renderer Settings."));
 				}
-#endif
 			}
 		}
 	}
diff --git a/Engine/Source/Runtime/Engine/Private/RendererSettings.cpp b/Engine/Source/Runtime/Engine/Private/RendererSettings.cpp
--- a/Engine/Source/Runtime/Engine/Private/RendererSettings.cpp
+++ b/Engine/Source/Runtime/Engine/Private/RendererSettings.cpp
@@ -86,6 +86,7 @@ URendererSettings::URendererSettings(const FObjectInitializer& ObjectInitializer
 	GreenChromaticityCoordinate = FVector2D::ZeroVector;
 	BlueChromaticityCoordinate = FVector2D::ZeroVector;
 	WhiteChromaticityCoordinate = FVector2D::ZeroVector;
+	bEnableVirtualTextureOpacityMask = false;
 }
 
 void URendererSettings::PostInitProperties()
diff --git a/Engine/Source/Runtime/Engine/Public/MaterialShared.h b/Engine/Source/Runtime/Engine/Public/MaterialShared.h
--- a/Engine/Source/Runtime/Engine/Public/MaterialShared.h
+++ b/Engine/Source/Runtime/Engine/Public/MaterialShared.h
@@ -98,8 +98,6 @@ template <class ElementType> class TLinkedList;
 // disallow debug data in shipping or on non-desktop Test
 #define ALLOW_SHADERMAP_DEBUG_DATA (!(UE_BUILD_SHIPPING || (UE_BUILD_TEST && !PLATFORM_DESKTOP)))
 
-#define MATERIAL_OPACITYMASK_DOESNT_SUPPORT_VIRTUALTEXTURE 1
-
 // Adds various checks to track FMaterial lifetime, may add some overhead
 #define UE_CHECK_FMATERIAL_LIFETIME PLATFORM_WINDOWS
 
