Expose GetEditingObjects() to script for asset editor tool menus

diff --git a/Engine/Source/Editor/UnrealEd/Private/Toolkits/AssetEditorToolkit.cpp b/Engine/Source/Editor/UnrealEd/Private/Toolkits/AssetEditorToolkit.cpp
--- a/Engine/Source/Editor/UnrealEd/Private/Toolkits/AssetEditorToolkit.cpp
+++ b/Engine/Source/Editor/UnrealEd/Private/Toolkits/AssetEditorToolkit.cpp
@@ -1292,5 +1292,22 @@ TSharedPtr<FExtender> FExtensibilityManager::GetAllExtenders(const TSharedRef<FU
 	}
 	return FExtender::Combine(OutExtenders);
 }
+
+TArray<UObject*> UAssetEditorToolkitMenuContext::GetEditingObjects() const
+{
+	TArray<UObject*> Result;
+	if (TSharedPtr<FAssetEditorToolkit> Pinned = Toolkit.Pin())
+	{
+		for (UObject* Object : Pinned->GetEditingObjects())
+		{
+			if (Object)
+			{
+				Result.Add(Object);
+			}
+		}
+	}
+
+	return Result;
+}
 	
 #undef LOCTEXT_NAMESPACE
diff --git a/Engine/Source/Editor/UnrealEd/Public/Toolkits/AssetEditorToolkit.h b/Engine/Source/Editor/UnrealEd/Public/Toolkits/AssetEditorToolkit.h
--- a/Engine/Source/Editor/UnrealEd/Public/Toolkits/AssetEditorToolkit.h
+++ b/Engine/Source/Editor/UnrealEd/Public/Toolkits/AssetEditorToolkit.h
@@ -273,6 +273,8 @@ public:
 	virtual void AddGraphEditorPinActionsToContextMenu(FToolMenuSection& InSection) const {};
 
 protected:
+	friend class UAssetEditorToolkitMenuContext;
+
 	/**
 	 * Perform any initialization that should happen after the basic toolkit needs are created for the asset editor.
 	 * For example, if you would like to activate a mode that is not the default from a mode manager, do that here.
diff --git a/Engine/Source/Editor/UnrealEd/Public/Toolkits/AssetEditorToolkitMenuContext.h b/Engine/Source/Editor/UnrealEd/Public/Toolkits/AssetEditorToolkitMenuContext.h
--- a/Engine/Source/Editor/UnrealEd/Public/Toolkits/AssetEditorToolkitMenuContext.h
+++ b/Engine/Source/Editor/UnrealEd/Public/Toolkits/AssetEditorToolkitMenuContext.h
@@ -16,5 +16,8 @@ class UNREALED_API UAssetEditorToolkitMenuContext : public UObject
 
 public:
 
+	UFUNCTION(BlueprintCallable, Category = "Tool Menus")
+	TArray<UObject*> GetEditingObjects() const;
+
 	TWeakPtr<FAssetEditorToolkit> Toolkit;
 };
