Exporting the Appearance setting will also export the current theme to the same folder.

diff --git a/Engine/Source/Editor/UnrealEd/Classes/Settings/EditorStyleSettings.h b/Engine/Source/Editor/UnrealEd/Classes/Settings/EditorStyleSettings.h
--- a/Engine/Source/Editor/UnrealEd/Classes/Settings/EditorStyleSettings.h
+++ b/Engine/Source/Editor/UnrealEd/Classes/Settings/EditorStyleSettings.h
@@ -160,6 +160,9 @@ public:
 
 	/** @return A subdued version of the users selection color (for use with inactive selection)*/
 	FLinearColor GetSubduedSelectionColor() const;
+	
+	bool OnExportBegin(const FString& ExportToPath); 
+
 protected:
 
 	// UObject overrides
diff --git a/Engine/Source/Editor/UnrealEd/Private/Settings/EditorStyleClasses.cpp b/Engine/Source/Editor/UnrealEd/Private/Settings/EditorStyleClasses.cpp
--- a/Engine/Source/Editor/UnrealEd/Private/Settings/EditorStyleClasses.cpp
+++ b/Engine/Source/Editor/UnrealEd/Private/Settings/EditorStyleClasses.cpp
@@ -12,6 +12,7 @@
 #include "UObject/UnrealNames.h"
 #include "UObject/UnrealType.h"
 #include "Styling/StyleColors.h"
+#include "HAL/FileManager.h"
 
 class FObjectInitializer;
 
@@ -92,3 +93,25 @@ void UEditorStyleSettings::PostEditChangeProperty(struct FPropertyChangedEvent&
 	SettingChangedEvent.Broadcast(PropertyName);
 }
 
+
+bool UEditorStyleSettings::OnExportBegin(const FString& ExportToPath)
+{
+	SaveConfig(CPF_Config, *ExportToPath);
+
+	const FStyleTheme& CurrentTheme = USlateThemeManager::Get().GetCurrentTheme();
+
+	FString PathPart;
+	FString Extension;
+	FString FilenameWithoutExtension;
+	FPaths::Split(ExportToPath, PathPart, FilenameWithoutExtension, Extension);
+
+	const FString DestinationPath = PathPart / CurrentTheme.DisplayName.ToString() + TEXT(".json");
+
+	// copy the theme file to the same destination
+	if (IPlatformFile::GetPlatformPhysical().CopyFile(*DestinationPath, *CurrentTheme.Filename))
+	{
+		return true; 
+	}
+	return false; 
+}
+
diff --git a/Engine/Source/Editor/UnrealEd/Private/UnrealEdEngine.cpp b/Engine/Source/Editor/UnrealEd/Private/UnrealEdEngine.cpp
--- a/Engine/Source/Editor/UnrealEd/Private/UnrealEdEngine.cpp
+++ b/Engine/Source/Editor/UnrealEd/Private/UnrealEdEngine.cpp
@@ -79,6 +79,7 @@
 #include "AssetRegistry/AssetRegistryModule.h"
 #include "ObjectTools.h"
 #include "Cooker/ExternalCookOnTheFlyServer.h"
+#include "ISettingsSection.h"
 
 DEFINE_LOG_CATEGORY_STATIC(LogUnrealEdEngine, Log, All);
 
@@ -172,8 +173,6 @@ void UUnrealEdEngine::Init(IEngineLoop* InEngineLoop)
 
 	if (FPaths::IsProjectFilePathSet() && GIsEditor && !FApp::IsUnattended())
 	{
-
-		UE_SCOPED_ENGINE_ACTIVITY(TEXT("Initializing AutoRepimportManager"));
 		AutoReimportManager = NewObject<UAutoReimportManager>();
 		AutoReimportManager->Initialize();
 	}
@@ -198,11 +197,13 @@ void UUnrealEdEngine::Init(IEngineLoop* InEngineLoop)
 
 		if (SettingsModule != nullptr)
 		{
-			SettingsModule->RegisterSettings("Editor", "General", "Appearance",
+			ISettingsSectionPtr StyleSettingsPtr = SettingsModule->RegisterSettings("Editor", "General", "Appearance",
 				NSLOCTEXT("UnrealEd", "Appearance_UserSettingsName", "Appearance"),
 				NSLOCTEXT("UnrealEd", "Appearance_UserSettingsDescription", "Customize the look of the editor."),
 				Settings
 			);
+			
+			StyleSettingsPtr->OnExport().BindUObject(Settings, &UEditorStyleSettings::OnExportBegin);
 		}
 
 
