Ensure that transaction diffing handles newly created objects consistently

In this context "newly created" means both created via NewObject but also objects that go from a dead->alive state via an undo/redo operation. Previously the latter would have produced a partial (usually empty) diff for the new object state, which put the burden of correctness on higher-level systems that were listening for changes.

Annotation diffing changed the rules about how objects might have been serialized for diffing (eg, managed components only diff some properties), so higher-level systems that are doing their own archetype diffing are now producing poor/incorrect results.

This change puts that burden onto the transaction record to handle correctly, as it is the only thing that accurately knows how to do so.

diff --git a/Engine/Plugins/Developer/Concert/ConcertSync/ConcertSyncClient/Source/ConcertSyncClient/Private/ConcertClientTransactionBridge.cpp b/Engine/Plugins/Developer/Concert/ConcertSync/ConcertSyncClient/Source/ConcertSyncClient/Private/ConcertClientTransactionBridge.cpp
--- a/Engine/Plugins/Developer/Concert/ConcertSync/ConcertSyncClient/Source/ConcertSyncClient/Private/ConcertClientTransactionBridge.cpp
+++ b/Engine/Plugins/Developer/Concert/ConcertSync/ConcertSyncClient/Source/ConcertSyncClient/Private/ConcertClientTransactionBridge.cpp
@@ -779,6 +779,12 @@ void FConcertClientTransactionBridge::HandleObjectTransacted(UObject* InObject,
 	TSharedPtr<ITransactionObjectAnnotation> TransactionAnnotation = InTransactionEvent.GetAnnotation();
 	const bool bUseSerializedAnnotationData = !bIncludeAnnotationObjectChanges || !TransactionAnnotation || !TransactionAnnotation->SupportsAdditionalObjectChanges();
 
+	if (!InTransactionEvent.HasNonPropertyChanges() && ExportedProperties.Num() == 0)
+	{
+		// This object has no changes to send
+		return;
+	}
+
 	// Track which packages were changed
 	OngoingTransaction.CommonData.ModifiedPackages.AddUnique(ChangedPackage->GetFName());
 
@@ -800,7 +806,7 @@ void FConcertClientTransactionBridge::HandleObjectTransacted(UObject* InObject,
 	if (InTransactionEvent.GetEventType() == ETransactionObjectEventType::Snapshot)
 	{
 		// Merge the snapshot property changes into pending snapshot list
-		if (OnLocalTransactionSnapshotDelegate.IsBound() && (InTransactionEvent.HasPropertyChanges() || TransactionAnnotation.IsValid()))
+		if (OnLocalTransactionSnapshotDelegate.IsBound() && (ExportedProperties.Num() > 0 || TransactionAnnotation.IsValid()))
 		{
 			// Find or add an entry for this object
 			FConcertExportedObject* ObjectUpdatePtr = OngoingTransaction.SnapshotData.SnapshotObjectUpdates.FindByPredicate([&ObjectId](FConcertExportedObject& ObjectUpdate)
@@ -842,8 +848,6 @@ void FConcertClientTransactionBridge::HandleObjectTransacted(UObject* InObject,
 	}
 	else if (OnLocalTransactionFinalizedDelegate.IsBound())
 	{
-		
-
 		const bool bIsNewlyCreated = InTransactionEvent.HasPendingKillChange() && IsValid(InObject);
 
 		FConcertExportedObject& ObjectUpdate = OngoingTransaction.FinalizedData.FinalizedObjectUpdates.AddDefaulted_GetRef();
@@ -863,57 +867,24 @@ void FConcertClientTransactionBridge::HandleObjectTransacted(UObject* InObject,
 			TransactionAnnotation->Serialize(AnnotationWriter);
 		}
 
-		auto PopulateObjectUpdateData = [this, &OngoingTransaction, &ObjectUpdate, InObject](const bool bHasNonPropertyChanges, const TArray<const FProperty*>* ExportedPropertiesPtr)
-		{
-			if (bHasNonPropertyChanges)
-			{
-				// The 'non-property changes' refers to custom data added by a deriver UObject before and/or after the standard serialized data. Since this is a custom
-				// data format, we don't know what changed, call the object to re-serialize this part, but still send the delta for the generic reflected properties (in RootPropertyNames).
-				ConcertSyncClientUtil::SerializeObject(&OngoingTransaction.FinalizedData.FinalizedLocalIdentifierTable, InObject, ExportedPropertiesPtr, bIncludeEditorOnlyProperties, ObjectUpdate.ObjectData.SerializedData);
-
-				if (ExportedPropertiesPtr)
-				{
-					// Track which properties changed. Not used to apply the transaction on the receiving side, the object specific serialization function will be used for that, but
-					// to be able to display, in the transaction detail view, which 'properties' changed in the transaction as transaction data is otherwise opaque to UI.
-					for (const FProperty* ExportedProperty : *ExportedPropertiesPtr)
-					{
-						FConcertSerializedPropertyData& PropertyData = ObjectUpdate.PropertyDatas.AddDefaulted_GetRef();
-						PropertyData.PropertyName = ExportedProperty->GetFName();
-					}
-				}
-			}
-			else if (ExportedPropertiesPtr) // Its possible to optimize the transaction payload, only sending a 'delta' update.
-			{
-				// Only send properties that changed. The receiving side will 'patch' the object using the reflection system. The specific object serialization function will NOT be called.
-				ConcertSyncClientUtil::SerializeProperties(&OngoingTransaction.FinalizedData.FinalizedLocalIdentifierTable, InObject, *ExportedPropertiesPtr, bIncludeEditorOnlyProperties, ObjectUpdate.PropertyDatas);
-			}
-		};
-
-		// If this object changed from being pending kill to not being pending kill then we'll send an update against the archetype state (rather than use 
-		// the current delta-update) and the receiving side will ensure that the target object matches the archetype state prior to applying the update
-		if (bIsNewlyCreated)
+		if (InTransactionEvent.HasNonPropertyChanges(/*SerializationOnly*/true))
 		{
-			UObject* ObjectArchetype = InObject->GetArchetype();
-			if (ObjectArchetype && !InObject->HasAnyFlags(RF_ClassDefaultObject | RF_ArchetypeObject))
-			{
-				UE::Transaction::FDiffableObject ArchetypeDiffableObject = UE::Transaction::DiffUtil::GetDiffableObject(ObjectArchetype);
-				ArchetypeDiffableObject.SetObject(InObject); // Use the real object here, as the archetype info will never be valid for an instance
+			// The 'non-property changes' refers to custom data added by a deriver UObject before and/or after the standard serialized data. Since this is a custom
+			// data format, we don't know what changed, call the object to re-serialize this part, but still send the delta for the generic reflected properties (in RootPropertyNames).
+			ConcertSyncClientUtil::SerializeObject(&OngoingTransaction.FinalizedData.FinalizedLocalIdentifierTable, InObject, &ExportedProperties, bIncludeEditorOnlyProperties, ObjectUpdate.ObjectData.SerializedData);
 
-				const UE::Transaction::FDiffableObject CurrentDiffableObject = UE::Transaction::DiffUtil::GetDiffableObject(InObject);
-
-				const FTransactionObjectDeltaChange DeltaChangeFromArchetype = UE::Transaction::DiffUtil::GenerateObjectDiff(ArchetypeDiffableObject, CurrentDiffableObject);
-				const TArray<const FProperty*> ChangedPropertiesFromArchetype = ConcertSyncClientUtil::GetExportedProperties(InObject->GetClass(), DeltaChangeFromArchetype.ChangedProperties, bIncludeEditorOnlyProperties);
-
-				PopulateObjectUpdateData(DeltaChangeFromArchetype.bHasNonPropertyChanges, &ChangedPropertiesFromArchetype);
-			}
-			else
+			// Track which properties changed. Not used to apply the transaction on the receiving side, the object specific serialization function will be used for that, but
+			// to be able to display, in the transaction detail view, which 'properties' changed in the transaction as transaction data is otherwise opaque to UI.
+			for (const FProperty* ExportedProperty : ExportedProperties)
 			{
-				PopulateObjectUpdateData(/*bHasNonPropertyChanges*/true, nullptr);
+				FConcertSerializedPropertyData& PropertyData = ObjectUpdate.PropertyDatas.AddDefaulted_GetRef();
+				PropertyData.PropertyName = ExportedProperty->GetFName();
 			}
 		}
-		else
+		else // Its possible to optimize the transaction payload, only sending a 'delta' update.
 		{
-			PopulateObjectUpdateData(InTransactionEvent.HasNonPropertyChanges(/*SerializationOnly*/true), &ExportedProperties);
+			// Only send properties that changed. The receiving side will 'patch' the object using the reflection system. The specific object serialization function will NOT be called.
+			ConcertSyncClientUtil::SerializeProperties(&OngoingTransaction.FinalizedData.FinalizedLocalIdentifierTable, InObject, ExportedProperties, bIncludeEditorOnlyProperties, ObjectUpdate.PropertyDatas);
 		}
 	}
 }
diff --git a/Engine/Plugins/Developer/Concert/ConcertSync/ConcertSyncCore/Source/ConcertSyncCore/Private/ConcertSyncArchives.cpp b/Engine/Plugins/Developer/Concert/ConcertSync/ConcertSyncCore/Source/ConcertSyncCore/Private/ConcertSyncArchives.cpp
--- a/Engine/Plugins/Developer/Concert/ConcertSync/ConcertSyncCore/Source/ConcertSyncCore/Private/ConcertSyncArchives.cpp
+++ b/Engine/Plugins/Developer/Concert/ConcertSync/ConcertSyncCore/Source/ConcertSyncCore/Private/ConcertSyncArchives.cpp
@@ -160,7 +160,7 @@ void ResetObjectPropertiesToArchetypeValues(UObject* Object, const bool InInclud
 		}
 	};
 
-	if (const UObject* ObjectArchetype = Object->GetArchetype(); ObjectArchetype && !Object->HasAnyFlags(RF_ClassDefaultObject | RF_ArchetypeObject))
+	if (const UObject* ObjectArchetype = Object->GetArchetype())
 	{
 		TArray<uint8> ArchetypeData;
 		FArchetypePropertyWriter(ObjectArchetype, ArchetypeData, InIncludeEditorOnlyData);
diff --git a/Engine/Source/Editor/UnrealEd/Classes/Editor/Transactor.h b/Engine/Source/Editor/UnrealEd/Classes/Editor/Transactor.h
--- a/Engine/Source/Editor/UnrealEd/Classes/Editor/Transactor.h
+++ b/Engine/Source/Editor/UnrealEd/Classes/Editor/Transactor.h
@@ -128,14 +128,16 @@ protected:
 		FObjectRecord& operator=( const FObjectRecord& ) = delete;
 
 	public:
+		UE::Transaction::FDiffableObject GetDiffableObject(TArrayView<const FProperty*> PropertiesToSerialize = TArrayView<const FProperty*>()) const;
+
 		// Functions.
-		void SerializeContents( FArchive& Ar, int32 InOper );
-		void SerializeObject( FArchive& Ar );
+		void SerializeContents( FArchive& Ar, int32 InOper ) const;
+		void SerializeObject( FArchive& Ar ) const;
 		void Restore( FTransaction* Owner );
 		void Save( FTransaction* Owner );
 		void Load( FTransaction* Owner );
-		void Finalize( FTransaction* Owner, TSharedPtr<ITransactionObjectAnnotation>& OutFinalizedObjectAnnotation );
-		void Snapshot( FTransaction* Owner, TArrayView<const FProperty*> Properties );
+		void Finalize( FTransaction* Owner, UE::Transaction::DiffUtil::FDiffableObjectArchetypeCache& ArchetypeCache, TSharedPtr<ITransactionObjectAnnotation>& OutFinalizedObjectAnnotation );
+		void Snapshot( FTransaction* Owner, UE::Transaction::DiffUtil::FDiffableObjectArchetypeCache& ArchetypeCache, TArrayView<const FProperty*> Properties );
 
 		/** Used by GC to collect referenced objects. */
 		void AddReferencedObjects( FReferenceCollector& Collector );
diff --git a/Engine/Source/Editor/UnrealEd/Private/EditorTransaction.cpp b/Engine/Source/Editor/UnrealEd/Private/EditorTransaction.cpp
--- a/Engine/Source/Editor/UnrealEd/Private/EditorTransaction.cpp
+++ b/Engine/Source/Editor/UnrealEd/Private/EditorTransaction.cpp
@@ -70,15 +70,32 @@ FTransaction::FObjectRecord::FObjectRecord(FTransaction* Owner, UObject* InObjec
 
 		if (!Array)
 		{
-			DiffableObject = MakeUnique<UE::Transaction::FDiffableObject>();
-			DiffableObject->SetObject(CurrentObject);
-			UE::Transaction::FDiffableObjectDataWriter Writer(*DiffableObject);
-			SerializeObject(Writer);
+			DiffableObject = MakeUnique<UE::Transaction::FDiffableObject>(GetDiffableObject());
 		}
 	}
 }
 
-void FTransaction::FObjectRecord::SerializeContents( FArchive& Ar, int32 InOper )
+UE::Transaction::FDiffableObject FTransaction::FObjectRecord::GetDiffableObject(TArrayView<const FProperty*> PropertiesToSerialize) const
+{
+	check(!Array);
+
+	if (UObject* CurrentObject = Object.Get())
+	{
+		UE::Transaction::DiffUtil::FGetDiffableObjectOptions ObjectOptions;
+		ObjectOptions.PropertiesToSerialize = PropertiesToSerialize;
+		ObjectOptions.ObjectSerializationMode = UE::Transaction::DiffUtil::EGetDiffableObjectMode::Custom;
+		ObjectOptions.CustomSerializer = [this](UE::Transaction::FDiffableObjectDataWriter& DiffWriter)
+		{
+			SerializeObject(DiffWriter);
+		};
+
+		return UE::Transaction::DiffUtil::GetDiffableObject(CurrentObject, ObjectOptions);
+	}
+
+	return UE::Transaction::FDiffableObject();
+}
+
+void FTransaction::FObjectRecord::SerializeContents( FArchive& Ar, int32 InOper ) const
 {
 	if( Array )
 	{
@@ -145,7 +162,7 @@ void FTransaction::FObjectRecord::SerializeContents( FArchive& Ar, int32 InOper
 	}
 }
 
-void FTransaction::FObjectRecord::SerializeObject( FArchive& Ar )
+void FTransaction::FObjectRecord::SerializeObject( FArchive& Ar ) const
 {
 	check(!Array);
 
@@ -250,7 +267,7 @@ void FTransaction::FObjectRecord::Load(FTransaction* Owner)
 	}
 }
 
-void FTransaction::FObjectRecord::Finalize( FTransaction* Owner, TSharedPtr<ITransactionObjectAnnotation>& OutFinalizedObjectAnnotation )
+void FTransaction::FObjectRecord::Finalize( FTransaction* Owner, UE::Transaction::DiffUtil::FDiffableObjectArchetypeCache& ArchetypeCache, TSharedPtr<ITransactionObjectAnnotation>& OutFinalizedObjectAnnotation )
 {
 	OutFinalizedObjectAnnotation.Reset();
 
@@ -291,15 +308,10 @@ void FTransaction::FObjectRecord::Finalize( FTransaction* Owner, TSharedPtr<ITra
 				|| SerializedObject.ReferencedObjects != SerializedObjectFlip.ReferencedObjects;
 
 			// Serialize the object so we can diff it
-			UE::Transaction::FDiffableObject CurrentDiffableObject;
-			{
-				CurrentDiffableObject.SetObject(CurrentObject);
-				UE::Transaction::FDiffableObjectDataWriter Writer(CurrentDiffableObject);
-				SerializeObject(Writer);
-			}
+			const UE::Transaction::FDiffableObject CurrentDiffableObject = GetDiffableObject();
 			
 			// Diff against the object state when the transaction started
-			UE::Transaction::DiffUtil::GenerateObjectDiff(*DiffableObject, CurrentDiffableObject, DeltaChange);
+			UE::Transaction::DiffUtil::GenerateObjectDiff(*DiffableObject, CurrentDiffableObject, DeltaChange, UE::Transaction::DiffUtil::FGenerateObjectDiffOptions(), &ArchetypeCache);
 
 			// If we have a previous snapshot then we need to consider that part of the diff for the finalized object, as systems may 
 			// have been tracking delta-changes between snapshots and this finalization will need to account for those changes too
@@ -308,7 +320,7 @@ void FTransaction::FObjectRecord::Finalize( FTransaction* Owner, TSharedPtr<ITra
 				UE::Transaction::DiffUtil::FGenerateObjectDiffOptions DiffOptions;
 				DiffOptions.bFullDiff = false;
 
-				UE::Transaction::DiffUtil::GenerateObjectDiff(*DiffableObjectSnapshot, CurrentDiffableObject, DeltaChange, DiffOptions);
+				UE::Transaction::DiffUtil::GenerateObjectDiff(*DiffableObjectSnapshot, CurrentDiffableObject, DeltaChange, DiffOptions, &ArchetypeCache);
 			}
 		}
 
@@ -320,7 +332,7 @@ void FTransaction::FObjectRecord::Finalize( FTransaction* Owner, TSharedPtr<ITra
 	}
 }
 
-void FTransaction::FObjectRecord::Snapshot( FTransaction* Owner, TArrayView<const FProperty*> Properties )
+void FTransaction::FObjectRecord::Snapshot( FTransaction* Owner, UE::Transaction::DiffUtil::FDiffableObjectArchetypeCache& ArchetypeCache, TArrayView<const FProperty*> Properties )
 {
 	if (Array)
 	{
@@ -363,14 +375,9 @@ void FTransaction::FObjectRecord::Snapshot( FTransaction* Owner, TArrayView<cons
 		}
 
 		// Serialize the object so we can diff it
-		UE::Transaction::FDiffableObject CurrentDiffableObject;
-		{
-			CurrentDiffableObject.SetObject(CurrentObject);
-			UE::Transaction::FDiffableObjectDataWriter Writer(CurrentDiffableObject, AllPropertiesSnapshot);
-			// although it would be preferable to use SerializeScriptProperties, this cause a false diff between the first snapshot and the base object
-			// since they were serialized with different algo and we don't record enough context to make the comparison appropriately
-			SerializeObject(Writer);
-		}
+		// Note: although it would be preferable to use SerializeScriptProperties, this cause a false diff between the first snapshot and the base object
+		// since they were serialized with different algo and we don't record enough context to make the comparison appropriately
+		UE::Transaction::FDiffableObject CurrentDiffableObject = GetDiffableObject(AllPropertiesSnapshot);
 
 		// Diff against the correct serialized data depending on whether we already had a snapshot
 		const UE::Transaction::FDiffableObject& InitialDiffableObject = DiffableObjectSnapshot ? *DiffableObjectSnapshot : *DiffableObject;
@@ -379,7 +386,7 @@ void FTransaction::FObjectRecord::Snapshot( FTransaction* Owner, TArrayView<cons
 			UE::Transaction::DiffUtil::FGenerateObjectDiffOptions DiffOptions;
 			DiffOptions.bFullDiff = false;
 
-			UE::Transaction::DiffUtil::GenerateObjectDiff(InitialDiffableObject, CurrentDiffableObject, SnapshotDeltaChange, DiffOptions);
+			UE::Transaction::DiffUtil::GenerateObjectDiff(InitialDiffableObject, CurrentDiffableObject, SnapshotDeltaChange, DiffOptions, &ArchetypeCache);
 		}
 
 		// Update the snapshot data for next time
@@ -666,12 +673,14 @@ void FTransaction::SnapshotObject( UObject* InObject, TArrayView<const FProperty
 
 	if (const FObjectRecords* ObjectRecords = ObjectRecordsMap.Find(UE::Transaction::FPersistentObjectRef(InObject)))
 	{
+		UE::Transaction::DiffUtil::FDiffableObjectArchetypeCache ArchetypeCache;
+
 		for (FObjectRecord* Record : ObjectRecords->Records)
 		{
 			checkSlow(Record->Object.Get() == InObject);
 			if (!Record->CustomChange)
 			{
-				Record->Snapshot(this, Properties);
+				Record->Snapshot(this, ArchetypeCache, Properties);
 			}
 		}
 	}
@@ -714,6 +723,8 @@ void FTransaction::Apply()
 	const int32 Start = Inc==1 ? 0             : Records.Num()-1;
 	const int32 End   = Inc==1 ? Records.Num() :              -1;
 
+	UE::Transaction::DiffUtil::FDiffableObjectArchetypeCache ArchetypeCache;
+
 	// Init objects.
 	for( int32 i=Start; i!=End; i+=Inc )
 	{
@@ -725,7 +736,7 @@ void FTransaction::Apply()
 		if (!Record.bFinalized)
 		{
 			TSharedPtr<ITransactionObjectAnnotation> FinalizedObjectAnnotation;
-			Record.Finalize(this, FinalizedObjectAnnotation);
+			Record.Finalize(this, ArchetypeCache, FinalizedObjectAnnotation);
 		}
 
 		UObject* Object = Record.Object.Get();
@@ -800,13 +811,32 @@ void FTransaction::Apply()
 		{
 			const FObjectRecord::FSerializedObject& InitialSerializedObject = ChangedObjectRecord.SerializedObject;
 
+			// If this object changed from being dead to alive, then we need to compute a new delta change against its archetype
+			// This is so that anything listening for changes will get a full update for the object, as if it had been newly constructed with its current state
+			FTransactionObjectDeltaChange PrimaryDeltaChange;
+			if (DeltaChange.bHasPendingKillChange && IsValid(ChangedObject))
+			{
+				UE::Transaction::FDiffableObject FakeDiffableComponent;
+				FakeDiffableComponent.SetObject(ChangedObject);
+				static_cast<FTransactionObjectId&>(FakeDiffableComponent.ObjectInfo) = InitialSerializedObject.ObjectId;
+				FakeDiffableComponent.ObjectInfo.bIsPendingKill = true;
+
+				const UE::Transaction::FDiffableObject CurrentDiffableObject = ChangedObjectRecord.GetDiffableObject();
+
+				PrimaryDeltaChange = UE::Transaction::DiffUtil::GenerateObjectDiff(FakeDiffableComponent, CurrentDiffableObject, UE::Transaction::DiffUtil::FGenerateObjectDiffOptions(), &ArchetypeCache);
+			}
+			else
+			{
+				PrimaryDeltaChange = DeltaChange;
+			}
+
 			TMap<UObject*, FTransactionObjectChange> AdditionalObjectChanges;
 			if (ChangedObjectTransactionAnnotation)
 			{
 				ChangedObjectTransactionAnnotation->ComputeAdditionalObjectChanges(InitialSerializedObject.ObjectAnnotation.Get(), AdditionalObjectChanges);
 			}
 
-			ChangedObject->PostTransacted(FTransactionObjectEvent(Id, OperationId, ETransactionObjectEventType::UndoRedo, ETransactionObjectChangeCreatedBy::TransactionRecord, FTransactionObjectChange{ InitialSerializedObject.ObjectId, DeltaChange }, ChangedObjectTransactionAnnotation));
+			ChangedObject->PostTransacted(FTransactionObjectEvent(Id, OperationId, ETransactionObjectEventType::UndoRedo, ETransactionObjectChangeCreatedBy::TransactionRecord, FTransactionObjectChange{ InitialSerializedObject.ObjectId, MoveTemp(PrimaryDeltaChange) }, ChangedObjectTransactionAnnotation));
 			for (const TTuple<UObject*, FTransactionObjectChange>& AdditionalObjectChangePair : AdditionalObjectChanges)
 			{
 				AdditionalObjectChangePair.Key->PostTransacted(FTransactionObjectEvent(Id, OperationId, ETransactionObjectEventType::UndoRedo, ETransactionObjectChangeCreatedBy::TransactionAnnotation, AdditionalObjectChangePair.Value, nullptr));
@@ -836,13 +866,15 @@ void FTransaction::Apply()
 
 void FTransaction::Finalize()
 {
+	UE::Transaction::DiffUtil::FDiffableObjectArchetypeCache ArchetypeCache;
+
 	for (int32 i = 0; i < Records.Num(); ++i)
 	{
 		FObjectRecord& ObjectRecord = Records[i];
 
 		TSharedPtr<ITransactionObjectAnnotation> FinalizedObjectAnnotation;
 		TSharedPtr<ITransactionObjectAnnotation> SnapshotObjectAnnotation = ObjectRecord.ObjectAnnotationSnapshot;
-		ObjectRecord.Finalize(this, FinalizedObjectAnnotation);
+		ObjectRecord.Finalize(this, ArchetypeCache, FinalizedObjectAnnotation);
 
 		UObject* Object = ObjectRecord.Object.Get();
 		if (Object)
diff --git a/Engine/Source/Runtime/Engine/Private/ActorEditor.cpp b/Engine/Source/Runtime/Engine/Private/ActorEditor.cpp
--- a/Engine/Source/Runtime/Engine/Private/ActorEditor.cpp
+++ b/Engine/Source/Runtime/Engine/Private/ActorEditor.cpp
@@ -353,7 +353,7 @@ void AActor::DebugShowOneComponentHierarchy( USceneComponent* SceneComp, int32&
 
 FActorTransactionAnnotation::FDiffableComponentInfo::FDiffableComponentInfo(const UActorComponent* Component)
 	: ComponentSourceInfo(Component)
-	, DiffableComponent(UE::Transaction::DiffUtil::GetDiffableObject(Component, UE::Transaction::DiffUtil::SerializeProperties))
+	, DiffableComponent(UE::Transaction::DiffUtil::GetDiffableObject(Component, UE::Transaction::DiffUtil::FGetDiffableObjectOptions{ UE::Transaction::DiffUtil::EGetDiffableObjectMode::SerializeProperties }))
 {
 }
 
@@ -512,21 +512,12 @@ void FActorTransactionAnnotation::ComputeAdditionalObjectChanges(const ITransact
 
 	if (DiffableComponentPairs.Num() > 0)
 	{
-		auto FindOrAddAdditionalObjectChange = [&OutAdditionalObjectChanges](UActorComponent* CurrentComponent, const FTransactionObjectId& OldObjectId, FTransactionObjectDeltaChange&& ComponentDeltaChange)
-		{
-			if (FTransactionObjectChange* ExistingComponentChange = OutAdditionalObjectChanges.Find(CurrentComponent))
-			{
-				ExistingComponentChange->DeltaChange.Merge(ComponentDeltaChange);
-			}
-			else
-			{
-				OutAdditionalObjectChanges.Add(CurrentComponent, FTransactionObjectChange{ OldObjectId, MoveTemp(ComponentDeltaChange) });
-			}
-		};
-
-		auto ComputeAdditionalObjectChange = [&FindOrAddAdditionalObjectChange](const UE::Transaction::FDiffableObject& OldDiffableComponent, const UE::Transaction::FDiffableObject& NewDiffableComponent, UActorComponent* CurrentComponent, const bool bForcePendingKillChange = false)
+		UE::Transaction::DiffUtil::FDiffableObjectArchetypeCache ArchetypeCache;
+		
+		auto ComputeAdditionalObjectChange = [&OutAdditionalObjectChanges, &ArchetypeCache](const UE::Transaction::FDiffableObject& OldDiffableComponent, const UE::Transaction::FDiffableObject& NewDiffableComponent, UActorComponent* CurrentComponent)
 		{
 			UE::Transaction::DiffUtil::FGenerateObjectDiffOptions DiffOptions;
+			DiffOptions.ArchetypeOptions.ObjectSerializationMode = UE::Transaction::DiffUtil::EGetDiffableObjectMode::SerializeProperties;
 			{
 				TSet<const FProperty*> PropertiesToSkip;
 
@@ -545,24 +536,38 @@ void FActorTransactionAnnotation::ComputeAdditionalObjectChanges(const ITransact
 					}
 				}
 
+				// Always skip the construction information for managed components, as this is managed by the construction script
+				if (CurrentComponent->CreationMethod != EComponentCreationMethod::Instance)
+				{
+					static const FName NAME_CreationMethod = "CreationMethod";
+					static const FName NAME_UCSSerializationIndex = "UCSSerializationIndex";
+
+					PropertiesToSkip.Add(FindFProperty<FProperty>(UActorComponent::StaticClass(), NAME_CreationMethod));
+					PropertiesToSkip.Add(FindFProperty<FProperty>(UActorComponent::StaticClass(), NAME_UCSSerializationIndex));
+				}
+
 				DiffOptions.ShouldSkipProperty = [CurrentComponent, PropertiesToSkip = MoveTemp(PropertiesToSkip)](FName PropertyName) -> bool
 				{
 					if (const FProperty* Property = FindFProperty<FProperty>(CurrentComponent->GetClass(), PropertyName))
 					{
-						return !Property->HasAnyPropertyFlags(CPF_Edit | CPF_Interp)
-							|| Property->IsA<FMulticastDelegateProperty>()
+						return Property->IsA<FMulticastDelegateProperty>()
 							|| PropertiesToSkip.Contains(Property);
 					}
 					return false;
 				};
 			}
 
-			FTransactionObjectDeltaChange ComponentDeltaChange = UE::Transaction::DiffUtil::GenerateObjectDiff(OldDiffableComponent, NewDiffableComponent, DiffOptions);
-			ComponentDeltaChange.bHasPendingKillChange |= bForcePendingKillChange;
-
+			FTransactionObjectDeltaChange ComponentDeltaChange = UE::Transaction::DiffUtil::GenerateObjectDiff(OldDiffableComponent, NewDiffableComponent, DiffOptions, &ArchetypeCache);
 			if (ComponentDeltaChange.HasChanged())
 			{
-				FindOrAddAdditionalObjectChange(CurrentComponent, OldDiffableComponent.ObjectInfo, MoveTemp(ComponentDeltaChange));
+				if (FTransactionObjectChange* ExistingComponentChange = OutAdditionalObjectChanges.Find(CurrentComponent))
+				{
+					ExistingComponentChange->DeltaChange.Merge(ComponentDeltaChange);
+				}
+				else
+				{
+					OutAdditionalObjectChanges.Add(CurrentComponent, FTransactionObjectChange{ OldDiffableComponent.ObjectInfo, MoveTemp(ComponentDeltaChange) });
+				}
 			}
 		};
 
@@ -578,23 +583,22 @@ void FActorTransactionAnnotation::ComputeAdditionalObjectChanges(const ITransact
 			else if (DiffableComponentPair.OldDiffableComponent && !DiffableComponentPair.NewDiffableComponent)
 			{
 				// Component is deleted; just mark it as pending kill
-				FTransactionObjectDeltaChange ComponentDeltaChange;
-				ComponentDeltaChange.bHasPendingKillChange = true;
+				UE::Transaction::FDiffableObject FakeDiffableComponent;
+				FakeDiffableComponent.SetObject(DiffableComponentPair.CurrentComponent);
+				FakeDiffableComponent.ObjectInfo.bIsPendingKill = true;
 
-				FindOrAddAdditionalObjectChange(DiffableComponentPair.CurrentComponent, DiffableComponentPair.OldDiffableComponent->ObjectInfo, MoveTemp(ComponentDeltaChange));
+				ComputeAdditionalObjectChange(*DiffableComponentPair.OldDiffableComponent, FakeDiffableComponent, DiffableComponentPair.CurrentComponent);
 			}
 			else if (DiffableComponentPair.NewDiffableComponent)
 			{
 				check(!DiffableComponentPair.OldDiffableComponent);
 
 				// Component is newly added; diff against its archetype
-				if (const UObject* ArchetypeObject = DiffableComponentPair.CurrentComponent->GetArchetype(); ArchetypeObject && !DiffableComponentPair.CurrentComponent->HasAnyFlags(RF_ClassDefaultObject | RF_ArchetypeObject))
-				{
-					UE::Transaction::FDiffableObject ArchetypeDiffableObject = UE::Transaction::DiffUtil::GetDiffableObject(ArchetypeObject, UE::Transaction::DiffUtil::SerializeProperties);
-					ArchetypeDiffableObject.SetObject(DiffableComponentPair.CurrentComponent); // Use the real component here, as the archetype info will never be valid for an instance
+				UE::Transaction::FDiffableObject FakeDiffableComponent;
+				FakeDiffableComponent.SetObject(DiffableComponentPair.CurrentComponent);
+				FakeDiffableComponent.ObjectInfo.bIsPendingKill = true;
 
-					ComputeAdditionalObjectChange(ArchetypeDiffableObject, *DiffableComponentPair.NewDiffableComponent, DiffableComponentPair.CurrentComponent, /*bForcePendingKillChange*/true);
-				}
+				ComputeAdditionalObjectChange(FakeDiffableComponent, *DiffableComponentPair.NewDiffableComponent, DiffableComponentPair.CurrentComponent);
 			}
 		}
 	}
diff --git a/Engine/Source/Runtime/Engine/Private/TransactionCommon.cpp b/Engine/Source/Runtime/Engine/Private/TransactionCommon.cpp
--- a/Engine/Source/Runtime/Engine/Private/TransactionCommon.cpp
+++ b/Engine/Source/Runtime/Engine/Private/TransactionCommon.cpp
@@ -448,36 +448,54 @@ FSerializedTaggedData& FDiffableObjectDataWriter::FCachedTaggedDataEntry::SyncCa
 namespace DiffUtil
 {
 
-FDiffableObject GetDiffableObject(const UObject* Object, const EGetDiffableObjectMode Mode)
+FDiffableObject GetDiffableObject(const UObject* Object, const FGetDiffableObjectOptions& Options)
 {
 	FDiffableObject DiffableObject;
 	DiffableObject.SetObject(Object);
 
-	FDiffableObjectDataWriter DiffWriter(DiffableObject);
-	switch (Mode)
+	if (Options.bSerializeEvenIfPendingKill || !DiffableObject.ObjectInfo.bIsPendingKill)
 	{
-	case EGetDiffableObjectMode::SerializeObject:
-		const_cast<UObject*>(Object)->Serialize(DiffWriter);
-		break;
-	case EGetDiffableObjectMode::SerializeProperties:
-		const_cast<UObject*>(Object)->SerializeScriptProperties(DiffWriter);
-		break;
-	default:
-		checkf(false, TEXT("Unknown EGetDiffableObjectMode!"));
-		break;
+		FDiffableObjectDataWriter DiffWriter(DiffableObject, Options.PropertiesToSerialize);
+		switch (Options.ObjectSerializationMode)
+		{
+		case EGetDiffableObjectMode::SerializeObject:
+			const_cast<UObject*>(Object)->Serialize(DiffWriter);
+			break;
+		case EGetDiffableObjectMode::SerializeProperties:
+			const_cast<UObject*>(Object)->SerializeScriptProperties(DiffWriter);
+			break;
+		case EGetDiffableObjectMode::Custom:
+			check(Options.CustomSerializer);
+			Options.CustomSerializer(DiffWriter);
+			break;
+		default:
+			checkf(false, TEXT("Unknown EGetDiffableObjectMode!"));
+			break;
+		}
 	}
 
 	return DiffableObject;
 }
 
-FTransactionObjectDeltaChange GenerateObjectDiff(const FDiffableObject& OldDiffableObject, const FDiffableObject& NewDiffableObject, const FGenerateObjectDiffOptions& DiffOptions)
+const FDiffableObject& FDiffableObjectArchetypeCache::GetArchetypeDiffableObject(const UObject* Archetype, const FGetDiffableObjectOptions& Options)
 {
-	FTransactionObjectDeltaChange DeltaChange;
-	GenerateObjectDiff(OldDiffableObject, NewDiffableObject, DeltaChange, DiffOptions);
-	return DeltaChange;
+	if (Archetype)
+	{
+		if (const FDiffableObject* ExistingDiffableObject = ArchetypeDiffableObjects.Find(Archetype))
+		{
+			return *ExistingDiffableObject;
+		}
+		else
+		{
+			return ArchetypeDiffableObjects.Add(Archetype, GetDiffableObject(Archetype, Options));
+		}
+	}
+
+	static const FDiffableObject EmptyDiffableObject;
+	return EmptyDiffableObject;
 }
 
-void GenerateObjectDiff(const FDiffableObject& OldDiffableObject, const FDiffableObject& NewDiffableObject, FTransactionObjectDeltaChange& OutDeltaChange, const FGenerateObjectDiffOptions& DiffOptions)
+void GenerateObjectDataDiff_Impl(const FDiffableObject& OldDiffableObject, const FDiffableObject& NewDiffableObject, FTransactionObjectDeltaChange& OutDeltaChange, const FGenerateObjectDiffOptions& DiffOptions)
 {
 	auto IsTaggedDataBlockIdentical = [&OldDiffableObject, &NewDiffableObject, &OutDeltaChange](const FName TaggedDataKey, const FSerializedTaggedData& OldSerializedTaggedData, const FSerializedTaggedData& NewSerializedTaggedData) -> bool
 	{
@@ -511,14 +529,6 @@ void GenerateObjectDiff(const FDiffableObject& OldDiffableObject, const FDiffabl
 		return TaggedDataKey.GetComparisonIndex() == TaggedDataKey_UnknownData.GetComparisonIndex();
 	};
 
-	if (DiffOptions.bFullDiff)
-	{
-		OutDeltaChange.bHasNameChange |= OldDiffableObject.ObjectInfo.ObjectName != NewDiffableObject.ObjectInfo.ObjectName;
-		OutDeltaChange.bHasOuterChange |= OldDiffableObject.ObjectInfo.ObjectOuterPathName != NewDiffableObject.ObjectInfo.ObjectOuterPathName;
-		OutDeltaChange.bHasExternalPackageChange |= OldDiffableObject.ObjectInfo.ObjectExternalPackageName != NewDiffableObject.ObjectInfo.ObjectExternalPackageName;
-		OutDeltaChange.bHasPendingKillChange |= OldDiffableObject.ObjectInfo.bIsPendingKill != NewDiffableObject.ObjectInfo.bIsPendingKill;
-	}
-
 	for (const TPair<FName, FSerializedTaggedData>& NewNamePropertyPair : NewDiffableObject.SerializedTaggedData)
 	{
 		if (!ShouldCompareTaggedData(NewNamePropertyPair.Key))
@@ -570,6 +580,44 @@ void GenerateObjectDiff(const FDiffableObject& OldDiffableObject, const FDiffabl
 	}
 }
 
+FTransactionObjectDeltaChange GenerateObjectDiff(const FDiffableObject& OldDiffableObject, const FDiffableObject& NewDiffableObject, const FGenerateObjectDiffOptions& DiffOptions, FDiffableObjectArchetypeCache* ArchetypeCache)
+{
+	FTransactionObjectDeltaChange DeltaChange;
+	GenerateObjectDiff(OldDiffableObject, NewDiffableObject, DeltaChange, DiffOptions, ArchetypeCache);
+	return DeltaChange;
+}
+
+void GenerateObjectDiff(const FDiffableObject& OldDiffableObject, const FDiffableObject& NewDiffableObject, FTransactionObjectDeltaChange& OutDeltaChange, const FGenerateObjectDiffOptions& DiffOptions, FDiffableObjectArchetypeCache* ArchetypeCache)
+{
+	if (DiffOptions.bFullDiff)
+	{
+		OutDeltaChange.bHasNameChange |= OldDiffableObject.ObjectInfo.ObjectName != NewDiffableObject.ObjectInfo.ObjectName;
+		OutDeltaChange.bHasOuterChange |= OldDiffableObject.ObjectInfo.ObjectOuterPathName != NewDiffableObject.ObjectInfo.ObjectOuterPathName;
+		OutDeltaChange.bHasExternalPackageChange |= OldDiffableObject.ObjectInfo.ObjectExternalPackageName != NewDiffableObject.ObjectInfo.ObjectExternalPackageName;
+		OutDeltaChange.bHasPendingKillChange |= OldDiffableObject.ObjectInfo.bIsPendingKill != NewDiffableObject.ObjectInfo.bIsPendingKill;
+	}
+
+	if (DiffOptions.bDiffDataEvenIfPendingKill || OldDiffableObject.ObjectInfo.bIsPendingKill == NewDiffableObject.ObjectInfo.bIsPendingKill)
+	{
+		GenerateObjectDataDiff_Impl(OldDiffableObject, NewDiffableObject, OutDeltaChange, DiffOptions);
+	}
+	else
+	{
+		if (!NewDiffableObject.ObjectInfo.bIsPendingKill)
+		{
+			check(OldDiffableObject.ObjectInfo.bIsPendingKill);
+
+			FDiffableObjectArchetypeCache LocalArchetypeCache;
+			FDiffableObjectArchetypeCache& ArchetypeCacheRef = ArchetypeCache ? *ArchetypeCache : LocalArchetypeCache;
+
+			// The object has been restored to life, so use the archetype state when generating the data diff (as there likely isn't any data cached for the old object; see GetDiffableObject)
+			const FDiffableObject& DiffableArchetypeObject = ArchetypeCacheRef.GetArchetypeDiffableObject(NewDiffableObject.ObjectArchetype.Get(), DiffOptions.ArchetypeOptions);
+			GenerateObjectDataDiff_Impl(DiffableArchetypeObject, NewDiffableObject, OutDeltaChange, DiffOptions);
+		}
+		// else, the object has become pending kill, so there's nothing more to diff after calculating the object info change above (as there likely isn't any cached data for the new object; see GetDiffableObject)
+	}
+}
+
 } // namespace DiffUtil
 
 } // namespace UE::Transaction
diff --git a/Engine/Source/Runtime/Engine/Public/TransactionCommon.h b/Engine/Source/Runtime/Engine/Public/TransactionCommon.h
--- a/Engine/Source/Runtime/Engine/Public/TransactionCommon.h
+++ b/Engine/Source/Runtime/Engine/Public/TransactionCommon.h
@@ -252,11 +252,13 @@ public:
 	void SetObject(const UObject* InObject)
 	{
 		ObjectInfo.SetObject(InObject);
+		ObjectArchetype = FPersistentObjectRef(InObject->GetArchetype());
 	}
 
 	void Reset()
 	{
 		ObjectInfo.Reset();
+		ObjectArchetype = FPersistentObjectRef();
 		SerializedData.Reset();
 		SerializedTaggedData.Reset();
 	}
@@ -264,6 +266,7 @@ public:
 	void Swap(FDiffableObject& Other)
 	{
 		ObjectInfo.Swap(Other.ObjectInfo);
+		Exchange(ObjectArchetype, Other.ObjectArchetype);
 		Exchange(SerializedData, Other.SerializedData);
 		Exchange(SerializedTaggedData, Other.SerializedTaggedData);
 	}
@@ -271,6 +274,9 @@ public:
 	/** Information about the object when it was serialized */
 	FSerializedObjectInfo ObjectInfo;
 
+	/** The archetype of the object when it was serialized */
+	FPersistentObjectRef ObjectArchetype;
+
 	/** The serialized data for the diffable object */
 	FSerializedObjectData SerializedData;
 
@@ -403,21 +409,78 @@ private:
 namespace DiffUtil
 {
 
-enum EGetDiffableObjectMode : uint8
+enum class EGetDiffableObjectMode : uint8
 {
+	/** Serialize the entire object state by calling its Serialize function */
 	SerializeObject,
+
+	/** Serialize the property state of the object by calling its SerializeScriptProperties function */
 	SerializeProperties,
+
+	/** Serialize the object via a custom serialize function */
+	Custom,
+};
+
+struct FGetDiffableObjectOptions
+{
+	/** How should we serialize the object for diffing? */
+	EGetDiffableObjectMode ObjectSerializationMode = EGetDiffableObjectMode::SerializeObject;
+
+	/** Custom serializer for the object (must be set when ObjectSerializationMode == Custom) */
+	TFunction<void(FDiffableObjectDataWriter&)> CustomSerializer;
+
+	/** Optional list of properties to serialize on the object, or an empty array to serialize all properties */
+	TArrayView<const FProperty*> PropertiesToSerialize;
+
+	/** Should we still serialize this object if it's considered pending kill? */
+	bool bSerializeEvenIfPendingKill = false;
 };
-ENGINE_API FDiffableObject GetDiffableObject(const UObject* Object, const EGetDiffableObjectMode Mode = EGetDiffableObjectMode::SerializeObject);
+
+/**
+ * Get an object snapshot that can be diffed later.
+ */
+ENGINE_API FDiffableObject GetDiffableObject(const UObject* Object, const FGetDiffableObjectOptions& Options = FGetDiffableObjectOptions());
 
 struct FGenerateObjectDiffOptions
 {
-	bool bFullDiff = true;
+	FGenerateObjectDiffOptions()
+	{
+		ArchetypeOptions.bSerializeEvenIfPendingKill = true;
+	}
+
+	/** Options used when getting the diffable state of archive objects */
+	FGetDiffableObjectOptions ArchetypeOptions;
+
+	/** Optional function used to skip comparing certain properties within the diffable data */
 	TFunction<bool(FName)> ShouldSkipProperty;
+
+	/** Should we perform a "full" diff? (compares object info, and considers any missing properties to have been changed) */
+	bool bFullDiff = true;
+
+	/**
+	 * Should we diff the object state even if it's considered pending kill?
+	 * When false:
+	 *   If the old object was pending kill and the new object isn't, then the diff will be performed against the archetype of the new object.
+	 *   If the old object wasn't pending kill and the new object is, then the data diff will be skipped and only the object info diffed.
+	 */
+	bool bDiffDataEvenIfPendingKill = false;
 };
 
-ENGINE_API FTransactionObjectDeltaChange GenerateObjectDiff(const FDiffableObject& OldDiffableObject, const FDiffableObject& NewDiffableObject, const FGenerateObjectDiffOptions& DiffOptions = FGenerateObjectDiffOptions());
-ENGINE_API void GenerateObjectDiff(const FDiffableObject& OldDiffableObject, const FDiffableObject& NewDiffableObject, FTransactionObjectDeltaChange& OutDeltaChange, const FGenerateObjectDiffOptions& DiffOptions = FGenerateObjectDiffOptions());
+/** Optional cache for requests to get the diffable state for a given archetype object */
+class FDiffableObjectArchetypeCache
+{
+public:
+	const FDiffableObject& GetArchetypeDiffableObject(const UObject* Archetype, const FGetDiffableObjectOptions& Options);
+
+private:
+	TMap<const UObject*, FDiffableObject> ArchetypeDiffableObjects;
+};
+
+/**
+ * Generate a diff between the two object snapshots.
+ */
+ENGINE_API FTransactionObjectDeltaChange GenerateObjectDiff(const FDiffableObject& OldDiffableObject, const FDiffableObject& NewDiffableObject, const FGenerateObjectDiffOptions& DiffOptions = FGenerateObjectDiffOptions(), FDiffableObjectArchetypeCache* ArchetypeCache = nullptr);
+ENGINE_API void GenerateObjectDiff(const FDiffableObject& OldDiffableObject, const FDiffableObject& NewDiffableObject, FTransactionObjectDeltaChange& OutDeltaChange, const FGenerateObjectDiffOptions& DiffOptions = FGenerateObjectDiffOptions(), FDiffableObjectArchetypeCache* ArchetypeCache = nullptr);
 
 } // namespace DiffUtil
 
