Remove the hardly used DependentBlueprintsToRefresh static member. This container is growing during BP compile, and only being flushed on hot reload. This means the only way to exercise this code was to compile blueprints with dependent blueprints, and then hotreload, which would cause BlueprintChanged to be run on those blueprint which were dependent of the previous, potentially unrelated, user compile requests

diff --git a/Engine/Source/Editor/UnrealEd/Private/Kismet2/KismetReinstanceUtilities.cpp b/Engine/Source/Editor/UnrealEd/Private/Kismet2/KismetReinstanceUtilities.cpp
--- a/Engine/Source/Editor/UnrealEd/Private/Kismet2/KismetReinstanceUtilities.cpp
+++ b/Engine/Source/Editor/UnrealEd/Private/Kismet2/KismetReinstanceUtilities.cpp
@@ -261,7 +261,6 @@ FReplaceInstancesOfClassParameters::FReplaceInstancesOfClassParameters(UClass* I
 /////////////////////////////////////////////////////////////////////////////////
 // FBlueprintCompileReinstancer
 
-TSet<TWeakObjectPtr<UBlueprint>> FBlueprintCompileReinstancer::DependentBlueprintsToRefresh = TSet<TWeakObjectPtr<UBlueprint>>();
 TSet<TWeakObjectPtr<UBlueprint>> FBlueprintCompileReinstancer::CompiledBlueprintsToSave = TSet<TWeakObjectPtr<UBlueprint>>();
 
 UClass* FBlueprintCompileReinstancer::HotReloadedOldClass = nullptr;
@@ -728,24 +727,12 @@ TSharedPtr<FReinstanceFinalizer> FBlueprintCompileReinstancer::ReinstanceInner(b
 	return Finalizer;
 }
 
-void FBlueprintCompileReinstancer::ListDependentBlueprintsToRefresh(const TArray<UBlueprint*>& DependentBPs)
-{
-	for (UBlueprint* Element : DependentBPs)
-	{
-		DependentBlueprintsToRefresh.Add(Element);
-	}
-}
-
 void FBlueprintCompileReinstancer::EnlistDependentBlueprintToRecompile(UBlueprint* BP, bool bBytecodeOnly)
 {
 }
 
 void FBlueprintCompileReinstancer::BlueprintWasRecompiled(UBlueprint* BP, bool bBytecodeOnly)
 {
-	if (IsValid(BP))
-	{
-		DependentBlueprintsToRefresh.Remove(BP);
-	}
 }
 
 extern UNREALED_API FSecondsCounterData BlueprintCompileAndLoadTimerData;
@@ -817,18 +804,6 @@ void FBlueprintCompileReinstancer::ReinstanceObjects(bool bForceAlwaysReinstance
 						}
 					}
 
-					{
-						BP_SCOPED_COMPILER_EVENT_STAT(EKismetCompilerStats_RefreshDependentBlueprintsInReinstancer);
-						for (TWeakObjectPtr<UBlueprint>& BPPtr : DependentBlueprintsToRefresh)
-						{
-							if (BPPtr.IsValid())
-							{
-								BPPtr->BroadcastChanged();
-							}
-						}
-						DependentBlueprintsToRefresh.Empty();
-					}
-
 					if (GEditor)
 					{
 						GEditor->BroadcastBlueprintCompiled();
@@ -837,7 +812,6 @@ void FBlueprintCompileReinstancer::ReinstanceObjects(bool bForceAlwaysReinstance
 				else
 				{
 					QueueToReinstance.Empty();
-					DependentBlueprintsToRefresh.Empty();
 				}
 			}
 		}
@@ -1071,11 +1045,6 @@ void FBlueprintCompileReinstancer::UpdateBytecodeReferences()
 				bBPWasChanged = true;
 				UE_LOG(LogBlueprint, Log, TEXT("UpdateBytecodeReferences: %d references from %s was replaced in BP %s"), ReplaceInBPAr.GetCount(), *GetPathNameSafe(ClassToReinstance), *GetPathNameSafe(DependentBP));
 			}
-
-			if (bBPWasChanged && CompiledBlueprint && !CompiledBlueprint->bIsRegeneratingOnLoad)
-			{
-				DependentBlueprintsToRefresh.Add(DependentBP);
-			}
 		}
 
 #if VALIDATE_BYTECODE_REFERENCE_DEPENDENCY_CACHE
diff --git a/Engine/Source/Editor/UnrealEd/Public/Kismet2/KismetReinstanceUtilities.h b/Engine/Source/Editor/UnrealEd/Public/Kismet2/KismetReinstanceUtilities.h
--- a/Engine/Source/Editor/UnrealEd/Public/Kismet2/KismetReinstanceUtilities.h
+++ b/Engine/Source/Editor/UnrealEd/Public/Kismet2/KismetReinstanceUtilities.h
@@ -86,7 +86,6 @@ protected:
 	friend struct FBlueprintCompilationManagerImpl;
 	friend struct FReinstancingJob;
 
-	static TSet<TWeakObjectPtr<UBlueprint>> DependentBlueprintsToRefresh;
 	static TSet<TWeakObjectPtr<UBlueprint>> CompiledBlueprintsToSave;
 
 	static UClass* HotReloadedOldClass;
@@ -146,7 +145,6 @@ public:
 
 	static void OptionallyRefreshNodes(UBlueprint* BP);
 
-	void ListDependentBlueprintsToRefresh(const TArray<UBlueprint*>& DependentBPs);
 	UE_DEPRECATED(5.0, "This method performs no function and isn't invoked by the base class.  Remove all calls to the base class method.")
 	virtual void EnlistDependentBlueprintToRecompile(UBlueprint* BP, bool bBytecodeOnly);
 	virtual void BlueprintWasRecompiled(UBlueprint* BP, bool bBytecodeOnly);
