COTF2 - refactored package store synchronization to handle Begin/End lock

diff --git a/Engine/Source/Editor/UnrealEd/Private/Cooker/IoStoreCookOnTheFlyRequestManager.cpp b/Engine/Source/Editor/UnrealEd/Private/Cooker/IoStoreCookOnTheFlyRequestManager.cpp
--- a/Engine/Source/Editor/UnrealEd/Private/Cooker/IoStoreCookOnTheFlyRequestManager.cpp
+++ b/Engine/Source/Editor/UnrealEd/Private/Cooker/IoStoreCookOnTheFlyRequestManager.cpp
@@ -25,6 +25,36 @@
 #include "Templates/Function.h"
 #include "UObject/SavePackage.h"
 
+class FPackageRegistry
+{
+public:
+	void Add(TArrayView<FAssetData> Assets)
+	{
+		FScopeLock _(&CriticalSection);
+
+		for (const FAssetData& Asset : Assets)
+		{
+			PackageIdToName.Add(FPackageId::FromName(Asset.PackageName), Asset.PackageName);
+		}
+	}
+
+	void Add(FName PackageName)
+	{
+		FScopeLock _(&CriticalSection);
+		PackageIdToName.Add(FPackageId::FromName(PackageName), PackageName);
+	}
+
+	FName Get(FPackageId PackageId)
+	{
+		FScopeLock _(&CriticalSection);
+		return PackageIdToName.FindRef(PackageId);
+	}
+
+private:
+	TMap<FPackageId, FName> PackageIdToName;
+	FCriticalSection CriticalSection;
+};
+
 class FIoStoreCookOnTheFlyRequestManager final
 	: public UE::Cook::ICookOnTheFlyRequestManager
 {
@@ -45,10 +75,7 @@ public:
 		MessageEndpoint = FMessageEndpoint::Builder("FCookOnTheFly");
 		TArray<FAssetData> AllAssets;
 		AssetRegistry->GetAllAssets(AllAssets, true);
-		for (const FAssetData& AssetData : AllAssets)
-		{
-			AllKnownPackagesMap.Add(FPackageId::FromName(AssetData.PackageName), AssetData.PackageName);
-		}
+		PackageRegistry.Add(AllAssets);
 	}
 
 private:
@@ -69,9 +96,10 @@ private:
 			FPackageStoreEntryResource Entry;
 		};
 
-		FPlatformContext(FName InPlatformName, IPackageStoreWriter* InPackageWriter)
+		FPlatformContext(FName InPlatformName, IPackageStoreWriter* InPackageWriter, FPackageRegistry& InPackageRegistry)
 			: PlatformName(InPlatformName)
 			, PackageWriter(InPackageWriter)
+			, PackageRegistry(InPackageRegistry)
 		{
 		}
 
@@ -104,7 +132,7 @@ private:
 			}
 		}
 
-		EPackageStoreEntryStatus RequestCook(UE::Cook::ICookOnTheFlyServer& InCookOnTheFlyServer, const FPackageId& PackageId, TFunctionRef<FName()> GetPackageName, FPackageStoreEntryResource& OutEntry)
+		EPackageStoreEntryStatus RequestCook(UE::Cook::ICookOnTheFlyServer& InCookOnTheFlyServer, const FPackageId& PackageId, FPackageStoreEntryResource& OutEntry)
 		{
 			FPackage& Package = GetPackage(PackageId);
 			if (Package.Status == EPackageStatus::Cooked)
@@ -123,7 +151,7 @@ private:
 				UE_LOG(LogCookOnTheFly, Verbose, TEXT("0x%llX was already cooking"), PackageId.ValueForDebugging());
 				return EPackageStoreEntryStatus::Pending;
 			}
-			FName PackageName = GetPackageName();
+			FName PackageName = PackageRegistry.Get(PackageId);
 			if (PackageName.IsNone())
 			{
 				UE_LOG(LogCookOnTheFly, Warning, TEXT("Received cook request for unknown package 0x%llX"), PackageId.ValueForDebugging());
@@ -255,6 +283,7 @@ private:
 		FName PlatformName;
 		TMap<FPackageId, TUniquePtr<FPackage>> Packages;
 		IPackageStoreWriter* PackageWriter = nullptr;
+		FPackageRegistry& PackageRegistry;
 		int32 SingleThreadedClientsCount = 0;
 		FEvent* PackageCookedEvent = nullptr;
 	};
@@ -282,9 +311,7 @@ private:
 	{
 		FPackageId PackageId = FPackageId::FromName(PackageName);
 		UE_LOG(LogCookOnTheFly, Verbose, TEXT("Package 0x%llX '%s' generated"), PackageId.ValueForDebugging(), *PackageName.ToString());
-
-		FScopeLock _(&AllKnownPackagesCriticalSection);
-		AllKnownPackagesMap.Add(PackageId, PackageName);
+		PackageRegistry.Add(PackageName);
 	}
 
 	void TickRecookPackages()
@@ -307,7 +334,7 @@ private:
 
 		for (FPackageId PackageId : PackageIds)
 		{
-			FName PackageName = AllKnownPackagesMap.FindRef(PackageId);
+			FName PackageName = PackageRegistry.Get(PackageId);
 			if (!PackageName.IsNone())
 			{
 				PackageNames.Add(PackageName);
@@ -369,7 +396,7 @@ private:
 			FScopeLock _(&ContextsCriticalSection);
 			if (!PlatformContexts.Contains(PlatformName))
 			{
-				TUniquePtr<FPlatformContext>& Context = PlatformContexts.Add(PlatformName, MakeUnique<FPlatformContext>(PlatformName, PackageWriter));
+				TUniquePtr<FPlatformContext>& Context = PlatformContexts.Add(PlatformName, MakeUnique<FPlatformContext>(PlatformName, PackageWriter, PackageRegistry));
 
 				PackageWriter->GetEntries([&Context](TArrayView<const FPackageStoreEntryResource> Entries,
 					TArrayView<const IPackageStoreWriter::FOplogCookInfo> CookInfos)
@@ -492,22 +519,20 @@ private:
 
 		if (PlatformName.IsNone())
 		{
-			UE_LOG(LogCookOnTheFly, Warning, TEXT("CookPackageRequest from editor client"));
+			UE_LOG(LogCookOnTheFly, Warning, TEXT("Received cook package request for unknown platform '%s'"), *PlatformName.ToString());
 			Response.SetStatus(UE::Cook::ECookOnTheFlyMessageStatus::Error);
 			return true;
 		}
 
 		TRACE_CPUPROFILER_EVENT_SCOPE(CookOnTheFly::HandleCookPackageRequest);
 
-		FCookPackageRequest CookRequest = Request.GetBodyAs<FCookPackageRequest>();
-		UE_LOG(LogCookOnTheFly, Verbose, TEXT("Received cook request 0x%llX"), CookRequest.PackageId.ValueForDebugging());
 		FPlatformContext& Context = GetContext(PlatformName);
+		FCookPackageRequest CookRequest = Request.GetBodyAs<FCookPackageRequest>();
+		FCookPackageResponse CookResponse;
+
+		for (const FPackageId& PackageId : CookRequest.PackageIds)
 		{
-			auto GetPackageNameFunc = [this, &CookRequest]()
-			{
-				FScopeLock _(&AllKnownPackagesCriticalSection);
-				return AllKnownPackagesMap.FindRef(CookRequest.PackageId);
-			};
+			UE_LOG(LogCookOnTheFly, Verbose, TEXT("Received cook request 0x%llX"), PackageId.ValueForDebugging());
 
 			FPackageStoreEntryResource Entry;
 			EPackageStoreEntryStatus PackageStatus = EPackageStoreEntryStatus::Pending;
@@ -517,7 +542,7 @@ private:
 				{
 					{
 						FScopeLock _(&Context.GetLock());
-						PackageStatus = Context.RequestCook(CookOnTheFlyServer, CookRequest.PackageId, GetPackageNameFunc, Entry);
+						PackageStatus = Context.RequestCook(CookOnTheFlyServer, PackageId, Entry);
 					}
 					if (PackageStatus != EPackageStoreEntryStatus::Pending)
 					{
@@ -529,11 +554,20 @@ private:
 			else
 			{
 				FScopeLock _(&Context.GetLock());
-				PackageStatus = Context.RequestCook(CookOnTheFlyServer, CookRequest.PackageId, GetPackageNameFunc, Entry);
+				PackageStatus = Context.RequestCook(CookOnTheFlyServer, PackageId, Entry);
 			}
 
-			Response.SetBodyTo(FCookPackageResponse{ PackageStatus, MoveTemp(Entry) });
+			if (PackageStatus == EPackageStoreEntryStatus::Ok)
+			{
+				CookResponse.CookedPackages.Add(MoveTemp(Entry));
+			}
+			else if (PackageStatus == EPackageStoreEntryStatus::Missing)
+			{
+				CookResponse.FailedPackages.Add(PackageId);
+			}
 		}
+
+		Response.SetBodyTo(MoveTemp(CookResponse));
 		Response.SetStatus(UE::Cook::ECookOnTheFlyMessageStatus::Ok);
 
 		return true;
@@ -604,12 +638,7 @@ private:
 		for (const FPackageId& ImportedPackageId : EventArgs.Entry.ImportedPackageIds)
 		{
 			FPackageStoreEntryResource DummyEntry;
-			auto GetPackageNameFunc = [this, &ImportedPackageId]()
-			{
-				FScopeLock _(&AllKnownPackagesCriticalSection);
-				return AllKnownPackagesMap.FindRef(ImportedPackageId);
-			};
-			EPackageStoreEntryStatus PackageStatus = Context.RequestCook(CookOnTheFlyServer, ImportedPackageId, GetPackageNameFunc, DummyEntry);
+			EPackageStoreEntryStatus PackageStatus = Context.RequestCook(CookOnTheFlyServer, ImportedPackageId, DummyEntry);
 		}
 	}
 
@@ -730,8 +759,7 @@ private:
 	const FString ServiceId;
 	FCriticalSection ContextsCriticalSection;
 	TMap<FName, TUniquePtr<FPlatformContext>> PlatformContexts;
-	FCriticalSection AllKnownPackagesCriticalSection;
-	TMap<FPackageId, FName> AllKnownPackagesMap;
+	FPackageRegistry PackageRegistry;
 	FCriticalSection PackagesToRecookCritical;
 	TSet<FPackageId> PackagesToRecook;
 };
diff --git a/Engine/Source/Runtime/CookOnTheFly/Internal/CookOnTheFlyMessages.h b/Engine/Source/Runtime/CookOnTheFly/Internal/CookOnTheFlyMessages.h
--- a/Engine/Source/Runtime/CookOnTheFly/Internal/CookOnTheFlyMessages.h
+++ b/Engine/Source/Runtime/CookOnTheFly/Internal/CookOnTheFlyMessages.h
@@ -18,18 +18,12 @@ struct FCompletedPackages
 
 struct FCookPackageRequest
 {
-	FPackageId PackageId;
+	TArray<FPackageId> PackageIds;
 	
 	COOKONTHEFLY_API friend FArchive& operator<<(FArchive& Ar, FCookPackageRequest& Request);
 };
 
-struct FCookPackageResponse
-{
-	EPackageStoreEntryStatus Status;
-	FPackageStoreEntryResource CookedEntry;
-
-	COOKONTHEFLY_API friend FArchive& operator<<(FArchive& Ar, FCookPackageResponse& Response);
-};
+using FCookPackageResponse = FCompletedPackages;
 
 struct FRecookPackagesRequest
 {
diff --git a/Engine/Source/Runtime/CookOnTheFly/Private/CookOnTheFlyMessages.cpp b/Engine/Source/Runtime/CookOnTheFly/Private/CookOnTheFlyMessages.cpp
--- a/Engine/Source/Runtime/CookOnTheFly/Private/CookOnTheFlyMessages.cpp
+++ b/Engine/Source/Runtime/CookOnTheFly/Private/CookOnTheFlyMessages.cpp
@@ -15,18 +15,7 @@ FArchive& operator<<(FArchive& Ar, FCompletedPackages& CompletedPackages)
 
 FArchive& operator<<(FArchive& Ar, FCookPackageRequest& Request)
 {
-	Ar << Request.PackageId;
-
-	return Ar;
-}
-
-FArchive& operator<<(FArchive& Ar, FCookPackageResponse& Response)
-{
-	Ar << Response.Status;
-	if (Response.Status == EPackageStoreEntryStatus::Ok)
-	{
-		Ar << Response.CookedEntry;
-	}
+	Ar << Request.PackageIds;
 
 	return Ar;
 }
diff --git a/Engine/Source/Runtime/StorageServerClient/Private/CookOnTheFlyPackageStore.cpp b/Engine/Source/Runtime/StorageServerClient/Private/CookOnTheFlyPackageStore.cpp
--- a/Engine/Source/Runtime/StorageServerClient/Private/CookOnTheFlyPackageStore.cpp
+++ b/Engine/Source/Runtime/StorageServerClient/Private/CookOnTheFlyPackageStore.cpp
@@ -47,82 +47,78 @@ FCookOnTheFlyPackageStoreBackend::FCookOnTheFlyPackageStoreBackend(UE::Cook::ICo
 	});
 }
 
+void FCookOnTheFlyPackageStoreBackend::BeginRead()
+{
+	EntriesLock.ReadLock();
+}
+
+void FCookOnTheFlyPackageStoreBackend::EndRead()
+{
+	TArray<FPackageId> PackageIds = MoveTemp(RequestedPackageIds);
+	EntriesLock.ReadUnlock();
+
+	if (PackageIds.Num() > 0)
+	{
+		SendCookRequest(MoveTemp(PackageIds));
+	}
+}
+
 bool FCookOnTheFlyPackageStoreBackend::DoesPackageExist(FPackageId PackageId)
 {
-	FScopeLock _(&CriticalSection);
+	FReadScopeLock ReadLock(EntriesLock);
 	FEntryInfo EntryInfo = PackageIdToEntryInfo.FindRef(PackageId);
 	return EntryInfo.Status != EPackageStoreEntryStatus::Missing;
 }
 
 EPackageStoreEntryStatus FCookOnTheFlyPackageStoreBackend::GetPackageStoreEntry(FPackageId PackageId, FPackageStoreEntry& OutPackageStoreEntry)
 {
-	using namespace UE::Cook;
-	using namespace UE::ZenCookOnTheFly::Messaging;
+	FEntryInfo& EntryInfo = PackageIdToEntryInfo.FindOrAdd(PackageId, FEntryInfo());
 
-	EPackageStoreEntryStatus OriginalStatus;
+	if (EntryInfo.Status == EPackageStoreEntryStatus::Ok)
+	{
+		check(EntryInfo.EntryIndex != INDEX_NONE);
+		return CreatePackageStoreEntry(EntryInfo, OutPackageStoreEntry);
+	}
+	else if (EntryInfo.Status == EPackageStoreEntryStatus::None)
 	{
-		FScopeLock _(&CriticalSection);
-		FEntryInfo& EntryInfo = PackageIdToEntryInfo.FindOrAdd(PackageId, FEntryInfo());
-
-		if (EntryInfo.Status == EPackageStoreEntryStatus::Ok)
-		{
-			check(EntryInfo.EntryIndex != INDEX_NONE);
-			return CreatePackageStoreEntry(EntryInfo, OutPackageStoreEntry);
-		}
-			
-		if (EntryInfo.Status == EPackageStoreEntryStatus::Missing)
-		{
-			return EPackageStoreEntryStatus::Missing;
-		}
-
-		if (EntryInfo.Status == EPackageStoreEntryStatus::Pending)
-		{
-			CheckActivity();
-			return EPackageStoreEntryStatus::Pending;
-		}
-
-		// The package hasn't been requested, set the status to pending.
 		EntryInfo.Status = EPackageStoreEntryStatus::Pending;
-
-		OriginalStatus = EntryInfo.Status;
+		RequestedPackageIds.Add(PackageId);
 	}
 
-	LastClientActivtyTime = FPlatformTime::Seconds();
-	UE_LOG(LogCookOnTheFly, Verbose, TEXT("Requesting package 0x%llX"), PackageId.ValueForDebugging());
+	return EntryInfo.Status; 
+}
 
+void FCookOnTheFlyPackageStoreBackend::SendCookRequest(TArray<FPackageId> PackageIds)
+{
+	using namespace UE::Cook;
+	using namespace UE::ZenCookOnTheFly::Messaging;
+
+	LastClientActivtyTime = FPlatformTime::Seconds();
+	
 	FCookOnTheFlyRequest Request(ECookOnTheFlyMessage::CookPackage);
-	Request.SetBodyTo(FCookPackageRequest { PackageId });
+	Request.SetBodyTo(FCookPackageRequest { PackageIds });
+
 	FCookOnTheFlyResponse Response = CookOnTheFlyServerConnection.SendRequest(Request).Get();
+	FCookPackageResponse CookResponse = Response.GetBodyAs<FCookPackageResponse>();
 
-	if (!Response.IsOk())
+	if (Response.IsOk())
 	{
-		UE_LOG(LogCookOnTheFly, Warning, TEXT("Failed to send 'CookPackage' request"));
-		FScopeLock _(&CriticalSection);
-		FEntryInfo& EntryInfo = PackageIdToEntryInfo.FindChecked(PackageId);
-		EntryInfo.Status = EPackageStoreEntryStatus::Missing;
-		return EPackageStoreEntryStatus::Missing;
+		AddPackages(MoveTemp(CookResponse.CookedPackages), MoveTemp(CookResponse.FailedPackages));
 	}
-
-	FCookPackageResponse CookPackageResponse = Response.GetBodyAs<FCookPackageResponse>();
+	else
 	{
-		FScopeLock _(&CriticalSection);
+		UE_LOG(LogCookOnTheFly, Warning, TEXT("Failed to send cook package(s) request"));
 
-		FEntryInfo& EntryInfo = PackageIdToEntryInfo.FindChecked(PackageId);
+		FWriteScopeLock WriteLock(EntriesLock);
 
-		// We might have received a PackagesCooked message while processing the request and in that case the response could be outdated so ignore it
-		if (EntryInfo.Status == OriginalStatus)
+		for (const FPackageId& PackageId : PackageIds)
 		{
-			EntryInfo.Status = CookPackageResponse.Status;
-			if (EntryInfo.Status == EPackageStoreEntryStatus::Ok)
+			FEntryInfo& EntryInfo = PackageIdToEntryInfo.FindChecked(PackageId);
+			if (EntryInfo.Status == EPackageStoreEntryStatus::Pending)
 			{
-				if (EntryInfo.EntryIndex == INDEX_NONE)
-				{
-					EntryInfo.EntryIndex = PackageEntries.Add();
-				}
-				PackageEntries[EntryInfo.EntryIndex] = MoveTemp(CookPackageResponse.CookedEntry);
+				EntryInfo.Status = EPackageStoreEntryStatus::Missing;
 			}
 		}
-		return CreatePackageStoreEntry(EntryInfo, OutPackageStoreEntry);
 	}
 }
 
@@ -143,7 +139,7 @@ EPackageStoreEntryStatus FCookOnTheFlyPackageStoreBackend::CreatePackageStoreEnt
 
 void FCookOnTheFlyPackageStoreBackend::AddPackages(TArray<FPackageStoreEntryResource> Entries, TArray<FPackageId> FailedPackageIds)
 {
-	FScopeLock _(&CriticalSection);
+	FWriteScopeLock WriteLock(EntriesLock);
 		
 	for (FPackageId FailedPackageId : FailedPackageIds)
 	{
@@ -153,8 +149,6 @@ void FCookOnTheFlyPackageStoreBackend::AddPackages(TArray<FPackageStoreEntryReso
 		PackageStats.Failed++;
 	}
 
-	const int32 NumPackageEntries = PackageEntries.Num();
-
 	for (FPackageStoreEntryResource& Entry : Entries)
 	{
 		const FPackageId PackageId = Entry.GetPackageId();
@@ -220,7 +214,7 @@ void FCookOnTheFlyPackageStoreBackend::CheckActivity()
 
 		UE_LOG(LogCookOnTheFly, Log, TEXT("=== Pending Packages ==="));
 		{
-			FScopeLock _(&CriticalSection);
+			FReadScopeLock ReadLock(EntriesLock);
 			for (const auto& KeyValue : PackageIdToEntryInfo)
 			{
 				const FEntryInfo& EntryInfo = KeyValue.Value;
diff --git a/Engine/Source/Runtime/StorageServerClient/Private/CookOnTheFlyPackageStore.h b/Engine/Source/Runtime/StorageServerClient/Private/CookOnTheFlyPackageStore.h
--- a/Engine/Source/Runtime/StorageServerClient/Private/CookOnTheFlyPackageStore.h
+++ b/Engine/Source/Runtime/StorageServerClient/Private/CookOnTheFlyPackageStore.h
@@ -8,6 +8,7 @@
 
 #include "IO/PackageStore.h"
 #include "CookOnTheFly.h"
+#include "Misc/ScopeRWLock.h"
 
 class FCookOnTheFlyPackageStoreBackend final
 	: public IPackageStoreBackend
@@ -32,13 +33,9 @@ public:
 		Context = InContext;
 	}
 
-	virtual void BeginRead() override
-	{
-	}
+	virtual void BeginRead() override;
 
-	virtual void EndRead() override
-	{
-	}
+	virtual void EndRead() override;
 
 	bool DoesPackageExist(FPackageId PackageId);
 	virtual EPackageStoreEntryStatus GetPackageStoreEntry(FPackageId PackageId, FPackageStoreEntry& OutPackageStoreEntry) override;
@@ -49,6 +46,7 @@ public:
 	}
 	
 private:
+	void SendCookRequest(TArray<FPackageId> PackageIds);
 	EPackageStoreEntryStatus CreatePackageStoreEntry(const FEntryInfo& EntryInfo, FPackageStoreEntry& OutPackageStoreEntry);
 	void AddPackages(TArray<FPackageStoreEntryResource> Entries, TArray<FPackageId> FailedPackageIds);
 	void OnCookOnTheFlyMessage(const UE::Cook::FCookOnTheFlyMessage& Message);
@@ -56,9 +54,10 @@ private:
 
 	UE::Cook::ICookOnTheFlyServerConnection& CookOnTheFlyServerConnection;
 	TSharedPtr<const FPackageStoreBackendContext> Context;
-	FCriticalSection CriticalSection;
+	FRWLock EntriesLock;
 	TMap<FPackageId, FEntryInfo> PackageIdToEntryInfo;
 	TChunkedArray<FPackageStoreEntryResource> PackageEntries;
+	TArray<FPackageId> RequestedPackageIds;
 	FPackageStats PackageStats;
 
 	const double MaxInactivityTime = 20;
