LevelInstance: Fix Break operation losing SoftObjectPaths

- Use CopyPasteId + ActorSourceName to find Source Actors this allows filling ExistingToNewMap which is used to remap properties

diff --git a/Engine/Source/Editor/UnrealEd/Private/EditorLevelUtils.cpp b/Engine/Source/Editor/UnrealEd/Private/EditorLevelUtils.cpp
--- a/Engine/Source/Editor/UnrealEd/Private/EditorLevelUtils.cpp
+++ b/Engine/Source/Editor/UnrealEd/Private/EditorLevelUtils.cpp
@@ -226,26 +226,30 @@ int32 UEditorLevelUtils::CopyOrMoveActorsToLevel(const TArray<AActor*>& ActorsTo
 					Actor->CopyPasteId = INDEX_NONE;
 				}
 
-				FAssetToolsModule& AssetToolsModule = FModuleManager::GetModuleChecked<FAssetToolsModule>(TEXT("AssetTools"));
-				TArray<FAssetRenameData> RenameData;
-
-				for (TTuple<FSoftObjectPath, FSoftObjectPath>& Pair : ActorPathMapping)
+				// Only do Asset Rename on Move (Copy should not affect existing references)
+				if (bMoveActors)
 				{
-					if (Pair.Value.IsValid())
-					{
-						RenameData.Add(FAssetRenameData(Pair.Key, Pair.Value, true));
-					}
-				}
+					FAssetToolsModule& AssetToolsModule = FModuleManager::GetModuleChecked<FAssetToolsModule>(TEXT("AssetTools"));
+					TArray<FAssetRenameData> RenameData;
 
-				if (RenameData.Num() > 0)
-				{
-					if (bWarnAboutRenaming)
+					for (TTuple<FSoftObjectPath, FSoftObjectPath>& Pair : ActorPathMapping)
 					{
-						AssetToolsModule.Get().RenameAssetsWithDialog(RenameData);
+						if (Pair.Value.IsValid())
+						{
+							RenameData.Add(FAssetRenameData(Pair.Key, Pair.Value, true));
+						}
 					}
-					else
+
+					if (RenameData.Num() > 0)
 					{
-						AssetToolsModule.Get().RenameAssets(RenameData);
+						if (bWarnAboutRenaming)
+						{
+							AssetToolsModule.Get().RenameAssetsWithDialog(RenameData);
+						}
+						else
+						{
+							AssetToolsModule.Get().RenameAssets(RenameData);
+						}
 					}
 				}
 
diff --git a/Engine/Source/Editor/UnrealEd/Private/Factories/EditorFactories.cpp b/Engine/Source/Editor/UnrealEd/Private/Factories/EditorFactories.cpp
--- a/Engine/Source/Editor/UnrealEd/Private/Factories/EditorFactories.cpp
+++ b/Engine/Source/Editor/UnrealEd/Private/Factories/EditorFactories.cpp
@@ -848,13 +848,7 @@ UObject* ULevelFactory::FactoryCreateText
 				FName ActorSourceName(NAME_None);
 				FParse::Value( Str, TEXT("NAME="), ActorSourceName );
 				ActorUniqueName = ActorSourceName;
-
-				AActor* Found=nullptr;
-				if( ActorUniqueName!=NAME_None )
-				{
-					// look in the current level for the same named actor
-					Found = FindObject<AActor>( World->GetCurrentLevel(), *ActorUniqueName.ToString() );
-				}
+				AActor* Found = nullptr;
 
 				// Make sure this name is unique. We need to do this upfront because we also want to potentially create the Associated BP class using the same name.
 				bool bNeedGloballyUniqueName = World->GetCurrentLevel()->IsUsingExternalActors() && CastChecked<AActor>(TempClass->GetDefaultObject())->SupportsExternalPackaging();
@@ -989,6 +983,30 @@ UObject* ULevelFactory::FactoryCreateText
 								{
 									check(CopyPasteId != INDEX_NONE);
 									NewActor->CopyPasteId = CopyPasteId;
+
+									// Use ActorSourceName + CopyPasteId to find matching source actor
+									if (ActorSourceName != NAME_None)
+									{
+										for (ULevel* Level : World->GetLevels())
+										{
+											if (Level && IsValid(Level))
+											{
+												if (AActor* Potential = FindObject<AActor>(Level, *ActorSourceName.ToString()))
+												{
+													if (Potential->CopyPasteId == CopyPasteId)
+													{
+														Found = Potential;
+														break;
+													}
+												}
+											}
+										}
+									}
+								}
+								else if(ActorSourceName != NAME_None)
+								{
+									// look in the current level for the same named actor
+									Found = FindObject<AActor>(World->GetCurrentLevel(), *ActorSourceName.ToString());
 								}
 
 								// Does the group exist?
