WorldPartition - Minimap: Added "-DebugCapture" option which writes bitmaps of all captured tiles & the final minimap

diff --git a/Engine/Source/Editor/UnrealEd/Private/WorldPartition/WorldPartitionMiniMapBuilder.cpp b/Engine/Source/Editor/UnrealEd/Private/WorldPartition/WorldPartitionMiniMapBuilder.cpp
--- a/Engine/Source/Editor/UnrealEd/Private/WorldPartition/WorldPartitionMiniMapBuilder.cpp
+++ b/Engine/Source/Editor/UnrealEd/Private/WorldPartition/WorldPartitionMiniMapBuilder.cpp
@@ -10,6 +10,7 @@
 #include "Engine/Texture2D.h"
 #include "Engine/World.h"
 #include "EngineUtils.h"
+#include "Misc/FileHelper.h"
 #include "Factories/TextureFactory.h"
 #include "SourceControlHelpers.h"
 #include "UObject/SavePackage.h"
@@ -54,6 +55,9 @@ bool UWorldPartitionMiniMapBuilder::PreRun(UWorld* World, FPackageSourceControlH
 		return false;
 	}
 
+	// Dump bitmaps for debugging purpose
+	bDebugCapture = FParse::Param(FCommandLine::Get(), TEXT("DebugCapture"));
+
 	// World bounds to process
 	IterativeWorldBounds = WorldMiniMap->GetMiniMapWorldBounds();
 
@@ -149,6 +153,14 @@ bool UWorldPartitionMiniMapBuilder::RunInternal(UWorld* World, const FCellInfo&
 			FMemory::Memcpy(DstCopy, SrcCopy, CopyWidthBytes);
 		}
 
+		// Write tile bitmap for debugging purpose
+		if (bDebugCapture)
+		{
+			const FString DirectoryPath = FPaths::ConvertRelativePathToFull(FPaths::ProjectIntermediateDir() + TEXT("Minimap"));
+			FString MinimapDebugImagePath = DirectoryPath / World->GetName() + "_" + TextureName + ".bmp";
+			FFileHelper::CreateBitmap(*MinimapDebugImagePath, TileTexture->Source.GetSizeX(), TileTexture->Source.GetSizeY(), (FColor*)SrcDataPtr);
+		}
+
 		TileTexture->Source.UnlockMip(0);
 	}
 
@@ -167,6 +179,14 @@ bool UWorldPartitionMiniMapBuilder::PostRun(UWorld* World, FPackageSourceControl
 
 	// Finalize texture
 	{
+		// Write minimap bitmap for debugging purpose
+		if (bDebugCapture)
+		{
+			const FString DirectoryPath = FPaths::ConvertRelativePathToFull(FPaths::ProjectIntermediateDir() + TEXT("Minimap"));
+			FString MinimapDebugImagePath = DirectoryPath / World->GetName() + "-Minimap.bmp";
+			FFileHelper::CreateBitmap(*MinimapDebugImagePath, WorldMiniMap->MiniMapTexture->Source.GetSizeX(), WorldMiniMap->MiniMapTexture->Source.GetSizeY(), (FColor*)MiniMapSourcePtr);
+		}
+
 		WorldMiniMap->MiniMapTexture->Source.UnlockMip(0);
 		WorldMiniMap->MiniMapTexture->PowerOfTwoMode = ETexturePowerOfTwoSetting::PadToPowerOfTwo;	// Required for VTs
 		WorldMiniMap->MiniMapTexture->AdjustMinAlpha = 1.f;
diff --git a/Engine/Source/Editor/UnrealEd/Public/WorldPartition/WorldPartitionMiniMapBuilder.h b/Engine/Source/Editor/UnrealEd/Public/WorldPartition/WorldPartitionMiniMapBuilder.h
--- a/Engine/Source/Editor/UnrealEd/Public/WorldPartition/WorldPartitionMiniMapBuilder.h
+++ b/Engine/Source/Editor/UnrealEd/Public/WorldPartition/WorldPartitionMiniMapBuilder.h
@@ -36,4 +36,6 @@ private:
 	int32 MinimapImageSizeY;
 
 	uint8* MiniMapSourcePtr;
+
+	bool bDebugCapture;
 };
\ No newline at end of file
