Fixed issue where undo/redo would crash after a change to the layout of a structure or class saved in an undo transaction.

diff --git a/Engine/Source/Editor/UnrealEd/Private/EditorTransaction.cpp b/Engine/Source/Editor/UnrealEd/Private/EditorTransaction.cpp
--- a/Engine/Source/Editor/UnrealEd/Private/EditorTransaction.cpp
+++ b/Engine/Source/Editor/UnrealEd/Private/EditorTransaction.cpp
@@ -37,16 +37,21 @@ FTransaction::FObjectRecord::FObjectRecord(FTransaction* Owner, UObject* InObjec
 	,	Serializer			( InSerializer )
 	,	Destructor			( InDestructor )
 {
+#if WITH_RELOAD
+	// With Hot Reload or Live Coding enabled, the layout of structures and classes can change.  We can't use binary serialization when this might happen.
+	bWantsBinarySerialization = false;
+#else
 	// Blueprint compile-in-place can alter class layout so use tagged serialization for objects relying on a UBlueprint's Class
 	if (UBlueprintGeneratedClass* Class = Cast<UBlueprintGeneratedClass>(InObject->GetClass()))
 	{
-		bWantsBinarySerialization = false; 
+		bWantsBinarySerialization = false;
 	}
 	// Data tables can contain user structs, so it's unsafe to use binary
 	if (UDataTable* DataTable = Cast<UDataTable>(InObject))
 	{
 		bWantsBinarySerialization = false;
 	}
+#endif
 
 	// Update any sub-object caches for use by ARO (to keep sub-objects alive during GC)
 	UObject* CurrentObject = Object.Get();
