Better detection of actors to regenerate in UWorldPartitionLandscapeSplineMeshesBuilder.

diff --git a/Engine/Source/Editor/UnrealEd/Private/WorldPartition/WorldPartitionLandscapeSplineMeshesBuilder.cpp b/Engine/Source/Editor/UnrealEd/Private/WorldPartition/WorldPartitionLandscapeSplineMeshesBuilder.cpp
--- a/Engine/Source/Editor/UnrealEd/Private/WorldPartition/WorldPartitionLandscapeSplineMeshesBuilder.cpp
+++ b/Engine/Source/Editor/UnrealEd/Private/WorldPartition/WorldPartitionLandscapeSplineMeshesBuilder.cpp
@@ -144,7 +144,9 @@ int32 UWorldPartitionLandscapeSplineMeshesBuilder::HashStaticMeshComponent(const
 	// Hash the owning actor's tags
 	if (const AActor* Owner = InComponent->GetOwner())
 	{
-		for (const FName& ActorTag : Owner->Tags)
+		TArray<FName> SortedTags = Owner->Tags;
+		SortedTags.Sort(FNameLexicalLess());
+		for (const FName& ActorTag : SortedTags)
 		{
 			if (!ActorTag.IsEqual(CreatedFromBuilderTag))
 			{
