Add a 'finalization' step to UpdateBytecodeReferences, for object graphs with dense lists of dependent blueprints we would repeatedly serialize the same objects. For large object graphs the slowdown could be substantial.

diff --git a/Engine/Source/Editor/Kismet/Private/BlueprintCompilationManager.cpp b/Engine/Source/Editor/Kismet/Private/BlueprintCompilationManager.cpp
--- a/Engine/Source/Editor/Kismet/Private/BlueprintCompilationManager.cpp
+++ b/Engine/Source/Editor/Kismet/Private/BlueprintCompilationManager.cpp
@@ -2176,18 +2176,24 @@ void FBlueprintCompilationManagerImpl::ReinstanceBatch(TArray<FReinstancingJob>&
 	}
 
 	// run UpdateBytecodeReferences:
-	for (const FReinstancingJob& ReinstancingJob : Reinstancers)
 	{
-		if (ReinstancingJob.OldToNew.Key)
-		{
-			InOutOldToNewClassMap.Add(ReinstancingJob.OldToNew.Key, ReinstancingJob.OldToNew.Value);
-		}
-			
-		if(ReinstancingJob.Reinstancer.IsValid())
+		TSet<UBlueprint*> DependentBPs;
+		TMap<FFieldVariant, FFieldVariant> FieldMappings;
+		for (const FReinstancingJob& ReinstancingJob : Reinstancers)
 		{
-			UBlueprint* CompiledBlueprint = UBlueprint::GetBlueprintFromClass(ReinstancingJob.OldToNew.Value);
-			ReinstancingJob.Reinstancer->UpdateBytecodeReferences();
+			if (ReinstancingJob.OldToNew.Key)
+			{
+				InOutOldToNewClassMap.Add(ReinstancingJob.OldToNew.Key, ReinstancingJob.OldToNew.Value);
+			}
+
+			if (ReinstancingJob.Reinstancer.IsValid())
+			{
+				UBlueprint* CompiledBlueprint = UBlueprint::GetBlueprintFromClass(ReinstancingJob.OldToNew.Value);
+				ReinstancingJob.Reinstancer->UpdateBytecodeReferences(DependentBPs, FieldMappings);
+			}
 		}
+
+		FBlueprintCompileReinstancer::FinishUpdateBytecodeReferences(DependentBPs, FieldMappings);
 	}
 	
 	// Now we can update templates and archetypes - note that we don't look for direct references to archetypes - doing
diff --git a/Engine/Source/Editor/UnrealEd/Private/Kismet2/KismetReinstanceUtilities.cpp b/Engine/Source/Editor/UnrealEd/Private/Kismet2/KismetReinstanceUtilities.cpp
--- a/Engine/Source/Editor/UnrealEd/Private/Kismet2/KismetReinstanceUtilities.cpp
+++ b/Engine/Source/Editor/UnrealEd/Private/Kismet2/KismetReinstanceUtilities.cpp
@@ -39,12 +39,6 @@
 #include "Subsystems/AssetEditorSubsystem.h"
 #include "Engine/ScopedMovementUpdate.h"
 
-// Enabling this will validate cached dependent Blueprints against the full set of loaded Blueprints when updating bytecode references.
-// Note: Enabling this may potentially increase editor/Blueprint load time and/or decrease performance related to Blueprint compilation.
-#ifndef VALIDATE_BYTECODE_REFERENCE_DEPENDENCY_CACHE
-#define VALIDATE_BYTECODE_REFERENCE_DEPENDENCY_CACHE 0
-#endif // VALIDATE_BYTECODE_REFERENCE_DEPENDENCY_CACHE
-
 DECLARE_CYCLE_STAT(TEXT("Replace Instances"), EKismetReinstancerStats_ReplaceInstancesOfClass, STATGROUP_KismetReinstancer );
 DECLARE_CYCLE_STAT(TEXT("Find Referencers"), EKismetReinstancerStats_FindReferencers, STATGROUP_KismetReinstancer );
 DECLARE_CYCLE_STAT(TEXT("Replace References"), EKismetReinstancerStats_ReplaceReferences, STATGROUP_KismetReinstancer );
@@ -957,7 +951,9 @@ protected:
 	const TMap<FFieldVariant, FFieldVariant>& ReplacementMap;
 };
 
-void FBlueprintCompileReinstancer::UpdateBytecodeReferences()
+void FBlueprintCompileReinstancer::UpdateBytecodeReferences(
+	TSet<UBlueprint*>& OutDependentBlueprints,
+	TMap<FFieldVariant, FFieldVariant>& OutFieldMapping)
 {
 	BP_SCOPED_COMPILER_EVENT_STAT(EKismetReinstancerStats_UpdateBytecodeReferences);
 
@@ -970,6 +966,7 @@ void FBlueprintCompileReinstancer::UpdateBytecodeReferences()
 	{
 		TMap<FFieldVariant, FFieldVariant> FieldMappings;
 		GenerateFieldMappings(FieldMappings);
+		OutFieldMapping.Append(FieldMappings);
 
 		// Note: This API returns a cached set of blueprints that's updated at compile time.
 		TArray<UBlueprint*> CachedDependentBPs;
@@ -978,12 +975,7 @@ void FBlueprintCompileReinstancer::UpdateBytecodeReferences()
 		// Determine whether or not we will be updating references for an Animation Blueprint class.
 		const bool bIsAnimBlueprintClass = !!Cast<UAnimBlueprint>(ClassToReinstance->ClassGeneratedBy);
 
-#if VALIDATE_BYTECODE_REFERENCE_DEPENDENCY_CACHE
-		TArray<UBlueprint*> ActualDependentBPs;
-		for (TObjectIterator<UBlueprint> BpIt; BpIt; ++BpIt)
-#else
 		for (auto BpIt = CachedDependentBPs.CreateIterator(); BpIt; ++BpIt)
-#endif // VALIDATE_BYTECODE_REFERENCE_DEPENDENCY_CACHE
 		{
 			UBlueprint* DependentBP = *BpIt;
 			UClass* BPClass = DependentBP->GeneratedClass;
@@ -1008,13 +1000,11 @@ void FBlueprintCompileReinstancer::UpdateBytecodeReferences()
 				BPClass->StaticLink(true);
 			}
 
-			bool bBPWasChanged = false;
 			// For each function defined in this blueprint, run through the bytecode, and update any refs from the old properties to the new
 			for( TFieldIterator<UFunction> FuncIter(BPClass, EFieldIteratorFlags::ExcludeSuper); FuncIter; ++FuncIter )
 			{
 				UFunction* CurrentFunction = *FuncIter;
 				FArchiveReplaceFieldReferences ReplaceAr(CurrentFunction, FieldMappings);
-				bBPWasChanged |= (0 != ReplaceAr.GetCount());
 			}
 
 			// Update any refs in called functions array, as the bytecode was just similarly updated:
@@ -1035,38 +1025,25 @@ void FBlueprintCompileReinstancer::UpdateBytecodeReferences()
 				}
 			}
 
-			FArchiveReplaceFieldReferences ReplaceInBPAr(DependentBP, FieldMappings);
-			if (ReplaceInBPAr.GetCount())
-			{
-#if VALIDATE_BYTECODE_REFERENCE_DEPENDENCY_CACHE
-				ActualDependentBPs.Add(DependentBP);
-#endif // VALIDATE_BYTECODE_REFERENCE_DEPENDENCY_CACHE
-
-				bBPWasChanged = true;
-				UE_LOG(LogBlueprint, Log, TEXT("UpdateBytecodeReferences: %d references from %s was replaced in BP %s"), ReplaceInBPAr.GetCount(), *GetPathNameSafe(ClassToReinstance), *GetPathNameSafe(DependentBP));
-			}
+			OutDependentBlueprints.Add(DependentBP);
 		}
+	}
+}
 
-#if VALIDATE_BYTECODE_REFERENCE_DEPENDENCY_CACHE
-		bool bHasMissingDependents = false;
-		for (UBlueprint* ChangedBP : ActualDependentBPs)
-		{
-			if (!CachedDependentBPs.Contains(ChangedBP))
-			{
-				UE_LOG(LogBlueprint, Error, TEXT("While updating %s, we needed to update %s but it wasn't cached as a dependent"), *ClassToReinstance->GetName(), *ChangedBP->GetName());
-				bHasMissingDependents = true;
-			}
-		}
+void FBlueprintCompileReinstancer::FinishUpdateBytecodeReferences(
+	const TSet<UBlueprint*>& DependentBPs,
+	const TMap<FFieldVariant, FFieldVariant>& FieldMappings)
+{
+	for (UBlueprint* DependentBP : DependentBPs)
+	{
+		FArchiveReplaceFieldReferences ReplaceInBPAr(DependentBP, FieldMappings);
 
-		if (bHasMissingDependents)
+		if (ReplaceInBPAr.GetCount())
 		{
-			UE_LOG(LogBlueprint, Error, TEXT("Class: %s, CachedDeps: [%s], ActualDeps: [%s]"),
-				*ClassToReinstance->GetName(),
-				*FString::JoinBy(CachedDependentBPs, TEXT(","), [](UBlueprint* Blueprint) { return Blueprint->GetName(); }),
-				*FString::JoinBy(ActualDependentBPs, TEXT(","), [](UBlueprint* Blueprint) { return Blueprint->GetName(); })
-			);
+			UE_LOG(LogBlueprint, Log, 
+				TEXT("UpdateBytecodeReferences: %d references were replaced in BP %s"), 
+				ReplaceInBPAr.GetCount(), *GetPathNameSafe(DependentBP));
 		}
-#endif // VALIDATE_BYTECODE_REFERENCE_DEPENDENCY_CACHE
 	}
 }
 
diff --git a/Engine/Source/Editor/UnrealEd/Public/Kismet2/KismetReinstanceUtilities.h b/Engine/Source/Editor/UnrealEd/Public/Kismet2/KismetReinstanceUtilities.h
--- a/Engine/Source/Editor/UnrealEd/Public/Kismet2/KismetReinstanceUtilities.h
+++ b/Engine/Source/Editor/UnrealEd/Public/Kismet2/KismetReinstanceUtilities.h
@@ -166,7 +166,10 @@ public:
 	void ReinstanceObjects(bool bForceAlwaysReinstance = false);
 
 	/** Updates references to properties and functions of the class that has in the bytecode of dependent blueprints */
-	void UpdateBytecodeReferences();
+	void UpdateBytecodeReferences( TSet<UBlueprint*>& OutDependentBlueprints, TMap<FFieldVariant, FFieldVariant>& OutFieldMapping);
+
+	/** Consumes the set and map populated by calls to UpdateBytecodeReferences */
+	static void FinishUpdateBytecodeReferences( const TSet<UBlueprint*>& DependentBPs, const TMap<FFieldVariant, FFieldVariant>& FieldMappings);
 
 	/** Worker function to replace all instances of OldClass with a new instance of NewClass */
 	static void ReplaceInstancesOfClass(UClass* OldClass, UClass* NewClass, UObject* OriginalCDO = nullptr, TSet<UObject*>* ObjectsThatShouldUseOldStuff = nullptr, bool bClassObjectReplaced = false, bool bPreserveRootComponent = true);
