Do not populate children in the debug tree if the referenced object was already present in the tree

diff --git a/Engine/Source/Editor/Kismet/Private/Debugging/SKismetDebugTreeView.cpp b/Engine/Source/Editor/Kismet/Private/Debugging/SKismetDebugTreeView.cpp
--- a/Engine/Source/Editor/Kismet/Private/Debugging/SKismetDebugTreeView.cpp
+++ b/Engine/Source/Editor/Kismet/Private/Debugging/SKismetDebugTreeView.cpp
@@ -1203,7 +1203,7 @@ protected:
 
 	virtual void GatherChildren(TArray<FDebugTreeItemPtr>& OutChildren, const FString& InSearchString, bool bRespectSearch) override
 	{
-		for (const TSharedPtr<FPropertyInstanceInfo>& ChildData : Data->Children)
+		for (const TSharedPtr<FPropertyInstanceInfo>& ChildData : Data->GetChildren())
 		{
 			EnsureChildIsAdded(OutChildren, FWatchChildLineItem(ChildData.ToSharedRef(), AsShared()), InSearchString, bRespectSearch);
 		}
@@ -1419,7 +1419,7 @@ protected:
 		TSharedPtr<FPropertyInstanceInfo> ThisDebugInfo = GetPropertyInfo();
 		if (ThisDebugInfo.IsValid())
 		{
-			for (const TSharedPtr<FPropertyInstanceInfo>& ChildData : ThisDebugInfo->Children)
+			for (const TSharedPtr<FPropertyInstanceInfo>& ChildData : ThisDebugInfo->GetChildren())
 			{
 				EnsureChildIsAdded(OutChildren, FWatchChildLineItem(ChildData.ToSharedRef(), AsShared()), InSearchString, bRespectSearch);
 			}
diff --git a/Engine/Source/Editor/Kismet/Private/WatchPointViewer.cpp b/Engine/Source/Editor/Kismet/Private/WatchPointViewer.cpp
--- a/Engine/Source/Editor/Kismet/Private/WatchPointViewer.cpp
+++ b/Engine/Source/Editor/Kismet/Private/WatchPointViewer.cpp
@@ -96,7 +96,7 @@ namespace
 			UPackage* Package = Cast<UPackage>(BP.IsValid() ? BP->GetOuter() : nullptr);
 			BlueprintPackageName = Package ? Package->GetFName() : FName();
 
-			for (const TSharedPtr<FPropertyInstanceInfo>& ChildInfo : Info.Children)
+			for (const TSharedPtr<FPropertyInstanceInfo>& ChildInfo : Info.GetChildren())
 			{
 				Children.Add(MakeShared<FWatchRow>(InBP, InNode, InPin, InObjectBeingDebugged, BlueprintName, GraphName, NodeName, MoveTemp(*ChildInfo)));
 			}
diff --git a/Engine/Source/Editor/UnrealEd/Private/Kismet2/KismetDebugUtilities.cpp b/Engine/Source/Editor/UnrealEd/Private/Kismet2/KismetDebugUtilities.cpp
--- a/Engine/Source/Editor/UnrealEd/Private/Kismet2/KismetDebugUtilities.cpp
+++ b/Engine/Source/Editor/UnrealEd/Private/Kismet2/KismetDebugUtilities.cpp
@@ -2137,15 +2137,35 @@ TSharedPtr<FPropertyInstanceInfo> FPropertyInstanceInfo::FindOrMake(FPropertyIns
 	{
 		return *Found;
 	}
-	
+
 	// wait to populate children until after inserting into VisitedNodes.
 	// this way we will catch circular references
-	TSharedPtr<FPropertyInstanceInfo> DebugInfo = VisitedNodes.Add(
+	TSharedPtr<FPropertyInstanceInfo> DebugInfo = MakeShared<FPropertyInstanceInfo>(PropertyInstance);
+
+	bool bFoundObject = false;
+	if (DebugInfo->Object.IsValid())
+	{
+		for (const TPair<FPropertyInstance, TSharedPtr<FPropertyInstanceInfo>>& VisitedNode : VisitedNodes)
+		{
+			if (VisitedNode.Value->Object == DebugInfo->Object)
+			{
+				bFoundObject = true;
+				DebugInfo->ReferencedObject = VisitedNode.Value;
+				break;
+			}
+		}
+	}
+	
+	VisitedNodes.Add(
 		PropertyInstance,
-		MakeShared<FPropertyInstanceInfo>(PropertyInstance)
+		DebugInfo
 	);
 
-	DebugInfo->PopulateChildren(PropertyInstance, VisitedNodes);
+	if (!bFoundObject)
+	{
+		DebugInfo->PopulateChildren(PropertyInstance, VisitedNodes);
+	}
+
 	return DebugInfo;
 }
 
@@ -2343,11 +2363,24 @@ FString FPropertyInstanceInfo::GetWatchText() const
 
 			return WatchText;
 		}
+
+
 	}
 
 	return Value.ToString();
 }
 
+const TArray<TSharedPtr<FPropertyInstanceInfo>>& FPropertyInstanceInfo::GetChildren() const
+{
+	// If this node represents an object that was already added elsewhere in the tree, return its children instead
+	if (ReferencedObject.IsValid())
+	{
+		return ReferencedObject->GetChildren();
+	}
+
+	return Children;
+}
+
 FText FKismetDebugUtilities::GetAndClearLastExceptionMessage()
 {
 	FKismetDebugUtilitiesData& Data = FKismetDebugUtilitiesData::Get();
diff --git a/Engine/Source/Editor/UnrealEd/Public/Kismet2/KismetDebugUtilities.h b/Engine/Source/Editor/UnrealEd/Public/Kismet2/KismetDebugUtilities.h
--- a/Engine/Source/Editor/UnrealEd/Public/Kismet2/KismetDebugUtilities.h
+++ b/Engine/Source/Editor/UnrealEd/Public/Kismet2/KismetDebugUtilities.h
@@ -84,14 +84,25 @@ struct FPropertyInstanceInfo
 	/** Returns the watch text for info popup bubbles on the graph */
 	UNREALED_API FString GetWatchText() const;
 	
+	UNREALED_API const TArray<TSharedPtr<FPropertyInstanceInfo>>& GetChildren() const;
+
 	FText Name;
 	FText DisplayName;
 	FText Value;
 	FText Type;
 	TWeakObjectPtr<UObject> Object = nullptr; // only filled if property is a UObject
 	TFieldPath<const FProperty> Property;
-	TArray<TSharedPtr<FPropertyInstanceInfo>> Children;
 	bool bIsInContainer = false;
+
+private:
+	TArray<TSharedPtr<FPropertyInstanceInfo>> Children;
+
+	/** 
+	 * Only populated when this is an object and was already seen in the debug tree
+	 * used to avoid circular references in cases where the same object may be referenced by multiple properties  
+	 */
+	TSharedPtr<FPropertyInstanceInfo> ReferencedObject; 
+
 };
 
 inline uint32 GetTypeHash(const FPropertyInstanceInfo::FPropertyInstance& PropertyInstance)
