fix EditorFactories incorrectly logging DDS load error when trying to load tiny non-DDS images
also log when the G8 auto-detect conversion is done

diff --git a/Engine/Source/Editor/UnrealEd/Private/Factories/EditorFactories.cpp b/Engine/Source/Editor/UnrealEd/Private/Factories/EditorFactories.cpp
--- a/Engine/Source/Editor/UnrealEd/Private/Factories/EditorFactories.cpp
+++ b/Engine/Source/Editor/UnrealEd/Private/Factories/EditorFactories.cpp
@@ -2935,7 +2935,10 @@ bool UTextureFactory::ImportImage(const uint8* Buffer, int64 Length, FFeedbackCo
 				return false;
 			}
 		
-			UE::TextureUtilitiesCommon::AutoDetectAndChangeGrayScale(LoadedImage);
+			if ( UE::TextureUtilitiesCommon::AutoDetectAndChangeGrayScale(LoadedImage) )
+			{
+				UE_LOG(LogEditorFactories,Display,TEXT("Auto-detected grayscale, image changed to G8"));
+			}
 
 			ETextureSourceFormat TextureFormat = FImageCoreUtils::ConvertToTextureSourceFormat(LoadedImage.Format);
 			bool bSRGB = LoadedImage.GammaSpace != EGammaSpace::Linear;
@@ -3320,8 +3323,9 @@ UTexture * UTextureFactory::ImportDDS(const uint8* Buffer,int64 Length,UObject*
 	if ( DDS == nullptr )
 	{
 		// NotADds is okay/expected , we try this on all image buffers
+		// IoError means buffer is too small to be a DDS
 		check( Error != UE::DDS::EDDSError::OK ); 
-		if ( Error != UE::DDS::EDDSError::NotADds )
+		if ( Error != UE::DDS::EDDSError::NotADds && Error != UE::DDS::EDDSError::IoError )
 		{
 			UE_LOG(LogEditorFactories,Warning,TEXT("Failed to load DDS (Error=%d)"),(int)Error);			
 		}
@@ -4223,6 +4227,8 @@ UObject* UTextureFactory::FactoryCreateBinary
 	{
 		if (UE::NormalMapIdentification::HandleAssetPostImport(Texture))
 		{
+			UE_LOG(LogEditorFactories,Display,TEXT("Auto-detected normal map"));
+
 			Texture->bFlipGreenChannel = bFlipNormalMapGreenChannel;
 		}
 	}
