Fix mismerge in CL 21790364

#rnx

- The "GeometryFramework" entry in ModuleNames needed a , after it to avoid cutting the array short, this is causing the VirtualizationEditor module not to be loaded.

diff --git a/Engine/Source/Editor/UnrealEd/Private/EditorEngine.cpp b/Engine/Source/Editor/UnrealEd/Private/EditorEngine.cpp
--- a/Engine/Source/Editor/UnrealEd/Private/EditorEngine.cpp
+++ b/Engine/Source/Editor/UnrealEd/Private/EditorEngine.cpp
@@ -1073,7 +1073,7 @@ void UEditorEngine::Init(IEngineLoop* InEngineLoop)
 			TEXT("EditorConfig"),
 			TEXT("DerivedDataEditor"),
 			TEXT("CSVtoSVG"),
-			TEXT("GeometryFramework")
+			TEXT("GeometryFramework"),
 			TEXT("VirtualizationEditor"),
 		};
 
