FNames containing asset paths are deprecated. FSoftObjectPath or FTopLevelAssetPath should be used instead.
Fixups for UnrealEd Fbx code.

diff --git a/Engine/Source/Editor/UnrealEd/Private/Fbx/FbxFactory.cpp b/Engine/Source/Editor/UnrealEd/Private/Fbx/FbxFactory.cpp
--- a/Engine/Source/Editor/UnrealEd/Private/Fbx/FbxFactory.cpp
+++ b/Engine/Source/Editor/UnrealEd/Private/Fbx/FbxFactory.cpp
@@ -889,7 +889,7 @@ void UFbxFactory::CancelObjectCreation(UnFbx::FFbxImporter* FbxImporter) const
 		const TWeakObjectPtr<UObject>& CurrentObject = CreatedObjects[CreatedObjectIndex];
 		if (CurrentObject.IsValid())
 		{
-			AssetsToDelete.Emplace(AssetRegistryModule.Get().GetAssetByObjectPath(FName(*CurrentObject->GetPathName())));
+			AssetsToDelete.Emplace(AssetRegistryModule.Get().GetAssetByObjectPath(FSoftObjectPath(CurrentObject.Get())));
 			PotentialPackageToDeleteReferences.Add(CurrentObject->GetOutermost());
 			CurrentObject->ClearFlags(RF_Standalone | RF_Public);
 			CurrentObject->RemoveFromRoot();
diff --git a/Engine/Source/Editor/UnrealEd/Private/Fbx/FbxReimportDialogs.cpp b/Engine/Source/Editor/UnrealEd/Private/Fbx/FbxReimportDialogs.cpp
--- a/Engine/Source/Editor/UnrealEd/Private/Fbx/FbxReimportDialogs.cpp
+++ b/Engine/Source/Editor/UnrealEd/Private/Fbx/FbxReimportDialogs.cpp
@@ -147,12 +147,12 @@ void FFbxImporter::ShowFbxSkeletonConflictWindow(USkeletalMesh* SkeletalMesh, US
 				for (const FName& AssetDependencyName : AllDependencies)
 				{
 					const FString PackageString = AssetDependencyName.ToString();
-					const FString FullAssetPathName = FString::Printf(TEXT("%s.%s"), *PackageString, *FPackageName::GetLongPackageAssetName(PackageString));
+					const FSoftObjectPath FullAssetPath(*PackageString, *FPackageName::GetLongPackageAssetName(PackageString), {});
 					
-					FAssetData AssetData = AssetRegistryModule.Get().GetAssetByObjectPath(FName(*FullAssetPathName));
+					FAssetData AssetData = AssetRegistryModule.Get().GetAssetByObjectPath(FullAssetPath);
 					if (AssetData.GetClass() != nullptr)
 					{
-						TSharedPtr<FString> AssetReferencing = MakeShareable(new FString(AssetData.AssetClassPath.ToString() + TEXT(" ") + FullAssetPathName));
+						TSharedPtr<FString> AssetReferencing = MakeShareable(new FString(AssetData.AssetClassPath.ToString() + TEXT(" ") + FullAssetPath.ToString()));
 						AssetReferencingSkeleton.Add(AssetReferencing);
 					}
 				}
diff --git a/Engine/Source/Editor/UnrealEd/Private/Fbx/ReimportFbxSceneFactory.cpp b/Engine/Source/Editor/UnrealEd/Private/Fbx/ReimportFbxSceneFactory.cpp
--- a/Engine/Source/Editor/UnrealEd/Private/Fbx/ReimportFbxSceneFactory.cpp
+++ b/Engine/Source/Editor/UnrealEd/Private/Fbx/ReimportFbxSceneFactory.cpp
@@ -553,7 +553,7 @@ EReimportResult::Type UReimportFbxSceneFactory::Reimport(UObject* Obj)
 			continue;
 		}
 		//Find the asset
-		AssetDataToDelete.Add(AssetRegistryModule.Get().GetAssetByObjectPath(FName(*(MeshInfo->GetFullImportName()))));
+		AssetDataToDelete.Add(AssetRegistryModule.Get().GetAssetByObjectPath(FSoftObjectPath((MeshInfo->GetFullImportName()))));
 	}
 
 	FbxNode* RootNodeToImport = FbxImporter->Scene->GetRootNode();
@@ -940,7 +940,7 @@ UBlueprint *UReimportFbxSceneFactory::UpdateOriginalBluePrint(FString &BluePrint
 	FbxSceneReimportStatusMapPtr NodeStatusMapPtr = (FbxSceneReimportStatusMapPtr)VoidNodeStatusMapPtr;
 	//Find the BluePrint
 	FAssetRegistryModule& AssetRegistryModule = FModuleManager::LoadModuleChecked<FAssetRegistryModule>("AssetRegistry");
-	FAssetData BlueprintAssetData = AssetRegistryModule.Get().GetAssetByObjectPath(FName(*(BluePrintFullName)));
+	FAssetData BlueprintAssetData = AssetRegistryModule.Get().GetAssetByObjectPath(FSoftObjectPath(BluePrintFullName));
 
 	UPackage* PkgExist = FindPackage(nullptr, *BlueprintAssetData.PackageName.ToString());
 	if (PkgExist == nullptr)
diff --git a/Engine/Source/Editor/UnrealEd/Private/Fbx/SSceneMaterialsListView.cpp b/Engine/Source/Editor/UnrealEd/Private/Fbx/SSceneMaterialsListView.cpp
--- a/Engine/Source/Editor/UnrealEd/Private/Fbx/SSceneMaterialsListView.cpp
+++ b/Engine/Source/Editor/UnrealEd/Private/Fbx/SSceneMaterialsListView.cpp
@@ -419,7 +419,7 @@ void SFbxSceneMaterialsListView::AssignMaterialAssetData(const FAssetData& Asset
 				//Override the MeshInfo with the new asset path
 				ItemPtr->SetOverridePath(true);
 				ItemPtr->OverrideImportPath = AssetData.PackageName.ToString();
-				ItemPtr->OverrideFullImportName = AssetData.ObjectPath.ToString();
+				ItemPtr->OverrideFullImportName = AssetData.GetObjectPathString();
 			}
 		}
 	}
@@ -453,7 +453,7 @@ void SFbxSceneMaterialsListView::AssignMaterialToExisting()
 					//Override the MeshInfo with the new asset path
 					ItemPtr->SetOverridePath(true);
 					ItemPtr->OverrideImportPath = AssetData[0].PackageName.ToString();
-					ItemPtr->OverrideFullImportName = AssetData[0].ObjectPath.ToString();
+					ItemPtr->OverrideFullImportName = AssetData[0].GetObjectPathString();
 				}
 			}
 		}
diff --git a/Engine/Source/Editor/UnrealEd/Private/Fbx/SSceneReimportStaticMeshListView.cpp b/Engine/Source/Editor/UnrealEd/Private/Fbx/SSceneReimportStaticMeshListView.cpp
--- a/Engine/Source/Editor/UnrealEd/Private/Fbx/SSceneReimportStaticMeshListView.cpp
+++ b/Engine/Source/Editor/UnrealEd/Private/Fbx/SSceneReimportStaticMeshListView.cpp
@@ -522,7 +522,7 @@ void SFbxSceneStaticMeshReimportListView::AssignToStaticMesh()
 			//Override the MeshInfo with the new asset path
 			ItemPtr->SetOverridePath(true);
 			ItemPtr->OverrideImportPath = AssetData[0].PackageName.ToString();
-			ItemPtr->OverrideFullImportName = AssetData[0].ObjectPath.ToString();
+			ItemPtr->OverrideFullImportName = AssetData[0].GetObjectPathString();
 			StatusFlag |= EFbxSceneReimportStatusFlags::FoundContentBrowserAsset;
 			
 			//Add the new reimport status
