UnrealEd: ensure cached file path exists and is mounted before using CommandLineMapCache cache entries

This ensures that we don't use a stale cache entry to an
unmounted file in case a previously opened map was copied
to a different location on disk and is now being mounted from
there, but the original file still exists.

diff --git a/Engine/Source/Editor/UnrealEd/Private/UnrealEdMisc.cpp b/Engine/Source/Editor/UnrealEd/Private/UnrealEdMisc.cpp
--- a/Engine/Source/Editor/UnrealEd/Private/UnrealEdMisc.cpp
+++ b/Engine/Source/Editor/UnrealEd/Private/UnrealEdMisc.cpp
@@ -562,7 +562,15 @@ FString FUnrealEdMisc::FindMapFileFromPartialName(const FString& PartialMapName)
 		
 			if (ShortMapName == PartialMapName)
 			{
-				if (FPaths::FileExists(RelativeFileName))
+				// Verify both that the file exists and that it's actually
+				// mounted. This ensures that we don't use a stale cache
+				// entry to an unmounted file in case a previously opened
+				// map was copied to a different location on disk and is
+				// now being mounted from there, but the original file still
+				// exists.
+				FString PackageName;
+				if (FPaths::FileExists(RelativeFileName) &&
+					FPackageName::TryConvertFilenameToLongPackageName(RelativeFileName, PackageName))
 				{
 					FoundMapFile = RelativeFileName;
 					break;
