VR Editor Mode: Screen Percentage support

Right now, FEditorViewportClient disables ScreenPercentage by default.
However, the EditorViewport DOES support the Preview Screen Percentage tool, unless the viewport is rendering in stereo.

Now, if VREM is enabled, we also support the Preview Screen Percentage. We could have enabled screen percentage explicitly for VREM editor viewports (via SetVREditView + FEngineShowFlags).
But that still wouldn't enable the Preview Screen Percentage. So this seems like the appropriate change, though the interactions between viewports, editor, screen percentage and VREM appear obtuse.

diff --git a/Engine/Source/Editor/UnrealEd/Private/EditorViewportClient.cpp b/Engine/Source/Editor/UnrealEd/Private/EditorViewportClient.cpp
--- a/Engine/Source/Editor/UnrealEd/Private/EditorViewportClient.cpp
+++ b/Engine/Source/Editor/UnrealEd/Private/EditorViewportClient.cpp
@@ -4035,8 +4035,9 @@ void FEditorViewportClient::Draw(FViewport* InViewport, FCanvas* Canvas)
 		{
 			float GlobalResolutionFraction = 1.0f;
 
-			// If not doing VR rendering, apply preview resolution fraction.
-			if (!bStereoRendering && SupportsPreviewResolutionFraction() && ViewFamily.SupportsScreenPercentage())
+			// Apply preview resolution fraction. Supported in stereo for VR Editor Mode only
+			if ((!bStereoRendering || bInVREditViewMode) && 
+				SupportsPreviewResolutionFraction() && ViewFamily.SupportsScreenPercentage())
 			{
 				if (PreviewResolutionFraction.IsSet())
 				{
