WorldPartitionBuilder: Allow skipping cells in Iterative mode

diff --git a/Engine/Source/Editor/UnrealEd/Private/WorldPartition/WorldPartitionBuilder.cpp b/Engine/Source/Editor/UnrealEd/Private/WorldPartition/WorldPartitionBuilder.cpp
--- a/Engine/Source/Editor/UnrealEd/Private/WorldPartition/WorldPartitionBuilder.cpp
+++ b/Engine/Source/Editor/UnrealEd/Private/WorldPartition/WorldPartitionBuilder.cpp
@@ -142,7 +142,7 @@ bool UWorldPartitionBuilder::RunBuilder(UWorld* World)
 	return bResult;
 }
 
-static FWorldBuilderCellCoord GetCellCoord(const FVector& InPos, const int32 InCellSize)
+FWorldBuilderCellCoord FCellInfo::GetCellCoord(const FVector& InPos, const int32 InCellSize)
 {
 	return FWorldBuilderCellCoord(
 		FMath::FloorToInt(InPos.X / InCellSize),
@@ -151,7 +151,7 @@ static FWorldBuilderCellCoord GetCellCoord(const FVector& InPos, const int32 InC
 	);
 }
 
-static FWorldBuilderCellCoord GetCellCount(const FBox& InBounds, const int32 InCellSize)
+FWorldBuilderCellCoord FCellInfo::GetCellCount(const FBox& InBounds, const int32 InCellSize)
 {
 	const FWorldBuilderCellCoord MinCellCoords = GetCellCoord(InBounds.Min, InCellSize);
 	const FWorldBuilderCellCoord MaxCellCoords = FWorldBuilderCellCoord(FMath::CeilToInt(InBounds.Max.X / InCellSize),
@@ -210,8 +210,8 @@ bool UWorldPartitionBuilder::Run(UWorld* World, FPackageSourceControlHelper& Pac
 	if ((LoadingMode == ELoadingMode::IterativeCells) || (LoadingMode == ELoadingMode::IterativeCells2D))
 	{
 		// do partial loading loop that calls RunInternal
-		const FWorldBuilderCellCoord MinCellCoords = GetCellCoord(CellInfo.EditorBounds.Min, IterativeCellSize);
-		const FWorldBuilderCellCoord NumCellsIterations = GetCellCount(CellInfo.EditorBounds, IterativeCellSize);
+		const FWorldBuilderCellCoord MinCellCoords = FCellInfo::GetCellCoord(CellInfo.EditorBounds.Min, IterativeCellSize);
+		const FWorldBuilderCellCoord NumCellsIterations = FCellInfo::GetCellCount(CellInfo.EditorBounds, IterativeCellSize);
 		const FWorldBuilderCellCoord BeginCellCoords = MinCellCoords;
 		const FWorldBuilderCellCoord EndCellCoords = BeginCellCoords + NumCellsIterations;
 
@@ -243,6 +243,13 @@ bool UWorldPartitionBuilder::Run(UWorld* World, FPackageSourceControlHelper& Pac
 				for (int64 x = BeginCellCoords.X; bResult && (x < EndCellCoords.X); x++)
 				{
 					IterationIndex++;
+					
+					if (ShouldSkipCell(FWorldBuilderCellCoord(x, y, z)))
+					{
+						UE_LOG(LogWorldPartitionBuilder, Display, TEXT("[%d / %d] Processing cells... (skipped)"), IterationIndex, IterationCount);
+						continue;
+					}
+					
 					UE_LOG(LogWorldPartitionBuilder, Display, TEXT("[%d / %d] Processing cells..."), IterationIndex, IterationCount);
 
 					FVector Min(x * IterativeCellSize, y * IterativeCellSize, z * IterativeCellSize);
diff --git a/Engine/Source/Editor/UnrealEd/Public/WorldPartition/WorldPartitionBuilder.h b/Engine/Source/Editor/UnrealEd/Public/WorldPartition/WorldPartitionBuilder.h
--- a/Engine/Source/Editor/UnrealEd/Public/WorldPartition/WorldPartitionBuilder.h
+++ b/Engine/Source/Editor/UnrealEd/Public/WorldPartition/WorldPartitionBuilder.h
@@ -28,6 +28,9 @@ struct FCellInfo
 
 	/** The size of a cell used by the World Partition Builder */
 	int32 IterativeCellSize;
+
+	static UNREALED_API FWorldBuilderCellCoord GetCellCoord(const FVector& InPos, const int32 InCellSize);
+	static UNREALED_API FWorldBuilderCellCoord GetCellCount(const FBox& InBounds, const int32 InCellSize);
 };
 
 UCLASS(Abstract, Config=Engine)
@@ -76,6 +79,11 @@ protected:
 	 */
 	virtual bool PostWorldTeardown(FPackageSourceControlHelper& PackageHelper) { return true; }
 
+	/**
+	 * When using EloadingMode::IterativeCells/IterativeCells2D return true to skip processing of cell.
+	 */
+	virtual bool ShouldSkipCell(const FWorldBuilderCellCoord& CellCoord) const { return false; }
+
 	int32 IterativeCellSize = 102400;
 	int32 IterativeCellOverlapSize = 0;
 	FBox  IterativeWorldBounds;
