FNames containing asset paths are deprecated. FSoftObjectPath or FTopLevelAssetPath should be used instead.
Fixups for commandlets.

diff --git a/Engine/Source/Editor/UnrealEd/Classes/Commandlets/ImportLocalizedDialogueCommandlet.h b/Engine/Source/Editor/UnrealEd/Classes/Commandlets/ImportLocalizedDialogueCommandlet.h
--- a/Engine/Source/Editor/UnrealEd/Classes/Commandlets/ImportLocalizedDialogueCommandlet.h
+++ b/Engine/Source/Editor/UnrealEd/Classes/Commandlets/ImportLocalizedDialogueCommandlet.h
@@ -41,5 +41,5 @@ private:
 	USoundWave* ConditionalImportSoundWave(const FString& InSoundWavePackageName, const FString& InSoundWaveAssetName, const FString& InWavFilename) const;
 	USoundWave* ImportSoundWave(const FString& InSoundWavePackageName, const FString& InSoundWaveAssetName, const FString& InWavFilename) const;
 
-	TSet<FName> AssetsToKeep;
+	TSet<FSoftObjectPath> AssetsToKeep;
 };
diff --git a/Engine/Source/Editor/UnrealEd/Private/Commandlets/AssetSizeQueryCommandlet.cpp b/Engine/Source/Editor/UnrealEd/Private/Commandlets/AssetSizeQueryCommandlet.cpp
--- a/Engine/Source/Editor/UnrealEd/Private/Commandlets/AssetSizeQueryCommandlet.cpp
+++ b/Engine/Source/Editor/UnrealEd/Private/Commandlets/AssetSizeQueryCommandlet.cpp
@@ -130,7 +130,7 @@ int32 UAssetSizeQueryCommandlet::Main(const FString& FullCommandLine)
 
 	struct FMatchedAssetInfo
 	{
-		FName ObjectPath;
+		FSoftObjectPath ObjectPath;
 		int64 CompressedSize;
 	};
 	TMap<FTopLevelAssetPath /* AssetClass */, int64> FilteredClassCompressedSizes;
@@ -158,7 +158,7 @@ int32 UAssetSizeQueryCommandlet::Main(const FString& FullCommandLine)
 		bool bMatched = true;
 		if (AssetFilter.Len())
 		{
-			FString ObjectPath = AssetData.ObjectPath.ToString();
+			FString ObjectPath = AssetData.GetObjectPathString();
 			bMatched = FWildcardString::IsMatchSubstring(*AssetFilter, *ObjectPath, *ObjectPath + ObjectPath.Len(), ESearchCase::IgnoreCase);
 		} // end if filter exists
 
@@ -167,7 +167,7 @@ int32 UAssetSizeQueryCommandlet::Main(const FString& FullCommandLine)
 			MatchedAssetCount++;
 
 			FMatchedAssetInfo& Info = FilteredClassMatchedAssets.FindOrAdd(AssetData.AssetClassPath).AddDefaulted_GetRef();
-			Info.ObjectPath = AssetData.ObjectPath;
+			Info.ObjectPath = AssetData.GetSoftObjectPath();
 			Info.CompressedSize = AssetCompressedSize;
 			
 			FilteredCompressedSize += AssetCompressedSize;
diff --git a/Engine/Source/Editor/UnrealEd/Private/Commandlets/CompileAllBlueprintsCommandlet.cpp b/Engine/Source/Editor/UnrealEd/Private/Commandlets/CompileAllBlueprintsCommandlet.cpp
--- a/Engine/Source/Editor/UnrealEd/Private/Commandlets/CompileAllBlueprintsCommandlet.cpp
+++ b/Engine/Source/Editor/UnrealEd/Private/Commandlets/CompileAllBlueprintsCommandlet.cpp
@@ -145,7 +145,7 @@ void UCompileAllBlueprintsCommandlet::BuildBlueprints()
 	{
 		if (ShouldBuildAsset(Asset))
 		{
-			FString const AssetPath = Asset.ObjectPath.ToString();
+			FString const AssetPath = Asset.GetObjectPathString();
 			UE_LOG(LogCompileAllBlueprintsCommandlet, Display, TEXT("Loading and Compiling: '%s'..."), *AssetPath);
 
 			//Load with LOAD_NoWarn and LOAD_DisableCompileOnLoad as we are covering those explicitly with CompileBlueprint errors.
@@ -191,7 +191,7 @@ bool UCompileAllBlueprintsCommandlet::ShouldBuildAsset(FAssetData const& Asset)
 
 	if (bCookedOnly && Asset.GetClass() && !Asset.GetClass()->bCooked)
 	{
-		FString const AssetPath = Asset.ObjectPath.ToString();
+		FString const AssetPath = Asset.GetObjectPathString();
 		UE_LOG(LogCompileAllBlueprintsCommandlet, Verbose, TEXT("Skipping Building %s: As is not cooked"), *AssetPath);
 		bShouldBuild = false;
 	}
@@ -200,9 +200,9 @@ bool UCompileAllBlueprintsCommandlet::ShouldBuildAsset(FAssetData const& Asset)
 	{
 		for (const FString& IgnoreFolder : IgnoreFolders)
 		{
-			if (Asset.ObjectPath.ToString().StartsWith(IgnoreFolder))
+			if (Asset.GetObjectPathString().StartsWith(IgnoreFolder))
 			{
-				FString const AssetPath = Asset.ObjectPath.ToString();
+				FString const AssetPath = Asset.GetObjectPathString();
 				UE_LOG(LogCompileAllBlueprintsCommandlet, Verbose, TEXT("Skipping Building %s: As Object is in an Ignored Folder"), *AssetPath);
 				bShouldBuild = false;
 			}
@@ -211,21 +211,21 @@ bool UCompileAllBlueprintsCommandlet::ShouldBuildAsset(FAssetData const& Asset)
 
 	if ((ExcludeAssetTags.Num() > 0) && (CheckHasTagInList(Asset, ExcludeAssetTags)))
 	{
-		FString const AssetPath = Asset.ObjectPath.ToString();
+		FString const AssetPath = Asset.GetObjectPathString();
 		UE_LOG(LogCompileAllBlueprintsCommandlet, Verbose, TEXT("Skipping Building %s: As has an excluded tag"), *AssetPath);
 		bShouldBuild = false;
 	}
 
 	if ((RequireAssetTags.Num() > 0) && (!CheckHasTagInList(Asset, RequireAssetTags)))
 	{
-		FString const AssetPath = Asset.ObjectPath.ToString();
+		FString const AssetPath = Asset.GetObjectPathString();
 		UE_LOG(LogCompileAllBlueprintsCommandlet, Verbose, TEXT("Skipping Building %s: As the asset is missing a required tag"), *AssetPath);
 		bShouldBuild = false;
 	}
 
 	if ((FileAllowList.Num() > 0) && (!IsAssetAllowed(Asset)))
 	{
-		FString const AssetPath = Asset.ObjectPath.ToString();
+		FString const AssetPath = Asset.GetObjectPathString();
 		UE_LOG(LogCompileAllBlueprintsCommandlet, Verbose, TEXT("Skipping Building %s: As the asset is not part of the allow list"), *AssetPath);
 		bShouldBuild = false;
 	}
@@ -235,7 +235,7 @@ bool UCompileAllBlueprintsCommandlet::ShouldBuildAsset(FAssetData const& Asset)
 		const UPackage* AssetPackage = Asset.GetPackage();
 		if ((AssetPackage == nullptr) || !AssetPackage->IsDirty())
 		{
-			FString const AssetPath = Asset.ObjectPath.ToString();
+			FString const AssetPath = Asset.GetObjectPathString();
 			UE_LOG(LogCompileAllBlueprintsCommandlet, Verbose, TEXT("Skipping Building %s: As Package is not dirty"), *AssetPath);
 			bShouldBuild = false;
 		}
@@ -281,7 +281,7 @@ bool UCompileAllBlueprintsCommandlet::IsAssetAllowed(FAssetData const& Asset) co
 {
 	bool bIsInList = false;
 
-	const FString& AssetFilePath = Asset.ObjectPath.ToString();
+	const FString& AssetFilePath = Asset.GetObjectPathString();
 	for (const FString& Entry : FileAllowList)
 	{
 		if (AssetFilePath == Entry)
diff --git a/Engine/Source/Editor/UnrealEd/Private/Commandlets/CompileShadersTestBedCommandlet.cpp b/Engine/Source/Editor/UnrealEd/Private/Commandlets/CompileShadersTestBedCommandlet.cpp
--- a/Engine/Source/Editor/UnrealEd/Private/Commandlets/CompileShadersTestBedCommandlet.cpp
+++ b/Engine/Source/Editor/UnrealEd/Private/Commandlets/CompileShadersTestBedCommandlet.cpp
@@ -62,7 +62,9 @@ int32 UCompileShadersTestBedCommandlet::Main(const FString& Params)
 			Filter.ClassPaths.Add(UMaterial::StaticClass()->GetClassPathName());
 
 			FCollectionManagerModule& CollectionManagerModule = FCollectionManagerModule::GetModule();
+PRAGMA_DISABLE_DEPRECATION_WARNINGS
 			CollectionManagerModule.Get().GetObjectsInCollection(FName(*CollectionName), ECollectionShareType::CST_All, Filter.ObjectPaths, ECollectionRecursionFlags::SelfAndChildren);
+PRAGMA_ENABLE_DEPRECATION_WARNINGS
 
 			AssetRegistry.GetAssets(Filter, MaterialList);
 
@@ -86,13 +88,13 @@ int32 UCompileShadersTestBedCommandlet::Main(const FString& Params)
 	{
 		// re-use the filter and only filter based on the passed in objects.
 		Filter.ClassPaths.Empty();
-		Filter.ObjectPaths.Empty();
-		for (const FString& MaterialPath : CmdLineMaterialEntries)
+		Filter.SoftObjectPaths.Empty();
+		for (const FString& MaterialPathString : CmdLineMaterialEntries)
 		{
-			const FName MaterialPathFName(MaterialPath);
-			if (!Filter.ObjectPaths.Contains(MaterialPathFName))
+			const FSoftObjectPath MaterialPath(MaterialPathString);
+			if (!Filter.SoftObjectPaths.Contains(MaterialPath))
 			{
-				Filter.ObjectPaths.Add(MaterialPathFName);
+				Filter.SoftObjectPaths.Add(MaterialPath);
 			}
 		}
 
@@ -134,7 +136,7 @@ int32 UCompileShadersTestBedCommandlet::Main(const FString& Params)
 			TRACE_CPUPROFILER_EVENT_SCOPE(MaterialShaders);
 
 			// Sort the material lists by name so the order is stable.
-			Algo::SortBy(MaterialList, [](const FAssetData& AssetData) { return AssetData.ObjectPath; }, FNameLexicalLess());
+			Algo::SortBy(MaterialList, [](const FAssetData& AssetData) { return AssetData.GetSoftObjectPath(); }, [](const FSoftObjectPath& A, const FSoftObjectPath& B) { return A.LexicalLess(B); });
 
 			for (const FAssetData& AssetData : MaterialList)
 			{
diff --git a/Engine/Source/Editor/UnrealEd/Private/Commandlets/ContentCommandlets.cpp b/Engine/Source/Editor/UnrealEd/Private/Commandlets/ContentCommandlets.cpp
--- a/Engine/Source/Editor/UnrealEd/Private/Commandlets/ContentCommandlets.cpp
+++ b/Engine/Source/Editor/UnrealEd/Private/Commandlets/ContentCommandlets.cpp
@@ -3296,7 +3296,7 @@ int32 UListMaterialsUsedWithMeshEmittersCommandlet::Main( const FString& Params
 
 	for(int32 AssetIdx = 0; AssetIdx < AssetList.Num(); ++AssetIdx )
 	{
-		const FString Filename = AssetList[AssetIdx].ObjectPath.ToString();
+		const FString Filename = AssetList[AssetIdx].GetObjectPathString();
 
 		UE_LOG(LogContentCommandlet, Display, TEXT("Processing particle system (%i/%i):  %s "), AssetIdx, AssetList.Num(), *Filename );
 
@@ -3382,7 +3382,7 @@ int32 UListStaticMeshesImportedFromSpeedTreesCommandlet::Main(const FString& Par
 
 	for (int32 AssetIdx = 0; AssetIdx < AssetList.Num(); ++AssetIdx)
 	{
-		const FString Filename = AssetList[AssetIdx].ObjectPath.ToString();
+		const FString Filename = AssetList[AssetIdx].GetObjectPathString();
 
 		UE_LOG(LogContentCommandlet, Display, TEXT("Processing static mesh (%i/%i):  %s "), AssetIdx, AssetList.Num(), *Filename);
 
@@ -3538,7 +3538,7 @@ int32 UStaticMeshMinLodCommandlet::Main(const FString& Params)
 	for (int32 AssetIdx = 0; AssetIdx < AssetList.Num(); ++AssetIdx)
 	{
 		bool SavePackage = false;
-		const FString Filename = AssetList[AssetIdx].ObjectPath.ToString();
+		const FString Filename = AssetList[AssetIdx].GetObjectPathString();
 		UE_LOG(LogContentCommandlet, Display, TEXT("Processing static mesh (%i/%i):  %s "), AssetIdx, AssetList.Num(), *Filename);
 
 		UPackage* Package = LoadPackage(NULL, *Filename, LOAD_Quiet);
diff --git a/Engine/Source/Editor/UnrealEd/Private/Commandlets/ConvertLevelsToExternalActorsCommandlet.cpp b/Engine/Source/Editor/UnrealEd/Private/Commandlets/ConvertLevelsToExternalActorsCommandlet.cpp
--- a/Engine/Source/Editor/UnrealEd/Private/Commandlets/ConvertLevelsToExternalActorsCommandlet.cpp
+++ b/Engine/Source/Editor/UnrealEd/Private/Commandlets/ConvertLevelsToExternalActorsCommandlet.cpp
@@ -65,9 +65,9 @@ void UConvertLevelsToExternalActorsCommandlet::GetSubLevelsToConvert(ULevel* Mai
 bool UConvertLevelsToExternalActorsCommandlet::CheckExternalActors(const FString& Level, bool bRepair)
 {
 	// Gather duplicated actor files.
-	TMultiMap<FName, FName> DuplicatedActorFiles;
+	TMultiMap<FSoftObjectPath, FName> DuplicatedActorFiles;
 	{
-		TMap<FName, FName> ActorFiles;
+		TMap<FSoftObjectPath, FName> ActorFiles;
 
 		IAssetRegistry& AssetRegistry = FModuleManager::LoadModuleChecked<FAssetRegistryModule>(TEXT("AssetRegistry")).Get();
 
@@ -76,19 +76,19 @@ bool UConvertLevelsToExternalActorsCommandlet::CheckExternalActors(const FString
 
 		FDelegateHandle AddedCheckHandle = AssetRegistry.OnAssetAdded().AddLambda([&ActorFiles](const FAssetData& AssetData)
 		{
-			check(!ActorFiles.Contains(AssetData.ObjectPath));
-			ActorFiles.Add(AssetData.ObjectPath, AssetData.PackageName);
+			check(!ActorFiles.Contains(AssetData.GetSoftObjectPath()));
+			ActorFiles.Add(AssetData.GetSoftObjectPath(), AssetData.PackageName);
 		});
 
 		FDelegateHandle UpdatedCheckHandle = AssetRegistry.OnAssetUpdated().AddLambda([&ActorFiles, &DuplicatedActorFiles](const FAssetData& AssetData)
 		{
 			FName ExistingPackageName;
-			if (ActorFiles.RemoveAndCopyValue(AssetData.ObjectPath, ExistingPackageName))
+			if (ActorFiles.RemoveAndCopyValue(AssetData.GetSoftObjectPath(), ExistingPackageName))
 			{
-				DuplicatedActorFiles.Add(AssetData.ObjectPath, ExistingPackageName);
+				DuplicatedActorFiles.Add(AssetData.GetSoftObjectPath(), ExistingPackageName);
 			}
 
-			DuplicatedActorFiles.Add(AssetData.ObjectPath, AssetData.PackageName);			
+			DuplicatedActorFiles.Add(AssetData.GetSoftObjectPath(), AssetData.PackageName);
 		});
 
 		AssetRegistry.ScanPathsSynchronous(ULevel::GetExternalObjectsPaths(Level), /*bForceRescan*/true, /*bIgnoreDenyListScanFilters*/true);
@@ -101,14 +101,14 @@ bool UConvertLevelsToExternalActorsCommandlet::CheckExternalActors(const FString
 	{
 		// Gather unique keys from the duplicated map.
 		// Note: TMultiMap::GenerateKeyArray will return duplicated keys, clean that.
-		TArray<FName> DuplicatedActorFilesKeys;
+		TArray<FSoftObjectPath> DuplicatedActorFilesKeys;
 		DuplicatedActorFiles.GenerateKeyArray(DuplicatedActorFilesKeys);
-		DuplicatedActorFilesKeys.Sort([](const FName& A, const FName& B) { return A.FastLess(B); });
+		DuplicatedActorFilesKeys.Sort([](const FSoftObjectPath& A, const FSoftObjectPath& B) { return A.FastLess(B); });
 		int32 EndIndex = Algo::Unique(DuplicatedActorFilesKeys);
 		DuplicatedActorFilesKeys.RemoveAt(EndIndex, DuplicatedActorFilesKeys.Num() - EndIndex);
 
 		// Report or delete duplicated entries, keeping the latest one
-		for (const FName& DuplicatedActorFileKey : DuplicatedActorFilesKeys)
+		for (const FSoftObjectPath& DuplicatedActorFileKey : DuplicatedActorFilesKeys)
 		{
 			TArray<FName> DuplicatedActorFilesPaths;
 			DuplicatedActorFiles.MultiFind(DuplicatedActorFileKey, DuplicatedActorFilesPaths);
diff --git a/Engine/Source/Editor/UnrealEd/Private/Commandlets/DiffAssetBulkDataCommandlet.cpp b/Engine/Source/Editor/UnrealEd/Private/Commandlets/DiffAssetBulkDataCommandlet.cpp
--- a/Engine/Source/Editor/UnrealEd/Private/Commandlets/DiffAssetBulkDataCommandlet.cpp
+++ b/Engine/Source/Editor/UnrealEd/Private/Commandlets/DiffAssetBulkDataCommandlet.cpp
@@ -332,7 +332,7 @@ int32 UDiffAssetBulkDataCommandlet::Main(const FString& FullCommandLine)
 					TArray<FDiffResult>& ClassResults = TagResults.FindOrAdd(Tag.BaseAssetData->AssetClassPath);
 
 					FDiffResult& Result = ClassResults.AddDefaulted_GetRef();
-					Result.ChangedAssetObjectPath = Tag.BaseAssetData->ObjectPath.ToString();
+					Result.ChangedAssetObjectPath = Tag.BaseAssetData->GetObjectPathString();
 					Result.TagBaseValue = MoveTemp(Tag.BaseValue);
 					Result.TagCurrentValue = MoveTemp(Tag.CurrentValue);
 					bFoundDiffTag = true;
diff --git a/Engine/Source/Editor/UnrealEd/Private/Commandlets/DumpMaterialExpressionsCommandlet.cpp b/Engine/Source/Editor/UnrealEd/Private/Commandlets/DumpMaterialExpressionsCommandlet.cpp
--- a/Engine/Source/Editor/UnrealEd/Private/Commandlets/DumpMaterialExpressionsCommandlet.cpp
+++ b/Engine/Source/Editor/UnrealEd/Private/Commandlets/DumpMaterialExpressionsCommandlet.cpp
@@ -283,7 +283,7 @@ void WriteOutMaterialFunctions(FArchive* FileWriter, const FString& OutputFilePa
 		// If this is a function that is selected to be exposed to the library
 		if (AssetData.GetTagValueRef<bool>("bExposeToLibrary"))
 		{
-			const FString FunctionPathName = AssetData.ObjectPath.ToString();
+			const FString FunctionPathName = AssetData.GetObjectPathString();
 			const FString Description = AssetData.GetTagValueRef<FText>("Description").ToString();
 
 			FString FunctionName = FunctionPathName;
diff --git a/Engine/Source/Editor/UnrealEd/Private/Commandlets/DumpMaterialShaderTypes.cpp b/Engine/Source/Editor/UnrealEd/Private/Commandlets/DumpMaterialShaderTypes.cpp
--- a/Engine/Source/Editor/UnrealEd/Private/Commandlets/DumpMaterialShaderTypes.cpp
+++ b/Engine/Source/Editor/UnrealEd/Private/Commandlets/DumpMaterialShaderTypes.cpp
@@ -378,7 +378,7 @@ int ProcessMaterials(const ITargetPlatform* TargetPlatform, const EShaderPlatfor
 			TotalShaders += TotalShadersForMaterial;
 
 			Output.Log(TEXT(""));
-			Output.Log(FString::Printf(TEXT("Material: %s - %d shaders"), *AssetData.ObjectPath.ToString(), TotalShadersForMaterial));
+			Output.Log(FString::Printf(TEXT("Material: %s - %d shaders"), *AssetData.GetObjectPathString(), TotalShadersForMaterial));
 
 			PrintDebugShaderInfo(Output, OutShaderInfo);
 		}
@@ -639,7 +639,9 @@ int32 UDumpMaterialShaderTypesCommandlet::Main(const FString& Params)
 				Filter.ClassPaths.Add(UMaterial::StaticClass()->GetClassPathName());
 
 				FCollectionManagerModule& CollectionManagerModule = FCollectionManagerModule::GetModule();
+PRAGMA_DISABLE_DEPRECATION_WARNINGS
 				CollectionManagerModule.Get().GetObjectsInCollection(FName(*CollectionName), ECollectionShareType::CST_All, Filter.ObjectPaths, ECollectionRecursionFlags::SelfAndChildren);
+PRAGMA_ENABLE_DEPRECATION_WARNINGS
 
 				AssetRegistry.GetAssets(Filter, MaterialList);
 
@@ -664,8 +666,8 @@ int32 UDumpMaterialShaderTypesCommandlet::Main(const FString& Params)
 	UE_LOG(LogDumpMaterialShaderTypesCommandlet, Display, TEXT("Asset scan took: %.3f"), AssetRegistryEnd - AssetRegistryStart);
 
 	// Sort the material lists by name so the order is stable.
-	Algo::SortBy(MaterialList, [](const FAssetData& AssetData) { return AssetData.ObjectPath; }, FNameLexicalLess());
-	Algo::SortBy(MaterialInstanceList, [](const FAssetData& AssetData) { return AssetData.ObjectPath; }, FNameLexicalLess());
+	Algo::SortBy(MaterialList, [](const FAssetData& AssetData) { return AssetData.GetSoftObjectPath(); }, [](const FSoftObjectPath& A, const FSoftObjectPath& B) { return A.LexicalLess(B); });
+	Algo::SortBy(MaterialInstanceList, [](const FAssetData& AssetData) { return AssetData.GetSoftObjectPath(); }, [](const FSoftObjectPath& A, const FSoftObjectPath& B) { return A.LexicalLess(B); });
 
 	// For all active platforms
 	ITargetPlatformManagerModule* TPM = GetTargetPlatformManager();
diff --git a/Engine/Source/Editor/UnrealEd/Private/Commandlets/ExportDialogueScriptCommandlet.cpp b/Engine/Source/Editor/UnrealEd/Private/Commandlets/ExportDialogueScriptCommandlet.cpp
--- a/Engine/Source/Editor/UnrealEd/Private/Commandlets/ExportDialogueScriptCommandlet.cpp
+++ b/Engine/Source/Editor/UnrealEd/Private/Commandlets/ExportDialogueScriptCommandlet.cpp
@@ -266,7 +266,7 @@ int32 UExportDialogueScriptCommandlet::Main(const FString& Params)
 			UDialogueWave* LocalizedDialogueWave = nullptr;
 			{
 				const FString LocalizedDialogueWavePackagePath = FPackageName::GetLocalizedPackagePath(AssetData.PackageName.ToString(), CultureName);
-				const FAssetData LocalizedDialogueWaveAssetData = AssetRegistry.GetAssetByObjectPath(*FString::Printf(TEXT("%s.%s"), *LocalizedDialogueWavePackagePath, *AssetData.AssetName.ToString()));
+				const FAssetData LocalizedDialogueWaveAssetData = AssetRegistry.GetAssetByObjectPath(FSoftObjectPath(*LocalizedDialogueWavePackagePath, AssetData.AssetName, {}));
 				LocalizedDialogueWave = Cast<UDialogueWave>(LocalizedDialogueWaveAssetData.GetAsset());
 				if (LocalizedDialogueWave == DialogueWave)
 				{
diff --git a/Engine/Source/Editor/UnrealEd/Private/Commandlets/GatherTextFromAssetsCommandlet.cpp b/Engine/Source/Editor/UnrealEd/Private/Commandlets/GatherTextFromAssetsCommandlet.cpp
--- a/Engine/Source/Editor/UnrealEd/Private/Commandlets/GatherTextFromAssetsCommandlet.cpp
+++ b/Engine/Source/Editor/UnrealEd/Private/Commandlets/GatherTextFromAssetsCommandlet.cpp
@@ -510,7 +510,9 @@ int32 UGatherTextFromAssetsCommandlet::Main(const FString& Params)
 			ICollectionManager& CollectionManager = CollectionManagerModule.Get();
 			for (const FString& CollectionName : CollectionFilters)
 			{
+PRAGMA_DISABLE_DEPRECATION_WARNINGS
 				if (!CollectionManager.GetObjectsInCollection(FName(*CollectionName), ECollectionShareType::CST_All, FirstPassFilter.ObjectPaths, ECollectionRecursionFlags::SelfAndChildren))
+PRAGMA_ENABLE_DEPRECATION_WARNINGS
 				{
 					UE_LOG(LogGatherTextFromAssetsCommandlet, Error, TEXT("Failed get objects in specified collection: %s"), *CollectionName);
 					HasFailedToGetACollection = true;
diff --git a/Engine/Source/Editor/UnrealEd/Private/Commandlets/ImportLocalizedDialogueCommandlet.cpp b/Engine/Source/Editor/UnrealEd/Private/Commandlets/ImportLocalizedDialogueCommandlet.cpp
--- a/Engine/Source/Editor/UnrealEd/Private/Commandlets/ImportLocalizedDialogueCommandlet.cpp
+++ b/Engine/Source/Editor/UnrealEd/Private/Commandlets/ImportLocalizedDialogueCommandlet.cpp
@@ -255,7 +255,7 @@ int32 UImportLocalizedDialogueCommandlet::Main(const FString& Params)
 	for (const FAssetData& LocalizedAssetData : LocalizedAssetsToPotentiallyDelete)
 	{
 		// Has this asset already keen marked as "keep"?
-		if (AssetsToKeep.Contains(LocalizedAssetData.ObjectPath))
+		if (AssetsToKeep.Contains(LocalizedAssetData.GetSoftObjectPath()))
 		{
 			continue;
 		}
@@ -308,7 +308,7 @@ bool UImportLocalizedDialogueCommandlet::ImportDialogueForCulture(FLocTextHelper
 		}
 
 		// Mark this localized dialogue wave as used so it doesn't get deleted later
-		AssetsToKeep.Add(*LocalizedDialogueWave->GetPathName());
+		AssetsToKeep.Add(FSoftObjectPath(LocalizedDialogueWave));
 	}
 
 	// First pass, handle any contexts that have an exact mapping to their audio file
@@ -324,7 +324,7 @@ bool UImportLocalizedDialogueCommandlet::ImportDialogueForCulture(FLocTextHelper
 			// We're skipping this context entry due to our manifest, but we don't want the sound wave it's using to be deleted
 			if (ContextMapping.SoundWave)
 			{
-				AssetsToKeep.Add(*ContextMapping.SoundWave->GetPathName());
+				AssetsToKeep.Add(FSoftObjectPath(ContextMapping.SoundWave));
 			}
 
 			UE_LOG(LogImportLocalizedDialogueCommandlet, Log, TEXT("No internationalization manifest entry was found for context '%s' in culture '%s'. This context will be skipped."), *ContextLocalizationKey, *InCultureImportInfo.Name);
@@ -352,7 +352,7 @@ bool UImportLocalizedDialogueCommandlet::ImportDialogueForCulture(FLocTextHelper
 		// This sound wave is in use, so shouldn't be deleted
 		if (ContextMapping.SoundWave)
 		{
-			AssetsToKeep.Add(*ContextMapping.SoundWave->GetPathName());
+			AssetsToKeep.Add(FSoftObjectPath(ContextMapping.SoundWave));
 		}
 	}
 
@@ -431,7 +431,7 @@ bool UImportLocalizedDialogueCommandlet::ImportDialogueForCulture(FLocTextHelper
 		// This sound wave is in use, so shouldn't be deleted
 		if (ContextMappingPtr->SoundWave)
 		{
-			AssetsToKeep.Add(*ContextMappingPtr->SoundWave->GetPathName());
+			AssetsToKeep.Add(FSoftObjectPath(ContextMappingPtr->SoundWave));
 		}
 	}
 
diff --git a/Engine/Source/Editor/UnrealEd/Private/Commandlets/ParticleSystemAuditCommandlet.cpp b/Engine/Source/Editor/UnrealEd/Private/Commandlets/ParticleSystemAuditCommandlet.cpp
--- a/Engine/Source/Editor/UnrealEd/Private/Commandlets/ParticleSystemAuditCommandlet.cpp
+++ b/Engine/Source/Editor/UnrealEd/Private/Commandlets/ParticleSystemAuditCommandlet.cpp
@@ -81,7 +81,9 @@ bool UParticleSystemAuditCommandlet::ProcessParticleSystems()
 	if (!FilterCollection.IsEmpty())
 	{
 		FCollectionManagerModule& CollectionManagerModule = FCollectionManagerModule::GetModule();
+PRAGMA_DISABLE_DEPRECATION_WARNINGS
 		CollectionManagerModule.Get().GetObjectsInCollection(FName(*FilterCollection), ECollectionShareType::CST_All, Filter.ObjectPaths, ECollectionRecursionFlags::SelfAndChildren);
+PRAGMA_ENABLE_DEPRECATION_WARNINGS
 	}
 
 	TArray<FAssetData> AssetList;
@@ -102,7 +104,7 @@ bool UParticleSystemAuditCommandlet::ProcessParticleSystems()
 	UPackage* CurrentPackage = NULL;
 	for (const FAssetData& AssetIt : AssetList)
 	{
-		const FString PSysName = AssetIt.ObjectPath.ToString();
+		const FString PSysName = AssetIt.GetObjectPathString();
 		const FString PackageName = AssetIt.PackageName.ToString();
 
 		if (PackageName.StartsWith(DevelopersFolder))
