World Partition: fix functional tests support for cooked builds by ensuring we gather asset registry tags from memory for packages that are generators (fix mostly provided by MattP).

#rnx

diff --git a/Engine/Source/Editor/UnrealEd/Private/Commandlets/AssetRegistryGenerator.cpp b/Engine/Source/Editor/UnrealEd/Private/Commandlets/AssetRegistryGenerator.cpp
--- a/Engine/Source/Editor/UnrealEd/Private/Commandlets/AssetRegistryGenerator.cpp
+++ b/Engine/Source/Editor/UnrealEd/Private/Commandlets/AssetRegistryGenerator.cpp
@@ -2515,7 +2515,8 @@ const FAssetData* FAssetRegistryGenerator::CreateOrFindAssetData(UObject& Object
 
 void FAssetRegistryGenerator::UpdateAssetRegistryData(const UPackage& Package,
 	FSavePackageResultStruct& SavePackageResult,
-	FCookTagList&& InArchiveCookTagList
+	FCookTagList&& InArchiveCookTagList,
+	bool bIncludeOnlyDiskAssets
 	)
 {
 	const FName PackageName = Package.GetFName();
@@ -2528,7 +2529,6 @@ void FAssetRegistryGenerator::UpdateAssetRegistryData(const UPackage& Package,
 	// Copy latest data for all Assets in the package into the cooked registry. This should be done even
 	// if not successful so that editor-only packages are recorded as well
 	TArray<FAssetData> AssetDatas;
-	constexpr bool bIncludeOnlyDiskAssets = true; // Enumerating memory assets is unnecessary; we waited on AR to update
 	AssetRegistry.GetAssetsByPackageName(PackageName, AssetDatas, bIncludeOnlyDiskAssets,
 		false /* SkipARFilteredAssets */);
 	for (FAssetData& AssetData : AssetDatas)
@@ -2591,7 +2591,7 @@ FAssetRegistryReporterRemote::FAssetRegistryReporterRemote(FCookWorkerClient& In
 }
 
 void FAssetRegistryReporterRemote::UpdateAssetRegistryData(FPackageData& PackageData, const UPackage& Package,
-	FSavePackageResultStruct& SavePackageResult, FCookTagList&& InArchiveCookTagList)
+	FSavePackageResultStruct& SavePackageResult, FCookTagList&& InArchiveCookTagList, bool bIncludeOnlyDiskAssets)
 {
 	uint32 NewPackageFlags = 0;
 	int64 DiskSize = -1;
@@ -2614,7 +2614,6 @@ void FAssetRegistryReporterRemote::UpdateAssetRegistryData(FPackageData& Package
 	Message.DiskSize = DiskSize;
 
 	// Add to the message all the AssetDatas in the package from the global AssetRegistry
-	constexpr bool bIncludeOnlyDiskAssets = true; // Enumerating memory assets is unnecessary; we waited on AR to update
 	IAssetRegistry::Get()->GetAssetsByPackageName(PackageData.GetPackageName(), Message.AssetDatas,
 		bIncludeOnlyDiskAssets, false /* SkipARFilteredAssets */);
 
diff --git a/Engine/Source/Editor/UnrealEd/Private/CookOnTheFlyServer.cpp b/Engine/Source/Editor/UnrealEd/Private/CookOnTheFlyServer.cpp
--- a/Engine/Source/Editor/UnrealEd/Private/CookOnTheFlyServer.cpp
+++ b/Engine/Source/Editor/UnrealEd/Private/CookOnTheFlyServer.cpp
@@ -5532,7 +5532,15 @@ void FSaveCookedPackageContext::FinishPlatform()
 		// Flush the AssetRegisty so any AssetData changes from the save are present
 		COTFS.AssetRegistry->WaitForCompletion();
 		IAssetRegistryReporter& Reporter = *(COTFS.PlatformManager->GetPlatformData(TargetPlatform)->RegistryReporter);
-		Reporter.UpdateAssetRegistryData(PackageData, *Package, SavePackageResult, MoveTemp(*ArchiveCookContext.GetCookTagList()));
+
+		// UpdateAssetRegistryData will pull data from the global asset registry
+		// For most types, our contract is that we use the data in the asset registry that was updated on load and is in
+		// the disk asset storage of the global asset registry. Most types do not take significant action during cooking
+		// so we do not need to allow them to update their tags after they finish loading.
+		// But generator packages do a lot of work during generation and may need to update their tags, so read the
+		// latest tags off of their asset by setting bIncludeOnlyDiskAssets=true.
+		const bool bIncludeOnlyDiskAssets = PackageData.GetGeneratorPackage() == nullptr;
+		Reporter.UpdateAssetRegistryData(PackageData, *Package, SavePackageResult, MoveTemp(*ArchiveCookContext.GetCookTagList()), bIncludeOnlyDiskAssets);
 	}
 
 	// If not retrying, mark the package as cooked, either successfully or with failure
diff --git a/Engine/Source/Editor/UnrealEd/Public/Commandlets/AssetRegistryGenerator.h b/Engine/Source/Editor/UnrealEd/Public/Commandlets/AssetRegistryGenerator.h
--- a/Engine/Source/Editor/UnrealEd/Public/Commandlets/AssetRegistryGenerator.h
+++ b/Engine/Source/Editor/UnrealEd/Public/Commandlets/AssetRegistryGenerator.h
@@ -201,8 +201,9 @@ public:
 	 *
 	 * @param Package The package to update info on
 	 * @param SavePackageResult The metadata to associate with the given package name
+	 * @param bIncludeOnlyDiskAssets Include only disk assets or else also enumerate memory assets
 	 */
-	void UpdateAssetRegistryData(const UPackage& Package, FSavePackageResultStruct& SavePackageResult, FCookTagList&& InArchiveCookTagList);
+	void UpdateAssetRegistryData(const UPackage& Package, FSavePackageResultStruct& SavePackageResult, FCookTagList&& InArchiveCookTagList, bool bIncludeOnlyDiskAssets = true);
 	void UpdateAssetRegistryData(UE::Cook::FPackageData& PackageData, UE::Cook::FAssetRegistryPackageMessage&& Message);
 
 	/**
@@ -434,7 +435,7 @@ public:
 	virtual ~IAssetRegistryReporter() {}
 
 	virtual void UpdateAssetRegistryData(FPackageData& PackageData, const UPackage& Package,
-		FSavePackageResultStruct& SavePackageResult, FCookTagList&& InArchiveCookTagList) = 0;
+		FSavePackageResultStruct& SavePackageResult, FCookTagList&& InArchiveCookTagList, bool bIncludeOnlyDiskAssets = true) = 0;
 
 };
 
@@ -447,9 +448,9 @@ public:
 	}
 
 	virtual void UpdateAssetRegistryData(FPackageData& PackageData, const UPackage& Package,
-		FSavePackageResultStruct& SavePackageResult, FCookTagList&& InArchiveCookTagList) override
+		FSavePackageResultStruct& SavePackageResult, FCookTagList&& InArchiveCookTagList, bool bIncludeOnlyDiskAssets = true) override
 	{
-		Generator.UpdateAssetRegistryData(Package, SavePackageResult, MoveTemp(InArchiveCookTagList));
+		Generator.UpdateAssetRegistryData(Package, SavePackageResult, MoveTemp(InArchiveCookTagList), bIncludeOnlyDiskAssets);
 	}
 
 private:
@@ -462,7 +463,7 @@ public:
 	FAssetRegistryReporterRemote(FCookWorkerClient& InClient, const ITargetPlatform* InTargetPlatform);
 
 	virtual void UpdateAssetRegistryData(FPackageData& PackageData, const UPackage& Package,
-		FSavePackageResultStruct& SavePackageResult, FCookTagList&& InArchiveCookTagList) override;
+		FSavePackageResultStruct& SavePackageResult, FCookTagList&& InArchiveCookTagList, bool bIncludeOnlyDiskAssets = true) override;
 
 private:
 	FCookWorkerClient& Client;
