Improved import setting button behaviour:

When users click on import and open an appearance setting, it will automatically apply the theme if it exists. If not, prompt users to import a theme from the same location as the Appearance setting config file. If import succeeds, the new theme is applied and the .json file is copied to user-level theme directoy. If users choose not to import, keep the current theme.

diff --git a/Engine/Source/Editor/UnrealEd/Classes/Settings/EditorStyleSettings.h b/Engine/Source/Editor/UnrealEd/Classes/Settings/EditorStyleSettings.h
--- a/Engine/Source/Editor/UnrealEd/Classes/Settings/EditorStyleSettings.h
+++ b/Engine/Source/Editor/UnrealEd/Classes/Settings/EditorStyleSettings.h
@@ -161,7 +161,8 @@ public:
 	/** @return A subdued version of the users selection color (for use with inactive selection)*/
 	FLinearColor GetSubduedSelectionColor() const;
 	
-	bool OnExportBegin(const FString& ExportToPath); 
+	bool OnImportBegin(const FString& ImportFromPath);
+	bool OnExportBegin(const FString& ExportToPath);
 
 protected:
 
diff --git a/Engine/Source/Editor/UnrealEd/Private/Settings/EditorStyleClasses.cpp b/Engine/Source/Editor/UnrealEd/Private/Settings/EditorStyleClasses.cpp
--- a/Engine/Source/Editor/UnrealEd/Private/Settings/EditorStyleClasses.cpp
+++ b/Engine/Source/Editor/UnrealEd/Private/Settings/EditorStyleClasses.cpp
@@ -13,6 +13,7 @@
 #include "UObject/UnrealType.h"
 #include "Styling/StyleColors.h"
 #include "HAL/FileManager.h"
+#include "EditorStyleSettingsCustomization.h"
 
 class FObjectInitializer;
 
@@ -93,6 +94,18 @@ void UEditorStyleSettings::PostEditChangeProperty(struct FPropertyChangedEvent&
 	SettingChangedEvent.Broadcast(PropertyName);
 }
 
+bool UEditorStyleSettings::OnImportBegin(const FString& ImportFromPath)
+{
+	LoadConfig(GetClass(), *ImportFromPath, UE::LCPF_PropagateToInstances);
+
+	// if theme exists, it would already be applied by import. 
+	// if theme doesn't exist, prompt to open to import a theme. 
+	if (!USlateThemeManager::Get().DoesThemeExist(CurrentAppliedTheme))
+	{
+		FEditorStyleSettingsCustomization::PromptToImportTheme(ImportFromPath); 
+	}
+	return true; 
+}
 
 bool UEditorStyleSettings::OnExportBegin(const FString& ExportToPath)
 {
diff --git a/Engine/Source/Editor/UnrealEd/Private/Settings/EditorStyleSettingsCustomization.cpp b/Engine/Source/Editor/UnrealEd/Private/Settings/EditorStyleSettingsCustomization.cpp
--- a/Engine/Source/Editor/UnrealEd/Private/Settings/EditorStyleSettingsCustomization.cpp
+++ b/Engine/Source/Editor/UnrealEd/Private/Settings/EditorStyleSettingsCustomization.cpp
@@ -347,9 +347,19 @@
 
 TSharedRef<IDetailCustomization> FEditorStyleSettingsCustomization::MakeInstance()
 {
-	return MakeShared<FEditorStyleSettingsCustomization>();
+	return MakeShared<FEditorStyleSettingsCustomization>(); 
 }
 
+FEditorStyleSettingsCustomization::FEditorStyleSettingsCustomization()
+{
+	USlateThemeManager::Get().OnThemeChanged().AddSP(StaticCastSharedRef<FEditorStyleSettingsCustomization>(AsShared()), &FEditorStyleSettingsCustomization::OnThemeChanged);
+}
 
+FEditorStyleSettingsCustomization::~FEditorStyleSettingsCustomization()
+{
+	USlateThemeManager& ThemeManager = USlateThemeManager::Get();
+	ThemeManager.OnThemeChanged().RemoveAll(this);
+}
+
 void FEditorStyleSettingsCustomization::CustomizeDetails(IDetailLayoutBuilder& DetailLayout)
 {
@@ -608,64 +618,7 @@
 
 FReply FEditorStyleSettingsCustomization::OnImportThemeClicked()
 {
-	TArray<FString> OutFiles;
-	const void* ParentWindowHandle = FSlateApplication::Get().FindBestParentWindowHandleForDialogs(nullptr); 
-	const FString ImportFrom = FPlatformProcess::UserDir();
-
-	if (FDesktopPlatformModule::Get()->OpenFileDialog(ParentWindowHandle, LOCTEXT("ImportThemeDialogTitle", "Import theme...").ToString(), FPaths::GetPath(ImportFrom), TEXT(""), TEXT("JSON files (*.json)|*.json"), EFileDialogFlags::None, OutFiles))
-	{
-		FString SourcePath = OutFiles[0];
-		const FString DestPath = USlateThemeManager::Get().GetUserThemeDir() / FPaths::GetCleanFilename(SourcePath);
-
-		FString PathPart; 
-		FString Extension; 
-		FString FilenameWithoutExtension; 
-		FPaths::Split(SourcePath, PathPart, FilenameWithoutExtension, Extension);
-
-		// if theme name exists, don't import (to prevent from overwriting existing theme files)
-		if (!IsThemeNameValid(FilenameWithoutExtension))
-		{
-			ShowNotification(LOCTEXT("ImportThemeFailureNameExists", "Import theme failed: Theme name already exists"), SNotificationItem::CS_Fail);
-		}
-		// if theme name is valid: copying the file is safe (as it will not overwrite existing theme files)
-		else 
-		{	
-			const int32 NumOfThemesBefore = USlateThemeManager::Get().GetThemes().Num();
-
-			if  (IPlatformFile::GetPlatformPhysical().CopyFile(*DestPath, *SourcePath)) 
-			{	
-				// update the number of valid themes: 
-				USlateThemeManager::Get().LoadThemes();
-
-				// if valid theme: the theme Num will be updated. 
-				if (USlateThemeManager::Get().GetThemes().Num() != NumOfThemesBefore)
-				{
-					// Extract ID as a FString directly from a JSON file. 
-					FString ImportedThemeID; 
-					GetThemeIdFromPath(SourcePath, ImportedThemeID);
-
-					// convert FString ID to a FGuid
-					FGuid ImportedThemeGUID = FGuid(ImportedThemeID); 
-					USlateThemeManager::Get().ApplyTheme(ImportedThemeGUID);
-					RefreshComboBox();
-
-					ShowNotification(LOCTEXT("ImportThemeSuccess", "Import theme succeeded"), SNotificationItem::CS_Success);
-				}
-				// if invalid theme: delete the copied file. 
-				else
-				{
-					// incomplete themes will not reach here. 
-					IPlatformFile::GetPlatformPhysical().DeleteFile(*DestPath); 
-					ShowNotification(LOCTEXT("ImportThemeFailureInvalidName", "Import theme failed: Invalid theme"), SNotificationItem::CS_Fail);
-				}
-			}
-			// if unable to copy the file to user-specific theme location, do nothing. 
-			else
-			{
-				ShowNotification(LOCTEXT("ImportThemeFailure", "Import theme failed"), SNotificationItem::CS_Fail);
-			}
-		}
-	}
+	PromptToImportTheme(FPlatformProcess::UserDir());
 	return FReply::Handled();
 }
 
@@ -782,7 +735,7 @@
 	return !ThemeEditorWindow.IsValid();
 }
 
-void FEditorStyleSettingsCustomization::ShowNotification(const FText& Text, SNotificationItem::ECompletionState CompletionState) const
+void FEditorStyleSettingsCustomization::ShowNotification(const FText& Text, SNotificationItem::ECompletionState CompletionState)
 {
 	FNotificationInfo Notification(Text);
 	Notification.ExpireDuration = 3.f;
@@ -791,6 +744,66 @@
 	FSlateNotificationManager::Get().AddNotification(Notification)->SetCompletionState(CompletionState);
 }
 
+void FEditorStyleSettingsCustomization::PromptToImportTheme(const FString& ImportPath)
+{
+	TArray<FString> OutFiles;
+	const void* ParentWindowHandle = FSlateApplication::Get().FindBestParentWindowHandleForDialogs(nullptr);
 
+	if (FDesktopPlatformModule::Get()->OpenFileDialog(ParentWindowHandle, LOCTEXT("ImportThemeDialogTitle", "Import theme...").ToString(), FPaths::GetPath(ImportPath), TEXT(""), TEXT("JSON files (*.json)|*.json"), EFileDialogFlags::None, OutFiles))
+	{
+		FString SourcePath = OutFiles[0];
+		const FString DestPath = USlateThemeManager::Get().GetUserThemeDir() / FPaths::GetCleanFilename(SourcePath);
+
+		FString PathPart;
+		FString Extension;
+		FString FilenameWithoutExtension;
+		FPaths::Split(SourcePath, PathPart, FilenameWithoutExtension, Extension);
+
+		// if theme name exists, don't import (to prevent from overwriting existing theme files)
+		if (!IsThemeNameValid(FilenameWithoutExtension))
+		{
+			ShowNotification(LOCTEXT("ImportThemeFailureNameExists", "Import theme failed: Theme name already exists"), SNotificationItem::CS_Fail);
+		}
+		// if theme name is valid: copying the file is safe (as it will not overwrite existing theme files)
+		else
+		{
+			const int32 NumOfThemesBefore = USlateThemeManager::Get().GetThemes().Num();
+
+			if (IPlatformFile::GetPlatformPhysical().CopyFile(*DestPath, *SourcePath))
+			{
+				// update the number of valid themes: 
+				USlateThemeManager::Get().LoadThemes();
+
+				// if valid theme: the theme Num will be updated. 
+				if (USlateThemeManager::Get().GetThemes().Num() != NumOfThemesBefore)
+				{
+					// Extract ID as a FString directly from a JSON file. 
+					FString ImportedThemeID;
+					GetThemeIdFromPath(SourcePath, ImportedThemeID);
+
+					// convert FString ID to a FGuid
+					FGuid ImportedThemeGUID = FGuid(ImportedThemeID);
+					USlateThemeManager::Get().ApplyTheme(ImportedThemeGUID);
+
+					ShowNotification(LOCTEXT("ImportThemeSuccess", "Import theme succeeded"), SNotificationItem::CS_Success);
+				}
+				// if invalid theme: delete the copied file. 
+				else
+				{
+					// incomplete themes will not reach here. 
+					IPlatformFile::GetPlatformPhysical().DeleteFile(*DestPath);
+					ShowNotification(LOCTEXT("ImportThemeFailureInvalidName", "Import theme failed: Invalid theme"), SNotificationItem::CS_Fail);
+				}
+			}
+			// if unable to copy the file to user-specific theme location, do nothing. 
+			else
+			{
+				ShowNotification(LOCTEXT("ImportThemeFailure", "Import theme failed"), SNotificationItem::CS_Fail);
+			}
+		}
+	}
+}
+
+
 #undef LOCTEXT_NAMESPACE
 
diff --git a/Engine/Source/Editor/UnrealEd/Private/Settings/EditorStyleSettingsCustomization.h b/Engine/Source/Editor/UnrealEd/Private/Settings/EditorStyleSettingsCustomization.h
--- a/Engine/Source/Editor/UnrealEd/Private/Settings/EditorStyleSettingsCustomization.h
+++ b/Engine/Source/Editor/UnrealEd/Private/Settings/EditorStyleSettingsCustomization.h
@@ -33,10 +33,18 @@ public:
 	/** Makes a new instance of this detail layout class for a specific detail view requesting it */
 	static TSharedRef<IDetailCustomization> MakeInstance();
 
+	FEditorStyleSettingsCustomization();
+
+	~FEditorStyleSettingsCustomization(); 
+
 	/** IDetailCustomization interface */
 	virtual void CustomizeDetails(IDetailLayoutBuilder& DetailLayout) override;
 	
 	void RefreshComboBox();
+
+	/** Import theme from OS */
+	static void PromptToImportTheme(const FString& ImportPath);
+
 private:
 	void GenerateThemeOptions(TSharedPtr<FString>& OutSelectedTheme);
 
@@ -50,7 +58,13 @@ private:
 	void OnThemePicked(TSharedPtr<FString> NewSelection, ESelectInfo::Type SelectInfo);
 	void OpenThemeEditorWindow(FOnThemeEditorClosed OnThemeEditorClosed);
 	bool IsThemeEditingEnabled() const;
-	void ShowNotification(const FText& Text, SNotificationItem::ECompletionState CompletionState) const;
+
+	/** Delegate for when theme gets changed by another class */
+	void OnThemeChanged(const FGuid ThemeId) { RefreshComboBox(); }
+
+	/** Show import status (success/failure) when users import a theme */
+	static void ShowNotification(const FText& Text, SNotificationItem::ECompletionState CompletionState);
+
 private:
 	TArray<TSharedPtr<FString>> ThemeOptions;
 	TSharedPtr<STextComboBox> ComboBox;
diff --git a/Engine/Source/Editor/UnrealEd/Private/UnrealEdEngine.cpp b/Engine/Source/Editor/UnrealEd/Private/UnrealEdEngine.cpp
--- a/Engine/Source/Editor/UnrealEd/Private/UnrealEdEngine.cpp
+++ b/Engine/Source/Editor/UnrealEd/Private/UnrealEdEngine.cpp
@@ -173,6 +173,8 @@ void UUnrealEdEngine::Init(IEngineLoop* InEngineLoop)
 
 	if (FPaths::IsProjectFilePathSet() && GIsEditor && !FApp::IsUnattended())
 	{
+
+		UE_SCOPED_ENGINE_ACTIVITY(TEXT("Initializing AutoReimportManager"));
 		AutoReimportManager = NewObject<UAutoReimportManager>();
 		AutoReimportManager->Initialize();
 	}
@@ -193,7 +195,8 @@ void UUnrealEdEngine::Init(IEngineLoop* InEngineLoop)
 
 
 		UEditorStyleSettings* Settings = GetMutableDefault<UEditorStyleSettings>();
-		ISettingsModule* SettingsModule = FModuleManager::GetModulePtr<ISettingsModule>("Settings");
+		ISettingsModule* SettingsModule = FModuleManager::GetModulePtr<ISettingsModule>("Settings"); 
+	
 
 		if (SettingsModule != nullptr)
 		{
@@ -202,9 +205,10 @@ void UUnrealEdEngine::Init(IEngineLoop* InEngineLoop)
 				NSLOCTEXT("UnrealEd", "Appearance_UserSettingsDescription", "Customize the look of the editor."),
 				Settings
 			);
-			
+
+			StyleSettingsPtr->OnImport().BindUObject(Settings, &UEditorStyleSettings::OnImportBegin);
 			StyleSettingsPtr->OnExport().BindUObject(Settings, &UEditorStyleSettings::OnExportBegin);
-		}
+		} 
 
 
 		FPropertyEditorModule& PropertyEditorModule = FModuleManager::Get().GetModuleChecked<FPropertyEditorModule>("PropertyEditor");
diff --git a/Engine/Source/Runtime/SlateCore/Private/Styling/StyleColors.cpp b/Engine/Source/Runtime/SlateCore/Private/Styling/StyleColors.cpp
--- a/Engine/Source/Runtime/SlateCore/Private/Styling/StyleColors.cpp
+++ b/Engine/Source/Runtime/SlateCore/Private/Styling/StyleColors.cpp
@@ -212,7 +212,7 @@ void USlateThemeManager::ApplyTheme(FGuid ThemeId)
 				// Unload existing colors
 				CurrentTheme->LoadedDefaultColors.Empty();
 			}
-
+			
 			FStyleTheme* Theme = Themes.FindByKey(ThemeId);
 			if (Theme)
 			{
@@ -230,6 +230,7 @@ void USlateThemeManager::ApplyTheme(FGuid ThemeId)
 		// Apply the new colors. Note that if the incoming theme is the same as the current theme we still apply the colors as they may have been overwritten by defaults
 		FMemory::Memcpy(ActiveColors.StyleColors, CurrentTheme->LoadedDefaultColors.GetData(), sizeof(FLinearColor)*CurrentTheme->LoadedDefaultColors.Num());
 	}
+	OnThemeChanged().Broadcast(CurrentThemeId); 
 }
 
 void USlateThemeManager::ApplyDefaultTheme()
@@ -456,6 +457,18 @@ void USlateThemeManager::LoadThemeColors(FStyleTheme& Theme)
 	}
 }
 
+bool USlateThemeManager::DoesThemeExist(const FGuid& ThemeID) const
+{
+	for (const FStyleTheme& Theme : Themes)
+	{
+		if (Theme.Id == ThemeID)
+		{
+			return true; 
+		}
+	}
+	return false; 
+}
+
 /*
 const FStyleTheme* UStyleColorTable::GetTheme(const FGuid ThemeId) const
 {
diff --git a/Engine/Source/Runtime/SlateCore/Public/Styling/StyleColors.h b/Engine/Source/Runtime/SlateCore/Public/Styling/StyleColors.h
--- a/Engine/Source/Runtime/SlateCore/Public/Styling/StyleColors.h
+++ b/Engine/Source/Runtime/SlateCore/Public/Styling/StyleColors.h
@@ -278,6 +278,20 @@ public:
 	 * @return the user theme dir. Themes in this dir are per-user and override engine and project themes
 	 */
 	FString GetUserThemeDir() const;
+		
+	/**
+	 * @return true if the ThemeID already exists in the theme dropdown. 
+	 */
+	bool DoesThemeExist(const FGuid& ThemeID) const;
+
+	DECLARE_EVENT_OneParam(USlateThemeManager, FThemeChangedEvent, FGuid);
+	/**
+	 * Returns an event delegate that is executed when the themeID has changed.
+	 *
+	 * @return the event delegate.
+	 */
+	FThemeChangedEvent& OnThemeChanged() { return ThemeChangedEvent; }
+
 private:
 
 	FStyleTheme& GetMutableCurrentTheme() { return *Themes.FindByKey(CurrentThemeId); }
@@ -286,11 +300,14 @@ private:
 	void EnsureValidCurrentTheme();
 	void LoadThemeColors(FStyleTheme& Theme);
 
-
 private:
 	FStyleTheme DefaultDarkTheme;
 	TArray<FStyleTheme> Themes;
 	FLinearColor DefaultColors[(int32)EStyleColor::MAX];
+
+	/** Broadcasts whenever the theme changes **/
+	FThemeChangedEvent ThemeChangedEvent;
+
 #endif // ALLOW_THEMES
 
 	FLinearColor GetDefaultColor(EStyleColor InColorId)
