FNames containing asset paths are deprecated. FSoftObjectPath or FTopLevelAssetPath should be used instead.
Fixups for editor support classes for actors.

diff --git a/Engine/Plugins/Experimental/ActorPalette/Source/ActorPalette/Private/ActorPaletteSettings.cpp b/Engine/Plugins/Experimental/ActorPalette/Source/ActorPalette/Private/ActorPaletteSettings.cpp
--- a/Engine/Plugins/Experimental/ActorPalette/Source/ActorPalette/Private/ActorPaletteSettings.cpp
+++ b/Engine/Plugins/Experimental/ActorPalette/Source/ActorPalette/Private/ActorPaletteSettings.cpp
@@ -6,7 +6,7 @@
 FAssetData FActorPaletteMapEntry::GetAsAssetData() const
 {
 	FAssetRegistryModule& AssetRegistryModule = FModuleManager::Get().LoadModuleChecked<FAssetRegistryModule>(TEXT("AssetRegistry"));
-	return AssetRegistryModule.Get().GetAssetByObjectPath(*MapPath);
+	return AssetRegistryModule.Get().GetAssetByObjectPath(FSoftObjectPath(MapPath));
 }
 
 UActorPaletteSettings::UActorPaletteSettings()
@@ -36,7 +36,7 @@ void UActorPaletteSettings::MarkAsRecentlyUsed(const FAssetData& MapAsset, int32
 {
 	if (MapAsset.IsValid())
 	{
-		const FString MapAssetPath = MapAsset.ObjectPath.ToString();
+		const FString MapAssetPath = MapAsset.GetObjectPathString();
 
 		// Remember as the most recent for this tab
 		if (!MostRecentLevelByTab.IsValidIndex(TabIndex))
@@ -62,7 +62,7 @@ void UActorPaletteSettings::MarkAsRecentlyUsed(const FAssetData& MapAsset, int32
 		{
 			// New entry
 			FActorPaletteMapEntry& NewEntry = SettingsPerLevel[SettingsPerLevel.AddDefaulted()];
-			NewEntry.MapPath = MapAsset.ObjectPath.ToString();
+			NewEntry.MapPath = MapAsset.GetObjectPathString();
 		}
 	}
 
@@ -76,13 +76,13 @@ void UActorPaletteSettings::MarkAsRecentlyUsed(const FAssetData& MapAsset, int32
 
 void UActorPaletteSettings::ToggleFavorite(const FAssetData& MapAsset)
 {
-	if (FavoritesList.Contains(MapAsset.ObjectPath.ToString()))
+	if (FavoritesList.Contains(MapAsset.GetObjectPathString()))
 	{
-		FavoritesList.Remove(MapAsset.ObjectPath.ToString());
+		FavoritesList.Remove(MapAsset.GetObjectPathString());
 	}
 	else
 	{
-		FavoritesList.Add(MapAsset.ObjectPath.ToString());
+		FavoritesList.Add(MapAsset.GetObjectPathString());
 	}
 
 #if WITH_EDITOR
diff --git a/Engine/Plugins/Experimental/ActorPalette/Source/ActorPalette/Private/ActorPaletteViewport.cpp b/Engine/Plugins/Experimental/ActorPalette/Source/ActorPalette/Private/ActorPaletteViewport.cpp
--- a/Engine/Plugins/Experimental/ActorPalette/Source/ActorPalette/Private/ActorPaletteViewport.cpp
+++ b/Engine/Plugins/Experimental/ActorPalette/Source/ActorPalette/Private/ActorPaletteViewport.cpp
@@ -65,7 +65,7 @@ public:
 						SNew(SImage)
 						.Image_Lambda([this]()
 						{
-							return FAppStyle::GetBrush(GetDefault<UActorPaletteSettings>()->FavoritesList.Contains(AssetData.ObjectPath.ToString()) ?
+							return FAppStyle::GetBrush(GetDefault<UActorPaletteSettings>()->FavoritesList.Contains(AssetData.GetObjectPathString()) ?
 								TEXT("Icons.Star") :
 								TEXT("PropertyWindow.Favorites_Disabled"));
 						})
diff --git a/Engine/Plugins/Experimental/AssetPlacementEdMode/Source/AssetPlacementEdMode/Private/SAssetPlacementPalette.cpp b/Engine/Plugins/Experimental/AssetPlacementEdMode/Source/AssetPlacementEdMode/Private/SAssetPlacementPalette.cpp
--- a/Engine/Plugins/Experimental/AssetPlacementEdMode/Source/AssetPlacementEdMode/Private/SAssetPlacementPalette.cpp
+++ b/Engine/Plugins/Experimental/AssetPlacementEdMode/Source/AssetPlacementEdMode/Private/SAssetPlacementPalette.cpp
@@ -570,7 +570,7 @@ void SAssetPlacementPalette::SetPaletteItems(TArrayView<const TObjectPtr<UPlacem
 	{
 		if (PaletteItem)
 		{
-			FAssetData AssetData = AssetRegistry.Get().GetAssetByObjectPath(PaletteItem->AssetPath.GetAssetPathName());
+			FAssetData AssetData = AssetRegistry.Get().GetAssetByObjectPath(PaletteItem->AssetPath.GetWithoutSubPath());
 			if (AssetData.IsValid())
 			{
 				PaletteItems.Add(MakeShared<FAssetPlacementPaletteItemModel>(AssetData, PaletteItem, SharedThis(this), ThumbnailPool));
diff --git a/Engine/Source/Editor/LevelEditor/Private/LevelEditorCreateActorMenu.cpp b/Engine/Source/Editor/LevelEditor/Private/LevelEditorCreateActorMenu.cpp
--- a/Engine/Source/Editor/LevelEditor/Private/LevelEditorCreateActorMenu.cpp
+++ b/Engine/Source/Editor/LevelEditor/Private/LevelEditorCreateActorMenu.cpp
@@ -467,7 +467,7 @@ void LevelEditorCreateActorMenu::FillAddReplaceActorMenu(UToolMenu* Menu, EActor
 			const TArray< FActorPlacementInfo > RecentlyPlaced = IPlacementModeModule::Get().GetRecentlyPlaced();
 			for (int Index = 0; Index < RecentlyPlaced.Num() && Index < 3; Index++)
 			{
-				FAssetData Asset = AssetRegistryModule.Get().GetAssetByObjectPath( *RecentlyPlaced[Index].ObjectPath );
+				FAssetData Asset = AssetRegistryModule.Get().GetAssetByObjectPath( FSoftObjectPath(RecentlyPlaced[Index].ObjectPath) );
 
 				if ( Asset.IsValid() )
 				{
diff --git a/Engine/Source/Editor/PlacementMode/Private/PlacementModeModule.cpp b/Engine/Source/Editor/PlacementMode/Private/PlacementModeModule.cpp
--- a/Engine/Source/Editor/PlacementMode/Private/PlacementModeModule.cpp
+++ b/Engine/Source/Editor/PlacementMode/Private/PlacementModeModule.cpp
@@ -438,7 +438,7 @@ void FPlacementModeModule::OnAssetRenamed(const FAssetData& AssetData, const FSt
 	{
 		if (RecentlyPlacedItem.ObjectPath == OldObjectPath)
 		{
-			RecentlyPlacedItem.ObjectPath = AssetData.ObjectPath.ToString();
+			RecentlyPlacedItem.ObjectPath = AssetData.GetObjectPathString();
 			break;
 		}
 	}
@@ -624,7 +624,7 @@ void FPlacementModeModule::RefreshRecentlyPlaced()
 		// If asset is pending delete, it will not be marked as RF_Standalone, in which case we skip it
 		if (Asset != nullptr && Asset->HasAnyFlags(RF_Standalone))
 		{
-			FAssetData AssetData = AssetRegistryModule.Get().GetAssetByObjectPath(*RecentlyPlacedItem.ObjectPath);
+			FAssetData AssetData = AssetRegistryModule.Get().GetAssetByObjectPath(FSoftObjectPath(RecentlyPlacedItem.ObjectPath));
 
 			if (AssetData.IsValid())
 			{
diff --git a/Engine/Source/Editor/PlacementMode/Private/SPlacementModeTools.cpp b/Engine/Source/Editor/PlacementMode/Private/SPlacementModeTools.cpp
--- a/Engine/Source/Editor/PlacementMode/Private/SPlacementModeTools.cpp
+++ b/Engine/Source/Editor/PlacementMode/Private/SPlacementModeTools.cpp
@@ -158,7 +158,7 @@ void SPlacementAssetEntry::Construct(const FArguments& InArgs, const TSharedPtr<
 		AssetEntryToolTip = FSlateApplicationBase::Get().MakeToolTip(
 			FText::Format(LOCTEXT("ItemInternalsTooltip", "Native Name: {0}\nAsset Path: {1}\nFactory Class: {2}"), 
 			FText::FromString(Item->NativeName), 
-			FText::FromName(Item->AssetData.ObjectPath),
+			FText::FromString(Item->AssetData.GetObjectPathString()),
 			FText::FromString(Item->Factory ? Item->Factory->GetClass()->GetName() : TEXT("None"))));
 	}
 
diff --git a/Engine/Source/Editor/UnrealEd/Classes/ActorFactories/ActorFactoryBasicShape.h b/Engine/Source/Editor/UnrealEd/Classes/ActorFactories/ActorFactoryBasicShape.h
--- a/Engine/Source/Editor/UnrealEd/Classes/ActorFactories/ActorFactoryBasicShape.h
+++ b/Engine/Source/Editor/UnrealEd/Classes/ActorFactories/ActorFactoryBasicShape.h
@@ -19,9 +19,9 @@ class UActorFactoryBasicShape : public UActorFactoryStaticMesh
 	virtual bool CanCreateActorFrom( const FAssetData& AssetData, FText& OutErrorMsg ) override;
 	virtual void PostSpawnActor(UObject* Asset, AActor* NewActor) override;
 
-	UNREALED_API static const FName BasicCube;
-	UNREALED_API static const FName BasicSphere;
-	UNREALED_API static const FName BasicCylinder;
-	UNREALED_API static const FName BasicCone;
-	UNREALED_API static const FName BasicPlane;
+	UNREALED_API static const FSoftObjectPath BasicCube;
+	UNREALED_API static const FSoftObjectPath BasicSphere;
+	UNREALED_API static const FSoftObjectPath BasicCylinder;
+	UNREALED_API static const FSoftObjectPath BasicCone;
+	UNREALED_API static const FSoftObjectPath BasicPlane;
 };
diff --git a/Engine/Source/Editor/UnrealEd/Private/Factories/ActorFactory.cpp b/Engine/Source/Editor/UnrealEd/Private/Factories/ActorFactory.cpp
--- a/Engine/Source/Editor/UnrealEd/Private/Factories/ActorFactory.cpp
+++ b/Engine/Source/Editor/UnrealEd/Private/Factories/ActorFactory.cpp
@@ -225,10 +225,19 @@ bool UActorFactory::CanCreateActorFrom( const FAssetData& AssetData, FText& OutE
 {
 	TRACE_CPUPROFILER_EVENT_SCOPE(UActorFactory::CanCreateActorFrom);
 
+	if (!AssetData.IsValid())
+	{
+		return false;
+	}
+
+	UObject* DefaultActor = GetDefaultActor(AssetData);
+	FSoftObjectPath DefaultActorPath(DefaultActor);
+	FSoftObjectPath DefaultActorClassPath(DefaultActor->GetClass());
+
 	// By Default we assume the factory can't work with existing asset data
 	return !AssetData.IsValid() || 
-		AssetData.ObjectPath == *GetDefaultActor( AssetData )->GetPathName() || 
-		AssetData.ObjectPath == *GetDefaultActor( AssetData )->GetClass()->GetPathName();
+		AssetData.GetSoftObjectPath() == DefaultActorPath || 
+		AssetData.GetSoftObjectPath() == DefaultActorClassPath;
 }
 
 AActor* UActorFactory::GetDefaultActor( const FAssetData& AssetData )
@@ -608,11 +617,11 @@ FQuat UActorFactoryStaticMesh::AlignObjectToSurfaceNormal(const FVector& InSurfa
 UActorFactoryBasicShape
 -----------------------------------------------------------------------------*/
 
-const FName UActorFactoryBasicShape::BasicCube("/Engine/BasicShapes/Cube.Cube");
-const FName UActorFactoryBasicShape::BasicSphere("/Engine/BasicShapes/Sphere.Sphere");
-const FName UActorFactoryBasicShape::BasicCylinder("/Engine/BasicShapes/Cylinder.Cylinder");
-const FName UActorFactoryBasicShape::BasicCone("/Engine/BasicShapes/Cone.Cone");
-const FName UActorFactoryBasicShape::BasicPlane("/Engine/BasicShapes/Plane.Plane");
+const FSoftObjectPath UActorFactoryBasicShape::BasicCube("/Engine/BasicShapes/Cube", "Cube", {});
+const FSoftObjectPath UActorFactoryBasicShape::BasicSphere("/Engine/BasicShapes/Sphere", "Sphere", {});
+const FSoftObjectPath UActorFactoryBasicShape::BasicCylinder("/Engine/BasicShapes/Cylinder", "Cylinder", {});
+const FSoftObjectPath UActorFactoryBasicShape::BasicCone("/Engine/BasicShapes/Cone", "Cone", {});
+const FSoftObjectPath UActorFactoryBasicShape::BasicPlane("/Engine/BasicShapes/Plane", "Plane", {});
 
 UActorFactoryBasicShape::UActorFactoryBasicShape(const FObjectInitializer& ObjectInitializer)
 	: Super(ObjectInitializer)
@@ -624,7 +633,8 @@ UActorFactoryBasicShape::UActorFactoryBasicShape(const FObjectInitializer& Objec
 
 bool UActorFactoryBasicShape::CanCreateActorFrom( const FAssetData& AssetData, FText& OutErrorMsg )
 {
-	if(AssetData.IsValid() && (AssetData.ObjectPath == BasicCube || AssetData.ObjectPath == BasicSphere || AssetData.ObjectPath == BasicCone || AssetData.ObjectPath == BasicCylinder || AssetData.ObjectPath == BasicPlane) )
+	FSoftObjectPath AssetPath = AssetData.GetSoftObjectPath();
+	if(AssetData.IsValid() && (AssetPath == BasicCube || AssetPath == BasicSphere || AssetPath == BasicCone || AssetPath == BasicCylinder || AssetPath == BasicPlane) )
 	{
 		return true;
 	}
@@ -700,7 +710,7 @@ bool UActorFactoryDeferredDecal::CanCreateActorFrom( const FAssetData& AssetData
 			return false;
 		}
 
-		CurrentAssetData = AssetRegistry.GetAssetByObjectPath( *ObjectPath );
+		CurrentAssetData = AssetRegistry.GetAssetByObjectPath( FSoftObjectPath(ObjectPath) );
 		if ( !CurrentAssetData.IsValid() )
 		{
 			OutErrorMsg = NSLOCTEXT("CanCreateActor", "NoMaterial", "A valid material must be specified.");
@@ -1038,7 +1048,7 @@ bool UActorFactoryAnimationAsset::CanCreateActorFrom( const FAssetData& AssetDat
 			return false;
 		}
 
-		FAssetData SkeletonData = AssetRegistry.GetAssetByObjectPath( *SkeletonPath );
+		FAssetData SkeletonData = AssetRegistry.GetAssetByObjectPath( FSoftObjectPath(SkeletonPath) );
 
 		if ( !SkeletonData.IsValid() )
 		{
@@ -1200,7 +1210,7 @@ bool UActorFactorySkeletalMesh::CanCreateActorFrom( const FAssetData& AssetData,
 			return false;
 		}
 
-		FAssetData TargetSkeleton = AssetRegistry.GetAssetByObjectPath( *TargetSkeletonPath );
+		FAssetData TargetSkeleton = AssetRegistry.GetAssetByObjectPath( FSoftObjectPath(TargetSkeletonPath) );
 		if ( !TargetSkeleton.IsValid() )
 		{
 			OutErrorMsg = NSLOCTEXT("CanCreateActor", "NoAnimBPTargetSkeleton", "UAnimBlueprints must have a valid Target Skeleton.");
@@ -1417,7 +1427,7 @@ UActorFactoryEmptyActor::UActorFactoryEmptyActor(const FObjectInitializer& Objec
 
 bool UActorFactoryEmptyActor::CanCreateActorFrom( const FAssetData& AssetData, FText& OutErrorMsg )
 {
-	return AssetData.ObjectPath == FName(*AActor::StaticClass()->GetPathName());
+	return AssetData.ToSoftObjectPath() == FSoftObjectPath(AActor::StaticClass());
 }
 
 AActor* UActorFactoryEmptyActor::SpawnActor(UObject* InAsset, ULevel* InLevel, const FTransform& InTransform, const FActorSpawnParameters& InSpawnParams)
@@ -1455,7 +1465,7 @@ UActorFactoryCharacter::UActorFactoryCharacter(const FObjectInitializer& ObjectI
 
 bool UActorFactoryCharacter::CanCreateActorFrom(const FAssetData& AssetData, FText& OutErrorMsg)
 {
-	return AssetData.ObjectPath == FName(*ACharacter::StaticClass()->GetPathName());
+	return AssetData.ToSoftObjectPath() == FSoftObjectPath(ACharacter::StaticClass());
 }
 
 /*-----------------------------------------------------------------------------
@@ -1470,7 +1480,7 @@ UActorFactoryPawn::UActorFactoryPawn(const FObjectInitializer& ObjectInitializer
 
 bool UActorFactoryPawn::CanCreateActorFrom(const FAssetData& AssetData, FText& OutErrorMsg)
 {
-	return AssetData.ObjectPath == FName(*APawn::StaticClass()->GetPathName());
+	return AssetData.GetSoftObjectPath() == FSoftObjectPath(APawn::StaticClass());
 }
 
 /*-----------------------------------------------------------------------------
@@ -1564,7 +1574,7 @@ AActor* UActorFactoryClass::GetDefaultActor( const FAssetData& AssetData )
 {
 	if ( AssetData.IsValid() && AssetData.IsInstanceOf( UClass::StaticClass() ) )
 	{
-		UClass* ActualClass = Cast<UClass>(StaticLoadObject(UClass::StaticClass(), nullptr, *AssetData.ObjectPath.ToString(), nullptr, LOAD_NoWarn, nullptr));
+		UClass* ActualClass = Cast<UClass>(StaticLoadObject(UClass::StaticClass(), nullptr, *AssetData.GetObjectPathString(), nullptr, LOAD_NoWarn, nullptr));
 			
 		//Cast<UClass>(AssetData.GetAsset());
 		if ( (nullptr != ActualClass) && ActualClass->IsChildOf(AActor::StaticClass()) )
