Added support for localization text gathering into the moderation commandlet.

diff --git a/Engine/Source/Editor/UnrealEd/Classes/Commandlets/BaseIteratePackagesCommandlet.h b/Engine/Source/Editor/UnrealEd/Classes/Commandlets/BaseIteratePackagesCommandlet.h
--- a/Engine/Source/Editor/UnrealEd/Classes/Commandlets/BaseIteratePackagesCommandlet.h
+++ b/Engine/Source/Editor/UnrealEd/Classes/Commandlets/BaseIteratePackagesCommandlet.h
@@ -227,7 +227,7 @@ protected:
 	/**
 	 * Allows the commandlet to do something after all the packages have been processed
 	 */
-	virtual void PostSavePackages() { }
+	virtual void PostProcessPackages() { }
 
 	/**
 	 * Removes any UClass exports from packages which aren't script packages.
diff --git a/Engine/Source/Editor/UnrealEd/Classes/Commandlets/GenerateModerationArtifactsCommandlet.h b/Engine/Source/Editor/UnrealEd/Classes/Commandlets/GenerateModerationArtifactsCommandlet.h
--- a/Engine/Source/Editor/UnrealEd/Classes/Commandlets/GenerateModerationArtifactsCommandlet.h
+++ b/Engine/Source/Editor/UnrealEd/Classes/Commandlets/GenerateModerationArtifactsCommandlet.h
@@ -7,6 +7,68 @@
 
 UNREALED_API DECLARE_LOG_CATEGORY_EXTERN(LogModerationArtifactsCommandlet, Log, All);
 
+
+
+
+USTRUCT()
+struct FModerationAsset
+{
+public:
+	GENERATED_BODY();
+
+	UObject* Object;
+
+	UPROPERTY()
+	FString FullPath;
+
+	UPROPERTY()
+	FString ClassName;
+
+	UPROPERTY()
+	TArray<FString> ModerationArtifactFilenames;
+
+};
+
+
+USTRUCT()
+struct FModerationPackage
+{
+public:
+	GENERATED_BODY();
+
+
+	FModerationAsset* FindOrCreateModerationAsset(const UObject* InObject);
+
+	UPackage* Package;
+
+	UPROPERTY()
+	FString PackagePath;
+
+	UPROPERTY()
+	FString PackageHash;
+
+	UPROPERTY()
+	TArray<struct FModerationAsset> Assets;
+
+
+};
+
+
+USTRUCT()
+struct FModerationManifest
+{
+	GENERATED_BODY();
+
+
+	FModerationPackage* FindOrCreateModerationPackage(UPackage* InPackage);
+
+	FString CreateModerationAssetFileName(const UObject* Object, const FString& Extension);
+
+	UPROPERTY()
+		TArray<struct FModerationPackage> Packages;
+
+};
+
 UCLASS()
 // Added UNREALED_API to expose this to the save packages test
 // Base commandlet used to iterate packages provided on a commandline
@@ -64,6 +126,8 @@ class UNREALED_API UGenerateModerationArtifactsCommandlet : public UBaseIterateP
 	*/
 	// virtual void PerformAdditionalOperations(class UWorld* World, bool& bSavePackage) { }
 
+	virtual void PostProcessPackages() override;
+
 private:
 	void GatherLocalizationFromPackage(class UPackage* Package);
 	void GatherFStringsFromObject(class UObject* Object);
@@ -76,4 +140,10 @@ private:
 
 	FString OutputPath;
 
+	FModerationManifest Manifest; 
+	TMap<const UObject*, TWeakPtr<FModerationAsset>> ModerationAssetMap;
+	TMap<const UPackage*, TWeakPtr<FModerationPackage>> ModerationPackageMap;
+
+	TMap<FName, float> TimerStats;
 };
+
diff --git a/Engine/Source/Editor/UnrealEd/Private/Commandlets/BaseIteratePackagesCommandlet.cpp b/Engine/Source/Editor/UnrealEd/Private/Commandlets/BaseIteratePackagesCommandlet.cpp
--- a/Engine/Source/Editor/UnrealEd/Private/Commandlets/BaseIteratePackagesCommandlet.cpp
+++ b/Engine/Source/Editor/UnrealEd/Private/Commandlets/BaseIteratePackagesCommandlet.cpp
@@ -1132,7 +1132,7 @@ int32 UBaseIteratePackagesCommandlet::Main( const FString& Params )
 	IAssetRegistry& AssetRegistry = AssetRegistryModule.Get();
 	AssetRegistry.Tick(-1.0f);
 
-	PostSavePackages();
+	PostProcessPackages();
 
 	// Submit the results to source control
 	CheckInFiles(FilesToSubmit, GetChangelistDescription());
diff --git a/Engine/Source/Editor/UnrealEd/Private/Commandlets/GeneratemoderationArtifactsCommandlet.cpp b/Engine/Source/Editor/UnrealEd/Private/Commandlets/GeneratemoderationArtifactsCommandlet.cpp
--- a/Engine/Source/Editor/UnrealEd/Private/Commandlets/GeneratemoderationArtifactsCommandlet.cpp
+++ b/Engine/Source/Editor/UnrealEd/Private/Commandlets/GeneratemoderationArtifactsCommandlet.cpp
@@ -12,6 +12,15 @@
 #include "Components/StaticMeshComponent.h"
 #include "Misc/FileHelper.h"
 
+// localization includes 
+#include "Serialization/PropertyLocalizationDataGathering.h"
+// string proxy includes
+//#include "Serialization/NameAsStringProxyArchive.h"
+// saving json manifest file
+#include "JsonObjectConverter.h"
+
+PRAGMA_DISABLE_OPTIMIZATION
+
 /**-----------------------------------------------------------------------------
  *	UGenerateModerationArtifactsCommandlet commandlet.
  *
@@ -22,6 +31,8 @@
 
 DEFINE_LOG_CATEGORY(LogModerationArtifactsCommandlet);
 
+
+
 UGenerateModerationArtifactsCommandlet::UGenerateModerationArtifactsCommandlet(const FObjectInitializer& ObjectInitializer)
 	: Super(ObjectInitializer)
 {
@@ -51,128 +62,326 @@ void UGenerateModerationArtifactsCommandlet::PerformAdditionalOperations(class U
 	GatherLocalizationFromPackage(Package);
 }
 
-
-
-void UGenerateModerationArtifactsCommandlet::GatherFStringsFromObject(class UObject* Object)
+class FDefaultObjectPropertyIterator
 {
-	
-}
+public:
 
+	virtual ~FDefaultObjectPropertyIterator() {}
 
-void UGenerateModerationArtifactsCommandlet::GatherLocalizationFromPackage(class UPackage* Package)
-{
+	void IterateObject(UObject* Object, bool bRecurse)
+	{
+		IterateUObjectProperty(Object->GetClass(), Object, Object->GetClass()->GetDefaultObject());
+	}
 
-/*
-	TArray<FGatherableTextData> GatherableTextDataArray;
+	virtual bool ProcessProperty(FProperty* Property, void* Ptr, void* DefaultPtr) = 0;
 
-	// Gathers from the given package
-	EPropertyLocalizationGathererResultFlags GatherableTextResultFlags = EPropertyLocalizationGathererResultFlags::Empty;
-	FPropertyLocalizationDataGatherer(GatherableTextDataArray, Package, GatherableTextResultFlags);
 
-	if (PackagePendingGather.PackageLocCacheState != EPackageLocCacheState::Cached)
+private:
+	TSet<void*> ProcessedObjectsSet;
+
+	void IterateProperty(FProperty* Property, void* ObjectPtr, void* DefaultObjectPtr)
 	{
-		ProcessGatherableTextDataArray(GatherableTextDataArray);
+		if (ProcessProperty(Property, ObjectPtr, DefaultObjectPtr) == false)
+		{
+			// stop recursion into this property
+			return;
+		}
+		if (Property->IsA(FStructProperty::StaticClass()))
+		{
+			IterateStructProperty((FStructProperty*)Property, ObjectPtr, DefaultObjectPtr);
+		}
+		else if (Property->IsA(FObjectPropertyBase::StaticClass()))
+		{
+			FObjectPropertyBase* ObjectProperty = (FObjectPropertyBase*)Property;
+			void* NewObject = ObjectProperty->GetObjectPropertyValue(ObjectPtr);
+			void* NewDefaultObject = DefaultObjectPtr ? ObjectProperty->GetObjectPropertyValue(DefaultObjectPtr) : nullptr;
+			IterateUObjectProperty(ObjectProperty->PropertyClass, NewObject, NewDefaultObject);
+		}
+		else if (Property->IsA(FArrayProperty::StaticClass()))
+		{
+			FArrayProperty* ArrayProperty = (FArrayProperty*)Property;
+			IterateArrayProperty(ArrayProperty, ObjectPtr, DefaultObjectPtr);
+		}
+		else if (Property->IsA(FMapProperty::StaticClass()))
+		{
+			FMapProperty* MapProperty = (FMapProperty*)Property;
+			IterateMapProperty(MapProperty, ObjectPtr, DefaultObjectPtr);
+		}
+		else if (Property->IsA(FSetProperty::StaticClass()))
+		{
+			FSetProperty* SetProperty = (FSetProperty*)Property;
+			IterateSetProperty(SetProperty, ObjectPtr, DefaultObjectPtr);
+		}
 	}
-*/
 
-#if 0
-	int32 NumPackagesFailed = 0;
-	TArray<FName> PackagesWithStaleGatherCache;
-	TArray<FGatherableTextData> GatherableTextDataArray;
-	while (PackagesPendingGather.Num() > 0)
+	void IterateSetProperty(FSetProperty* SetProperty, void* ObjectPtr, void* DefaultObjectPtr)
 	{
-		const FPackagePendingGather PackagePendingGather = PackagesPendingGather.Pop(/ *bAllowShrinking* /false);
-		const FNameBuilder PackageNameStr(PackagePendingGather.PackageName);
+		FScriptSetHelper SetHelper(SetProperty, ObjectPtr);
+		FScriptSetHelper DefaultSetHelper(SetProperty, DefaultObjectPtr);
+		bool bUseDefault = DefaultObjectPtr ? SetHelper.Num() == DefaultSetHelper.Num() : false;
+		for (int32 Index = 0; Index < SetHelper.Num(); ++Index)
+		{
+			void* KeyPtr = SetHelper.GetElementPtr(Index);
+			void* DefaultKeyPtr = bUseDefault ? DefaultSetHelper.GetElementPtr(Index) : nullptr;
+			IterateProperty(SetHelper.GetElementProperty(), KeyPtr, DefaultKeyPtr);
+		}
+	}
 
-		const int32 CurrentPackageNum = ++NumPackagesProcessed + NumPackagesFailed;
-		const float PercentageComplete = (((float)CurrentPackageNum / (float)PackageCount) * 100.0f);
-		UE_LOG(LogGatherTextFromAssetsCommandlet, Display, TEXT("[%6.2f%%] Loading package: '%s'..."), PercentageComplete, *PackageNameStr);
+	void IterateMapProperty(FMapProperty* MapProperty, void* ObjectPtr, void* DefaultObjectPtr)
+	{
+		FScriptMapHelper MapHelper(MapProperty, ObjectPtr);
+		FScriptMapHelper DefaultMapHelper(MapProperty, DefaultObjectPtr);
+		bool bUseDefault = DefaultObjectPtr ? DefaultMapHelper.Num() == MapHelper.Num() : false;
+		for (int32 Index = 0; Index < MapHelper.Num(); ++Index)
+		{
+			void* KeyPtr = MapHelper.GetKeyPtr(Index);
+			void* DefaultKeyPtr = bUseDefault ? DefaultMapHelper.GetKeyPtr(Index) : nullptr;
+			IterateProperty(MapHelper.GetKeyProperty(), KeyPtr, DefaultKeyPtr);
+			void* ValuePtr = MapHelper.GetValuePtr(Index);
+			void* DefaultValuePtr = bUseDefault ? DefaultMapHelper.GetValuePtr(Index) : nullptr;
+			IterateProperty(MapHelper.GetValueProperty(), ValuePtr, DefaultValuePtr);
+		}
+	}
 
-		UPackage* Package = nullptr;
+	void IterateArrayProperty(FArrayProperty* ArrayProperty, void* ObjectPtr, void* DefaultObjectPtr)
+	{
+		FScriptArrayHelper ArrayHelper(ArrayProperty, ObjectPtr);
+		FScriptArrayHelper DefaultArrayHelper(ArrayProperty, DefaultObjectPtr);
+		bool bUseDefault = DefaultObjectPtr ? DefaultArrayHelper.Num() == ArrayHelper.Num() : false;
+		for (int32 Index = 0; Index < ArrayHelper.Num(); ++Index)
 		{
-			FLoadPackageLogOutputRedirector::FScopedCapture ScopedCapture(&LogOutputRedirector, *PackageNameStr);
-			Package = LoadPackage(nullptr, *PackageNameStr, LOAD_NoWarn | LOAD_Quiet);
+			void* ArrayValuePtr = ArrayHelper.GetRawPtr(Index);
+			void* DefaultValuePtr = bUseDefault ? DefaultArrayHelper.GetRawPtr(Index) : nullptr;
+			IterateProperty(ArrayProperty->Inner, ArrayValuePtr, DefaultValuePtr);
 		}
+	}
 
-		if (!Package)
+	void IterateStructProperty(const FStructProperty* StructProperty, void* ObjectPtr, void* DefaultObjectPtr)
+	{
+		for (TFieldIterator<FProperty> PropertyIt(StructProperty->Struct); PropertyIt; ++PropertyIt)
 		{
-			UE_LOG(LogGatherTextFromAssetsCommandlet, Warning, TEXT("Failed to load package: '%s'."), *PackageNameStr);
-			++NumPackagesFailed;
-			continue;
+			FProperty* Property = *PropertyIt;
+			for (int32 Index = 0; Index < Property->ArrayDim; ++Index)
+			{
+				void* PropertyPtr = Property->ContainerPtrToValuePtr<void>(ObjectPtr, Index);
+				void* DefaultPropertyPtr = nullptr;
+				if(DefaultObjectPtr)
+				{
+					DefaultPropertyPtr = Property->ContainerPtrToValuePtr<void>(DefaultObjectPtr, Index);
+				}
+				IterateProperty(Property, PropertyPtr, DefaultPropertyPtr);
+			}
 		}
+	}
 
-		// Tick background tasks
-		if (GShaderCompilingManager)
+	void IterateUObjectProperty(const UClass* Class, void* ObjectPtr, void* DefaultObjectPtr)
+	{
+		UObject* Object = (UObject*)ObjectPtr;
+		if (Object == nullptr)
 		{
-			GShaderCompilingManager->ProcessAsyncResults(true, false);
+			return;
 		}
-		if (GDistanceFieldAsyncQueue)
+		UObject* DefaultObject = (UObject*)DefaultObjectPtr;
+		// make sure they are the same type
+		bool bUseDefaultObject = true;
+		if (DefaultObject == nullptr || Object->GetClass() != DefaultObject->GetClass())
 		{
-			GDistanceFieldAsyncQueue->ProcessAsyncTasks();
+			bUseDefaultObject = false;
 		}
-
-		// Because packages may not have been resaved after this flagging was implemented, we may have added packages to load that weren't flagged - potential false positives.
-		// The loading process should have reflagged said packages so that only true positives will have this flag.
-		if (Package->RequiresLocalizationGather())
+		if (ProcessedObjectsSet.Contains(ObjectPtr))
 		{
-			UE_LOG(LogGatherTextFromAssetsCommandlet, Display, TEXT("[%6.2f%%] Gathering package: '%s'..."), PercentageComplete, *PackageNameStr);
-
-			// Gathers from the given package
-			EPropertyLocalizationGathererResultFlags GatherableTextResultFlags = EPropertyLocalizationGathererResultFlags::Empty;
-			FPropertyLocalizationDataGatherer(GatherableTextDataArray, Package, GatherableTextResultFlags);
-
-			bool bSavePackage = false;
+			// ignore ciricular references
+			return;
+		}
+		ProcessedObjectsSet.Add(ObjectPtr);
 
-			// Optionally check to see whether the clean gather we did is in-sync with the gather cache and deal with it accordingly
-			if ((bReportStaleGatherCache || bFixStaleGatherCache) && PackagePendingGather.PackageLocCacheState == EPackageLocCacheState::Cached)
+		for (TFieldIterator<FProperty> PropertyIt(Class); PropertyIt; ++PropertyIt)
+		{
+			FProperty* Property = *PropertyIt;
+			// static const FName NAME_Outer = GET_MEMBER_NAME_CHECKED(UObject, OuterPrivate);
+			if (Property->GetFName() == NAME_Outer)
+			{
+				// don't follow owner property
+				continue;
+			}
+			for (int32 Index = 0; Index < Property->ArrayDim; ++Index)
 			{
-				// Look for any structurally significant changes (missing, added, or changed texts) in the cache
-				// Ignore insignificant things (like source changes caused by assets moving or being renamed)
-				if (EnumHasAnyFlags(GatherableTextResultFlags, EPropertyLocalizationGathererResultFlags::HasTextWithInvalidPackageLocalizationID)
-					|| !IsGatherableTextDataIdentical(GatherableTextDataArray, PackagePendingGather.GatherableTextDataArray))
+				void* PropertyPtr = Property->ContainerPtrToValuePtr<void>(ObjectPtr, Index);
+				void* DefaultPropertyPtr = nullptr;
+				if (DefaultObjectPtr && bUseDefaultObject)
 				{
-					PackagesWithStaleGatherCache.Add(PackagePendingGather.PackageName);
-
-					if (bFixStaleGatherCache)
-					{
-						bSavePackage = true;
-					}
+					DefaultPropertyPtr = Property->ContainerPtrToValuePtr<void>(DefaultObjectPtr, Index);
 				}
+				IterateProperty(Property, PropertyPtr, DefaultPropertyPtr);
 			}
+		}
+	}
+};
 
-			// Optionally save the package if it is missing a gather cache
-			if (bFixMissingGatherCache && PackagePendingGather.PackageLocCacheState == EPackageLocCacheState::Uncached_TooOld)
+class FStringCollector : public FDefaultObjectPropertyIterator
+{
+public:
+	TSet<FString> AllStrings;
+	virtual bool ProcessProperty(FProperty* Property, void* Ptr, void* DefaultPtr) override
+	{
+		if (Property->IsA(FNameProperty::StaticClass()))
+		{
+			const FName* NameValue = (FName*)Ptr;
+			if (*NameValue != NAME_None)
 			{
-				bSavePackage = true;
+				if (Property->Identical(Ptr, DefaultPtr))
+				{
+					UE_LOG(LogModerationArtifactsCommandlet, Display, TEXT("Found property name %s value %s is same as default property, ignoring string"), *Property->GetName(), *NameValue->ToString());
+				}
+				else
+				{
+					AllStrings.Add(NameValue->ToString());
+				}
 			}
+		}
 
-			// Re-save the package to attempt to fix it?
-			if (bSavePackage)
+		if (Property->IsA(FStrProperty::StaticClass()))
+		{
+			const FString* StringValue = (FString*)Ptr;
+			if (StringValue != nullptr && !StringValue->IsEmpty())
 			{
-				UE_LOG(LogGatherTextFromAssetsCommandlet, Display, TEXT("Resaving package: '%s'..."), *PackageNameStr);
-
-				bool bSavedPackage = false;
+				if( Property->Identical(Ptr, DefaultPtr))
 				{
-					FLoadPackageLogOutputRedirector::FScopedCapture ScopedCapture(&LogOutputRedirector, PackageNameStr);
-					bSavedPackage = FLocalizedAssetSCCUtil::SavePackageWithSCC(SourceControlInfo, Package, PackagePendingGather.PackageFilename);
+					UE_LOG(LogModerationArtifactsCommandlet, Display, TEXT("Found property name %s value %s is same as default property, ignoring string"), *Property->GetName(), **StringValue);
 				}
-
-				if (!bSavedPackage)
+				else
 				{
-					UE_LOG(LogGatherTextFromAssetsCommandlet, Warning, TEXT("Failed to resave package: '%s'."), *PackageNameStr);
+					AllStrings.Add(*StringValue);
 				}
 			}
+		}
+		return true;
+	}
+};
+
+
+void UGenerateModerationArtifactsCommandlet::GatherFStringsFromObject(class UObject* Object)
+{
+	/*FGatherStringArchive GatherStringArchive;
+	GatherStringArchive << Object;*/
+
+	FStringCollector StringCollector;
+	StringCollector.IterateObject(Object, true);
+
+	if (StringCollector.AllStrings.Num() > 0)
+	{
+		FString FileName = CreateOutputFileName(Object, TEXT("str"));
+		TArray<FString> AllStringsArray;
+		for (const FString& Str : StringCollector.AllStrings)
+		{
+			AllStringsArray.Add(Str);
+		}
+		FFileHelper::SaveStringArrayToFile(AllStringsArray, *FileName);
+	}
+
+#if 0
 
-			// This package may have already been cached in cases where we're reporting or fixing assets with a stale gather cache
-			// This check prevents it being gathered a second time
-			if (PackagePendingGather.PackageLocCacheState != EPackageLocCacheState::Cached)
+
+	TArray<FString> AllStrings;
+	for (TPropertyValueIterator<FStrProperty> It(Object->GetClass(), Object); It; ++It)
+	{
+		const FStrProperty* Property = It->Key;
+		const FString* StringValue = (FString*)It->Value;
+		if (Property->GetClass())
+		{
+			UE_LOG(LogModerationArtifactsCommandlet, Display, TEXT("Found property %s of type class %s"), *Property->GetName(), *Property->GetClass()->GetName());
+		}
+		if (StringValue != nullptr && !StringValue->IsEmpty())
+		{
+			AllStrings.Add(*StringValue);
+		}
+	}
+
+	for (TPropertyValueIterator<FNameProperty> It(Object->GetClass(), Object); It; ++It)
+	{
+		const FNameProperty* Property = It->Key;
+		const FName* NameValue = (FName*)It->Value;
+		if (Property->GetClass())
+		{
+			UE_LOG(LogModerationArtifactsCommandlet, Display, TEXT("Found property %s of type class %s"), *Property->GetName(), *Property->GetClass()->GetName());
+		}
+
+		// Property->Identical()
+		// Property->GetClass()->GetDefaultObject()->G
+
+		if (NameValue != nullptr && *NameValue != NAME_None)
+		{
+			AllStrings.Add(NameValue->ToString());
+		}
+	}
+
+
+	if(AllStrings.Num() > 0)
+	{
+		FString FileName = CreateOutputFileName(Object, TEXT("str"));
+		FFileHelper::SaveStringArrayToFile(AllStrings, *FileName);
+	}
+#endif
+}
+
+
+void UGenerateModerationArtifactsCommandlet::GatherLocalizationFromPackage(class UPackage* Package)
+{
+
+
+	TArray<FGatherableTextData> GatherableTextDataArray;
+	EPropertyLocalizationGathererResultFlags GatherableTextResultFlags = EPropertyLocalizationGathererResultFlags::Empty;
+	FPropertyLocalizationDataGatherer(GatherableTextDataArray, Package, GatherableTextResultFlags);
+
+	FString OutputFilePath = CreateOutputFileName(Package, TEXT("loc"));
+
+
+	FString LocalizationTargetName = FPaths::GetBaseFilename(OutputFilePath);
+	/* {
+		FString ManifestName;
+		GetStringFromConfig(TEXT("CommonSettings"), TEXT("ManifestName"), ManifestName, GatherTextConfigPath);
+		LocalizationTargetName = FPaths::GetBaseFilename(ManifestName);
+	}*/
+
+	// Basic helper that can be used only to gather a new manifest for writing
+	TSharedRef<FLocTextHelper> GatherManifestHelper = MakeShared<FLocTextHelper>(LocalizationTargetName, nullptr);
+	GatherManifestHelper->LoadManifest(ELocTextHelperLoadFlags::Create);
+	bool bContainsText = false;
+	for (const FGatherableTextData& GatherableTextData : GatherableTextDataArray)
+	{
+		for (const FTextSourceSiteContext& TextSourceSiteContext : GatherableTextData.SourceSiteContexts)
+		{
+
+			if (TextSourceSiteContext.KeyName.IsEmpty())
 			{
-				ProcessGatherableTextDataArray(GatherableTextDataArray);
+				UE_LOG(LogModerationArtifactsCommandlet, Warning, TEXT("Detected missing key on asset \"%s\"."), *TextSourceSiteContext.SiteDescription);
+				continue;
 			}
 
-			GatherableTextDataArray.Reset();
+			static const FLocMetadataObject DefaultMetadataObject;
+
+			FManifestContext Context;
+			Context.Key = TextSourceSiteContext.KeyName;
+			Context.KeyMetadataObj = !(FLocMetadataObject::IsMetadataExactMatch(&TextSourceSiteContext.KeyMetaData, &DefaultMetadataObject)) ? MakeShareable(new FLocMetadataObject(TextSourceSiteContext.KeyMetaData)) : nullptr;
+			Context.InfoMetadataObj = !(FLocMetadataObject::IsMetadataExactMatch(&TextSourceSiteContext.InfoMetaData, &DefaultMetadataObject)) ? MakeShareable(new FLocMetadataObject(TextSourceSiteContext.InfoMetaData)) : nullptr;
+			Context.bIsOptional = TextSourceSiteContext.IsOptional;
+			Context.SourceLocation = TextSourceSiteContext.SiteDescription;
+			//Context.PlatformName = GetSplitPlatformNameFromPath(TextSourceSiteContext.SiteDescription);
+
+			FLocItem Source(GatherableTextData.SourceData.SourceString);
+
+			GatherManifestHelper->AddSourceText(GatherableTextData.NamespaceName, Source, Context);
+			bContainsText = true;
 		}
-	#endif
+	}
+	if (bContainsText)
+	{
+		FText Error;
+		if (GatherManifestHelper->SaveManifest(OutputFilePath, &Error)==false)
+		{
+			UE_LOG(LogModerationArtifactsCommandlet, Error, TEXT("Failed to save localization manifest for package %s, error: %s"), *Package->GetFullName(), *Error.ToString());
+		}
+	}
 }
 
 
@@ -200,13 +409,8 @@ void UGenerateModerationArtifactsCommandlet::PerformAdditionalOperations(class U
 
 FString UGenerateModerationArtifactsCommandlet::CreateOutputFileName(class UObject* Object, const FString& Extension)
 {
-	UPackage* Package = Object->GetOutermost();
-	const FPackagePath& PackagePath = Package->GetLoadedPath();
-	
-	FMD5Hash FileHash = FMD5Hash::HashFile(*PackagePath.GetLocalFullPath());
-	FString FileName = FString::Printf(TEXT("%s-%s.%s"), *Object->GetClass()->GetName(), *LexToString(FileHash), *Extension);
+	FString FileName = Manifest.CreateModerationAssetFileName(Object, Extension);
 	FString FullPath = FPaths::Combine(OutputPath, FileName);
-	UE_LOG(LogModerationArtifactsCommandlet, Display, TEXT("Created moderation file %s for asset %s"), *FullPath, *Object->GetPathName());
 	return FullPath;
 }
 
@@ -302,4 +506,82 @@ void UGenerateModerationArtifactsCommandlet::GenerateArtifact(UDataTable* DataTa
 void UGenerateModerationArtifactsCommandlet::GenerateArtifact(UStaticMeshComponent* StaticMesh)
 {
 	UE_LOG(LogModerationArtifactsCommandlet, Display, TEXT("Found staticmesh %s"), *StaticMesh->GetFullName());
-}
\ No newline at end of file
+}
+
+void UGenerateModerationArtifactsCommandlet::PostProcessPackages()
+{
+	FString JsonManifestFile;
+	if (FJsonObjectConverter::UStructToJsonObjectString(Manifest, JsonManifestFile))
+	{
+		FString ManifestFilename = FPaths::Combine(OutputPath, TEXT("ModerationManifest.manifest"));
+		if (FFileHelper::SaveStringToFile(JsonManifestFile, *ManifestFilename) == false)
+		{
+			UE_LOG(LogModerationArtifactsCommandlet, Error, TEXT("Unable to save manifest file to %s"), *ManifestFilename);
+		}
+	}
+	else
+	{
+		UE_LOG(LogModerationArtifactsCommandlet,Error, TEXT("Unable to generate json manifest file"));
+	}
+
+}
+
+
+FModerationPackage* FModerationManifest::FindOrCreateModerationPackage(UPackage* InPackage)
+{
+	
+	FModerationPackage* ModerationPackage = Packages.FindByPredicate([InPackage](const FModerationPackage& Other) {
+				return (Other.Package == InPackage);
+			});
+
+	if (ModerationPackage == nullptr)
+	{
+		int32 Index = Packages.AddDefaulted();
+		ModerationPackage = &Packages[Index];
+		ModerationPackage->Package = InPackage;
+		const FPackagePath& PackagePath = InPackage->GetLoadedPath();
+		ModerationPackage->PackagePath = PackagePath.GetLocalFullPath();
+		FMD5Hash FileHash = FMD5Hash::HashFile(*ModerationPackage->PackagePath);
+		ModerationPackage->PackageHash = LexToString(FileHash);
+	}
+	return ModerationPackage;
+}
+
+FString FModerationManifest::CreateModerationAssetFileName(const UObject* Object, const FString& Extension)
+{
+	UPackage* Package = Object->GetOutermost();
+	check(Package);
+
+	FModerationPackage* ModerationPackage = FindOrCreateModerationPackage(Package);
+	check(ModerationPackage);
+	FModerationAsset* ModerationAsset = ModerationPackage->FindOrCreateModerationAsset(Object);
+	check(ModerationAsset);
+	
+	FString Path = ModerationAsset->FullPath;
+	Path.RemoveFromStart(TEXT("/"));
+	Path.ReplaceInline(TEXT("/"), TEXT("+"));
+	Path.ReplaceInline(TEXT(":"), TEXT(""));
+	FString FileName = FString::Printf(TEXT("%s-%s-%s.%s"), *ModerationAsset->ClassName, *Path, *ModerationPackage->PackageHash, *Extension);
+	ModerationAsset->ModerationArtifactFilenames.Add(FileName);
+
+	UE_LOG(LogModerationArtifactsCommandlet, Display, TEXT("Created moderation file %s for asset %s"), *FileName, *Object->GetPathName());
+
+	return FileName;
+}
+
+
+FModerationAsset* FModerationPackage::FindOrCreateModerationAsset(const UObject* InObject)
+{
+	FModerationAsset* Asset = Assets.FindByPredicate([InObject](const FModerationAsset& Other){ return Other.Object == InObject; });
+
+	if (Asset == nullptr)
+	{
+		int32 Index = Assets.AddDefaulted();
+		Asset = &Assets[Index];
+		Asset->FullPath = InObject->GetPathName();
+		Asset->ClassName = InObject->GetClass()->GetName();
+	}
+	return Asset;
+}
+
+PRAGMA_ENABLE_OPTIMIZATION
\ No newline at end of file
