Fix for UnrealEditor NonUnity Win64 CIS build missing header file

diff --git a/Engine/Source/Editor/UnrealEd/Private/Commandlets/GeneratemoderationArtifactsCommandlet.cpp b/Engine/Source/Editor/UnrealEd/Private/Commandlets/GeneratemoderationArtifactsCommandlet.cpp
--- a/Engine/Source/Editor/UnrealEd/Private/Commandlets/GeneratemoderationArtifactsCommandlet.cpp
+++ b/Engine/Source/Editor/UnrealEd/Private/Commandlets/GeneratemoderationArtifactsCommandlet.cpp
@@ -14,6 +14,7 @@
 
 // localization includes 
 #include "Serialization/PropertyLocalizationDataGathering.h"
+#include "LocTextHelper.h"
 // string proxy includes
 //#include "Serialization/NameAsStringProxyArchive.h"
 // saving json manifest file
