fix static analysis warning

#rnx

diff --git a/Engine/Source/Editor/UnrealEd/Private/PlayLevelViaLauncher.cpp b/Engine/Source/Editor/UnrealEd/Private/PlayLevelViaLauncher.cpp
--- a/Engine/Source/Editor/UnrealEd/Private/PlayLevelViaLauncher.cpp
+++ b/Engine/Source/Editor/UnrealEd/Private/PlayLevelViaLauncher.cpp
@@ -426,7 +426,10 @@ void UEditorEngine::StartPlayUsingLauncherSession(FRequestPlaySessionParams& InR
 
 	// analytics for launch on
 	TArray<FAnalyticsEventAttribute> AnalyticsParamArray;
-	LaunchPlatform->GetPlatformSpecificProjectAnalytics(AnalyticsParamArray);
+	if (LaunchPlatform != nullptr)
+	{
+		LaunchPlatform->GetPlatformSpecificProjectAnalytics(AnalyticsParamArray);
+	}
 	FEditorAnalytics::ReportEvent(TEXT("Editor.LaunchOn.Started"), LaunchPlatformName, LauncherSessionInfo->bPlayUsingLauncherHasCode, AnalyticsParamArray);
 
 
