HotReload/Live Coding re-instancing refactor to re-instance classes more orderly and support the referencing of the new default object prior to the class being reinstanced.

diff --git a/Engine/Source/Editor/UnrealEd/Private/Kismet2/ReloadUtilities.cpp b/Engine/Source/Editor/UnrealEd/Private/Kismet2/ReloadUtilities.cpp
--- a/Engine/Source/Editor/UnrealEd/Private/Kismet2/ReloadUtilities.cpp
+++ b/Engine/Source/Editor/UnrealEd/Private/Kismet2/ReloadUtilities.cpp
@@ -17,10 +17,7 @@
 #include "UObject/UObjectIterator.h"
 #include "Misc/StringBuilder.h"
 
-/**
- * Helper class used for re-instancing native and blueprint classes after hot-reload
- */
-class FReloadClassReinstancer : public FBlueprintCompileReinstancer
+namespace UE::Reload::Private
 {
 	/** Holds a property and its offset in the serialized properties data array */
 	struct FCDOProperty
@@ -45,648 +42,761 @@ class FReloadClassReinstancer : public FBlueprintCompileReinstancer
 		TMap<FName, FCDOProperty> Properties;
 	};
 
-	/** Hot-reloaded version of the old class */
-	UClass* NewClass;
+	/**
+	 * Helper class used for re-instancing native and blueprint classes after hot-reload
+	 */
+	class FReloadClassReinstancer : public FBlueprintCompileReinstancer
+	{
 
-	/** Serialized properties of the original CDO (before hot-reload) */
-	FCDOPropertyData OriginalCDOProperties;
+		/** Hot-reloaded version of the old class */
+		UClass* NewClass;
 
-	/** Serialized properties of the new CDO (after hot-reload) */
-	FCDOPropertyData ReconstructedCDOProperties;
+		/** Necessary for delta serialization */
+		UObject* CopyOfPreviousCDO;
 
-	/** True if the provided native class needs re-instancing */
-	bool bNeedsReinstancing;
+		/**
+		 * Sets the re-instancer up for new class re-instancing
+		 *
+		 * @param InNewClass Class that has changed after hot-reload
+		 * @param InOldClass Class before it was hot-reloaded
+		 */
+		void SetupNewClassReinstancing(UClass* InNewClass, UClass* InOldClass);
 
-	/** Necessary for delta serialization */
-	UObject* CopyOfPreviousCDO;
+		/**
+		* Sets the re-instancer up for old class re-instancing. Always re-creates the CDO.
+		*
+		* @param InOldClass Class that has NOT changed after hot-reload
+		*/
+		void RecreateCDOAndSetupOldClassReinstancing(UClass* InOldClass);
 
-	/**
-	 * Sets the re-instancer up for new class re-instancing
-	 *
-	 * @param InNewClass Class that has changed after hot-reload
-	 * @param InOldClass Class before it was hot-reloaded
-	 */
-	void SetupNewClassReinstancing(UClass* InNewClass, UClass* InOldClass);
+	public:
 
-	/**
-	* Sets the re-instancer up for old class re-instancing. Always re-creates the CDO.
-	*
-	* @param InOldClass Class that has NOT changed after hot-reload
-	*/
-	void RecreateCDOAndSetupOldClassReinstancing(UClass* InOldClass);
+		/** Sets the re-instancer up to re-instance native classes */
+		FReloadClassReinstancer(UClass* InNewClass, UClass* InOldClass, UObject* InOldClassCDO,
+			const TSet<UObject*>& InReinstancingObjects, TSet<UBlueprint*>& InCompiledBlueprints);
 
-	/**
-	* Creates a mem-comparable array of data containing CDO property values.
-	*
-	* @param InObject CDO
-	* @param OutData Data containing all of the CDO property values
-	*/
-	void SerializeCDOProperties(UObject* InObject, FCDOPropertyData& OutData);
+		/** Destructor */
+		virtual ~FReloadClassReinstancer();
 
-	/**
-	 * Re-creates class default object.
-	 *
-	 * @param InClass Class that has NOT changed after hot-reload.
-	 * @param InOuter Outer for the new CDO.
-	 * @param InName Name of the new CDO.
-	 * @param InFlags Flags of the new CDO.
-	 */
-	void ReconstructClassDefaultObject(UClass* InClass, UObject* InOuter, FName InName, EObjectFlags InFlags);
+		// FSerializableObject interface
+		virtual void AddReferencedObjects(FReferenceCollector& Collector) override;
+		// End of FSerializableObject interface
 
-	/** Updates property values on instances of the hot-reloaded class */
-	void UpdateDefaultProperties();
+		virtual bool IsClassObjectReplaced() const override { return true; }
 
-	/** Returns true if the properties of the CDO have changed during hot-reload */
-	FORCEINLINE bool DefaultPropertiesHaveChanged() const
-	{
-		return OriginalCDOProperties.Bytes.Num() != ReconstructedCDOProperties.Bytes.Num() ||
-			FMemory::Memcmp(OriginalCDOProperties.Bytes.GetData(), ReconstructedCDOProperties.Bytes.GetData(), OriginalCDOProperties.Bytes.Num());
-	}
+		virtual void BlueprintWasRecompiled(UBlueprint* BP, bool bBytecodeOnly) override;
 
-public:
+	protected:
 
-	/** Sets the re-instancer up to re-instance native classes */
-	FReloadClassReinstancer(UClass* InNewClass, UClass* InOldClass, const TSet<UObject*>& InReinstancingObjects, TMap<UObject*, UObject*>& OutReconstructedCDOsMap, TSet<UBlueprint*>& InCompiledBlueprints);
+		// FBlueprintCompileReinstancer interface
+		virtual bool ShouldPreserveRootComponentOfReinstancedActor() const override { return false; }
+		// End of FBlueprintCompileReinstancer interface
 
-	/** Destructor */
-	virtual ~FReloadClassReinstancer();
+	private:
+		/** Collection of blueprints already recompiled */
+		TSet<UBlueprint*>& CompiledBlueprints;
+	};
 
-	/** If true, the class needs re-instancing */
-	FORCEINLINE bool ClassNeedsReinstancing() const
+	FReloadClassReinstancer::FReloadClassReinstancer(UClass* InNewClass, UClass* InOldClass, UObject* InOldClassCDO, const TSet<UObject*>& InReinstancingObjects, TSet<UBlueprint*>& InCompiledBlueprints)
+		: NewClass(nullptr)
+		, CopyOfPreviousCDO(nullptr)
+		, CompiledBlueprints(InCompiledBlueprints)
 	{
-		return bNeedsReinstancing;
-	}
+		ensure(InOldClass);
+		ensure(!HotReloadedOldClass && !HotReloadedNewClass);
+		HotReloadedOldClass = InOldClass;
+		HotReloadedNewClass = InNewClass ? InNewClass : InOldClass;
+		OriginalCDO = InOldClassCDO;
 
-	/** Reinstances all objects of the hot-reloaded class and update their properties to match the new CDO */
-	void ReinstanceObjectsAndUpdateDefaults();
+		for (UObject* Object : InReinstancingObjects)
+		{
+			ObjectsThatShouldUseOldStuff.Add(Object);
+		}
 
-	/** Creates the reinstancer as a sharable object */
-	static TSharedPtr<FReloadClassReinstancer> Create(UClass* InNewClass, UClass* InOldClass, const TSet<UObject*>& InReinstancingObjects, TMap<UObject*, UObject*>& OutReconstructedCDOsMap, TSet<UBlueprint*>& InCompiledBlueprints)
-	{
-		return MakeShareable(new FReloadClassReinstancer(InNewClass, InOldClass, InReinstancingObjects, OutReconstructedCDOsMap, InCompiledBlueprints));
-	}
+		// If InNewClass is NULL, then the old class has not changed after hot-reload.
+		// However, we still need to check for changes to its constructor code (CDO values).
+		if (InNewClass)
+		{
+			SetupNewClassReinstancing(InNewClass, InOldClass);
 
-	// FSerializableObject interface
-	virtual void AddReferencedObjects(FReferenceCollector& Collector) override;
-	// End of FSerializableObject interface
+			TMap<UObject*, UObject*> ClassRedirects;
+			ClassRedirects.Add(InOldClass, InNewClass);
 
-	virtual bool IsClassObjectReplaced() const override { return true; }
+			for (TObjectIterator<UBlueprint> BlueprintIt; BlueprintIt; ++BlueprintIt)
+			{
+				constexpr EArchiveReplaceObjectFlags ReplaceObjectArchFlags = (EArchiveReplaceObjectFlags::IgnoreOuterRef | EArchiveReplaceObjectFlags::IgnoreArchetypeRef);
+				FArchiveReplaceObjectRef<UObject> ReplaceObjectArch(*BlueprintIt, ClassRedirects, ReplaceObjectArchFlags);
+			}
+		}
+		else
+		{
+			RecreateCDOAndSetupOldClassReinstancing(InOldClass);
+		}
+	}
 
-	virtual void BlueprintWasRecompiled(UBlueprint* BP, bool bBytecodeOnly) override;
+	FReloadClassReinstancer::~FReloadClassReinstancer()
+	{
+		// Make sure the base class does not remove the DuplicatedClass from root, we not always want it.
+		// For example when we're just reconstructing CDOs. Other cases are handled by HotReloadClassReinstancer.
+		DuplicatedClass = nullptr;
 
-protected:
+		ensure(HotReloadedOldClass);
+		HotReloadedOldClass = nullptr;
+		HotReloadedNewClass = nullptr;
+	}
 
-	// FBlueprintCompileReinstancer interface
-	virtual bool ShouldPreserveRootComponentOfReinstancedActor() const override { return false; }
-	// End of FBlueprintCompileReinstancer interface
+	void FReloadClassReinstancer::SetupNewClassReinstancing(UClass* InNewClass, UClass* InOldClass)
+	{
+		// Set base class members to valid values
+		ClassToReinstance = InNewClass;
+		DuplicatedClass = InOldClass;
+		bHasReinstanced = false;
+		NewClass = InNewClass;
 
-private:
-	/** Reference to reconstructed CDOs map in this hot-reload session. */
-	TMap<UObject*, UObject*>& ReconstructedCDOsMap;
+		SaveClassFieldMapping(InOldClass);
 
-	/** Collection of blueprints already recompiled */
-	TSet<UBlueprint*>& CompiledBlueprints;
- };
+		ObjectsThatShouldUseOldStuff.Add(InOldClass); //CDO of REINST_ class can be used as archetype
 
-void FReloadClassReinstancer::SetupNewClassReinstancing(UClass* InNewClass, UClass* InOldClass)
-{
-	// Set base class members to valid values
-	ClassToReinstance = InNewClass;
-	DuplicatedClass = InOldClass;
-	OriginalCDO = InOldClass->GetDefaultObject();
-	bHasReinstanced = false;
-	bNeedsReinstancing = true;
-	NewClass = InNewClass;
-
-	// Collect the original CDO property values
-	SerializeCDOProperties(InOldClass->GetDefaultObject(), OriginalCDOProperties);
-	// Collect the property values of the new CDO
-	SerializeCDOProperties(InNewClass->GetDefaultObject(), ReconstructedCDOProperties);
-
-	SaveClassFieldMapping(InOldClass);
-
-	ObjectsThatShouldUseOldStuff.Add(InOldClass); //CDO of REINST_ class can be used as archetype
-
-	TArray<UClass*> ChildrenOfClass;
-	GetDerivedClasses(InOldClass, ChildrenOfClass);
-	for (auto ClassIt = ChildrenOfClass.CreateConstIterator(); ClassIt; ++ClassIt)
-	{
-		UClass* ChildClass = *ClassIt;
-		UBlueprint* ChildBP = Cast<UBlueprint>(ChildClass->ClassGeneratedBy);
-		if (ChildBP && !ChildBP->HasAnyFlags(RF_BeingRegenerated))
+		TArray<UClass*> ChildrenOfClass;
+		GetDerivedClasses(InOldClass, ChildrenOfClass);
+		for (auto ClassIt = ChildrenOfClass.CreateConstIterator(); ClassIt; ++ClassIt)
 		{
-			// If this is a direct child, change the parent and relink so the property chain is valid for reinstancing
-			if (!ChildBP->HasAnyFlags(RF_NeedLoad))
+			UClass* ChildClass = *ClassIt;
+			UBlueprint* ChildBP = Cast<UBlueprint>(ChildClass->ClassGeneratedBy);
+			if (ChildBP && !ChildBP->HasAnyFlags(RF_BeingRegenerated))
 			{
-				if (ChildClass->GetSuperClass() == InOldClass)
+				// If this is a direct child, change the parent and relink so the property chain is valid for reinstancing
+				if (!ChildBP->HasAnyFlags(RF_NeedLoad))
 				{
-					ReparentChild(ChildBP);
-				}
+					if (ChildClass->GetSuperClass() == InOldClass)
+					{
+						ReparentChild(ChildBP);
+					}
 
-				Children.AddUnique(ChildBP);
-				if (ChildBP->ParentClass == InOldClass)
+					Children.AddUnique(ChildBP);
+					if (ChildBP->ParentClass == InOldClass)
+					{
+						ChildBP->ParentClass = NewClass;
+					}
+				}
+				else
 				{
-					ChildBP->ParentClass = NewClass;
+					// If this is a child that caused the load of their parent, relink to the REINST class so that we can still serialize in the CDO, but do not add to later processing
+					ReparentChild(ChildClass);
 				}
 			}
-			else
+		}
+
+		// Finally, remove the old class from Root so that it can get GC'd and mark it as CLASS_NewerVersionExists
+		InOldClass->RemoveFromRoot();
+		InOldClass->ClassFlags |= CLASS_NewerVersionExists;
+	}
+
+	void FReloadClassReinstancer::RecreateCDOAndSetupOldClassReinstancing(UClass* InOldClass)
+	{
+		// Set base class members to valid values
+		ClassToReinstance = InOldClass;
+		DuplicatedClass = InOldClass;
+		bHasReinstanced = false;
+		NewClass = InOldClass; // The class doesn't change in this case
+
+		SaveClassFieldMapping(InOldClass);
+
+		TArray<UClass*> ChildrenOfClass;
+		GetDerivedClasses(InOldClass, ChildrenOfClass);
+		for (auto ClassIt = ChildrenOfClass.CreateConstIterator(); ClassIt; ++ClassIt)
+		{
+			UClass* ChildClass = *ClassIt;
+			UBlueprint* ChildBP = Cast<UBlueprint>(ChildClass->ClassGeneratedBy);
+			if (ChildBP && !ChildBP->HasAnyFlags(RF_BeingRegenerated))
 			{
-				// If this is a child that caused the load of their parent, relink to the REINST class so that we can still serialize in the CDO, but do not add to later processing
-				ReparentChild(ChildClass);
+				if (!ChildBP->HasAnyFlags(RF_NeedLoad))
+				{
+					Children.AddUnique(ChildBP);
+					UBlueprintGeneratedClass* BPGC = Cast<UBlueprintGeneratedClass>(ChildBP->GeneratedClass);
+					UObject* CurrentCDO = BPGC ? BPGC->GetDefaultObject(false) : nullptr;
+					if (CurrentCDO && (OriginalCDO == CurrentCDO->GetArchetype()))
+					{
+						BPGC->OverridenArchetypeForCDO = OriginalCDO;
+					}
+				}
 			}
 		}
 	}
 
-	// Finally, remove the old class from Root so that it can get GC'd and mark it as CLASS_NewerVersionExists
-	InOldClass->RemoveFromRoot();
-	InOldClass->ClassFlags |= CLASS_NewerVersionExists;
-}
+	void FReloadClassReinstancer::AddReferencedObjects(FReferenceCollector& Collector)
+	{
+		FBlueprintCompileReinstancer::AddReferencedObjects(Collector);
+		Collector.AllowEliminatingReferences(false);
+		Collector.AddReferencedObject(CopyOfPreviousCDO);
+		Collector.AllowEliminatingReferences(true);
+	}
 
-void FReloadClassReinstancer::SerializeCDOProperties(UObject* InObject, FReloadClassReinstancer::FCDOPropertyData& OutData)
-{
-	// Creates a mem-comparable CDO data
-	class FCDOWriter : public FMemoryWriter
+	void FReloadClassReinstancer::BlueprintWasRecompiled(UBlueprint* BP, bool bBytecodeOnly)
 	{
-		/** Objects already visited by this archive */
-		TSet<UObject*>& VisitedObjects;
-		/** Output property data */
-		FCDOPropertyData& PropertyData;
-		/** Current subobject being serialized */
-		FName SubobjectName;
+		CompiledBlueprints.Add(BP);
 
-	public:
-		/** Serializes all script properties of the provided DefaultObject */
-		FCDOWriter(FCDOPropertyData& InOutData, TSet<UObject*>& InVisitedObjects, FName InSubobjectName)
-			: FMemoryWriter(InOutData.Bytes, /* bIsPersistent = */ false, /* bSetOffset = */ true)
-			, VisitedObjects(InVisitedObjects)
-			, PropertyData(InOutData)
-			, SubobjectName(InSubobjectName)
-		{
-			// Disable delta serialization, we want to serialize everything
-			ArNoDelta = true;
-		}
-		virtual void Serialize(void* Data, int64 Num) override
+		FBlueprintCompileReinstancer::BlueprintWasRecompiled(BP, bBytecodeOnly);
+	}
+
+	/**
+	* Creates a mem-comparable array of data containing CDO property values.
+	*
+	* @param InObject CDO
+	* @param OutData Data containing all of the CDO property values
+	*/
+	void SerializeCDOProperties(UObject* InObject, FCDOPropertyData& OutData)
+	{
+		// Creates a mem-comparable CDO data
+		class FCDOWriter : public FMemoryWriter
 		{
-			// Collect serialized properties so we can later update their values on instances if they change
-			FProperty* SerializedProperty = GetSerializedProperty();
-			if (SerializedProperty != nullptr)
+			/** Objects already visited by this archive */
+			TSet<UObject*>& VisitedObjects;
+			/** Output property data */
+			FCDOPropertyData& PropertyData;
+			/** Current subobject being serialized */
+			FName SubobjectName;
+
+		public:
+			/** Serializes all script properties of the provided DefaultObject */
+			FCDOWriter(FCDOPropertyData& InOutData, TSet<UObject*>& InVisitedObjects, FName InSubobjectName)
+				: FMemoryWriter(InOutData.Bytes, /* bIsPersistent = */ false, /* bSetOffset = */ true)
+				, VisitedObjects(InVisitedObjects)
+				, PropertyData(InOutData)
+				, SubobjectName(InSubobjectName)
+			{
+				// Disable delta serialization, we want to serialize everything
+				ArNoDelta = true;
+			}
+			virtual void Serialize(void* Data, int64 Num) override
+			{
+				// Collect serialized properties so we can later update their values on instances if they change
+				FProperty* SerializedProperty = GetSerializedProperty();
+				if (SerializedProperty != nullptr)
+				{
+					FCDOProperty& PropertyInfo = PropertyData.Properties.FindOrAdd(SerializedProperty->GetFName());
+					if (PropertyInfo.Property == nullptr)
+					{
+						PropertyInfo.Property = SerializedProperty;
+						PropertyInfo.SubobjectName = SubobjectName;
+						PropertyInfo.SerializedValueOffset = Tell();
+						PropertyInfo.SerializedValueSize = Num;
+					}
+					else
+					{
+						PropertyInfo.SerializedValueSize += Num;
+					}
+				}
+				FMemoryWriter::Serialize(Data, Num);
+			}
+			/** Serializes an object. Only name and class for normal references, deep serialization for DSOs */
+			virtual FArchive& operator<<(class UObject*& InObj) override
 			{
-				FCDOProperty& PropertyInfo = PropertyData.Properties.FindOrAdd(SerializedProperty->GetFName());
-				if (PropertyInfo.Property == nullptr)
+				FArchive& Ar = *this;
+				if (InObj)
 				{
-					PropertyInfo.Property = SerializedProperty;
-					PropertyInfo.SubobjectName = SubobjectName;
-					PropertyInfo.SerializedValueOffset = Tell();
-					PropertyInfo.SerializedValueSize = Num;
+					FName ClassName = InObj->GetClass()->GetFName();
+					FName ObjectName = InObj->GetFName();
+					Ar << ClassName;
+					Ar << ObjectName;
+					if (!VisitedObjects.Contains(InObj))
+					{
+						VisitedObjects.Add(InObj);
+						if (Ar.GetSerializedProperty() && Ar.GetSerializedProperty()->ContainsInstancedObjectProperty())
+						{
+							// Serialize all DSO properties too
+							FCDOWriter DefaultSubobjectWriter(PropertyData, VisitedObjects, InObj->GetFName());
+							InObj->SerializeScriptProperties(DefaultSubobjectWriter);
+							Seek(PropertyData.Bytes.Num());
+						}
+					}
 				}
 				else
 				{
-					PropertyInfo.SerializedValueSize += Num;
+					FName UnusedName = NAME_None;
+					Ar << UnusedName;
+					Ar << UnusedName;
 				}
+
+				return *this;
 			}
-			FMemoryWriter::Serialize(Data, Num);
-		}
-		/** Serializes an object. Only name and class for normal references, deep serialization for DSOs */
-		virtual FArchive& operator<<(class UObject*& InObj) override
-		{
-			FArchive& Ar = *this;
-			if (InObj)
+			virtual FArchive& operator<<(FObjectPtr& InObj) override
 			{
-				FName ClassName = InObj->GetClass()->GetFName();
-				FName ObjectName = InObj->GetFName();
-				Ar << ClassName;
-				Ar << ObjectName;
-				if (!VisitedObjects.Contains(InObj))
+				// Invoke the method above
+				return FArchiveUObject::SerializeObjectPtr(*this, InObj);
+			}
+			/** Serializes an FName as its index and number */
+			virtual FArchive& operator<<(FName& InName) override
+			{
+				FArchive& Ar = *this;
+				FNameEntryId ComparisonIndex = InName.GetComparisonIndex();
+				FNameEntryId DisplayIndex = InName.GetDisplayIndex();
+				int32 Number = InName.GetNumber();
+				Ar << ComparisonIndex;
+				Ar << DisplayIndex;
+				Ar << Number;
+				return Ar;
+			}
+			virtual FArchive& operator<<(FLazyObjectPtr& LazyObjectPtr) override
+			{
+				FArchive& Ar = *this;
+				FUniqueObjectGuid UniqueID = LazyObjectPtr.GetUniqueID();
+				Ar << UniqueID;
+				return *this;
+			}
+			virtual FArchive& operator<<(FSoftObjectPtr& Value) override
+			{
+				FArchive& Ar = *this;
+				FSoftObjectPath UniqueID = Value.GetUniqueID();
+				Ar << UniqueID;
+				return Ar;
+			}
+			virtual FArchive& operator<<(FSoftObjectPath& Value) override
+			{
+				FArchive& Ar = *this;
+
+				FString Path = Value.ToString();
+
+				Ar << Path;
+
+				if (IsLoading())
 				{
-					VisitedObjects.Add(InObj);
-					if (Ar.GetSerializedProperty() && Ar.GetSerializedProperty()->ContainsInstancedObjectProperty())
-					{
-						// Serialize all DSO properties too
-						FCDOWriter DefaultSubobjectWriter(PropertyData, VisitedObjects, InObj->GetFName());
-						InObj->SerializeScriptProperties(DefaultSubobjectWriter);
-						Seek(PropertyData.Bytes.Num());
-					}
+					Value.SetPath(MoveTemp(Path));
 				}
+
+				return Ar;
 			}
-			else
+			FArchive& operator<<(FWeakObjectPtr& WeakObjectPtr) override
 			{
-				FName UnusedName = NAME_None;
-				Ar << UnusedName;
-				Ar << UnusedName;
+				return FArchiveUObject::SerializeWeakObjectPtr(*this, WeakObjectPtr);
 			}
+			/** Archive name, for debugging */
+			virtual FString GetArchiveName() const override { return TEXT("FCDOWriter"); }
+		};
+		TSet<UObject*> VisitedObjects;
+		VisitedObjects.Add(InObject);
+		FCDOWriter Ar(OutData, VisitedObjects, NAME_None);
+		InObject->SerializeScriptProperties(Ar);
+	}
 
-			return *this;
-		}
-		virtual FArchive& operator<<(FObjectPtr& InObj) override
-		{
-			// Invoke the method above
-			return FArchiveUObject::SerializeObjectPtr(*this, InObj);
-		}
-		/** Serializes an FName as its index and number */
-		virtual FArchive& operator<<(FName& InName) override
-		{
-			FArchive& Ar = *this;
-			FNameEntryId ComparisonIndex = InName.GetComparisonIndex();
-			FNameEntryId DisplayIndex = InName.GetDisplayIndex();
-			int32 Number = InName.GetNumber();
-			Ar << ComparisonIndex;
-			Ar << DisplayIndex;
-			Ar << Number;
-			return Ar;
-		}
-		virtual FArchive& operator<<(FLazyObjectPtr& LazyObjectPtr) override
+	/** Returns true if the properties of the CDO have changed during hot-reload */
+	bool HavePropertiesChanged(const FCDOPropertyData& lhs, const FCDOPropertyData& rhs)
+	{
+		return lhs.Bytes.Num() != rhs.Bytes.Num() || FMemory::Memcmp(lhs.Bytes.GetData(), rhs.Bytes.GetData(), lhs.Bytes.Num());
+	}
+
+	/** Helper for finding subobject in an array. Usually there's not that many subobjects on a class to justify a TMap */
+	UObject* FindDefaultSubobject(TArray<UObject*>& InDefaultSubobjects, FName SubobjectName)
+	{
+		for (UObject* Subobject : InDefaultSubobjects)
 		{
-			FArchive& Ar = *this;
-			FUniqueObjectGuid UniqueID = LazyObjectPtr.GetUniqueID();
-			Ar << UniqueID;
-			return *this;
+			if (Subobject->GetFName() == SubobjectName)
+			{
+				return Subobject;
+			}
 		}
-		virtual FArchive& operator<<(FSoftObjectPtr& Value) override
+		return nullptr;
+	}
+
+	/** Update the properties from the old class CDO to the new class CDO */
+	void UpdateDefaultProperties(UClass* NewClass, const FCDOPropertyData& OldClassCDOProperties, const FCDOPropertyData& NewClassCDOProperties)
+	{
+		struct FPropertyToUpdate
 		{
-			FArchive& Ar = *this;
-			FSoftObjectPath UniqueID = Value.GetUniqueID();
-			Ar << UniqueID;
-			return Ar;
-		}
-		virtual FArchive& operator<<(FSoftObjectPath& Value) override
+			FProperty* Property;
+			FName SubobjectName;
+			const uint8* OldSerializedValuePtr;
+			uint8* NewValuePtr;
+			int64 OldSerializedSize;
+		};
+		/** Memory writer archive that supports UObject values the same way as FCDOWriter. */
+		class FPropertyValueMemoryWriter : public FMemoryWriter
 		{
-			FArchive& Ar = *this;
+		public:
+			FPropertyValueMemoryWriter(TArray<uint8>& OutData)
+				: FMemoryWriter(OutData)
+			{}
+			virtual FArchive& operator<<(class UObject*& InObj) override
+			{
+				FArchive& Ar = *this;
+				if (InObj)
+				{
+					FName ClassName = InObj->GetClass()->GetFName();
+					FName ObjectName = InObj->GetFName();
+					Ar << ClassName;
+					Ar << ObjectName;
+				}
+				else
+				{
+					FName UnusedName = NAME_None;
+					Ar << UnusedName;
+					Ar << UnusedName;
+				}
+				return *this;
+			}
+			virtual FArchive& operator<<(FObjectPtr& InObj) override
+			{
+				// Invoke the method above
+				return FArchiveUObject::SerializeObjectPtr(*this, InObj);
+			}
+			virtual FArchive& operator<<(FName& InName) override
+			{
+				FArchive& Ar = *this;
+				FNameEntryId ComparisonIndex = InName.GetComparisonIndex();
+				FNameEntryId DisplayIndex = InName.GetDisplayIndex();
+				int32 Number = InName.GetNumber();
+				Ar << ComparisonIndex;
+				Ar << DisplayIndex;
+				Ar << Number;
+				return Ar;
+			}
+			virtual FArchive& operator<<(FLazyObjectPtr& LazyObjectPtr) override
+			{
+				FArchive& Ar = *this;
+				FUniqueObjectGuid UniqueID = LazyObjectPtr.GetUniqueID();
+				Ar << UniqueID;
+				return *this;
+			}
+			virtual FArchive& operator<<(FSoftObjectPtr& Value) override
+			{
+				FArchive& Ar = *this;
+				FSoftObjectPath UniqueID = Value.GetUniqueID();
+				Ar << UniqueID;
+				return Ar;
+			}
+			virtual FArchive& operator<<(FSoftObjectPath& Value) override
+			{
+				FArchive& Ar = *this;
 
-			FString Path = Value.ToString();
+				FString Path = Value.ToString();
 
-			Ar << Path;
+				Ar << Path;
 
-			if (IsLoading())
+				if (IsLoading())
+				{
+					Value.SetPath(MoveTemp(Path));
+				}
+
+				return Ar;
+			}
+			FArchive& operator<<(FWeakObjectPtr& WeakObjectPtr) override
 			{
-				Value.SetPath(MoveTemp(Path));
+				return FArchiveUObject::SerializeWeakObjectPtr(*this, WeakObjectPtr);
 			}
+		};
 
-			return Ar;
-		}
-		FArchive& operator<<(FWeakObjectPtr& WeakObjectPtr) override
-		{
-			return FArchiveUObject::SerializeWeakObjectPtr(*this, WeakObjectPtr);
-		}
-		/** Archive name, for debugging */
-		virtual FString GetArchiveName() const override { return TEXT("FCDOWriter"); }
-	};
-	TSet<UObject*> VisitedObjects;
-	VisitedObjects.Add(InObject);
-	FCDOWriter Ar(OutData, VisitedObjects, NAME_None);
-	InObject->SerializeScriptProperties(Ar);
-}
+		// Collect default subobjects to update their properties too
+		const int32 DefaultSubobjectArrayCapacity = 16;
+		TArray<UObject*> DefaultSubobjectArray;
+		DefaultSubobjectArray.Empty(DefaultSubobjectArrayCapacity);
+		NewClass->GetDefaultObject()->CollectDefaultSubobjects(DefaultSubobjectArray);
 
-void FReloadClassReinstancer::ReconstructClassDefaultObject(UClass* InClass, UObject* InOuter, FName InName, EObjectFlags InFlags)
-{
-	// Get the parent CDO
-	UClass* ParentClass = InClass->GetSuperClass();
-	UObject* ParentDefaultObject = NULL;
-	if (ParentClass != NULL)
-	{
-		ParentDefaultObject = ParentClass->GetDefaultObject(); // Force the default object to be constructed if it isn't already
-	}
+		TArray<FPropertyToUpdate> PropertiesToUpdate;
+		// Collect all properties that have actually changed
+		for (const TPair<FName, FCDOProperty>& Pair : NewClassCDOProperties.Properties)
+		{
+			const FCDOProperty* OldPropertyInfo = OldClassCDOProperties.Properties.Find(Pair.Key);
+			if (OldPropertyInfo)
+			{
+				const FCDOProperty& NewPropertyInfo = Pair.Value;
 
-	// Re-create
-	InClass->ClassDefaultObject = StaticAllocateObject(InClass, InOuter, InName, InFlags, EInternalObjectFlags::None, false);
-	check(InClass->ClassDefaultObject);
-	(*InClass->ClassConstructor)(FObjectInitializer(InClass->ClassDefaultObject, ParentDefaultObject, EObjectInitializerOptions::None));
-}
+				const uint8* OldSerializedValuePtr = OldClassCDOProperties.Bytes.GetData() + OldPropertyInfo->SerializedValueOffset;
+				const uint8* NewSerializedValuePtr = NewClassCDOProperties.Bytes.GetData() + NewPropertyInfo.SerializedValueOffset;
+				if (OldPropertyInfo->SerializedValueSize != NewPropertyInfo.SerializedValueSize ||
+					FMemory::Memcmp(OldSerializedValuePtr, NewSerializedValuePtr, OldPropertyInfo->SerializedValueSize) != 0)
+				{
+					// Property value has changed so add it to the list of properties that need updating on instances
+					FPropertyToUpdate PropertyToUpdate;
+					PropertyToUpdate.Property = NewPropertyInfo.Property;
+					PropertyToUpdate.NewValuePtr = nullptr;
+					PropertyToUpdate.SubobjectName = NewPropertyInfo.SubobjectName;
 
-void FReloadClassReinstancer::RecreateCDOAndSetupOldClassReinstancing(UClass* InOldClass)
-{
-	// Set base class members to valid values
-	ClassToReinstance = InOldClass;
-	DuplicatedClass = InOldClass;
-	OriginalCDO = InOldClass->GetDefaultObject();
-	bHasReinstanced = false;
-	bNeedsReinstancing = false;
-	NewClass = InOldClass; // The class doesn't change in this case
-
-	// Collect the original property values
-	SerializeCDOProperties(InOldClass->GetDefaultObject(), OriginalCDOProperties);
-
-	// Remember all the basic info about the object before we rename it
-	EObjectFlags CDOFlags = OriginalCDO->GetFlags();
-	UObject* CDOOuter = OriginalCDO->GetOuter();
-	FName CDOName = OriginalCDO->GetFName();
-
-	// Rename original CDO, so we can store this one as OverridenArchetypeForCDO
-	// and create new one with the same name and outer.
-	OriginalCDO->Rename(
-		*MakeUniqueObjectName(
-			GetTransientPackage(),
-			OriginalCDO->GetClass(),
-			*FString::Printf(TEXT("BPGC_ARCH_FOR_CDO_%s"), *InOldClass->GetName())
-		).ToString(),
-		GetTransientPackage(),
-		REN_DoNotDirty | REN_DontCreateRedirectors | REN_NonTransactional | REN_SkipGeneratedClasses | REN_ForceNoResetLoaders);
-
-	// Re-create the CDO, re-running its constructor
-	ReconstructClassDefaultObject(InOldClass, CDOOuter, CDOName, CDOFlags);
-
-	ReconstructedCDOsMap.Add(OriginalCDO, InOldClass->GetDefaultObject());
-
-	// Collect the property values after re-constructing the CDO
-	SerializeCDOProperties(InOldClass->GetDefaultObject(), ReconstructedCDOProperties);
-
-	// We only want to re-instance the old class if its CDO's values have changed or any of its DSOs' property values have changed
-	if (DefaultPropertiesHaveChanged())
-	{
-		bNeedsReinstancing = true;
-		SaveClassFieldMapping(InOldClass);
+					if (NewPropertyInfo.Property->GetOwner<UObject>() == NewClass)
+					{
+						PropertyToUpdate.NewValuePtr = PropertyToUpdate.Property->ContainerPtrToValuePtr<uint8>(NewClass->GetDefaultObject());
+					}
+					else if (NewPropertyInfo.SubobjectName != NAME_None)
+					{
+						UObject* DefaultSubobjectPtr = FindDefaultSubobject(DefaultSubobjectArray, NewPropertyInfo.SubobjectName);
+						if (DefaultSubobjectPtr && NewPropertyInfo.Property->GetOwner<UObject>() == DefaultSubobjectPtr->GetClass())
+						{
+							PropertyToUpdate.NewValuePtr = PropertyToUpdate.Property->ContainerPtrToValuePtr<uint8>(DefaultSubobjectPtr);
+						}
+					}
+					if (PropertyToUpdate.NewValuePtr)
+					{
+						PropertyToUpdate.OldSerializedValuePtr = OldSerializedValuePtr;
+						PropertyToUpdate.OldSerializedSize = OldPropertyInfo->SerializedValueSize;
 
-		TArray<UClass*> ChildrenOfClass;
-		GetDerivedClasses(InOldClass, ChildrenOfClass);
-		for (auto ClassIt = ChildrenOfClass.CreateConstIterator(); ClassIt; ++ClassIt)
+						PropertiesToUpdate.Add(PropertyToUpdate);
+					}
+				}
+			}
+		}
+		if (PropertiesToUpdate.Num())
 		{
-			UClass* ChildClass = *ClassIt;
-			UBlueprint* ChildBP = Cast<UBlueprint>(ChildClass->ClassGeneratedBy);
-			if (ChildBP && !ChildBP->HasAnyFlags(RF_BeingRegenerated))
+			TArray<uint8> CurrentValueSerializedData;
+
+			// Update properties on all existing instances of the class
+			const UPackage* TransientPackage = GetTransientPackage();
+			for (FThreadSafeObjectIterator It(NewClass); It; ++It)
 			{
-				if (!ChildBP->HasAnyFlags(RF_NeedLoad))
+				UObject* ObjectPtr = *It;
+				if (!IsValidChecked(ObjectPtr) || ObjectPtr->GetOutermost() == TransientPackage)
 				{
-					Children.AddUnique(ChildBP);
-					UBlueprintGeneratedClass* BPGC = Cast<UBlueprintGeneratedClass>(ChildBP->GeneratedClass);
-					UObject* CurrentCDO = BPGC ? BPGC->GetDefaultObject(false) : nullptr;
-					if (CurrentCDO && (OriginalCDO == CurrentCDO->GetArchetype()))
+					continue;
+				}
+
+				DefaultSubobjectArray.Empty(DefaultSubobjectArrayCapacity);
+				ObjectPtr->CollectDefaultSubobjects(DefaultSubobjectArray);
+
+				for (auto& PropertyToUpdate : PropertiesToUpdate)
+				{
+					uint8* InstanceValuePtr = nullptr;
+					if (PropertyToUpdate.SubobjectName == NAME_None)
 					{
-						BPGC->OverridenArchetypeForCDO = OriginalCDO;
+						InstanceValuePtr = PropertyToUpdate.Property->ContainerPtrToValuePtr<uint8>(ObjectPtr);
+					}
+					else
+					{
+						UObject* DefaultSubobjectPtr = FindDefaultSubobject(DefaultSubobjectArray, PropertyToUpdate.SubobjectName);
+						if (DefaultSubobjectPtr && PropertyToUpdate.Property->GetOwner<UObject>() == DefaultSubobjectPtr->GetClass())
+						{
+							InstanceValuePtr = PropertyToUpdate.Property->ContainerPtrToValuePtr<uint8>(DefaultSubobjectPtr);
+						}
+					}
+
+					if (InstanceValuePtr)
+					{
+						// Serialize current value to a byte array as we don't have the previous CDO to compare against, we only have its serialized property data
+						CurrentValueSerializedData.Empty(CurrentValueSerializedData.Num() + CurrentValueSerializedData.GetSlack());
+						FPropertyValueMemoryWriter CurrentValueWriter(CurrentValueSerializedData);
+						PropertyToUpdate.Property->SerializeItem(FStructuredArchiveFromArchive(CurrentValueWriter).GetSlot(), InstanceValuePtr);
+
+						// Update only when the current value on the instance is identical to the original CDO
+						if (CurrentValueSerializedData.Num() == PropertyToUpdate.OldSerializedSize &&
+							FMemory::Memcmp(CurrentValueSerializedData.GetData(), PropertyToUpdate.OldSerializedValuePtr, CurrentValueSerializedData.Num()) == 0)
+						{
+							// Update with the new value
+							PropertyToUpdate.Property->CopyCompleteValue(InstanceValuePtr, PropertyToUpdate.NewValuePtr);
+						}
 					}
 				}
 			}
 		}
 	}
-}
 
-FReloadClassReinstancer::FReloadClassReinstancer(UClass* InNewClass, UClass* InOldClass, const TSet<UObject*>& InReinstancingObjects, TMap<UObject*, UObject*>& OutReconstructedCDOsMap, TSet<UBlueprint*>& InCompiledBlueprints)
-	: NewClass(nullptr)
-	, bNeedsReinstancing(false)
-	, CopyOfPreviousCDO(nullptr)
-	, ReconstructedCDOsMap(OutReconstructedCDOsMap)
-	, CompiledBlueprints(InCompiledBlueprints)
-{
-	ensure(InOldClass);
-	ensure(!HotReloadedOldClass && !HotReloadedNewClass);
-	HotReloadedOldClass = InOldClass;
-	HotReloadedNewClass = InNewClass ? InNewClass : InOldClass;
+	/** Helper class that handles class reloading */
 
-	for (UObject* Object : InReinstancingObjects)
+	class FReloadClassHelper
 	{
-		ObjectsThatShouldUseOldStuff.Add(Object);
-	}
+	private:
+		enum class TopologicalState
+		{
+			None,
+			Visited,
+			Finished,
+		};
 
-	// If InNewClass is NULL, then the old class has not changed after hot-reload.
-	// However, we still need to check for changes to its constructor code (CDO values).
-	if (InNewClass)
-	{
-		SetupNewClassReinstancing(InNewClass, InOldClass);
+		struct ClassReinstanceState
+		{
+			UClass* OldClass;
+			UObject* OldClassCDO;
+			TopologicalState State = TopologicalState::None;
+			FCDOPropertyData OldClassCDOProperties;
+		};
 
-		TMap<UObject*, UObject*> ClassRedirects;
-		ClassRedirects.Add(InOldClass, InNewClass);
+	public:
+		FReloadClassHelper(FOutputDevice& InAr, const TSet<UObject*>& InReinstancingObjects, TSet<UBlueprint*>& InCompiledBlueprints, TMap<UObject*, UObject*>& InReconstructedCDOsMap);
 
-		for (TObjectIterator<UBlueprint> BlueprintIt; BlueprintIt; ++BlueprintIt)
-		{
-			constexpr EArchiveReplaceObjectFlags ReplaceObjectArchFlags = (EArchiveReplaceObjectFlags::IgnoreOuterRef | EArchiveReplaceObjectFlags::IgnoreArchetypeRef);
-			FArchiveReplaceObjectRef<UObject> ReplaceObjectArch(*BlueprintIt, ClassRedirects, ReplaceObjectArchFlags);
-		}
-	}
-	else
-	{
-		RecreateCDOAndSetupOldClassReinstancing(InOldClass);
-	}
-}
+		/** 
+		 * Re-instance the collection of classes.  
+		 * The map is from old to new class object.
+		 * If the old class is not changing structure, then the new class pointer is null.
+		 */
+		void ReinstanceClasses(const TMap<UClass*, UClass*>& ClassesToReinstance);
 
-FReloadClassReinstancer::~FReloadClassReinstancer()
-{
-	// Make sure the base class does not remove the DuplicatedClass from root, we not always want it.
-	// For example when we're just reconstructing CDOs. Other cases are handled by HotReloadClassReinstancer.
-	DuplicatedClass = nullptr;
+	private:
 
-	ensure(HotReloadedOldClass);
-	HotReloadedOldClass = nullptr;
-	HotReloadedNewClass = nullptr;
-}
+		/** Given a new/old class pair, re-instance the class */
+		void ReinstanceClass(UClass* NewClass, ClassReinstanceState& State);
 
-/** Helper for finding subobject in an array. Usually there's not that many subobjects on a class to justify a TMap */
-FORCEINLINE static UObject* FindDefaultSubobject(TArray<UObject*>& InDefaultSubobjects, FName SubobjectName)
-{
-	for (UObject* Subobject : InDefaultSubobjects)
+		/** Scan the properties looking for any referenced classes that need to be re-instanced first */
+		void ReinstanceClassPropertyScan(FProperty* PropertyLink);
+
+		/** Checked the property looking for any referenced classes that need to be re-instanced first */
+		void ReinstanceClassProperty(FProperty* Property);
+
+	private:
+		FOutputDevice& Ar;
+		const TSet<UObject*>& ReinstancingObjects;
+		TSet<UBlueprint*>& CompiledBlueprints;
+		TMap<UObject*, UObject*>& ReconstructedCDOsMap;
+
+		TMap<UClass*, ClassReinstanceState> ReinstanceStates;
+	};
+
+	FReloadClassHelper::FReloadClassHelper(FOutputDevice& InAr, const TSet<UObject*>& InReinstancingObjects, TSet<UBlueprint*>& InCompiledBlueprints, TMap<UObject*, UObject*>& InReconstructedCDOsMap)
+		: Ar(InAr)
+		, ReinstancingObjects(InReinstancingObjects)
+		, CompiledBlueprints(InCompiledBlueprints)
+		, ReconstructedCDOsMap(InReconstructedCDOsMap)
 	{
-		if (Subobject->GetFName() == SubobjectName)
-		{
-			return Subobject;
-		}
 	}
-	return nullptr;
-}
 
-void FReloadClassReinstancer::UpdateDefaultProperties()
-{
-	struct FPropertyToUpdate
-	{
-		FProperty* Property;
-		FName SubobjectName;
-		uint8* OldSerializedValuePtr;
-		uint8* NewValuePtr;
-		int64 OldSerializedSize;
-	};
-	/** Memory writer archive that supports UObject values the same way as FCDOWriter. */
-	class FPropertyValueMemoryWriter : public FMemoryWriter
+	void FReloadClassHelper::ReinstanceClasses(const TMap<UClass*, UClass*>& ClassesToReinstance)
 	{
-	public:
-		FPropertyValueMemoryWriter(TArray<uint8>& OutData)
-			: FMemoryWriter(OutData)
-		{}
-		virtual FArchive& operator<<(class UObject*& InObj) override
+		// Create the collection of re-instance states so we can re-instance with topological sorting.
+		for (const TPair<UClass*, UClass*>& Pair : ClassesToReinstance)
 		{
-			FArchive& Ar = *this;
-			if (InObj)
-			{
-				FName ClassName = InObj->GetClass()->GetFName();
-				FName ObjectName = InObj->GetFName();
-				Ar << ClassName;
-				Ar << ObjectName;
-			}
-			else
-			{
-				FName UnusedName = NAME_None;
-				Ar << UnusedName;
-				Ar << UnusedName;
-			}
-			return *this;
-		}
-		virtual FArchive& operator<<(FObjectPtr& InObj) override
-		{
-			// Invoke the method above
-			return FArchiveUObject::SerializeObjectPtr(*this, InObj);
+			UClass* NewClass = Pair.Value != nullptr ? Pair.Value : Pair.Key;
+			UClass* OldClass = Pair.Key;
+			UObject* OldClassCDO = OldClass->GetDefaultObject();
+			ClassReinstanceState State{ OldClass, OldClassCDO };
+			UE::Reload::Private::SerializeCDOProperties(OldClassCDO, State.OldClassCDOProperties);
+			ReinstanceStates.Add(NewClass, MoveTemp(State));
 		}
-		virtual FArchive& operator<<(FName& InName) override
+
+		// Rename and clear the old defaults
+		for (TPair<UClass*, ClassReinstanceState>& Pair : ReinstanceStates)
 		{
-			FArchive& Ar = *this;
-			FNameEntryId ComparisonIndex = InName.GetComparisonIndex();
-			FNameEntryId DisplayIndex = InName.GetDisplayIndex();
-			int32 Number = InName.GetNumber();
-			Ar << ComparisonIndex;
-			Ar << DisplayIndex;
-			Ar << Number;
-			return Ar;
+
+			// If we are re-instancing a class that didn't change structure, then we need to replace
+			// the existing default object
+			if (Pair.Key == Pair.Value.OldClass)
+			{
+				Pair.Value.OldClassCDO->Rename(
+					*MakeUniqueObjectName(
+						GetTransientPackage(),
+						Pair.Value.OldClassCDO->GetClass(),
+						*FString::Printf(TEXT("BPGC_ARCH_FOR_CDO_%s"), *Pair.Value.OldClass->GetName())
+					).ToString(),
+					GetTransientPackage(),
+					REN_DoNotDirty | REN_DontCreateRedirectors | REN_NonTransactional | REN_SkipGeneratedClasses | REN_ForceNoResetLoaders);
+
+				// Clear the class default object so it gets recreated.
+				Pair.Value.OldClass->ClassDefaultObject = nullptr;
+			}
 		}
-		virtual FArchive& operator<<(FLazyObjectPtr& LazyObjectPtr) override
+
+		// Re-instance the classes
+		for (TPair<UClass*, ClassReinstanceState>& Pair : ReinstanceStates)
 		{
-			FArchive& Ar = *this;
-			FUniqueObjectGuid UniqueID = LazyObjectPtr.GetUniqueID();
-			Ar << UniqueID;
-			return *this;
+			ReinstanceClass(Pair.Key, Pair.Value);
 		}
-		virtual FArchive& operator<<(FSoftObjectPtr& Value) override
+	}
+
+	void FReloadClassHelper::ReinstanceClass(UClass* NewClass, ClassReinstanceState& State)
+	{
+		// If this isn't a class we care about or we are already processing, then skip
+		if (State.State != TopologicalState::None)
 		{
-			FArchive& Ar = *this;
-			FSoftObjectPath UniqueID = Value.GetUniqueID();
-			Ar << UniqueID;
-			return Ar;
+			return;
 		}
-		virtual FArchive& operator<<(FSoftObjectPath& Value) override
-		{
-			FArchive& Ar = *this;
-
-			FString Path = Value.ToString();
 
-			Ar << Path;
+		// Mark that we are in the process of re-instancing this class
+		State.State = TopologicalState::Visited;
 
-			if (IsLoading())
+		// Run the parents
+		for (UClass* Super = NewClass->GetSuperClass(); Super != nullptr; Super = Super->GetSuperClass())
+		{
+			ClassReinstanceState* SuperState = ReinstanceStates.Find(Super);
+			if (SuperState != nullptr)
 			{
-				Value.SetPath(MoveTemp(Path));
+				ReinstanceClass(Super, *SuperState);
 			}
-
-			return Ar;
 		}
-		FArchive& operator<<(FWeakObjectPtr& WeakObjectPtr) override
-		{
-			return FArchiveUObject::SerializeWeakObjectPtr(*this, WeakObjectPtr);
-		}
-	};
 
-	// Collect default subobjects to update their properties too
-	const int32 DefaultSubobjectArrayCapacity = 16;
-	TArray<UObject*> DefaultSubobjectArray;
-	DefaultSubobjectArray.Empty(DefaultSubobjectArrayCapacity);
-	NewClass->GetDefaultObject()->CollectDefaultSubobjects(DefaultSubobjectArray);
+		// Look for any properties that reference other types needing re-instancing.  They will be re-instanced first
+		ReinstanceClassPropertyScan(NewClass->PropertyLink);
 
-	TArray<FPropertyToUpdate> PropertiesToUpdate;
-	// Collect all properties that have actually changed
-	for (const TPair<FName, FCDOProperty>& Pair : ReconstructedCDOProperties.Properties)
-	{
-		FCDOProperty* OldPropertyInfo = OriginalCDOProperties.Properties.Find(Pair.Key);
-		if (OldPropertyInfo)
+		// Now we re-instance this class
 		{
-			const FCDOProperty& NewPropertyInfo = Pair.Value;
+			UClass* OldClass = State.OldClass;
+			bool bNeedsReinstancing = true;
 
-			uint8* OldSerializedValuePtr = OriginalCDOProperties.Bytes.GetData() + OldPropertyInfo->SerializedValueOffset;
-			uint8* NewSerializedValuePtr = ReconstructedCDOProperties.Bytes.GetData() + NewPropertyInfo.SerializedValueOffset;
-			if (OldPropertyInfo->SerializedValueSize != NewPropertyInfo.SerializedValueSize ||
-				FMemory::Memcmp(OldSerializedValuePtr, NewSerializedValuePtr, OldPropertyInfo->SerializedValueSize) != 0)
-			{
-				// Property value has changed so add it to the list of properties that need updating on instances
-				FPropertyToUpdate PropertyToUpdate;
-				PropertyToUpdate.Property = NewPropertyInfo.Property;
-				PropertyToUpdate.NewValuePtr = nullptr;
-				PropertyToUpdate.SubobjectName = NewPropertyInfo.SubobjectName;
+			// Collect the property values of the new CDO
+			FCDOPropertyData NewClassCDOProperties;
+			SerializeCDOProperties(NewClass->GetDefaultObject(), NewClassCDOProperties);
 
-				if (NewPropertyInfo.Property->GetOwner<UObject>() == NewClass)
-				{
-					PropertyToUpdate.NewValuePtr = PropertyToUpdate.Property->ContainerPtrToValuePtr<uint8>(NewClass->GetDefaultObject());
-				}
-				else if (NewPropertyInfo.SubobjectName != NAME_None)
-				{
-					UObject* DefaultSubobjectPtr = FindDefaultSubobject(DefaultSubobjectArray, NewPropertyInfo.SubobjectName);
-					if (DefaultSubobjectPtr && NewPropertyInfo.Property->GetOwner<UObject>() == DefaultSubobjectPtr->GetClass())
-					{
-						PropertyToUpdate.NewValuePtr = PropertyToUpdate.Property->ContainerPtrToValuePtr<uint8>(DefaultSubobjectPtr);
-					}
-				}
-				if (PropertyToUpdate.NewValuePtr)
-				{
-					PropertyToUpdate.OldSerializedValuePtr = OldSerializedValuePtr;
-					PropertyToUpdate.OldSerializedSize = OldPropertyInfo->SerializedValueSize;
+			// If the structure didn't change, then we might not need to re-instance.  We always add to the collection
+			// of reconstructed CDOs.
+			if (OldClass == NewClass)
+			{
+				bNeedsReinstancing = HavePropertiesChanged(State.OldClassCDOProperties, NewClassCDOProperties);
+				ReconstructedCDOsMap.Add(State.OldClassCDO, NewClass->GetDefaultObject());
+			}
 
-					PropertiesToUpdate.Add(PropertyToUpdate);
-				}
+			if (bNeedsReinstancing)
+			{
+				UClass* NullableNewClass = NewClass == OldClass ? nullptr : NewClass;
+				TSharedPtr<FReloadClassReinstancer> ReinstanceHelper = MakeShareable(new FReloadClassReinstancer(
+					NullableNewClass, OldClass, State.OldClassCDO, ReinstancingObjects, CompiledBlueprints));
+				Ar.Logf(ELogVerbosity::Log, TEXT("Re-instancing %s after reload."), *NewClass->GetName());
+				ReinstanceHelper->ReinstanceObjects(true);
+				UpdateDefaultProperties(NewClass, State.OldClassCDOProperties, NewClassCDOProperties);
 			}
 		}
+
+		State.State = TopologicalState::Finished;
 	}
-	if (PropertiesToUpdate.Num())
-	{
-		TArray<uint8> CurrentValueSerializedData;
 
-		// Update properties on all existing instances of the class
-		const UPackage* TransientPackage = GetTransientPackage();
-		for (FThreadSafeObjectIterator It(NewClass); It; ++It)
+	void FReloadClassHelper::ReinstanceClassPropertyScan(FProperty* PropertyLink)
+	{
+		for (FProperty* Property = PropertyLink; Property; Property = Property->PropertyLinkNext)
 		{
-			UObject* ObjectPtr = *It;
-			if (!IsValidChecked(ObjectPtr) || ObjectPtr->GetOutermost() == TransientPackage)
-			{
-				continue;
-			}
+			ReinstanceClassProperty(Property);
+		}
+	}
 
-			DefaultSubobjectArray.Empty(DefaultSubobjectArrayCapacity);
-			ObjectPtr->CollectDefaultSubobjects(DefaultSubobjectArray);
+	void FReloadClassHelper::ReinstanceClassProperty(FProperty* Property)
+	{
+		if (Property == nullptr)
+		{
+			return;
+		}
 
-			for (auto& PropertyToUpdate : PropertiesToUpdate)
+		if (FObjectPropertyBase* objProp = CastField<FObjectPropertyBase>(Property))
+		{
+			if (objProp->PropertyClass != nullptr)
 			{
-				uint8* InstanceValuePtr = nullptr;
-				if (PropertyToUpdate.SubobjectName == NAME_None)
+				ClassReinstanceState* State = ReinstanceStates.Find(objProp->PropertyClass);
+				if (State != nullptr)
 				{
-					InstanceValuePtr = PropertyToUpdate.Property->ContainerPtrToValuePtr<uint8>(ObjectPtr);
-				}
-				else
-				{
-					UObject* DefaultSubobjectPtr = FindDefaultSubobject(DefaultSubobjectArray, PropertyToUpdate.SubobjectName);
-					if (DefaultSubobjectPtr && PropertyToUpdate.Property->GetOwner<UObject>() == DefaultSubobjectPtr->GetClass())
-					{
-						InstanceValuePtr = PropertyToUpdate.Property->ContainerPtrToValuePtr<uint8>(DefaultSubobjectPtr);
-					}
+					ReinstanceClass(objProp->PropertyClass, *State);
 				}
-
-				if (InstanceValuePtr)
+			}
+			if (FClassProperty* classProp = CastField<FClassProperty>(Property))
+			{
+				if (classProp->MetaClass != nullptr)
 				{
-					// Serialize current value to a byte array as we don't have the previous CDO to compare against, we only have its serialized property data
-					CurrentValueSerializedData.Empty(CurrentValueSerializedData.Num() + CurrentValueSerializedData.GetSlack());
-					FPropertyValueMemoryWriter CurrentValueWriter(CurrentValueSerializedData);
-					PropertyToUpdate.Property->SerializeItem(FStructuredArchiveFromArchive(CurrentValueWriter).GetSlot(), InstanceValuePtr);
-
-					// Update only when the current value on the instance is identical to the original CDO
-					if (CurrentValueSerializedData.Num() == PropertyToUpdate.OldSerializedSize &&
-						FMemory::Memcmp(CurrentValueSerializedData.GetData(), PropertyToUpdate.OldSerializedValuePtr, CurrentValueSerializedData.Num()) == 0)
+					ClassReinstanceState* State = ReinstanceStates.Find(classProp->MetaClass);
+					if (State != nullptr)
 					{
-						// Update with the new value
-						PropertyToUpdate.Property->CopyCompleteValue(InstanceValuePtr, PropertyToUpdate.NewValuePtr);
+						ReinstanceClass(classProp->MetaClass, *State);
 					}
 				}
 			}
 		}
+		else if (FStructProperty* structProp = CastField<FStructProperty>(Property))
+		{
+			ReinstanceClassPropertyScan(structProp->Struct->PropertyLink);
+		}
+		else if (FMapProperty* MapProp = CastField<FMapProperty>(Property))
+		{
+			ReinstanceClassProperty(MapProp->KeyProp);
+			ReinstanceClassProperty(MapProp->ValueProp);
+		}
+		else if (FSetProperty* SetProp = CastField<FSetProperty>(Property))
+		{
+			ReinstanceClassProperty(SetProp->ElementProp);
+		}
+		else if (FArrayProperty* ArrayProp = CastField<FArrayProperty>(Property))
+		{
+			ReinstanceClassProperty(ArrayProp->Inner);
+		}
 	}
-}
-
-void FReloadClassReinstancer::ReinstanceObjectsAndUpdateDefaults()
-{
-	ReinstanceObjects(true);
-	UpdateDefaultProperties();
-}
-
-void FReloadClassReinstancer::AddReferencedObjects(FReferenceCollector& Collector)
-{
-	FBlueprintCompileReinstancer::AddReferencedObjects(Collector);
-	Collector.AllowEliminatingReferences(false);
-	Collector.AddReferencedObject(CopyOfPreviousCDO);
-	Collector.AllowEliminatingReferences(true);
-}
-
-void FReloadClassReinstancer::BlueprintWasRecompiled(UBlueprint* BP, bool bBytecodeOnly)
-{
-	CompiledBlueprints.Add(BP);
-
-	FBlueprintCompileReinstancer::BlueprintWasRecompiled(BP, bBytecodeOnly);
-}
+} // namespace UE::Reload::Private
 
 FReload::FReload(EActiveReloadType InType, const TCHAR* InPrefix, const TArray<UPackage*>& InPackages, FOutputDevice& InAr)
 	: Type(InType)
@@ -949,7 +1059,7 @@ void FReload::Reinstance()
 	}
 	NumScriptStructsRemapped = ScriptStructs.Num();
 
-	// Collect all the classes being reinstanced
+	// Collect all the classes being re-instanced
 	TSet<UObject*> ReinstancingObjects;
 	ReinstancingObjects.Reserve(ClassesToReinstance.Num() + ReinstancedStructs.Num() + ReinstancedEnums.Num());
 	for (const TPair<UClass*, UClass*>& Pair : ClassesToReinstance)
@@ -1015,11 +1125,10 @@ void FReload::Reinstance()
 		}
 	}
 
+	// Re-instance the classes
 	TSet<UBlueprint*> CompiledBlueprints;
-	for (const TPair<UClass*, UClass*>& Pair : ClassesToReinstance)
-	{
-		ReinstanceClass(Pair.Value, Pair.Key, ReinstancingObjects, CompiledBlueprints);
-	}
+	UE::Reload::Private::FReloadClassHelper rch(Ar, ReinstancingObjects, CompiledBlueprints, ReconstructedCDOsMap);
+	rch.ReinstanceClasses(ClassesToReinstance);
 
 	// Recompile blueprints if they haven't already been recompiled)
 	for (TPair<UBlueprint*, FBlueprintUpdateInfo>& KVP : ModifiedBlueprints)
@@ -1039,16 +1148,6 @@ void FReload::Reinstance()
 	FCoreUObjectDelegates::ReloadReinstancingCompleteDelegate.Broadcast();
 }
 
-void FReload::ReinstanceClass(UClass* NewClass, UClass* OldClass, const TSet<UObject*>& ReinstancingObjects, TSet<UBlueprint*>& CompiledBlueprints)
-{
-	TSharedPtr<FReloadClassReinstancer> ReinstanceHelper = FReloadClassReinstancer::Create(NewClass, OldClass, ReinstancingObjects, ReconstructedCDOsMap, CompiledBlueprints);
-	if (ReinstanceHelper->ClassNeedsReinstancing())
-	{
-		Ar.Logf(ELogVerbosity::Log, TEXT("Re-instancing %s after reload."), NewClass ? *NewClass->GetName() : *OldClass->GetName());
-		ReinstanceHelper->ReinstanceObjectsAndUpdateDefaults();
-	}
-}
-
 void FReload::Finalize(bool bRunGC)
 {
 
diff --git a/Engine/Source/Editor/UnrealEd/Public/Kismet2/ReloadUtilities.h b/Engine/Source/Editor/UnrealEd/Public/Kismet2/ReloadUtilities.h
--- a/Engine/Source/Editor/UnrealEd/Public/Kismet2/ReloadUtilities.h
+++ b/Engine/Source/Editor/UnrealEd/Public/Kismet2/ReloadUtilities.h
@@ -108,11 +108,6 @@ private:
 	 */
 	void ReplaceReferencesToReconstructedCDOs();
 
-	/**
-	 * Given a new/old class pair, re-instance the class
-	 */
-	void ReinstanceClass(UClass* NewClass, UClass* OldClass, const TSet<UObject*>& ReinstancingObjects, TSet<UBlueprint*>& CompiledBlueprints);
-
 	/**
 	 * Based on the pointers, update the given stat
 	 */
