Ensure construction scripts are executed prior to running the HLOD generation.

diff --git a/Engine/Source/Editor/UnrealEd/Private/HierarchicalLOD.cpp b/Engine/Source/Editor/UnrealEd/Private/HierarchicalLOD.cpp
--- a/Engine/Source/Editor/UnrealEd/Private/HierarchicalLOD.cpp
+++ b/Engine/Source/Editor/UnrealEd/Private/HierarchicalLOD.cpp
@@ -146,6 +146,9 @@ void FHierarchicalLODBuilder::BuildClusters(ULevel* InLevel)
 
 	SCOPE_LOG_TIME(TEXT("STAT_HLOD_BuildClusters"), nullptr);
 
+	// This may execute pending construction scripts.
+	FAssetCompilingManager::Get().ProcessAsyncTasks();
+
 	const TArray<FHierarchicalSimplification>& BuildLODLevelSettings = InLevel->GetWorldSettings()->GetHierarchicalLODSetup();
 	
 	LODLevelLODActors.Empty();
