FNames containing asset paths are deprecated. FSoftObjectPath or FTopLevelAssetPath should be used instead.
Fixups for misc editor files.

diff --git a/Engine/Source/Editor/SourceControlWindows/Private/SSourceControlCommon.cpp b/Engine/Source/Editor/SourceControlWindows/Private/SSourceControlCommon.cpp
--- a/Engine/Source/Editor/SourceControlWindows/Private/SSourceControlCommon.cpp
+++ b/Engine/Source/Editor/SourceControlWindows/Private/SSourceControlCommon.cpp
@@ -206,7 +206,7 @@ FString FFileTreeItem::RetrieveAssetName(const FAssetData& InAssetData) const
 FString FFileTreeItem::RetrieveAssetPath(const FAssetData& InAssetData) const
 {
 	int32 LastDot = -1;
-	FString Path = InAssetData.ObjectPath.ToString();
+	FString Path = InAssetData.GetObjectPathString();
 
 	// Strip asset name from object path
 	if (Path.FindLastChar('.', LastDot))
@@ -243,7 +243,7 @@ FOfflineFileTreeItem::FOfflineFileTreeItem(const FString& InFilename)
 	if (Assets.Num() > 0)
 	{
 		const FAssetData& AssetData = Assets[0];
-		AssetPath = FText::FromName(AssetData.ObjectPath);
+		AssetPath = FText::FromString(AssetData.GetObjectPathString());
 
 		// Find name, asset type & color only if there is exactly one asset
 		if (Assets.Num() == 1)
diff --git a/Engine/Source/Editor/UnrealEd/Private/AssetSelection.cpp b/Engine/Source/Editor/UnrealEd/Private/AssetSelection.cpp
--- a/Engine/Source/Editor/UnrealEd/Private/AssetSelection.cpp
+++ b/Engine/Source/Editor/UnrealEd/Private/AssetSelection.cpp
@@ -702,7 +702,7 @@ namespace AssetUtil
 				{
 					if (DroppedAssetString.Len() < NAME_SIZE && FName::IsValidXName(DroppedAssetString, INVALID_OBJECTPATH_CHARACTERS))
 					{
-						FAssetData AssetData = AssetRegistry.GetAssetByObjectPath(FName(*DroppedAssetString));
+						FAssetData AssetData = AssetRegistry.GetAssetByObjectPath(FSoftObjectPath(DroppedAssetString));
 						if (AssetData.IsValid())
 						{
 							DroppedAssetData.Add(AssetData);
diff --git a/Engine/Source/Editor/UnrealEd/Private/AssetThumbnail.cpp b/Engine/Source/Editor/UnrealEd/Private/AssetThumbnail.cpp
--- a/Engine/Source/Editor/UnrealEd/Private/AssetThumbnail.cpp
+++ b/Engine/Source/Editor/UnrealEd/Private/AssetThumbnail.cpp
@@ -818,9 +818,9 @@ FSlateShaderResource* FAssetThumbnail::GetViewportRenderTargetTexture() const
 
 UObject* FAssetThumbnail::GetAsset() const
 {
-	if ( AssetData.ObjectPath != NAME_None )
+	if ( AssetData.IsValid() )
 	{
-		return FindObject<UObject>(NULL, *AssetData.ObjectPath.ToString());
+		return AssetData.GetSoftObjectPath().ResolveObject();
 	}
 	else
 	{
@@ -988,7 +988,7 @@ void FAssetThumbnailPool::ReleaseResources()
 		const TSharedRef<FThumbnailInfo>& Thumb = *ThumbIt;
 		if ( !Thumb.IsUnique() )
 		{
-			ensureMsgf(0, TEXT("Thumbnail info for '%s' is still referenced by '%d' other objects"), *Thumb->AssetData.ObjectPath.ToString(), Thumb.GetSharedReferenceCount());
+			ensureMsgf(0, TEXT("Thumbnail info for '%s' is still referenced by '%d' other objects"), *Thumb->AssetData.GetObjectPathString(), Thumb.GetSharedReferenceCount());
 		}
 	}
 }
@@ -1097,7 +1097,7 @@ void FAssetThumbnailPool::Tick( float DeltaTime )
 					{
 						if (FPackageName::IsValidObjectPath(CustomThumbnailTagValue))
 						{
-							CustomThumbnailAsset = FModuleManager::LoadModuleChecked<FAssetRegistryModule>(AssetRegistryConstants::ModuleName).Get().GetAssetByObjectPath(*CustomThumbnailTagValue);
+							CustomThumbnailAsset = FModuleManager::LoadModuleChecked<FAssetRegistryModule>(AssetRegistryConstants::ModuleName).Get().GetAssetByObjectPath(FSoftObjectPath(CustomThumbnailTagValue));
 						}
 					}
 
@@ -1235,13 +1235,13 @@ bool FAssetThumbnailPool::LoadThumbnail(TSharedRef<FThumbnailInfo> ThumbnailInfo
 
 FSlateTexture2DRHIRef* FAssetThumbnailPool::AccessTexture( const FAssetData& AssetData, uint32 Width, uint32 Height )
 {
-	if(AssetData.ObjectPath == NAME_None || Width == 0 || Height == 0)
+	if(AssetData.IsValid() || Width == 0 || Height == 0)
 	{
 		return NULL;
 	}
 	else
 	{
-		FThumbId ThumbId( AssetData.ObjectPath, Width, Height ) ;
+		FThumbId ThumbId( AssetData.GetSoftObjectPath(), Width, Height ) ;
 		// Check to see if a thumbnail for this asset exists.  If so we don't need to render it
 		const TSharedRef<FThumbnailInfo>* ThumbnailInfoPtr = ThumbnailToTextureMap.Find( ThumbId );
 		TSharedPtr<FThumbnailInfo> ThumbnailInfo;
@@ -1330,14 +1330,14 @@ FSlateTexture2DRHIRef* FAssetThumbnailPool::AccessTexture( const FAssetData& Ass
 void FAssetThumbnailPool::AddReferencer( const FAssetThumbnail& AssetThumbnail )
 {
 	FIntPoint Size = AssetThumbnail.GetSize();
-	if ( AssetThumbnail.GetAssetData().ObjectPath == NAME_None || Size.X == 0 || Size.Y == 0 )
+	if ( AssetThumbnail.GetAssetData().IsValid() || Size.X == 0 || Size.Y == 0 )
 	{
 		// Invalid referencer
 		return;
 	}
 
 	// Generate a key and look up the number of references in the RefCountMap
-	FThumbId ThumbId( AssetThumbnail.GetAssetData().ObjectPath, Size.X, Size.Y ) ;
+	FThumbId ThumbId( AssetThumbnail.GetAssetData().GetSoftObjectPath(), Size.X, Size.Y ) ;
 	int32* RefCountPtr = RefCountMap.Find(ThumbId);
 
 	if ( RefCountPtr )
@@ -1355,8 +1355,8 @@ void FAssetThumbnailPool::AddReferencer( const FAssetThumbnail& AssetThumbnail )
 void FAssetThumbnailPool::RemoveReferencer( const FAssetThumbnail& AssetThumbnail )
 {
 	FIntPoint Size = AssetThumbnail.GetSize();
-	const FName ObjectPath = AssetThumbnail.GetAssetData().ObjectPath;
-	if ( ObjectPath == NAME_None || Size.X == 0 || Size.Y == 0 )
+	const FSoftObjectPath ObjectPath = AssetThumbnail.GetAssetData().GetSoftObjectPath();
+	if ( ObjectPath.IsNull() || Size.X == 0 || Size.Y == 0 )
 	{
 		// Invalid referencer
 		return;
@@ -1391,9 +1391,9 @@ bool FAssetThumbnailPool::IsInRenderStack( const TSharedPtr<FAssetThumbnail>& Th
 	const uint32 Width = Thumbnail->GetSize().X;
 	const uint32 Height = Thumbnail->GetSize().Y;
 
-	if ( ensure(AssetData.ObjectPath != NAME_None) && ensure(Width > 0) && ensure(Height > 0) )
+	if ( ensure(AssetData.IsValid()) && ensure(Width > 0) && ensure(Height > 0) )
 	{
-		FThumbId ThumbId( AssetData.ObjectPath, Width, Height ) ;
+		FThumbId ThumbId(AssetData.GetSoftObjectPath(), Width, Height);
 		const TSharedRef<FThumbnailInfo>* ThumbnailInfoPtr = ThumbnailToTextureMap.Find( ThumbId );
 		if ( ThumbnailInfoPtr )
 		{
@@ -1410,9 +1410,9 @@ bool FAssetThumbnailPool::IsRendered(const TSharedPtr<FAssetThumbnail>& Thumbnai
 	const uint32 Width = Thumbnail->GetSize().X;
 	const uint32 Height = Thumbnail->GetSize().Y;
 
-	if (ensure(AssetData.ObjectPath != NAME_None) && ensure(Width > 0) && ensure(Height > 0))
+	if (ensure(AssetData.IsValid()) && ensure(Width > 0) && ensure(Height > 0))
 	{
-		FThumbId ThumbId(AssetData.ObjectPath, Width, Height);
+		FThumbId ThumbId(AssetData.GetSoftObjectPath(), Width, Height);
 		const TSharedRef<FThumbnailInfo>* ThumbnailInfoPtr = ThumbnailToTextureMap.Find(ThumbId);
 		if (ThumbnailInfoPtr)
 		{
@@ -1427,17 +1427,17 @@ void FAssetThumbnailPool::PrioritizeThumbnails( const TArray< TSharedPtr<FAssetT
 {
 	if ( ensure(Width > 0) && ensure(Height > 0) )
 	{
-		TSet<FName> ObjectPathList;
+		TSet<FSoftObjectPath> ObjectPathList;
 		for ( int32 ThumbIdx = 0; ThumbIdx < ThumbnailsToPrioritize.Num(); ++ThumbIdx )
 		{
-			ObjectPathList.Add(ThumbnailsToPrioritize[ThumbIdx]->GetAssetData().ObjectPath);
+			ObjectPathList.Add(ThumbnailsToPrioritize[ThumbIdx]->GetAssetData().GetSoftObjectPath());
 		}
 
 		TArray< TSharedRef<FThumbnailInfo> > FoundThumbnails;
 		for ( int32 ThumbIdx = ThumbnailsToRenderStack.Num() - 1; ThumbIdx >= 0; --ThumbIdx )
 		{
 			const TSharedRef<FThumbnailInfo>& ThumbnailInfo = ThumbnailsToRenderStack[ThumbIdx];
-			if ( ThumbnailInfo->Width == Width && ThumbnailInfo->Height == Height && ObjectPathList.Contains(ThumbnailInfo->AssetData.ObjectPath) )
+			if (ThumbnailInfo->Width == Width && ThumbnailInfo->Height == Height && ObjectPathList.Contains(ThumbnailInfo->AssetData.GetSoftObjectPath()))
 			{
 				FoundThumbnails.Add(ThumbnailInfo);
 				ThumbnailsToRenderStack.RemoveAt(ThumbIdx);
@@ -1457,9 +1457,9 @@ void FAssetThumbnailPool::RefreshThumbnail( const TSharedPtr<FAssetThumbnail>& T
 	const uint32 Width = ThumbnailToRefresh->GetSize().X;
 	const uint32 Height = ThumbnailToRefresh->GetSize().Y;
 
-	if ( ensure(AssetData.ObjectPath != NAME_None) && ensure(Width > 0) && ensure(Height > 0) )
+	if ( ensure(AssetData.IsValid()) && ensure(Width > 0) && ensure(Height > 0) )
 	{
-		FThumbId ThumbId( AssetData.ObjectPath, Width, Height ) ;
+		FThumbId ThumbId(AssetData.GetSoftObjectPath(), Width, Height);
 		const TSharedRef<FThumbnailInfo>* ThumbnailInfoPtr = ThumbnailToTextureMap.Find( ThumbId );
 		if ( ThumbnailInfoPtr )
 		{
@@ -1474,9 +1474,9 @@ void FAssetThumbnailPool::SetRealTimeThumbnail( const TSharedPtr<FAssetThumbnail
 	const uint32 Width = Thumbnail->GetSize().X;
 	const uint32 Height = Thumbnail->GetSize().Y;
 
-	if (ensure(AssetData.ObjectPath != NAME_None) && ensure(Width > 0) && ensure(Height > 0) )
+	if (ensure(AssetData.IsValid()) && ensure(Width > 0) && ensure(Height > 0) )
 	{
-		FThumbId ThumbId( AssetData.ObjectPath, Width, Height );
+		FThumbId ThumbId(AssetData.GetSoftObjectPath(), Width, Height);
 		const TSharedRef<FThumbnailInfo>* ThumbnailInfoPtr = ThumbnailToTextureMap.Find( ThumbId );
 		if ( ThumbnailInfoPtr )
 		{
@@ -1492,9 +1492,9 @@ void FAssetThumbnailPool::SetRealTimeThumbnail( const TSharedPtr<FAssetThumbnail
 	}
 }
 
-void FAssetThumbnailPool::FreeThumbnail( const FName& ObjectPath, uint32 Width, uint32 Height )
+void FAssetThumbnailPool::FreeThumbnail( const FSoftObjectPath& ObjectPath, uint32 Width, uint32 Height )
 {
-	if(ObjectPath != NAME_None && Width != 0 && Height != 0)
+	if(ObjectPath.IsValid() && Width != 0 && Height != 0)
 	{
 		FThumbId ThumbId( ObjectPath, Width, Height ) ;
 
@@ -1520,7 +1520,7 @@ void FAssetThumbnailPool::FreeThumbnail( const FName& ObjectPath, uint32 Width,
 			
 }
 
-void FAssetThumbnailPool::RefreshThumbnailsFor( FName ObjectPath )
+void FAssetThumbnailPool::RefreshThumbnailsFor( const FSoftObjectPath& ObjectPath )
 {
 	for ( auto ThumbIt = ThumbnailToTextureMap.CreateIterator(); ThumbIt; ++ThumbIt)
 	{
@@ -1535,7 +1535,7 @@ void FAssetThumbnailPool::OnAssetLoaded( UObject* Asset )
 {
 	if ( Asset != NULL )
 	{
-		RecentlyLoadedAssets.Add( FName(*Asset->GetPathName()) );
+		RecentlyLoadedAssets.Add( FSoftObjectPath(Asset) );
 	}
 }
 
@@ -1612,6 +1612,6 @@ void FAssetThumbnailPool::DirtyThumbnailForObject(UObject* ObjectBeingModified)
 			Thumbnail->MarkAsDirty();
 		}
 
-		RefreshThumbnailsFor( FName(*ObjectBeingModified->GetPathName()) );
+		RefreshThumbnailsFor(FSoftObjectPath(ObjectBeingModified));
 	}
 }
diff --git a/Engine/Source/Editor/UnrealEd/Private/Commandlets/PackageUtilities.cpp b/Engine/Source/Editor/UnrealEd/Private/Commandlets/PackageUtilities.cpp
--- a/Engine/Source/Editor/UnrealEd/Private/Commandlets/PackageUtilities.cpp
+++ b/Engine/Source/Editor/UnrealEd/Private/Commandlets/PackageUtilities.cpp
@@ -305,22 +305,34 @@ public:
 		return CollectionManagerModule.Get().DestroyCollection(InSetName, InSetType);
 	}
 
-	static bool RemoveAssetsFromSet(FName InSetName, ECollectionShareType::Type InSetType, const TArray<FName>& InAssetPathNames )
+	static bool RemoveAssetsFromSet(FName InSetName, ECollectionShareType::Type InSetType, const TArray<FSoftObjectPath>& InAssetPathNames )
 	{
 		FCollectionManagerModule& CollectionManagerModule = FCollectionManagerModule::GetModule();
-		return CollectionManagerModule.Get().RemoveFromCollection(InSetName, InSetType, InAssetPathNames);
+PRAGMA_DISABLE_DEPRECATION_WARNINGS
+		return CollectionManagerModule.Get().RemoveFromCollection(InSetName, InSetType, UE::SoftObjectPath::Private::ConvertSoftObjectPaths(InAssetPathNames));
+PRAGMA_ENABLE_DEPRECATION_WARNINGS
 	}
 
-	static bool AddAssetsToSet(FName InSetName, ECollectionShareType::Type InSetType, const TArray<FName>& InAssetPathNames )
+	static bool AddAssetsToSet(FName InSetName, ECollectionShareType::Type InSetType, const TArray<FSoftObjectPath>& InAssetPathNames )
 	{
 		FCollectionManagerModule& CollectionManagerModule = FCollectionManagerModule::GetModule();
-		return CollectionManagerModule.Get().AddToCollection(InSetName, InSetType, InAssetPathNames);
+PRAGMA_DISABLE_DEPRECATION_WARNINGS
+		return CollectionManagerModule.Get().AddToCollection(InSetName, InSetType, UE::SoftObjectPath::Private::ConvertSoftObjectPaths(InAssetPathNames));
+PRAGMA_ENABLE_DEPRECATION_WARNINGS
 	}
 
-	static bool QueryAssetsInSet(FName InSetName, ECollectionShareType::Type InSetType, TArray<FName>& OutAssetPathNames )
+	static bool QueryAssetsInSet(FName InSetName, ECollectionShareType::Type InSetType, TArray<FSoftObjectPath>& OutAssetPathNames )
 	{
 		FCollectionManagerModule& CollectionManagerModule = FCollectionManagerModule::GetModule();
-		return CollectionManagerModule.Get().GetAssetsInCollection(InSetName, InSetType, OutAssetPathNames);
+		TArray<FName> Temp;
+PRAGMA_DISABLE_DEPRECATION_WARNINGS
+		if (CollectionManagerModule.Get().GetAssetsInCollection(InSetName, InSetType, Temp))
+		{
+			OutAssetPathNames.Append(UE::SoftObjectPath::Private::ConvertObjectPathNames(Temp));
+PRAGMA_ENABLE_DEPRECATION_WARNINGS
+			return true;
+		}		
+		return false;
 	}
 };
 
@@ -352,7 +364,7 @@ bool FContentHelper::ClearAssetSet(FName InSetName, ECollectionShareType::Type I
 
 /** Sets the contents of a Tag or Collection to be the InAssetList. Assets not mentioned in the list will be untagged. */
 template <class AssetSetPolicy>	
-bool FContentHelper::AssignSetContent(FName InSetName, ECollectionShareType::Type InType, const TArray<FName>& InAssetList )
+bool FContentHelper::AssignSetContent(FName InSetName, ECollectionShareType::Type InType, const TArray<FSoftObjectPath>& InAssetList )
 {
 	bool bResult = true;
 
@@ -373,19 +385,19 @@ bool FContentHelper::AssignSetContent(FName InSetName, ECollectionShareType::Typ
 		{
 			bool bAddCompleteInAssetList = true;
 
-			TArray<FName> AssetsInCollection;
+			TArray<FSoftObjectPath> AssetsInCollection;
 			AssetSetPolicy::QueryAssetsInSet(InSetName, InType, AssetsInCollection);
 			int32 CurrentAssetCount = AssetsInCollection.Num();
 			if (CurrentAssetCount != 0)
 			{
 				// Generate the lists
-				TArray<FName> TrueAddList;
-				TArray<FName> TrueRemoveList;
+				TArray<FSoftObjectPath> TrueAddList;
+				TArray<FSoftObjectPath> TrueRemoveList;
 
 				// See how many items are really being added/removed
 				for (int32 CheckIdx = 0; CheckIdx < AssetsInCollection.Num(); CheckIdx++)
 				{
-					FName CheckAsset = AssetsInCollection[CheckIdx];
+					FSoftObjectPath CheckAsset = FSoftObjectPath(AssetsInCollection[CheckIdx]);
 					if (InAssetList.Find(CheckAsset) != INDEX_NONE)
 					{
 						TrueAddList.AddUnique(CheckAsset);
@@ -452,7 +464,7 @@ bool FContentHelper::AssignSetContent(FName InSetName, ECollectionShareType::Typ
 
 /** Add and remove assets for the specified Tag or Connection. Assets from InAddList are added; assets from InRemoveList are removed. */
 template <class AssetSetPolicy>	
-bool FContentHelper::UpdateSetContent(FName InSetName, ECollectionShareType::Type InType, const TArray<FName>& InAddList, const TArray<FName>& InRemoveList )
+bool FContentHelper::UpdateSetContent(FName InSetName, ECollectionShareType::Type InType, const TArray<FSoftObjectPath>& InAddList, const TArray<FSoftObjectPath>& InRemoveList )
 {
 	bool bResult = true;
 
@@ -471,13 +483,13 @@ bool FContentHelper::UpdateSetContent(FName InSetName, ECollectionShareType::Typ
 		// If there is nothing to update, we are done.
 		if ((InAddList.Num() >= 0) || (InRemoveList.Num() >= 0))
 		{
-			TArray<FName> AssetsInCollection;
+			TArray<FSoftObjectPath> AssetsInCollection;
 			AssetSetPolicy::QueryAssetsInSet(InSetName, InType, AssetsInCollection);
 			if (AssetsInCollection.Num() != 0)
 			{
 				// Clean up the lists
-				TArray<FName> TrueAddList;
-				TArray<FName> TrueRemoveList;
+				TArray<FSoftObjectPath> TrueAddList;
+				TArray<FSoftObjectPath> TrueRemoveList;
 
 				// Generate the true Remove list, only removing items that are actually in the collection.
 				for (int32 RemoveIdx = 0; RemoveIdx < InRemoveList.Num(); RemoveIdx++)
@@ -537,7 +549,7 @@ bool FContentHelper::UpdateSetContent(FName InSetName, ECollectionShareType::Typ
 
 /** Get the list of all assets in the specified Collection or Tag */
 template <class AssetSetPolicy>	
-bool FContentHelper::QuerySetContent(FName InSetName, ECollectionShareType::Type InType, TArray<FName>& OutAssetPathNames)
+bool FContentHelper::QuerySetContent(FName InSetName, ECollectionShareType::Type InType, TArray<FSoftObjectPath>& OutAssetPathNames)
 {
 	if (bInitialized == false)
 	{
@@ -597,7 +609,7 @@ bool FContentHelper::ClearCollection(FName InCollectionName, ECollectionShareTyp
  *
  *	@return	bool				true if successful, false if not.
  */
-bool FContentHelper::SetCollection(FName InCollectionName, ECollectionShareType::Type InType, const TArray<FName>& InAssetList)
+bool FContentHelper::SetCollection(FName InCollectionName, ECollectionShareType::Type InType, const TArray<FSoftObjectPath>& InAssetList)
 {
 	return this->AssignSetContent<FCollectionPolicy>(InCollectionName, InType, InAssetList);
 }
@@ -612,7 +624,7 @@ bool FContentHelper::SetCollection(FName InCollectionName, ECollectionShareType:
  *
  *	@return	bool				true if successful, false if not.
  */
-bool FContentHelper::UpdateCollection(FName InCollectionName, ECollectionShareType::Type InType, const TArray<FName>& InAddList, const TArray<FName>& InRemoveList)
+bool FContentHelper::UpdateCollection(FName InCollectionName, ECollectionShareType::Type InType, const TArray<FSoftObjectPath>& InAddList, const TArray<FSoftObjectPath>& InRemoveList)
 {
 	return this->UpdateSetContent<FCollectionPolicy>( InCollectionName, InType, InAddList, InRemoveList );
 }
@@ -622,15 +634,28 @@ bool FContentHelper::UpdateCollection(FName InCollectionName, ECollectionShareTy
  *
  *	@param	InCollectionName	Name of collection to query
  *	@param	InType				Type of collection
- *	@param	OutAssetPathNames	The assets contained in the collection
+ *	@param	OutAssetPaths		The assets contained in the collection
  * 
  *	@return True if collection was created successfully
  */
-bool FContentHelper::QueryAssetsInCollection(FName InCollectionName, ECollectionShareType::Type InType, TArray<FName>& OutAssetPathNames)
+bool FContentHelper::QueryAssetsInCollection(FName InCollectionName, ECollectionShareType::Type InType, TArray<FSoftObjectPath>& OutAssetPaths)
 {
-	return this->QuerySetContent<FCollectionPolicy>(InCollectionName, InType, OutAssetPathNames);
+	return this->QuerySetContent<FCollectionPolicy>(InCollectionName, InType, OutAssetPaths);
 }
 
+bool FContentHelper::QueryAssetsInCollection(FName InCollectionName, ECollectionShareType::Type InType, TArray<FName>& OutAssetPathNames)
+{
+	TArray<FSoftObjectPath> AssetPaths;
+PRAGMA_DISABLE_DEPRECATION_WARNINGS
+	if (this->QuerySetContent<FCollectionPolicy>(InCollectionName, InType, AssetPaths))
+	{
+		OutAssetPathNames.Append(UE::SoftObjectPath::Private::ConvertSoftObjectPaths(AssetPaths));
+		return true;
+	}
+PRAGMA_ENABLE_DEPRECATION_WARNINGS
+
+	return false;
+}
 
 /*-----------------------------------------------------------------------------
 ULoadPackageCommandlet
@@ -1386,7 +1411,7 @@ void FPkgInfoReporter_Log::GeneratePackageReport( FLinkerLoad* InLinker /*=nullp
 			for (FAssetData* AssetData : AssetDatas)
 			{
 				// Display the asset class and path
-				Out.Logf(ELogVerbosity::Display, TEXT("\t\t%d) %s'%s' (%d Tags)"), AssetIdx++, *AssetData->AssetClassPath.ToString(), *AssetData->ObjectPath.ToString(), AssetData->TagsAndValues.Num());
+				Out.Logf(ELogVerbosity::Display, TEXT("\t\t%d) %s'%s' (%d Tags)"), AssetIdx++, *AssetData->AssetClassPath.ToString(), *AssetData->GetObjectPathString(), AssetData->TagsAndValues.Num());
 
 				// Display all tags on the asset
 				for (const TPair<FName, FAssetTagValueRef>& Pair : AssetData->TagsAndValues)
diff --git a/Engine/Source/Editor/UnrealEd/Private/ComponentTypeRegistry.cpp b/Engine/Source/Editor/UnrealEd/Private/ComponentTypeRegistry.cpp
--- a/Engine/Source/Editor/UnrealEd/Private/ComponentTypeRegistry.cpp
+++ b/Engine/Source/Editor/UnrealEd/Private/ComponentTypeRegistry.cpp
@@ -32,8 +32,8 @@ namespace UE::Editor::ComponentTypeRegistry::Private
 	{
 	public:
 		FUnloadedBlueprintData(const FAssetData& InAssetData)
-			:ClassPath(NAME_None)
-			,ParentClassPath(NAME_None)
+			:ClassPath()
+			,ParentClassPath()
 			,ClassFlags(CLASS_None)
 			,bIsNormalBlueprintType(false)
 		{
@@ -43,17 +43,17 @@ namespace UE::Editor::ComponentTypeRegistry::Private
 			const UClass* AssetClass = InAssetData.GetClass();
 			if (AssetClass && AssetClass->IsChildOf(UBlueprintGeneratedClass::StaticClass()))
 			{
-				ClassPath = InAssetData.ObjectPath;
+				ClassPath = InAssetData.GetSoftObjectPath().GetAssetPath();
 			}
 			else if (InAssetData.GetTagValue(FBlueprintTags::GeneratedClassPath, GeneratedClassPath))
 			{
-				ClassPath = FName(*FPackageName::ExportTextPathToObjectPath(GeneratedClassPath));
+				ClassPath = FTopLevelAssetPath(FPackageName::ExportTextPathToObjectPath(GeneratedClassPath));
 			}
 
 			FString ParentClassPathString;
 			if (InAssetData.GetTagValue(FBlueprintTags::ParentClassPath, ParentClassPathString))
 			{
-				ParentClassPath = FName(*FPackageName::ExportTextPathToObjectPath(ParentClassPathString));
+				ParentClassPath = FTopLevelAssetPath(FPackageName::ExportTextPathToObjectPath(ParentClassPathString));
 			}
 
 			FEditorClassUtils::GetImplementedInterfaceClassPathsFromAsset(InAssetData, ImplementedInterfaces);
@@ -231,19 +231,21 @@ namespace UE::Editor::ComponentTypeRegistry::Private
 		UE_DEPRECATED(5.1, "Class names are now represented by path names. Please use GetClassPathName.")
 		virtual FName GetClassPath() const override
 		{
-			return ClassPath;
+PRAGMA_DISABLE_DEPRECATION_WARNINGS
+			return ClassPath.ToFName();
+PRAGMA_ENABLE_DEPRECATION_WARNINGS
 		}
 		
 		virtual FTopLevelAssetPath GetClassPathName() const override
 		{
-			return FTopLevelAssetPath(ClassPath.ToString());
+			return ClassPath;
 		}
 		// End IUnloadedBlueprintData interface
 
 	private:
 		TSharedPtr<FString> ClassName;
-		FName ClassPath;
-		FName ParentClassPath;
+		FTopLevelAssetPath ClassPath;
+		FTopLevelAssetPath ParentClassPath;
 		uint32 ClassFlags;
 		TArray<FString> ImplementedInterfaces;
 		bool bIsNormalBlueprintType;
@@ -287,7 +289,7 @@ public:
 	TArray<FComponentClassComboEntryPtr> ComponentClassList;
 	TArray<FComponentTypeEntry> ComponentTypeList;
 	TArray<FAssetData> PendingAssetData;
-	TMap<FName, int32> ClassPathToClassListIndexMap;
+	TMap<FTopLevelAssetPath, int32> ClassPathToClassListIndexMap;
 	FOnComponentTypeListChanged ComponentListChanged;
 	bool bNeedsRefreshNextTick;
 };
@@ -659,7 +661,7 @@ void FComponentTypeRegistryData::ForceRefreshComponentList()
 			int32 EntryIndex = ComponentClassList.Add(CurrentEntry);
 			if (CurrentEntry->IsClass())
 			{
-				ClassPathToClassListIndexMap.FindOrAdd(FName(*CurrentEntry->GetComponentPath()), EntryIndex);
+				ClassPathToClassListIndexMap.FindOrAdd(FTopLevelAssetPath(CurrentEntry->GetComponentPath()), EntryIndex);
 			}
 		}
 	}
@@ -780,7 +782,7 @@ void FComponentTypeRegistry::Invalidate()
 	Data->Invalidate();
 }
 
-FComponentClassComboEntryPtr FComponentTypeRegistry::FindClassEntryForObjectPath(FName InObjectPath) const
+FComponentClassComboEntryPtr FComponentTypeRegistry::FindClassEntryForObjectPath(FTopLevelAssetPath InObjectPath) const
 {
 	if (int32* ClassListIndexPtr = Data->ClassPathToClassListIndexMap.Find(InObjectPath))
 	{
diff --git a/Engine/Source/Editor/UnrealEd/Private/EditorEngine.cpp b/Engine/Source/Editor/UnrealEd/Private/EditorEngine.cpp
--- a/Engine/Source/Editor/UnrealEd/Private/EditorEngine.cpp
+++ b/Engine/Source/Editor/UnrealEd/Private/EditorEngine.cpp
@@ -3007,7 +3007,7 @@ void UEditorEngine::GetAssetsToSyncToContentBrowser(TArray<FAssetData>& Assets,
 
 						for (const FSoftObjectPath& SoftObject : SoftObjects)
 						{
-							FAssetData AssetData = AssetRegistry.GetAssetByObjectPath(*SoftObject.ToString());
+							FAssetData AssetData = AssetRegistry.GetAssetByObjectPath(SoftObject);
 
 							if (AssetData.IsValid())
 							{
diff --git a/Engine/Source/Editor/UnrealEd/Private/EditorScriptingHelpers.cpp b/Engine/Source/Editor/UnrealEd/Private/EditorScriptingHelpers.cpp
--- a/Engine/Source/Editor/UnrealEd/Private/EditorScriptingHelpers.cpp
+++ b/Engine/Source/Editor/UnrealEd/Private/EditorScriptingHelpers.cpp
@@ -234,7 +234,7 @@ bool EditorScriptingHelpers::IsAValidPathForCreateNewAsset(const FString& Object
 	}
 
 	FAssetRegistryModule& AssetRegistryModule = FModuleManager::LoadModuleChecked<FAssetRegistryModule>("AssetRegistry");
-	FAssetData AssetData = AssetRegistryModule.Get().GetAssetByObjectPath(*ObjectPath);
+	FAssetData AssetData = AssetRegistryModule.Get().GetAssetByObjectPath(FSoftObjectPath(ObjectPath));
 	if (AssetData.IsValid())
 	{
 		OutFailureReason = TEXT("An asset already exists at this location.");
diff --git a/Engine/Source/Editor/UnrealEd/Private/FileHelpers.cpp b/Engine/Source/Editor/UnrealEd/Private/FileHelpers.cpp
--- a/Engine/Source/Editor/UnrealEd/Private/FileHelpers.cpp
+++ b/Engine/Source/Editor/UnrealEd/Private/FileHelpers.cpp
@@ -831,7 +831,7 @@ static bool SaveWorld(UWorld* World,
 
 					// We're changing the world path, add a path redirector so that soft object paths get fixed on save
 					FSoftObjectPath NewPath( World );
-					GRedirectCollector.AddAssetPathRedirection( *OldPath.GetAssetPathString(), *NewPath.GetAssetPathString() );
+					GRedirectCollector.AddAssetPathRedirection( OldPath.GetWithoutSubPath(), NewPath.GetWithoutSubPath() );
 					bAddedAssetPathRedirection = true;
 				}
 			}
diff --git a/Engine/Source/Editor/UnrealEd/Private/LevelEditorDragDropHandler.cpp b/Engine/Source/Editor/UnrealEd/Private/LevelEditorDragDropHandler.cpp
--- a/Engine/Source/Editor/UnrealEd/Private/LevelEditorDragDropHandler.cpp
+++ b/Engine/Source/Editor/UnrealEd/Private/LevelEditorDragDropHandler.cpp
@@ -25,7 +25,7 @@ bool ULevelEditorDragDropHandler::PreviewDropObjectsAtCoordinates(int32 MouseX,
 	bCanDrop = false;
 	HintText = FText::GetEmpty();
 
-	if ( AssetData.ObjectPath.IsNone() )
+	if ( !AssetData.IsValid() )
 	{
 		bCanDrop = false;
 		return false;
diff --git a/Engine/Source/Editor/UnrealEd/Private/MRUList.cpp b/Engine/Source/Editor/UnrealEd/Private/MRUList.cpp
--- a/Engine/Source/Editor/UnrealEd/Private/MRUList.cpp
+++ b/Engine/Source/Editor/UnrealEd/Private/MRUList.cpp
@@ -182,9 +182,9 @@ bool FMRUList::VerifyMRUFile(int32 InItem, FString& OutPackageName)
 
 	// Handle redirector
 	const FString OriginalPackageName = Items[InItem];
-	const FName OriginalObjectPath = FName(*(OriginalPackageName + TEXT('.') + FPackageName::GetShortName(OriginalPackageName)));
+	const FSoftObjectPath OriginalObjectPath = FSoftObjectPath(*OriginalPackageName, *FPackageName::GetShortName(OriginalPackageName), {});
 	FAssetRegistryModule& AssetRegistryModule = FModuleManager::LoadModuleChecked<FAssetRegistryModule>("AssetRegistry");
-	const FName RedirectedObjectPath = AssetRegistryModule.Get().GetRedirectedObjectPath(OriginalObjectPath);
+	const FSoftObjectPath RedirectedObjectPath = AssetRegistryModule.Get().GetRedirectedObjectPath(OriginalObjectPath);
 
 	FString PackageName;
 	FString RedirectedPackageName;
diff --git a/Engine/Source/Editor/UnrealEd/Private/MaterialGraphSchema.cpp b/Engine/Source/Editor/UnrealEd/Private/MaterialGraphSchema.cpp
--- a/Engine/Source/Editor/UnrealEd/Private/MaterialGraphSchema.cpp
+++ b/Engine/Source/Editor/UnrealEd/Private/MaterialGraphSchema.cpp
@@ -942,7 +942,7 @@ void UMaterialGraphSchema::GetMaterialFunctionActions(FGraphActionMenuBuilder& A
 			if (!ActionMenuBuilder.FromPin || HasCompatibleConnection(AssetData, FromPinType, ActionMenuBuilder.FromPin->Direction))
 			{
 				// Gather the relevant information from the asset data
-				const FString FunctionPathName = AssetData.ObjectPath.ToString();
+				const FString FunctionPathName = AssetData.GetObjectPathString();
 				const FText Description = AssetData.GetTagValueRef<FText>("Description");
 				TArray<FString> LibraryCategories;
 				{
diff --git a/Engine/Source/Editor/UnrealEd/Private/ObjectTools.cpp b/Engine/Source/Editor/UnrealEd/Private/ObjectTools.cpp
--- a/Engine/Source/Editor/UnrealEd/Private/ObjectTools.cpp
+++ b/Engine/Source/Editor/UnrealEd/Private/ObjectTools.cpp
@@ -2009,7 +2009,7 @@ namespace ObjectTools
 				}
 				else
 				{
-					TArray<FName> ObjectsToAdd;
+					TArray<FSoftObjectPath> ObjectsToAdd;
 					for(TSet<UObject*>::TConstIterator SetIt(ReferencedObjects); SetIt; ++SetIt)
 					{
 						UObject* RefObj = *SetIt;
@@ -2017,7 +2017,7 @@ namespace ObjectTools
 						{
 							if (RefObj != Object)
 							{
-								ObjectsToAdd.Add(FName(*RefObj->GetPathName()));
+								ObjectsToAdd.Emplace(RefObj);
 							}
 						}
 					}
diff --git a/Engine/Source/Editor/UnrealEd/Private/SkinWeightsUtilities.cpp b/Engine/Source/Editor/UnrealEd/Private/SkinWeightsUtilities.cpp
--- a/Engine/Source/Editor/UnrealEd/Private/SkinWeightsUtilities.cpp
+++ b/Engine/Source/Editor/UnrealEd/Private/SkinWeightsUtilities.cpp
@@ -196,7 +196,7 @@ bool FSkinWeightsUtilities::ImportAlternateSkinWeight(USkeletalMesh* SkeletalMes
 
 		for (FString AssetPath : Task->ImportedObjectPaths)
 		{
-			FAssetData AssetData = AssetRegistryModule.Get().GetAssetByObjectPath(FName(*AssetPath));
+			FAssetData AssetData = AssetRegistryModule.Get().GetAssetByObjectPath(FSoftObjectPath(AssetPath));
 			ImportedObject = AssetData.GetAsset();
 			if (ImportedObject != nullptr)
 			{
diff --git a/Engine/Source/Editor/UnrealEd/Private/Tests/AutomationEditorCommon.cpp b/Engine/Source/Editor/UnrealEd/Private/Tests/AutomationEditorCommon.cpp
--- a/Engine/Source/Editor/UnrealEd/Private/Tests/AutomationEditorCommon.cpp
+++ b/Engine/Source/Editor/UnrealEd/Private/Tests/AutomationEditorCommon.cpp
@@ -577,7 +577,7 @@ FAssetData FAutomationEditorCommonUtils::GetAssetDataFromPackagePath(const FStri
 	if (AssetPath.Len() > 0)
 	{
 		IAssetRegistry& AssetRegistry = FModuleManager::LoadModuleChecked<FAssetRegistryModule>(TEXT("AssetRegistry")).Get();
-		return AssetRegistry.GetAssetByObjectPath(*AssetPath);
+		return AssetRegistry.GetAssetByObjectPath(FSoftObjectPath(AssetPath));
 	}
 
 	return FAssetData();
@@ -605,14 +605,14 @@ void FAutomationEditorCommonUtils::CollectTestsByClass(UClass * Class, TArray<FS
 	for (auto ObjIter = ObjectList.CreateConstIterator(); ObjIter; ++ObjIter)
 	{
 		const FAssetData& Asset = *ObjIter;
-		FString Filename = Asset.ObjectPath.ToString();
+		FString Filename = Asset.GetObjectPathString();
 		//convert to full paths
 		Filename = FPackageName::LongPackageNameToFilename(Filename);
 		if (FAutomationTestFramework::Get().ShouldTestContent(Filename))
 		{
 			FString BeautifiedFilename = Asset.AssetName.ToString();
 			OutBeautifiedNames.Add(BeautifiedFilename);
-			OutTestCommands.Add(Asset.ObjectPath.ToString());
+			OutTestCommands.Add(Asset.GetObjectPathString());
 		}
 	}
 }
@@ -644,7 +644,7 @@ void FAutomationEditorCommonUtils::CollectGameContentTestsByClass(UClass * Class
 	for (auto ObjIter = ObjectList.CreateConstIterator(); ObjIter; ++ObjIter)
 	{
 		const FAssetData& Asset = *ObjIter;
-		FString Filename = Asset.ObjectPath.ToString();
+		FString Filename = Asset.GetObjectPathString();
 
 		if (Filename.StartsWith("/Game"))
 		{
@@ -654,7 +654,7 @@ void FAutomationEditorCommonUtils::CollectGameContentTestsByClass(UClass * Class
 			{
 				FString BeautifiedFilename = Asset.AssetName.ToString();
 				OutBeautifiedNames.Add(BeautifiedFilename);
-				OutTestCommands.Add(Asset.ObjectPath.ToString());
+				OutTestCommands.Add(Asset.GetObjectPathString());
 			}
 		}
 	}
@@ -710,11 +710,11 @@ void FAutomationEditorCommonUtils::CollectGameContentTests(TArray<FString>& OutB
 		if (Asset.GetClass() == nullptr)
 		{
 			// a nullptr class is bad !
-			UE_LOG(LogAutomationEditorCommon, Warning, TEXT("GetClass for %s (%s) returned nullptr. Asset ignored"), *Asset.AssetName.ToString(), *Asset.ObjectPath.ToString());
+			UE_LOG(LogAutomationEditorCommon, Warning, TEXT("GetClass for %s (%s) returned nullptr. Asset ignored"), *Asset.AssetName.ToString(), *Asset.GetObjectPathString());
 		}
 		else 
 		{
-			FString Filename = Asset.ObjectPath.ToString();
+			FString Filename = Asset.GetObjectPathString();
 
 			if (Filename.StartsWith("/Game"))
 			{
@@ -724,7 +724,7 @@ void FAutomationEditorCommonUtils::CollectGameContentTests(TArray<FString>& OutB
 				{
 					FString BeautifiedFilename = FString::Printf(TEXT("%s.%s"), *Asset.AssetClassPath.ToString(), *Asset.AssetName.ToString());
 					OutBeautifiedNames.Add(BeautifiedFilename);
-					OutTestCommands.Add(Asset.ObjectPath.ToString());
+					OutTestCommands.Add(Asset.GetObjectPathString());
 				}
 			}
 		}
diff --git a/Engine/Source/Editor/UnrealEd/Private/UnrealEdMisc.cpp b/Engine/Source/Editor/UnrealEd/Private/UnrealEdMisc.cpp
--- a/Engine/Source/Editor/UnrealEd/Private/UnrealEdMisc.cpp
+++ b/Engine/Source/Editor/UnrealEd/Private/UnrealEdMisc.cpp
@@ -1631,7 +1631,7 @@ void FUnrealEdMisc::OnGotoAsset(const FString& InAssetPath) const
 {
 	FAssetRegistryModule& AssetRegistryModule = FModuleManager::LoadModuleChecked<FAssetRegistryModule>(AssetRegistryName);
 	IAssetRegistry& AssetRegistry = AssetRegistryModule.Get();
-	FAssetData AssetData = AssetRegistry.GetAssetByObjectPath( *InAssetPath );
+	FAssetData AssetData = AssetRegistry.GetAssetByObjectPath( FSoftObjectPath(InAssetPath) );
 	if ( AssetData.IsValid() )
 	{
 		TArray<FAssetData> AssetDataToSync;
diff --git a/Engine/Source/Editor/UnrealEd/Public/AssetThumbnail.h b/Engine/Source/Editor/UnrealEd/Public/AssetThumbnail.h
--- a/Engine/Source/Editor/UnrealEd/Public/AssetThumbnail.h
+++ b/Engine/Source/Editor/UnrealEd/Public/AssetThumbnail.h
@@ -256,10 +256,10 @@ private:
 	 * @param Width 		The width of the thumbnail to free
 	 * @param Height		The height of the thumbnail to free
 	 */
-	void FreeThumbnail( const FName& ObjectPath, uint32 Width, uint32 Height );
+	void FreeThumbnail( const FSoftObjectPath& ObjectPath, uint32 Width, uint32 Height );
 
 	/** Adds the thumbnails associated with the object found at ObjectPath to the render stack */
-	void RefreshThumbnailsFor( FName ObjectPath );
+	void RefreshThumbnailsFor( const FSoftObjectPath& ObjectPath );
 
 	/** Handler for when an asset is loaded */
 	void OnAssetLoaded( UObject* Asset );
@@ -326,12 +326,12 @@ private:
 	/** Key for looking up thumbnails in a map */
 	struct FThumbId
 	{
-		FName ObjectPath;
+		FSoftObjectPath ObjectPath;
 		uint32 Width;
 		uint32 Height;
 
-		FThumbId( const FName& InObjectPath, uint32 InWidth, uint32 InHeight )
-			: ObjectPath( InObjectPath )
+		FThumbId( FSoftObjectPath InObjectPath, uint32 InWidth, uint32 InHeight )
+			: ObjectPath( MoveTemp(InObjectPath) )
 			, Width( InWidth )
 			, Height( InHeight )
 		{}
@@ -371,7 +371,7 @@ private:
 	TMap< FThumbId, int32 > RefCountMap;
 
 	/** A list of object paths for recently loaded assets whose thumbnails need to be refreshed. */
-	TArray<FName> RecentlyLoadedAssets;
+	TArray<FSoftObjectPath> RecentlyLoadedAssets;
 
 	/** Attribute that determines if real-time thumbnails are allowed. Called every frame. */
 	TAttribute<bool> AreRealTimeThumbnailsAllowed;
diff --git a/Engine/Source/Editor/UnrealEd/Public/ComponentTypeRegistry.h b/Engine/Source/Editor/UnrealEd/Public/ComponentTypeRegistry.h
--- a/Engine/Source/Editor/UnrealEd/Public/ComponentTypeRegistry.h
+++ b/Engine/Source/Editor/UnrealEd/Public/ComponentTypeRegistry.h
@@ -49,7 +49,7 @@ struct UNREALED_API FComponentTypeRegistry
 	/**
 	 * Attempts to locate the class entry corresponding to the given object path.
 	 */
-	FComponentClassComboEntryPtr FindClassEntryForObjectPath(FName InObjectPath) const;
+	FComponentClassComboEntryPtr FindClassEntryForObjectPath(FTopLevelAssetPath InObjectPath) const;
 
 private:
 	void OnReloadComplete(EReloadCompleteReason Reason);
diff --git a/Engine/Source/Editor/UnrealEd/Public/PackageHelperFunctions.h b/Engine/Source/Editor/UnrealEd/Public/PackageHelperFunctions.h
--- a/Engine/Source/Editor/UnrealEd/Public/PackageHelperFunctions.h
+++ b/Engine/Source/Editor/UnrealEd/Public/PackageHelperFunctions.h
@@ -138,7 +138,7 @@ public:
 	 *
 	 *	@return	bool				true if successful, false if not.
 	 */
-	bool SetCollection(FName InCollectionName, ECollectionShareType::Type InType, const TArray<FName>& InAssetList);
+	bool SetCollection(FName InCollectionName, ECollectionShareType::Type InType, const TArray<FSoftObjectPath>& InAssetList);
 
 	/**
 	 *	Update the given collection with the lists of adds/removes
@@ -150,17 +150,20 @@ public:
 	 *
 	 *	@return	bool				true if successful, false if not.
 	 */
-	bool UpdateCollection(FName InCollectionName, ECollectionShareType::Type InType, const TArray<FName>& InAddList, const TArray<FName>& InRemoveList);
+	bool UpdateCollection(FName InCollectionName, ECollectionShareType::Type InType, const TArray<FSoftObjectPath>& InAddList, const TArray<FSoftObjectPath>& InRemoveList);
 
 	/**
 	 *	Retrieve the assets contained in the given collection.
 	 *
 	 *	@param	InCollectionName	Name of collection to query
 	 *	@param	InType				Type of collection
-	 *	@param	OutAssetPathNames	The assets contained in the collection
+	 *	@param	OutAssetPaths		The assets contained in the collection
 	 * 
 	 *	@return True if collection was created successfully
 	 */
+	UNREALED_API bool QueryAssetsInCollection(FName InCollectionName, ECollectionShareType::Type InType, TArray<FSoftObjectPath>& OutAssetPaths);
+	
+	UE_DEPRECATED(5.1, "FNames containing full asset paths are deprecated. Use FSoftObjectPath instead.")
 	UNREALED_API bool QueryAssetsInCollection(FName InCollectionName, ECollectionShareType::Type InType, TArray<FName>& OutAssetPathNames);
 
 
@@ -172,13 +175,13 @@ protected:
 	template <class AssetSetPolicy>	bool ClearAssetSet ( FName InSetName, ECollectionShareType::Type InSetType );
 	
 	/** Sets the contents of a Tag or Collection to be the InAssetList. Assets not mentioned in the list will be untagged. */
-	template <class AssetSetPolicy>	bool AssignSetContent( FName InSetName, ECollectionShareType::Type InType, const TArray<FName>& InAssetList );
+	template <class AssetSetPolicy>	bool AssignSetContent( FName InSetName, ECollectionShareType::Type InType, const TArray<FSoftObjectPath>& InAssetList );
 	
 	/** Add and remove assets for the specified Tag or Connection. Assets from InAddList are added; assets from InRemoveList are removed. */
-	template <class AssetSetPolicy>	bool UpdateSetContent( FName InSetName, ECollectionShareType::Type InType, const TArray<FName>& InAddList, const TArray<FName>& InRemoveList );
+	template <class AssetSetPolicy>	bool UpdateSetContent( FName InSetName, ECollectionShareType::Type InType, const TArray<FSoftObjectPath>& InAddList, const TArray<FSoftObjectPath>& InRemoveList );
 	
 	/** Get the list of all assets in the specified Collection or Tag */
-	template <class AssetSetPolicy>	bool QuerySetContent( FName InCollectionName, ECollectionShareType::Type InType, TArray<FName>& OutAssetFullNames );
+	template <class AssetSetPolicy>	bool QuerySetContent( FName InCollectionName, ECollectionShareType::Type InType, TArray<FSoftObjectPath>& OutAssetPaths );
 
 	bool bInitialized;
 };
