Classes that are detected by live coding that do not change in structure will no longer invoke the reinstancing system since it isn't needed.  Changes to default values will be respected.

diff --git a/Engine/Source/Editor/UnrealEd/Private/Kismet2/ReloadUtilities.cpp b/Engine/Source/Editor/UnrealEd/Private/Kismet2/ReloadUtilities.cpp
--- a/Engine/Source/Editor/UnrealEd/Private/Kismet2/ReloadUtilities.cpp
+++ b/Engine/Source/Editor/UnrealEd/Private/Kismet2/ReloadUtilities.cpp
@@ -713,29 +713,29 @@ namespace UE::Reload::Private
 		// Now we re-instance this class
 		{
 			UClass* OldClass = State.OldClass;
-			bool bNeedsReinstancing = true;
 
 			// Collect the property values of the new CDO
 			FCDOPropertyData NewClassCDOProperties;
 			SerializeCDOProperties(NewClass->GetDefaultObject(), NewClassCDOProperties);
 
-			// If the structure didn't change, then we might not need to re-instance.  We always add to the collection
-			// of reconstructed CDOs.
+			// If the structure didn't change, always add to the list of reconstructed CDOs
 			if (OldClass == NewClass)
 			{
-				bNeedsReinstancing = HavePropertiesChanged(State.OldClassCDOProperties, NewClassCDOProperties);
 				ReconstructedCDOsMap.Add(State.OldClassCDO, NewClass->GetDefaultObject());
 			}
 
-			if (bNeedsReinstancing)
+			// We only need to do re-instancing when we have a new UClass.
+			else
 			{
 				UClass* NullableNewClass = NewClass == OldClass ? nullptr : NewClass;
 				TSharedPtr<FReloadClassReinstancer> ReinstanceHelper = MakeShareable(new FReloadClassReinstancer(
 					NullableNewClass, OldClass, State.OldClassCDO, ReinstancingObjects, CompiledBlueprints));
 				Ar.Logf(ELogVerbosity::Log, TEXT("Re-instancing %s after reload."), *NewClass->GetName());
 				ReinstanceHelper->ReinstanceObjects(true);
-				UpdateDefaultProperties(NewClass, State.OldClassCDOProperties, NewClassCDOProperties);
 			}
+
+			// Update the default values
+			UpdateDefaultProperties(NewClass, State.OldClassCDOProperties, NewClassCDOProperties);
 		}
 
 		State.State = TopologicalState::Finished;
