Fix for FEditorStyleSettingsCustomization::MakeInstance crash resulting from calling AsShared in FEditorStyleSettingsCustomization's constructor

diff --git a/Engine/Source/Editor/UnrealEd/Private/Settings/EditorStyleSettingsCustomization.cpp b/Engine/Source/Editor/UnrealEd/Private/Settings/EditorStyleSettingsCustomization.cpp
--- a/Engine/Source/Editor/UnrealEd/Private/Settings/EditorStyleSettingsCustomization.cpp
+++ b/Engine/Source/Editor/UnrealEd/Private/Settings/EditorStyleSettingsCustomization.cpp
@@ -347,12 +347,12 @@ bool FStyleColorListCustomization::IsResetToDefaultVisible(TSharedPtr<IPropertyH
 
 TSharedRef<IDetailCustomization> FEditorStyleSettingsCustomization::MakeInstance()
 {
-	return MakeShared<FEditorStyleSettingsCustomization>(); 
-}
+	TSharedRef<FEditorStyleSettingsCustomization> Instance = MakeShared<FEditorStyleSettingsCustomization>();
+	// unable to perform this operation in FEditorStyleSettingsCustomization's constructor since by then the shared ref 
+	// controller has not been created yet
+	USlateThemeManager::Get().OnThemeChanged().AddSP(Instance, &FEditorStyleSettingsCustomization::OnThemeChanged);
 
-FEditorStyleSettingsCustomization::FEditorStyleSettingsCustomization()
-{
-	USlateThemeManager::Get().OnThemeChanged().AddSP(StaticCastSharedRef<FEditorStyleSettingsCustomization>(AsShared()), &FEditorStyleSettingsCustomization::OnThemeChanged);
+	return Instance;
 }
 
 FEditorStyleSettingsCustomization::~FEditorStyleSettingsCustomization()
diff --git a/Engine/Source/Editor/UnrealEd/Private/Settings/EditorStyleSettingsCustomization.h b/Engine/Source/Editor/UnrealEd/Private/Settings/EditorStyleSettingsCustomization.h
--- a/Engine/Source/Editor/UnrealEd/Private/Settings/EditorStyleSettingsCustomization.h
+++ b/Engine/Source/Editor/UnrealEd/Private/Settings/EditorStyleSettingsCustomization.h
@@ -33,8 +33,6 @@ public:
 	/** Makes a new instance of this detail layout class for a specific detail view requesting it */
 	static TSharedRef<IDetailCustomization> MakeInstance();
 
-	FEditorStyleSettingsCustomization();
-
 	~FEditorStyleSettingsCustomization(); 
 
 	/** IDetailCustomization interface */
