Construct component class registry with full path to classes not just name.

diff --git a/Engine/Source/Editor/UnrealEd/Private/ComponentTypeRegistry.cpp b/Engine/Source/Editor/UnrealEd/Private/ComponentTypeRegistry.cpp
--- a/Engine/Source/Editor/UnrealEd/Private/ComponentTypeRegistry.cpp
+++ b/Engine/Source/Editor/UnrealEd/Private/ComponentTypeRegistry.cpp
@@ -569,7 +569,7 @@ void FComponentTypeRegistryData::ForceRefreshComponentList()
 			FString ClassName = ClassPath.GetAssetName().ToString();
 			ClassName.RemoveFromEnd(TEXT("_C"));
 
-			FComponentClassComboEntryPtr NewEntry(new FComponentClassComboEntry(BlueprintComponents, ClassName, ClassPath.GetAssetName(), IconClass, /*bIncludeInFilter=*/true));
+			FComponentClassComboEntryPtr NewEntry(new FComponentClassComboEntry(BlueprintComponents, ClassName, ClassPath, IconClass, /*bIncludeInFilter=*/true));
 			SortedClassList.Add(NewEntry);
 
 			// Create an unloaded blueprint data object to assist with class filtering
diff --git a/Engine/Source/Editor/UnrealEd/Public/SComponentClassCombo.h b/Engine/Source/Editor/UnrealEd/Public/SComponentClassCombo.h
--- a/Engine/Source/Editor/UnrealEd/Public/SComponentClassCombo.h
+++ b/Engine/Source/Editor/UnrealEd/Public/SComponentClassCombo.h
@@ -81,7 +81,7 @@ public:
 		, CustomizationArgs(InCustomizationArgs)
 	{}
 
-	FComponentClassComboEntry(const FString& InHeadingText, const FString& InComponentName, FName InComponentPath, const UClass* InIconClass, bool InIncludedInFilter)
+	FComponentClassComboEntry(const FString& InHeadingText, const FString& InComponentName, FTopLevelAssetPath InComponentPath, const UClass* InIconClass, bool InIncludedInFilter)
 		: ComponentClass(nullptr)
 		, IconClass(InIconClass)
 		, ComponentName(InComponentName)
@@ -177,7 +177,7 @@ private:
 	// loading occurs when the blueprint is spawned, which should also trigger a refresh
 	// of the component list:
 	FString ComponentName;
-	FName ComponentPath;
+	FTopLevelAssetPath ComponentPath;
 	FString HeadingText;
 	bool bIncludedInFilter;
 	EComponentCreateAction::Type ComponentCreateAction;
