Fix missing thumbnails.

diff --git a/Engine/Source/Editor/UnrealEd/Private/AssetThumbnail.cpp b/Engine/Source/Editor/UnrealEd/Private/AssetThumbnail.cpp
--- a/Engine/Source/Editor/UnrealEd/Private/AssetThumbnail.cpp
+++ b/Engine/Source/Editor/UnrealEd/Private/AssetThumbnail.cpp
@@ -1235,7 +1235,7 @@ bool FAssetThumbnailPool::LoadThumbnail(TSharedRef<FThumbnailInfo> ThumbnailInfo
 
 FSlateTexture2DRHIRef* FAssetThumbnailPool::AccessTexture( const FAssetData& AssetData, uint32 Width, uint32 Height )
 {
-	if(AssetData.IsValid() || Width == 0 || Height == 0)
+	if(!AssetData.IsValid() || Width == 0 || Height == 0)
 	{
 		return NULL;
 	}
@@ -1330,7 +1330,7 @@ FSlateTexture2DRHIRef* FAssetThumbnailPool::AccessTexture( const FAssetData& Ass
 void FAssetThumbnailPool::AddReferencer( const FAssetThumbnail& AssetThumbnail )
 {
 	FIntPoint Size = AssetThumbnail.GetSize();
-	if ( AssetThumbnail.GetAssetData().IsValid() || Size.X == 0 || Size.Y == 0 )
+	if ( !AssetThumbnail.GetAssetData().IsValid() || Size.X == 0 || Size.Y == 0 )
 	{
 		// Invalid referencer
 		return;
