Changes to UE-160593. Added virtual to the destructor for readability.

diff --git a/Engine/Source/Editor/UnrealEd/Private/Settings/EditorStyleSettingsCustomization.h b/Engine/Source/Editor/UnrealEd/Private/Settings/EditorStyleSettingsCustomization.h
--- a/Engine/Source/Editor/UnrealEd/Private/Settings/EditorStyleSettingsCustomization.h
+++ b/Engine/Source/Editor/UnrealEd/Private/Settings/EditorStyleSettingsCustomization.h
@@ -33,7 +33,7 @@ public:
 	/** Makes a new instance of this detail layout class for a specific detail view requesting it */
 	static TSharedRef<IDetailCustomization> MakeInstance();
 
-	~FEditorStyleSettingsCustomization(); 
+	virtual ~FEditorStyleSettingsCustomization(); 
 
 	/** IDetailCustomization interface */
 	virtual void CustomizeDetails(IDetailLayoutBuilder& DetailLayout) override;
