Block on shader map compilation and clear cached platform data after we compile shaders for each platform.

diff --git a/Engine/Source/Editor/UnrealEd/Private/Commandlets/CompileShadersTestBedCommandlet.cpp b/Engine/Source/Editor/UnrealEd/Private/Commandlets/CompileShadersTestBedCommandlet.cpp
--- a/Engine/Source/Editor/UnrealEd/Private/Commandlets/CompileShadersTestBedCommandlet.cpp
+++ b/Engine/Source/Editor/UnrealEd/Private/Commandlets/CompileShadersTestBedCommandlet.cpp
@@ -150,28 +150,28 @@ PRAGMA_ENABLE_DEPRECATION_WARNINGS
 				}
 			}
 		}
-	}
-
-	// Block on all the jobs submitted above.
-	{
-		TRACE_CPUPROFILER_EVENT_SCOPE(BlockOnShaderCompiles);
 
-		GShaderCompilingManager->FinishAllCompilation();
-	}
+		// Block on all the jobs submitted above.
+		{
+			TRACE_CPUPROFILER_EVENT_SCOPE(BlockOnShaderCompiles);
 
-	// Perform cleanup and clear cached data for cooking.
-	{
-		TRACE_CPUPROFILER_EVENT_SCOPE(ClearCachedCookedPlatformData);
+			GShaderCompilingManager->FinishAllCompilation();
+		}
 
-		for (const FAssetData& AssetData : MaterialList)
+		// Perform cleanup and clear cached data for cooking.
 		{
-			if (UMaterial* Material = Cast<UMaterial>(AssetData.GetAsset()))
-			{
-				Material->ClearAllCachedCookedPlatformData();
-			}
-			else if (UMaterialInstance* MaterialInstance = Cast<UMaterialInstance>(AssetData.GetAsset()))
+			TRACE_CPUPROFILER_EVENT_SCOPE(ClearCachedCookedPlatformData);
+
+			for (const FAssetData& AssetData : MaterialList)
 			{
-				MaterialInstance->ClearAllCachedCookedPlatformData();
+				if (UMaterial* Material = Cast<UMaterial>(AssetData.GetAsset()))
+				{
+					Material->ClearAllCachedCookedPlatformData();
+				}
+				else if (UMaterialInstance* MaterialInstance = Cast<UMaterialInstance>(AssetData.GetAsset()))
+				{
+					MaterialInstance->ClearAllCachedCookedPlatformData();
+				}
 			}
 		}
 	}
