Fix Material/MaterialInstance being created next to a texture in a read-only folder

diff --git a/Engine/Source/Editor/UnrealEd/Private/LevelEditorViewport.cpp b/Engine/Source/Editor/UnrealEd/Private/LevelEditorViewport.cpp
--- a/Engine/Source/Editor/UnrealEd/Private/LevelEditorViewport.cpp
+++ b/Engine/Source/Editor/UnrealEd/Private/LevelEditorViewport.cpp
@@ -35,6 +35,7 @@
 #include "UObject/UObjectIterator.h"
 #include "EngineUtils.h"
 #include "Editor.h"
+#include "EditorDirectories.h"
 #include "EditorModeRegistry.h"
 #include "EditorModes.h"
 #include "PhysicsManipulationMode.h"
@@ -61,6 +62,7 @@
 #include "AssetSelection.h"
 #include "Kismet2/KismetEditorUtilities.h"
 #include "AssetRegistry/AssetRegistryModule.h"
+#include "AssetToolsModule.h"
 #include "IPlacementModeModule.h"
 #include "Engine/Polys.h"
 #include "ActorEditorUtils.h"
@@ -508,29 +510,74 @@ static bool TryAndCreateMaterialInput( UMaterial* UnrealMaterial, EMaterialKind
 
 UObject* FLevelEditorViewportClient::GetOrCreateMaterialFromTexture(UTexture* UnrealTexture)
 {
+	IAssetRegistry& AssetRegistry = FModuleManager::LoadModuleChecked<FAssetRegistryModule>(AssetRegistryConstants::ModuleName).Get();
+	IAssetTools& AssetTools = FModuleManager::LoadModuleChecked<FAssetToolsModule>("AssetTools").Get();
+
 	// Check if a base material and corresponding params are set in the Settings to determine whether to create a Material or MIC
 	const ULevelEditorViewportSettings* ViewportSettings = GetDefault<ULevelEditorViewportSettings>();
 	UMaterialInterface* BaseMaterial = ViewportSettings->MaterialForDroppedTextures.LoadSynchronous();
 	const bool bCreateMaterialInstance = BaseMaterial && ViewportSettings->MaterialParamsForDroppedTextures.Num() > 0;
 
-	FString TextureShortName = FPackageName::GetShortName( UnrealTexture->GetOutermost()->GetName() );
+	const FString TexturePackageName = UnrealTexture->GetPackage()->GetName();
+	FString TextureShortName = FPackageName::GetShortName(TexturePackageName);
 
 	// See if we can figure out what kind of material it is, based on a suffix, like _S for Specular, _D for Base/Detail/Diffuse.
 	// if it can determine which type of texture it was, it will return the base name of the texture minus the suffix.
 	EMaterialKind MaterialKind;
 	TextureShortName = GetSharedTextureNameAndKind( TextureShortName, MaterialKind );
-	
-	const FString MaterialFullName = bCreateMaterialInstance ? (TEXT("MI_") + TextureShortName) : (TextureShortName + TEXT("_Mat"));
-	FString NewPackageName = FPackageName::GetLongPackagePath( UnrealTexture->GetOutermost()->GetName() ) + TEXT( "/" ) + MaterialFullName;
-	NewPackageName = UPackageTools::SanitizePackageName( NewPackageName );
-	UPackage* Package = CreatePackage( *NewPackageName );
 
-	FAssetRegistryModule& AssetRegistryModule = FModuleManager::LoadModuleChecked<FAssetRegistryModule>( TEXT( "AssetRegistry" ) );
+	FString NewPackageFolder;
+	if (AssetTools.GetWritableFolderPermissionList()->PassesStartsWithFilter(TexturePackageName))
+	{
+		// Create the material in the source texture folder
+		NewPackageFolder = FPackageName::GetLongPackagePath(TexturePackageName);
+	}
+	if (NewPackageFolder.IsEmpty())
+	{
+		// Create the material in the current world folder?
+		if (UWorld* EditorWorld = GEditor->GetEditorWorldContext().World())
+		{
+			const FString EditorWorldPackageName = EditorWorld->GetPackage()->GetName();
+			if (!EditorWorldPackageName.StartsWith(TEXT("/Temp/")))
+			{
+				NewPackageFolder = FPackageName::GetLongPackagePath(EditorWorldPackageName);
+				if (!AssetTools.GetWritableFolderPermissionList()->PassesStartsWithFilter(NewPackageFolder))
+				{
+					NewPackageFolder.Reset();
+				}
+			}
+		}
+	}
+	if (NewPackageFolder.IsEmpty())
+	{
+		// Create the material in the last directory an asset was created in?
+		if (FPackageName::TryConvertFilenameToLongPackageName(FEditorDirectories::Get().GetLastDirectory(ELastDirectory::NEW_ASSET), NewPackageFolder))
+		{
+			if (!AssetTools.GetWritableFolderPermissionList()->PassesStartsWithFilter(NewPackageFolder))
+			{
+				NewPackageFolder.Reset();
+			}
+		}
+		else
+		{
+			NewPackageFolder.Reset();
+		}
+	}
+	if (NewPackageFolder.IsEmpty())
+	{
+		// Create the material in the game root folder
+		NewPackageFolder = TEXT("/Game");
+	}
+
+	const FString MaterialFullName = bCreateMaterialInstance ? (TEXT("MI_") + TextureShortName) : (TextureShortName + TEXT("_Mat"));
+	FString NewPackageName = FPaths::Combine(NewPackageFolder, MaterialFullName);
+	NewPackageName = UPackageTools::SanitizePackageName(NewPackageName);
+	UPackage* Package = CreatePackage(*NewPackageName);
 
 	// See if the material asset already exists with the expected name, if it does, just return
 	// an instance of it.
 	TArray<FAssetData> OutAssetData;
-	if ( AssetRegistryModule.Get().GetAssetsByPackageName( *NewPackageName, OutAssetData ) && OutAssetData.Num() > 0 )
+	if (AssetRegistry.GetAssetsByPackageName(*NewPackageName, OutAssetData) && (OutAssetData.Num() > 0))
 	{
 		UObject* FoundAsset = OutAssetData[0].GetAsset();
 		if (FoundAsset->IsA(UMaterialInterface::StaticClass()))
@@ -567,7 +614,7 @@ UObject* FLevelEditorViewportClient::GetOrCreateMaterialFromTexture(UTexture* Un
 	EmissiveSuffixes.Add("_Emissive");
 
 	// The asset path for the base texture, we need this to try and append different suffixes to to find other textures we can use.
-	const FString BaseTexturePackage = FPackageName::GetLongPackagePath(UnrealTexture->GetOutermost()->GetName()) + TEXT("/") + TextureShortName;
+	const FString BaseTexturePackage = FPackageName::GetLongPackagePath(TexturePackageName) + TEXT("/") + TextureShortName;
 
 	UMaterialInterface* CreatedMaterialInterface = nullptr;
 
