Removed private header from UnrealEdPrivatePCH.h.

diff --git a/Engine/Source/Editor/UnrealEd/Private/UnrealEdPrivatePCH.h b/Engine/Source/Editor/UnrealEd/Private/UnrealEdPrivatePCH.h
--- a/Engine/Source/Editor/UnrealEd/Private/UnrealEdPrivatePCH.h
+++ b/Engine/Source/Editor/UnrealEd/Private/UnrealEdPrivatePCH.h
@@ -396,7 +396,6 @@
 #include "Framework/Text/ShapedTextCacheFwd.h"
 #include "Framework/Application/SlateApplication.h"
 #include "Framework/Application/MenuStack.h"
-#include "Runtime/Slate/Private/Framework/Application/Menu.h"
 #include "Widgets/Text/STextBlock.h"
 #include "Widgets/Input/SCheckBox.h"
 #include "Framework/Docking/TabManager.h"
