MultiprocessCook: Collect arguments needed by LoadBalanceCookBurden and pass them into the CookDirector's LoadBalance function.
Reorganize Cookdirector config parsing to handle more config values.
Conform Cookdirector logging to the CookMultiprocess format.
#rnx

diff --git a/Engine/Source/Editor/UnrealEd/Classes/CookOnTheSide/CookOnTheFlyServer.h b/Engine/Source/Editor/UnrealEd/Classes/CookOnTheSide/CookOnTheFlyServer.h
--- a/Engine/Source/Editor/UnrealEd/Classes/CookOnTheSide/CookOnTheFlyServer.h
+++ b/Engine/Source/Editor/UnrealEd/Classes/CookOnTheSide/CookOnTheFlyServer.h
@@ -445,9 +445,10 @@ private:
 	void PumpRequests(UE::Cook::FTickStackData& StackData, int32& OutNumPushed);
 	/**
 	 * Assign the requests found in PumpRequests; either pushing them to ReadyRequests if SingleProcess Cook or
-	 * to CookWorkers if MultiProcess.
+	 * to CookWorkers if MultiProcess. Input Requests have been sorted by leaf to root load order.
 	 */
-	void AssignRequests(TArrayView<UE::Cook::FPackageData*> Requests, UE::Cook::FRequestQueue& RequestQueue);
+	void AssignRequests(TArrayView<UE::Cook::FPackageData*> Requests, UE::Cook::FRequestQueue& RequestQueue,
+		TMap<UE::Cook::FPackageData*, TArray<UE::Cook::FPackageData*>>&& RequestGraph);
 	/**
 	 * Multiprocess cook: Notify the CookDirector that a package it assigned to a worker was demoted
 	 * due to e.g. cancelled cook and should be removed from the CookWorker.
diff --git a/Engine/Source/Editor/UnrealEd/Private/CookOnTheFlyServer.cpp b/Engine/Source/Editor/UnrealEd/Private/CookOnTheFlyServer.cpp
--- a/Engine/Source/Editor/UnrealEd/Private/CookOnTheFlyServer.cpp
+++ b/Engine/Source/Editor/UnrealEd/Private/CookOnTheFlyServer.cpp
@@ -2381,8 +2381,9 @@ void UCookOnTheFlyServer::PumpRequests(UE::Cook::FTickStackData& StackData, int3
 		{
 			TArray<FPackageData*> RequestsToLoad;
 			TArray<TPair<FPackageData*, ESuppressCookReason>> RequestsToDemote;
-			RequestCluster.ClearAndDetachOwnedPackageDatas(RequestsToLoad, RequestsToDemote);
-			AssignRequests(RequestsToLoad, RequestQueue);
+			TMap<FPackageData*, TArray<FPackageData*>> RequestGraph;
+			RequestCluster.ClearAndDetachOwnedPackageDatas(RequestsToLoad, RequestsToDemote, RequestGraph);
+			AssignRequests(RequestsToLoad, RequestQueue, MoveTemp(RequestGraph));
 			for (TPair<FPackageData*, ESuppressCookReason>& Pair : RequestsToDemote)
 			{
 				DemoteToIdle(*Pair.Key, ESendFlags::QueueAdd, Pair.Value);
@@ -2420,7 +2421,8 @@ void UCookOnTheFlyServer::PumpRequests(UE::Cook::FTickStackData& StackData, int3
 	OutNumPushed += NumInBatch;
 }
 
-void UCookOnTheFlyServer::AssignRequests(TArrayView<UE::Cook::FPackageData*> Requests, UE::Cook::FRequestQueue& RequestQueue)
+void UCookOnTheFlyServer::AssignRequests(TArrayView<UE::Cook::FPackageData*> Requests, UE::Cook::FRequestQueue& RequestQueue,
+	TMap<UE::Cook::FPackageData*, TArray<UE::Cook::FPackageData*>>&& RequestGraph)
 {
 	using namespace UE::Cook;
 
@@ -2432,7 +2434,7 @@ void UCookOnTheFlyServer::AssignRequests(TArrayView<UE::Cook::FPackageData*> Req
 			return;
 		}
 		TArray<FWorkerId> Assignments;
-		CookDirector->AssignRequests(Requests, Assignments);
+		CookDirector->AssignRequests(Requests, Assignments, MoveTemp(RequestGraph));
 		check(Assignments.Num() == NumRequests);
 		for (int32 Index = 0; Index < NumRequests; ++Index)
 		{
@@ -8592,7 +8594,8 @@ void UCookOnTheFlyServer::CancelAllQueues()
 	{
 		TArray<FPackageData*> RequestsToLoad;
 		TArray<TPair<FPackageData*, ESuppressCookReason>> RequestsToDemote;
-		RequestCluster.ClearAndDetachOwnedPackageDatas(RequestsToLoad, RequestsToDemote);
+		TMap<FPackageData*, TArray<FPackageData*>> UnusedRequestGraph;
+		RequestCluster.ClearAndDetachOwnedPackageDatas(RequestsToLoad, RequestsToDemote, UnusedRequestGraph);
 		for (FPackageData* PackageData : RequestsToLoad)
 		{
 			DemoteToIdle(*PackageData, ESendFlags::QueueAdd, ESuppressCookReason::CookCanceled);
@@ -9668,7 +9671,7 @@ void UCookOnTheFlyServer::BeginCookEditorSystems()
 namespace UE::Cook
 {
 
-/** MultiprocessCook collector for FEDLCookChecker data. */
+/** CookMultiprocess collector for FEDLCookChecker data. */
 class FEDLMPCollector : public IMPCollector
 {
 public:
diff --git a/Engine/Source/Editor/UnrealEd/Private/Cooker/CookDirector.cpp b/Engine/Source/Editor/UnrealEd/Private/Cooker/CookDirector.cpp
--- a/Engine/Source/Editor/UnrealEd/Private/Cooker/CookDirector.cpp
+++ b/Engine/Source/Editor/UnrealEd/Private/Cooker/CookDirector.cpp
@@ -8,6 +8,7 @@
 #include "CookWorkerServer.h"
 #include "CookOnTheSide/CookOnTheFlyServer.h"
 #include "CoreGlobals.h"
+#include "LoadBalanceCookBurden.h"
 #include "Math/NumericLimits.h"
 #include "Misc/CommandLine.h"
 #include "Misc/ConfigCacheIni.h"
@@ -26,39 +27,26 @@ FCookDirector::FCookDirector(UCookOnTheFlyServer& InCOTFS)
 	WorkersStalledStartTimeSeconds = MAX_flt;
 	WorkersStalledWarnTimeSeconds = MAX_flt;
 
-	ParseDesiredNumRemoteWorkers();
-	WorkerConnectPort = Sockets::COOKDIRECTOR_DEFAULT_REQUEST_CONNECTION_PORT;
-	FParse::Value(FCommandLine::Get(), TEXT("-CookDirectorListenPort="), WorkerConnectPort);
-	ParseShowWorkerOption();
-
-	if (DesiredNumRemoteWorkers > 0)
-	{
-		ISocketSubsystem* SocketSubsystem = ISocketSubsystem::Get();
-		if (!SocketSubsystem)
-		{
-			UE_LOG(LogCook, Error, TEXT("CookDirector initialization failure: platform does not support network sockets. CookWorkers will be disabled."));
-			DesiredNumRemoteWorkers = 0;
-		}
-	}
-
-	if (DesiredNumRemoteWorkers > 0)
-	{
-		TryCreateWorkerConnectSocket();
-	}
-	UE_LOG(LogCook, Display, TEXT("MultiprocessCook is enabled with %d CookWorker processes."), DesiredNumRemoteWorkers);
-}
-
-void FCookDirector::ParseDesiredNumRemoteWorkers()
-{
-	DesiredNumRemoteWorkers = 4;
-	GConfig->GetInt(TEXT("CookSettings"), TEXT("CookWorkerCount"), DesiredNumRemoteWorkers, GEditorIni);
-	FParse::Value(FCommandLine::Get(), TEXT("-CookWorkerCount="), DesiredNumRemoteWorkers);
+	ParseConfig();
+	ResetCookWorkerCount();
+	UE_LOG(LogCook, Display, TEXT("CookMultiprocess is enabled with %d CookWorker processes."), CookWorkerCount);
 }
 
-void FCookDirector::ParseShowWorkerOption()
+void FCookDirector::ParseConfig()
 {
-	FString Text;
 	const TCHAR* CommandLine = FCommandLine::Get();
+	FString Text;
+
+	// CookWorkerCount
+	RequestedCookWorkerCount = 3;
+	GConfig->GetInt(TEXT("CookSettings"), TEXT("CookWorkerCount"), RequestedCookWorkerCount, GEditorIni);
+	FParse::Value(CommandLine, TEXT("-CookWorkerCount="), RequestedCookWorkerCount);
+
+	// CookDirectorListenPort
+	WorkerConnectPort = Sockets::COOKDIRECTOR_DEFAULT_REQUEST_CONNECTION_PORT;
+	FParse::Value(CommandLine, TEXT("-CookDirectorListenPort="), WorkerConnectPort);
+
+	// ShowCookWorker
 	if (!FParse::Value(CommandLine, TEXT("-ShowCookWorker="), Text))
 	{
 		if (FParse::Param(CommandLine, TEXT("ShowCookWorker")))
@@ -66,7 +54,6 @@ void FCookDirector::ParseShowWorkerOption()
 			Text = TEXT("SeparateWindows");
 		}
 	}
-
 	if (Text == TEXT("CombinedLogs")) { ShowWorkerOption = EShowWorker::CombinedLogs; }
 	else if (Text == TEXT("SeparateLogs")) { ShowWorkerOption = EShowWorker::SeparateLogs; }
 	else if (Text == TEXT("SeparateWindows")) { ShowWorkerOption = EShowWorker::SeparateWindows; }
@@ -78,6 +65,38 @@ void FCookDirector::ParseShowWorkerOption()
 		}
 		ShowWorkerOption = EShowWorker::CombinedLogs;
 	}
+
+	// LoadBalanceAlgorithm
+	LoadBalanceAlgorithm = ELoadBalanceAlgorithm::CookBurden;
+	if (FParse::Value(CommandLine, TEXT("-CookLoadBalance="), Text))
+	{
+		if (Text == TEXT("Striped")) { LoadBalanceAlgorithm = ELoadBalanceAlgorithm::Striped; }
+		else if (Text == TEXT("CookBurden")) { LoadBalanceAlgorithm = ELoadBalanceAlgorithm::CookBurden; }
+		else
+		{
+			UE_LOG(LogCook, Warning, TEXT("Invalid selection \"%s\" for -CookLoadBalance."), *Text);
+		}
+	}
+}
+
+void FCookDirector::ResetCookWorkerCount()
+{
+	CookWorkerCount = RequestedCookWorkerCount;
+	if (CookWorkerCount > 0)
+	{
+		ISocketSubsystem* SocketSubsystem = ISocketSubsystem::Get();
+		if (!SocketSubsystem)
+		{
+			UE_LOG(LogCook, Error, TEXT("CookDirector initialization failure: platform does not support network sockets. CookWorkers will be disabled."));
+			CookWorkerCount = 0;
+		}
+	}
+
+	if (CookWorkerCount > 0)
+	{
+		TryCreateWorkerConnectSocket();
+	}
+
 }
 
 FCookDirector::~FCookDirector()
@@ -98,7 +117,8 @@ FCookDirector::~FCookDirector()
 	Sockets::CloseSocket(WorkerConnectSocket);
 }
 
-void FCookDirector::AssignRequests(TArrayView<UE::Cook::FPackageData*> Requests, TArray<FWorkerId>& OutAssignments)
+void FCookDirector::AssignRequests(TArrayView<UE::Cook::FPackageData*> Requests, TArray<FWorkerId>& OutAssignments,
+	TMap<FPackageData*, TArray<FPackageData*>>&& RequestGraph)
 {
 	InitializeWorkers();
 
@@ -125,9 +145,8 @@ void FCookDirector::AssignRequests(TArrayView<UE::Cook::FPackageData*> Requests,
 	check(MaxRemoteIndex >= 0);
 	SortedWorkers.Sort([](const FCookWorkerServer& A, const FCookWorkerServer& B) { return A.GetWorkerId() < B.GetWorkerId(); });
 
-	// Call a LoadBalancing algorithm to split the requests among the LocalWorker and RemoteWorkers
-	// MPCOOKTODO: Implement LoadBalanceGreedy
-	LoadBalanceStriped(SortedWorkers, Requests, OutAssignments);
+	// Call the LoadBalancing algorithm to split the requests among the LocalWorker and RemoteWorkers
+	LoadBalance(SortedWorkers, Requests, MoveTemp(RequestGraph), OutAssignments);
 
 	// Split the output array of WorkerId assignments into a batch for each of the RemoteWorkers 
 	TArray<TArray<FPackageData*>> RemoteBatches; // Indexed by WorkerId.GetRemoteIndex()
@@ -275,7 +294,7 @@ void FCookDirector::ShutdownCookSession()
 	PendingConnections.Reset();
 
 	// Restore the FCookDirector to its original state so that it is ready for a new session
-	ParseDesiredNumRemoteWorkers();
+	ResetCookWorkerCount();
 }
 
 void FCookDirector::Register(IMPCollector* Collector)
@@ -364,6 +383,10 @@ FGuid FWorkerConnectMessage::MessageType(TEXT("302096E887DA48F7B079FAFAD0EE5695"
 
 bool FCookDirector::TryCreateWorkerConnectSocket()
 {
+	if (WorkerConnectSocket)
+	{
+		return true;
+	}
 	FString ErrorReason;
 	TSharedPtr<FInternetAddr> ListenAddr;
 	WorkerConnectSocket = Sockets::CreateListenSocket(WorkerConnectPort, ListenAddr, WorkerConnectAuthority,
@@ -372,7 +395,7 @@ bool FCookDirector::TryCreateWorkerConnectSocket()
 	{
 		UE_LOG(LogCook, Error, TEXT("CookDirector could not create listen socket, CookWorkers will be disabled. Reason: %s."),
 			*ErrorReason);
-		DesiredNumRemoteWorkers = 0;
+		CookWorkerCount = 0;
 		return false;
 	}
 	return true;
@@ -381,7 +404,7 @@ bool FCookDirector::TryCreateWorkerConnectSocket()
 
 void FCookDirector::InitializeWorkers()
 {
-	if (RemoteWorkers.Num() >= DesiredNumRemoteWorkers)
+	if (RemoteWorkers.Num() >= CookWorkerCount)
 	{
 		return;
 	}
@@ -400,7 +423,7 @@ void FCookDirector::InitializeWorkers()
 	}
 	// Add RemoteWorkers, pulling the RemoteIndex id from the UnusedRemoteIndexes if any exist
 	// otherwise use the next integer because all indexes up to RemoteWorkers.Num() are in use.
-	while (RemoteWorkers.Num() < DesiredNumRemoteWorkers)
+	while (RemoteWorkers.Num() < CookWorkerCount)
 	{
 		int32 RemoteIndex;
 		if (UnusedRemoteIndexes.Num())
@@ -586,7 +609,8 @@ bool FDirectorConnectionInfo::TryParseCommandLine()
 	return true;
 }
 
-void FCookDirector::LoadBalanceStriped(TArrayView<FCookWorkerServer*> SortedWorkers, TArrayView<FPackageData*> Requests, TArray<FWorkerId>& OutAssignments)
+void FCookDirector::LoadBalance(TConstArrayView<FCookWorkerServer*> SortedWorkers, TArrayView<FPackageData*> Requests,
+	TMap<FPackageData*, TArray<FPackageData*>>&& RequestGraph, TArray<FWorkerId>& OutAssignments)
 {
 	TArray<FWorkerId> AllWorkers;
 	int32 NumAllWorkers = SortedWorkers.Num() + 1;
@@ -597,12 +621,16 @@ void FCookDirector::LoadBalanceStriped(TArrayView<FCookWorkerServer*> SortedWork
 		AllWorkers.Add(Worker->GetWorkerId());
 	}
 	OutAssignments.Reset(Requests.Num());
-	int32 AllWorkersIndex = 0;
-	for (FPackageData* Request : Requests)
+
+	switch (LoadBalanceAlgorithm)
 	{
-		OutAssignments.Add(AllWorkers[AllWorkersIndex]);
-		AllWorkersIndex = (AllWorkersIndex + 1) % NumAllWorkers;
+	case ELoadBalanceAlgorithm::Striped:
+		return LoadBalanceStriped(AllWorkers, Requests, MoveTemp(RequestGraph), OutAssignments);
+	case ELoadBalanceAlgorithm::CookBurden:
+		return LoadBalanceCookBurden(AllWorkers, Requests, MoveTemp(RequestGraph), OutAssignments);
 	}
+	checkNoEntry();
+	return LoadBalanceCookBurden(AllWorkers, Requests, MoveTemp(RequestGraph), OutAssignments);
 }
 
 void FCookDirector::AbortWorker(FWorkerId WorkerId)
@@ -615,7 +643,7 @@ void FCookDirector::AbortWorker(FWorkerId WorkerId)
 	{
 		return;
 	}
-	--DesiredNumRemoteWorkers;
+	--CookWorkerCount;
 	TSet<FPackageData*> PackagesToReassign;
 	RemoteWorker->AbortAssignments(PackagesToReassign);
 	if (!RemoteWorker->IsShuttingDown())
diff --git a/Engine/Source/Editor/UnrealEd/Private/Cooker/CookDirector.h b/Engine/Source/Editor/UnrealEd/Private/Cooker/CookDirector.h
--- a/Engine/Source/Editor/UnrealEd/Private/Cooker/CookDirector.h
+++ b/Engine/Source/Editor/UnrealEd/Private/Cooker/CookDirector.h
@@ -34,8 +34,12 @@ public:
 	FCookDirector(UCookOnTheFlyServer& InCOTFS);
 	~FCookDirector();
 
-	/** Assign the given requests out to CookWorkers (or keep on local COTFS), return the list of assignments. */
-	void AssignRequests(TArrayView<UE::Cook::FPackageData*> Requests, TArray<FWorkerId>& OutAssignments);
+	/**
+	 * Assign the given requests out to CookWorkers (or keep on local COTFS), return the list of assignments.
+	 * Input requests have been sorted by leaf to root load order.
+	 */
+	void AssignRequests(TArrayView<UE::Cook::FPackageData*> Requests, TArray<FWorkerId>& OutAssignments,
+		TMap<FPackageData*, TArray<FPackageData*>>&& RequestGraph);
 	/** Notify the CookWorker that owns the cook of the package that the Director wants to take it back. */
 	void RemoveFromWorker(FPackageData& PackageData);
 	/** Periodic tick function. Sends/Receives messages to CookWorkers. */
@@ -77,9 +81,9 @@ private:
 		UE::CompactBinaryTCP::FReceiveBuffer Buffer;
 	};
 	/** Helper for constructor parsing. */
-	void ParseDesiredNumRemoteWorkers();
-	/** Helper for constructor parsing. */
-	void ParseShowWorkerOption();
+	void ParseConfig();
+	/** Sets CookWorkerCount for a new session */
+	void ResetCookWorkerCount();
 	/** Initialization helper: create the listen socket. */
 	bool TryCreateWorkerConnectSocket();
 	/** Initialization helper: add the local Server for a remote worker, worker process not yet created. */
@@ -90,9 +94,10 @@ private:
 	void TickWorkerShutdowns();
 	/** Get the commandline to launch a worker process with. */
 	FString GetWorkerCommandLine(FWorkerId WorkerId);
-	/** Simple assignment that divides requests evenly without considering dependencies or load burden. */
-	void LoadBalanceStriped(TArrayView<FCookWorkerServer*> SortedWorkers, TArrayView<FPackageData*> Requests,
-		TArray<FWorkerId>& OutAssignments);
+	/** Calls the configured LoadBalanceAlgorithm. Input Requests have been sorted by leaf to root load order. */
+	void LoadBalance(TConstArrayView<FCookWorkerServer*> Workers, TArrayView<FPackageData*> Requests,
+		TMap<FPackageData*, TArray<FPackageData*>>&& RequestGraph, TArray<FWorkerId>& OutAssignments);
+
 	/** Move the given worker from active workers to the list of workers shutting down. */
 	void AbortWorker(FWorkerId WorkerId);
 	/**
@@ -102,6 +107,11 @@ private:
 	void SetWorkersStalled(bool bInWorkersStalled);
 
 private:
+	enum class ELoadBalanceAlgorithm
+	{
+		Striped,
+		CookBurden,
+	};
 	TMap<int32, TUniquePtr<FCookWorkerServer>> RemoteWorkers;
 	TMap<FCookWorkerServer*, TUniquePtr<FCookWorkerServer>> ShuttingDownWorkers;
 	TMap<FGuid, TRefCountPtr<IMPCollector>> MessageHandlers;
@@ -111,10 +121,12 @@ private:
 	FSocket* WorkerConnectSocket = nullptr;
 	double WorkersStalledStartTimeSeconds = 0.;
 	double WorkersStalledWarnTimeSeconds = 0.;
-	int32 DesiredNumRemoteWorkers = 0;
+	int32 RequestedCookWorkerCount = 0;
+	int32 CookWorkerCount = 0;
 	int32 WorkerConnectPort = 0;
 	bool bWorkersStalled = false;
 	EShowWorker ShowWorkerOption = EShowWorker::CombinedLogs;
+	ELoadBalanceAlgorithm LoadBalanceAlgorithm = ELoadBalanceAlgorithm::CookBurden;
 
 	friend class UE::Cook::FCookWorkerServer;
 };
diff --git a/Engine/Source/Editor/UnrealEd/Private/Cooker/CookRequestCluster.cpp b/Engine/Source/Editor/UnrealEd/Private/Cooker/CookRequestCluster.cpp
--- a/Engine/Source/Editor/UnrealEd/Private/Cooker/CookRequestCluster.cpp
+++ b/Engine/Source/Editor/UnrealEd/Private/Cooker/CookRequestCluster.cpp
@@ -446,25 +446,30 @@ bool FRequestCluster::Contains(FPackageData* PackageData) const
 	return OwnedPackageDatas.Contains(PackageData);
 }
 
-void FRequestCluster::ClearAndDetachOwnedPackageDatas(TArray<FPackageData*>& OutRequestsToLoad, TArray<TPair<FPackageData*, ESuppressCookReason>>& OutRequestsToDemote)
+void FRequestCluster::ClearAndDetachOwnedPackageDatas(TArray<FPackageData*>& OutRequestsToLoad,
+	TArray<TPair<FPackageData*, ESuppressCookReason>>& OutRequestsToDemote,
+	TMap<FPackageData*, TArray<FPackageData*>>& OutRequestGraph)
 {
 	if (bStartAsyncComplete)
 	{
 		check(!GraphSearch);
 		OutRequestsToLoad = MoveTemp(Requests);
 		OutRequestsToDemote = MoveTemp(RequestsToDemote);
+		OutRequestGraph = MoveTemp(RequestGraph);
 		check(OutRequestsToLoad.Num() + OutRequestsToDemote.Num() == OwnedPackageDatas.Num())
 	}
 	else
 	{
 		OutRequestsToLoad = OwnedPackageDatas.Array();
 		OutRequestsToDemote.Reset();
+		OutRequestGraph.Reset();
 	}
 	InRequests.Empty();
 	Requests.Empty();
 	RequestsToDemote.Empty();
 	OwnedPackageDatas.Empty();
 	GraphSearch.Reset();
+	RequestGraph.Reset();
 	NextRequest = 0;
 }
 
@@ -551,6 +556,7 @@ void FRequestCluster::FetchDependencies(const FCookerTimer& CookerTimer, bool& b
 	Algo::TopologicalSort(NewRequests, GetElementDependencies, Algo::ETopologicalSort::AllowCycles);
 	Requests = MoveTemp(NewRequests);
 
+	RequestGraph = MoveTemp(GraphSearch->GetGraphEdges());
 	GraphSearch.Reset();
 	bDependenciesComplete = true;
 }
diff --git a/Engine/Source/Editor/UnrealEd/Private/Cooker/CookRequestCluster.h b/Engine/Source/Editor/UnrealEd/Private/Cooker/CookRequestCluster.h
--- a/Engine/Source/Editor/UnrealEd/Private/Cooker/CookRequestCluster.h
+++ b/Engine/Source/Editor/UnrealEd/Private/Cooker/CookRequestCluster.h
@@ -77,7 +77,8 @@ public:
 	 * If called before Process sets bOutComplete=true, all packages are put in OutRequestToLoad and are unsorted.
 	 */
 	void ClearAndDetachOwnedPackageDatas(TArray<FPackageData*>& OutRequestsToLoad,
-		TArray<TPair<FPackageData*, ESuppressCookReason>>& OutRequestsToDemote);
+		TArray<TPair<FPackageData*, ESuppressCookReason>>& OutRequestsToDemote,
+		TMap<FPackageData*, TArray<FPackageData*>>& OutRequestGraph);
 
 	/**
 	 * Create clusters(s) for all the given name or packagedata requests and append them to OutClusters.
@@ -319,6 +320,7 @@ private:
 	TArray<const ITargetPlatform*> Platforms;
 	TArray<ICookedPackageWriter*> PackageWriters;
 	FPackageDataSet OwnedPackageDatas;
+	TMap<FPackageData*, TArray<FPackageData*>> RequestGraph;
 	FString DLCPath;
 	TUniquePtr<FGraphSearch> GraphSearch; // Needs to be dynamic-allocated because of large alignment
 	UCookOnTheFlyServer& COTFS;
diff --git a/Engine/Source/Editor/UnrealEd/Private/Cooker/LoadBalanceCookBurden.cpp b/Engine/Source/Editor/UnrealEd/Private/Cooker/LoadBalanceCookBurden.cpp
new file mode 100644
--- /dev/null
+++ b/Engine/Source/Editor/UnrealEd/Private/Cooker/LoadBalanceCookBurden.cpp
@@ -0,0 +1,29 @@
+// Copyright Epic Games, Inc. All Rights Reserved.
+
+#include "LoadBalanceCookBurden.h"
+
+#include "Cooker/CookPackageData.h"
+#include "Cooker/CookTypes.h"
+
+namespace UE::Cook
+{
+
+void LoadBalanceStriped(TConstArrayView<FWorkerId> AllWorkers, TArrayView<FPackageData*> Requests,
+	TMap<FPackageData*, TArray<FPackageData*>>&& RequestGraph, TArray<FWorkerId>& OutAssignments)
+{
+	int32 AllWorkersIndex = 0;
+	int32 NumWorkers = AllWorkers.Num();
+	for (FPackageData* Request : Requests)
+	{
+		OutAssignments.Add(AllWorkers[AllWorkersIndex]);
+		AllWorkersIndex = (AllWorkersIndex + 1) % NumWorkers;
+	}
+}
+
+void LoadBalanceCookBurden(TConstArrayView<FWorkerId> AllWorkers, TArrayView<FPackageData*> Requests,
+	TMap<FPackageData*, TArray<FPackageData*>>&& RequestGraph, TArray<FWorkerId>& OutAssignments)
+{
+	return LoadBalanceStriped(AllWorkers, Requests, MoveTemp(RequestGraph), OutAssignments);
+}
+
+}
\ No newline at end of file
diff --git a/Engine/Source/Editor/UnrealEd/Private/Cooker/LoadBalanceCookBurden.h b/Engine/Source/Editor/UnrealEd/Private/Cooker/LoadBalanceCookBurden.h
new file mode 100644
--- /dev/null
+++ b/Engine/Source/Editor/UnrealEd/Private/Cooker/LoadBalanceCookBurden.h
@@ -0,0 +1,23 @@
+// Copyright Epic Games, Inc. All Rights Reserved.
+
+#pragma once
+
+#include "Containers/Array.h"
+#include "Containers/ArrayView.h"
+#include "Containers/Map.h"
+
+namespace UE::Cook { struct FPackageData; }
+namespace UE::Cook { struct FWorkerId; }
+
+namespace UE::Cook
+{
+
+/** Divides requests evenly without considering dependencies or load/save burden. */
+void LoadBalanceStriped(TConstArrayView<FWorkerId> Workers, TArrayView<FPackageData*> Requests,
+	TMap<FPackageData*, TArray<FPackageData*>>&& RequestGraph, TArray<FWorkerId>& OutAssignments);
+
+/** Balances "CookBurden" (considers load and save cost) across the CookWorkers. */
+void LoadBalanceCookBurden(TConstArrayView<FWorkerId> Workers, TArrayView<FPackageData*> Requests,
+	TMap<FPackageData*, TArray<FPackageData*>>&& RequestGraph, TArray<FWorkerId>& OutAssignments);
+
+}
\ No newline at end of file
