Add a new cmdline option to the Resave commandlet '-SkipPayloadTrailers' which allows the command to skip any package that already contains a package trailer.

#rnx

diff --git a/Engine/Source/Editor/UnrealEd/Classes/Commandlets/ResavePackagesCommandlet.h b/Engine/Source/Editor/UnrealEd/Classes/Commandlets/ResavePackagesCommandlet.h
--- a/Engine/Source/Editor/UnrealEd/Classes/Commandlets/ResavePackagesCommandlet.h
+++ b/Engine/Source/Editor/UnrealEd/Classes/Commandlets/ResavePackagesCommandlet.h
@@ -108,6 +108,9 @@ protected:
 	/** Only save packages containing FPayloadTrailers **/
 	bool bOnlyPayloadTrailers;
 
+	/** Only save packages without a FPayloadTrailer **/
+	bool bSkipPayloadTrailers;
+
 	/** Should we build navigation data for the packages we are saving? **/
 	bool bShouldBuildNavigationData;
 
diff --git a/Engine/Source/Editor/UnrealEd/Private/Commandlets/ContentCommandlets.cpp b/Engine/Source/Editor/UnrealEd/Private/Commandlets/ContentCommandlets.cpp
--- a/Engine/Source/Editor/UnrealEd/Private/Commandlets/ContentCommandlets.cpp
+++ b/Engine/Source/Editor/UnrealEd/Private/Commandlets/ContentCommandlets.cpp
@@ -1127,6 +1127,8 @@ int32 UResavePackagesCommandlet::Main( const FString& Params )
 	bOnlyVirtualized = Switches.Contains(TEXT("OnlyVirtualized"));
 	/** whether we should only save packages containing FPayloadTrailers */
 	bOnlyPayloadTrailers = Switches.Contains(TEXT("OnlyPayloadTrailers"));
+	/** whether we should skip saving packages that already have a payload trailer */
+	bSkipPayloadTrailers = Switches.Contains(TEXT("SkipPayloadTrailers"));
 	/** only process packages containing materials */
 	bOnlyMaterials = Switches.Contains(TEXT("onlymaterials"));
 	/** determine if we are building navigation data for the map packages on the pass. **/
@@ -1526,15 +1528,16 @@ void UResavePackagesCommandlet::PerformPreloadOperations( FLinkerLoad* PackageLi
 		}
 	}
 
-	// Check if the package contains a FPackageTrailer or not
-	if (bOnlyPayloadTrailers)
+	if (bOnlyPayloadTrailers && PackageLinker->GetPackageTrailer() == nullptr)
 	{
-		const UE::FPackageTrailer* Trailer = PackageLinker->GetPackageTrailer();
-		if (Trailer == nullptr)
-		{
-			bSavePackage = false;
-			return;
-		}
+		bSavePackage = false;
+		return;
+	}
+
+	if (bSkipPayloadTrailers && PackageLinker->GetPackageTrailer() != nullptr)
+	{
+		bSavePackage = false;
+		return;
 	}
 
 	// Check if the package contains any instances of the class that needs to be resaved.
