When reconstructing a material node, if we don't find a material pin by name, then find it by index.
Before 20122451 we only found by index.
Fixes issue for pasting material functions where finding by name fails because pasted nodes drop the type from the pin name.
The reason they drop the type from the pin name is because name is calculated before expression inputs have been determined.

diff --git a/Engine/Source/Editor/UnrealEd/Private/MaterialGraphNode_Base.cpp b/Engine/Source/Editor/UnrealEd/Private/MaterialGraphNode_Base.cpp
--- a/Engine/Source/Editor/UnrealEd/Private/MaterialGraphNode_Base.cpp
+++ b/Engine/Source/Editor/UnrealEd/Private/MaterialGraphNode_Base.cpp
@@ -283,13 +283,29 @@ void UMaterialGraphNode_Base::ReconstructNode()
 		UEdGraphPin* NewPin = nullptr;
 		if (OldPin->PinType.PinCategory == UMaterialGraphSchema::PC_Exec)
 		{
-			if (OldPin->Direction == EGPD_Input) NewPin = GetExecInputPin();
-			else NewPin = GetExecOutputPin(OldPin->SourceIndex);
+			if (OldPin->Direction == EGPD_Input) 
+			{
+				NewPin = GetExecInputPin();
+			}
+			else 
+			{
+				NewPin = GetExecOutputPin(OldPin->SourceIndex);
+			}
 		}
 		else
 		{
-			if (OldPin->Direction == EGPD_Input) NewPin = GetInputPin(OldPin->PinName);
-			else NewPin = GetOutputPin(OldPin->SourceIndex);
+			if (OldPin->Direction == EGPD_Input)
+			{
+				NewPin = GetInputPin(OldPin->PinName);
+				if (NewPin == nullptr)
+				{
+					NewPin = GetInputPin(OldPin->SourceIndex);
+				}
+			}
+			else 
+			{
+				NewPin = GetOutputPin(OldPin->SourceIndex);
+			}
 		}
 		if (NewPin)
 		{
