Change dialog "Updating source control status" to be delayed (will only show after 0.5sec).

diff --git a/Engine/Source/Editor/UnrealEd/Private/FileHelpers.cpp b/Engine/Source/Editor/UnrealEd/Private/FileHelpers.cpp
--- a/Engine/Source/Editor/UnrealEd/Private/FileHelpers.cpp
+++ b/Engine/Source/Editor/UnrealEd/Private/FileHelpers.cpp
@@ -1594,7 +1594,7 @@ bool FEditorFileUtils::AddCheckoutPackageItems(bool bCheckDirty, TArray<UPackage
 		{
 			// Update the source control status of all potentially relevant packages
 			FScopedSlowTask SlowTask(SourceControlCheckPackages.Num(), LOCTEXT("UpdatingSourceControlStatus", "Updating source control status..."));
-			SlowTask.MakeDialog();
+			SlowTask.MakeDialogDelayed(0.5f);
 			SourceControlProvider.Execute(ISourceControlOperation::Create<FUpdateStatus>(), SourceControlCheckPackages);
 			SlowTask.EnterProgressFrame(SourceControlCheckPackages.Num());
 		}
