World Partition Convert Commandlet: normalize config file cache paths to avoid a warning.

#rnx

diff --git a/Engine/Source/Editor/UnrealEd/Private/Commandlets/WorldPartitionConvertCommandlet.cpp b/Engine/Source/Editor/UnrealEd/Private/Commandlets/WorldPartitionConvertCommandlet.cpp
--- a/Engine/Source/Editor/UnrealEd/Private/Commandlets/WorldPartitionConvertCommandlet.cpp
+++ b/Engine/Source/Editor/UnrealEd/Private/Commandlets/WorldPartitionConvertCommandlet.cpp
@@ -749,7 +749,7 @@ int32 UWorldPartitionConvertCommandlet::Main(const FString& Params)
 	FString LevelLongPackageName;
 	if (FPackageName::SearchForPackageOnDisk(Tokens[0], nullptr, &LevelLongPackageName))
 	{
-		LevelConfigFilename = FPaths::ChangeExtension(LevelLongPackageName, TEXT("ini"));
+		LevelConfigFilename = FConfigCacheIni::NormalizeConfigIniPath(FPaths::ChangeExtension(LevelLongPackageName, TEXT("ini")));
 
 		if (FPlatformFileManager::Get().GetPlatformFile().FileExists(*LevelConfigFilename))
 		{
