Editor: fix BSP brush drag-in messing up the selection undo/redo history and potentially causing a crash.

diff --git a/Engine/Source/Editor/UnrealEd/Private/LevelEditorViewport.cpp b/Engine/Source/Editor/UnrealEd/Private/LevelEditorViewport.cpp
--- a/Engine/Source/Editor/UnrealEd/Private/LevelEditorViewport.cpp
+++ b/Engine/Source/Editor/UnrealEd/Private/LevelEditorViewport.cpp
@@ -1042,6 +1042,19 @@ static bool AreAllDroppedObjectsBrushBuilders(const TArray<UObject*>& DroppedObj
 	return true;
 }
 
+static bool AreAnyDroppedObjectsBrushBuilders(const TArray<UObject*>& DroppedObjects)
+{
+	for (UObject* DroppedObject : DroppedObjects)
+	{
+		if (DroppedObject->IsA(UBrushBuilder::StaticClass()))
+		{
+			return true;
+		}
+	}
+
+	return false;
+}
+
 bool FLevelEditorViewportClient::DropObjectsOnBackground(FViewportCursorLocation& Cursor, const TArray<UObject*>& DroppedObjects, EObjectFlags ObjectFlags, TArray<AActor*>& OutNewActors, bool bCreateDropPreview, bool bSelectActors, UActorFactory* FactoryToUse)
 {
 	if (DroppedObjects.Num() == 0)
@@ -1645,7 +1658,13 @@ bool FLevelEditorViewportClient::DropObjectsAtCoordinates(int32 MouseX, int32 Mo
 				//Viewport->InvalidateHitProxy();
 			}
 			// Dropping the actors rather than a preview? Probably want to select them all then. 
-			else if(!bCreateDropPreview && SelectActors && OutNewActors.Num() > 0)
+			else if(!bCreateDropPreview && SelectActors && OutNewActors.Num() > 0
+
+				// If we're dropping bsp brushes, that geometry actually gets created later, in 
+				// FBrushBuilderDragDropOp::OnDrop(), so we don't want to select the (preview) builder brush.
+				// (for reference, this function gets called via CurWidget.Widget->OnDrop() in FSlateApplication::RoutePointerUpEvent,
+				// whereas the FBrushBuilderDragDropOp::OnDrop() gets called via SlateUser->NotifyPointerReleased).
+				&& !AreAnyDroppedObjectsBrushBuilders(DroppedObjects))
 			{
 				for (auto It = OutNewActors.CreateConstIterator(); It; ++It)
 				{
diff --git a/Engine/Source/Editor/UnrealEd/Public/DragAndDrop/BrushBuilderDragDropOp.h b/Engine/Source/Editor/UnrealEd/Public/DragAndDrop/BrushBuilderDragDropOp.h
--- a/Engine/Source/Editor/UnrealEd/Public/DragAndDrop/BrushBuilderDragDropOp.h
+++ b/Engine/Source/Editor/UnrealEd/Public/DragAndDrop/BrushBuilderDragDropOp.h
@@ -34,9 +34,8 @@ public:
 		UWorld* World = GEditor->GetEditorWorldContext().World();
 		if ((World != nullptr) && (World->GetDefaultBrush() != nullptr))
 		{
-			// Deselect & hide the builder brush
+			// Hide the builder brush
 			World->GetDefaultBrush()->SetIsTemporarilyHiddenInEditor(true);
-			GEditor->SelectActor(World->GetDefaultBrush(), false, false);
 		}
 	}
 
@@ -56,9 +55,8 @@ public:
 					GEditor->Exec(World, Command);
 				}
 
-				// Deselect & hide the builder brush
+				// Hide the builder brush
 				World->GetDefaultBrush()->SetIsTemporarilyHiddenInEditor(true);
-				GEditor->SelectActor(World->GetDefaultBrush(), false, false);
 			}
 		}
 
@@ -72,9 +70,8 @@ public:
 		{
 			if(CursorDecoratorWindow->IsVisible() && !World->GetDefaultBrush()->IsTemporarilyHiddenInEditor())
 			{
-				// Deselect & hide the builder brush if the decorator is visible
+				// Hide the builder brush if the decorator is visible
 				World->GetDefaultBrush()->SetIsTemporarilyHiddenInEditor(true);
-				GEditor->SelectActor(World->GetDefaultBrush(), false, false);
 			}
 		}
 
@@ -109,12 +106,11 @@ private:
 		, IconBrush(InIconBrush)
 		, bIsAdditive(bInIsAdditive)
 	{
-		// show & select the builder brush
+		// show the builder brush
 		UWorld* World = GEditor->GetEditorWorldContext().World();
 		if ((World != nullptr) && (World->GetDefaultBrush() != nullptr))
 		{
 			World->GetDefaultBrush()->SetIsTemporarilyHiddenInEditor(false);
-			GEditor->SelectActor(World->GetDefaultBrush(), true, false);
 		}
 	}
 
