Emit Pre/Post edit change events when dropping materials on actors to trigger construction script

diff --git a/Engine/Source/Editor/UnrealEd/Private/Kismet2/ComponentEditorUtils.cpp b/Engine/Source/Editor/UnrealEd/Private/Kismet2/ComponentEditorUtils.cpp
--- a/Engine/Source/Editor/UnrealEd/Private/Kismet2/ComponentEditorUtils.cpp
+++ b/Engine/Source/Editor/UnrealEd/Private/Kismet2/ComponentEditorUtils.cpp
@@ -942,7 +942,9 @@ bool FComponentEditorUtils::AttemptApplyMaterialToComponent(USceneComponent* Sce
 	{
 		bResult = true;
 		const FScopedTransaction Transaction(LOCTEXT("DropTarget_UndoSetComponentMaterial", "Assign Material to Component (Drag and Drop)"));
+		FProperty* Property = FindFProperty<FProperty>(SceneComponent->GetClass(), MeshComponent ? "OverrideMaterials" : "DecalMaterial");
 		SceneComponent->Modify();
+		SceneComponent->PreEditChange(Property);
 
 		if (MeshComponent)
 		{
@@ -975,7 +977,8 @@ bool FComponentEditorUtils::AttemptApplyMaterialToComponent(USceneComponent* Sce
 		}
 
 		SceneComponent->MarkRenderStateDirty();
-		SceneComponent->PostEditChange();
+		FPropertyChangedEvent PropertyChangedEvent(Property, EPropertyChangeType::ValueSet, { SceneComponent });
+		SceneComponent->PostEditChangeProperty(PropertyChangedEvent);
 		GEditor->OnSceneMaterialsModified();
 	}
 
