#UE Add a ini configurable option to ignore assets in certain paths from being considered by IsModifiedInOtherBranch before showing the dialog.

diff --git a/Engine/Source/Editor/UnrealEd/Private/FileHelpers.cpp b/Engine/Source/Editor/UnrealEd/Private/FileHelpers.cpp
--- a/Engine/Source/Editor/UnrealEd/Private/FileHelpers.cpp
+++ b/Engine/Source/Editor/UnrealEd/Private/FileHelpers.cpp
@@ -83,6 +83,7 @@
 #include "PackageSourceControlHelper.h"
 #include "ActorFolder.h"
 #include "InterchangeManager.h"
+#include "SourceControlHelpers.h"
 #include "InterchangeProjectSettings.h"
 
 DEFINE_LOG_CATEGORY_STATIC(LogFileHelpers, Log, All);
@@ -230,7 +231,9 @@ namespace FileDialogHelpers
 static bool ConfirmPackageBranchCheckOutStatus(const TArray<UPackage*>& PackagesToCheckOut)
 {
 	//@TODO: Need more info here (in the event multiple packages are trying to be saved at once; the prompt shown is misleading in that case (you might be OK with stomping over one file but not others later on in the list))
-	for (UPackage* CurPackage : PackagesToCheckOut)
+	bool bModifiedInOtherBranchIgnorePathsInitialized = false;
+	TArray<FString> ModifiedInOtherBranchIgnorePaths;
+ 	for (UPackage* CurPackage : PackagesToCheckOut)
 	{
 		ISourceControlProvider& SourceControlProvider = ISourceControlModule::Get().GetProvider();
 		FSourceControlStatePtr SourceControlState = SourceControlProvider.GetState(CurPackage, EStateCacheUsage::Use);
@@ -248,22 +251,41 @@ static bool ConfirmPackageBranchCheckOutStatus(const TArray<UPackage*>& Packages
 			SourceControlState->GetOtherBranchHeadModification(HeadBranch, HeadAction, HeadCL);
 
 			FText InfoText;
-
+			bool bShowDialog = true;
+			const FString CurPackageName = CurPackage->GetName();
 			if (SourceControlState->IsModifiedInOtherBranch())
 			{
-				int32 CurrentBranchIdx = SourceControlProvider.GetStateBranchIndex(CurrentBranch);
-				int32 HeadBranchIdx = SourceControlProvider.GetStateBranchIndex(HeadBranch);
+				// Lazy-initialize ModifiedInOtherBranchIgnorePaths since most of the time, assets are not modified in other branches.
+				if (!bModifiedInOtherBranchIgnorePathsInitialized)
 				{
-					if (CurrentBranchIdx != INDEX_NONE && HeadBranchIdx != INDEX_NONE)
+					if (GConfig)
 					{
-						// modified
-						if (CurrentBranchIdx < HeadBranchIdx)
-						{
-							InfoText = LOCTEXT("WarningModifiedOtherBranchHigher", "Modified in higher branch, consider waiting for package to be merged down.");
-						}
-						else
+						GConfig->GetArray(TEXT("SourceControl.SourceControlSettings"), TEXT("ModifiedInOtherBranchIgnorePaths"), ModifiedInOtherBranchIgnorePaths, SourceControlHelpers::GetGlobalSettingsIni());
+					}
+					bModifiedInOtherBranchIgnorePathsInitialized = true;
+				}
+
+				if (ModifiedInOtherBranchIgnorePaths.ContainsByPredicate([&CurPackageName](const FString& Path) { return CurPackageName.StartsWith(Path); }))
+				{
+					// This file was modified in another branch, but we are configured to ignore it and not show a dialog for files in this path.
+					bShowDialog = false;
+				}
+				else
+				{
+					int32 CurrentBranchIdx = SourceControlProvider.GetStateBranchIndex(CurrentBranch);
+					int32 HeadBranchIdx = SourceControlProvider.GetStateBranchIndex(HeadBranch);
+					{
+						if (CurrentBranchIdx != INDEX_NONE && HeadBranchIdx != INDEX_NONE)
 						{
-							InfoText = LOCTEXT("WarningModifiedOtherBranchLower", "Modified in lower branch, keep track of your work. You may need to redo it during the merge.");
+							// modified
+							if (CurrentBranchIdx < HeadBranchIdx)
+							{
+								InfoText = LOCTEXT("WarningModifiedOtherBranchHigher", "Modified in higher branch, consider waiting for package to be merged down.");
+							}
+							else
+							{
+								InfoText = LOCTEXT("WarningModifiedOtherBranchLower", "Modified in lower branch, keep track of your work. You may need to redo it during the merge.");
+							}
 						}
 					}
 				}
@@ -282,14 +304,17 @@ static bool ConfirmPackageBranchCheckOutStatus(const TArray<UPackage*>& Packages
 
 			}
 
-			const FText PackageNameText = FText::FromName(CurPackage->GetFName());
+			if (bShowDialog)
+			{
+				const FText PackageNameText = FText::FromString(CurPackageName);
 
-			const FText Message = SourceControlState->IsModifiedInOtherBranch() ? FText::Format(LOCTEXT("WarningModifiedOtherBranch", "WARNING: Package {3} modified in {0} CL {1}\n\n{2}\n\nCheck out packages anyway?"), FText::FromString(HeadBranch), FText::AsNumber(HeadCL, &NoCommas), InfoText, PackageNameText)
-				: FText::Format(LOCTEXT("WarningCheckedOutOtherBranch", "WARNING: Package {2} checked out in {0}\n\n{1}\n\nCheck out packages anyway?"), FText::FromString(SourceControlState->GetOtherUserBranchCheckedOuts()), InfoText, PackageNameText);
+				const FText Message = SourceControlState->IsModifiedInOtherBranch() ? FText::Format(LOCTEXT("WarningModifiedOtherBranch", "WARNING: Package {3} modified in {0} CL {1}\n\n{2}\n\nCheck out packages anyway?"), FText::FromString(HeadBranch), FText::AsNumber(HeadCL, &NoCommas), InfoText, PackageNameText)
+					: FText::Format(LOCTEXT("WarningCheckedOutOtherBranch", "WARNING: Package {2} checked out in {0}\n\n{1}\n\nCheck out packages anyway?"), FText::FromString(SourceControlState->GetOtherUserBranchCheckedOuts()), InfoText, PackageNameText);
 
-			const FText Title = SourceControlState->IsModifiedInOtherBranch() ? FText::FromString("Package Branch Modifications") : FText::FromString("Package Branch Checkouts");
+				const FText Title = SourceControlState->IsModifiedInOtherBranch() ? FText::FromString("Package Branch Modifications") : FText::FromString("Package Branch Checkouts");
 
-			return FMessageDialog::Open(EAppMsgType::YesNo, Message, &Title) == EAppReturnType::Yes;
+				return FMessageDialog::Open(EAppMsgType::YesNo, Message, &Title) == EAppReturnType::Yes;
+			}
 		}
 	}
 
