Fixed merge conflict for patrick.boutot

Prevent the ActorEditorContext to access the UObject and Editor when the editor shutdown.

diff --git a/Engine/Source/Editor/UnrealEd/Private/SActorEditorContext.cpp b/Engine/Source/Editor/UnrealEd/Private/SActorEditorContext.cpp
--- a/Engine/Source/Editor/UnrealEd/Private/SActorEditorContext.cpp
+++ b/Engine/Source/Editor/UnrealEd/Private/SActorEditorContext.cpp
@@ -16,6 +16,7 @@ void SActorEditorContext::Construct(const FArguments& InArgs)
 {
 	World = (InArgs._World);
 	bIsContextExpanded = true;
+	check(GEditor);
 	GEditor->GetEditorWorldContext().AddRef(World);
 	UActorEditorContextSubsystem::Get()->OnActorEditorContextSubsystemChanged().AddSP(this, &SActorEditorContext::Rebuild);
 	FEditorDelegates::MapChange.AddSP(this, &SActorEditorContext::OnEditorMapChange);
@@ -25,12 +26,12 @@ void SActorEditorContext::Construct(const FArguments& InArgs)
 SActorEditorContext::~SActorEditorContext()
 {
 	FEditorDelegates::MapChange.RemoveAll(this);
-	if (UObjectInitialized())
+	if (GEditor)
 	{
 		GEditor->GetEditorWorldContext().RemoveRef(World);
-		if (UActorEditorContextSubsystem* ActorEditorContextSubsystem = UActorEditorContextSubsystem::Get())
+		if (UObjectInitialized())
 		{
-			ActorEditorContextSubsystem->OnActorEditorContextSubsystemChanged().RemoveAll(this);
+			UActorEditorContextSubsystem::Get()->OnActorEditorContextSubsystemChanged().RemoveAll(this);
 		}
 	}
 }
