World Partition: don't setup a loading region when saving partitioned worlds with streaming disabled.

#rnx

diff --git a/Engine/Source/Editor/UnrealEd/Private/FileHelpers.cpp b/Engine/Source/Editor/UnrealEd/Private/FileHelpers.cpp
--- a/Engine/Source/Editor/UnrealEd/Private/FileHelpers.cpp
+++ b/Engine/Source/Editor/UnrealEd/Private/FileHelpers.cpp
@@ -926,7 +926,7 @@ static bool SaveWorld(UWorld* World,
 				AssetRegistry.ScanPathsSynchronous( ULevel::GetExternalObjectsPaths(NewPackageName) , true);
 			}
 
-			if (RenamedWorldPartition)
+			if (RenamedWorldPartition && RenamedWorldPartition->IsStreamingEnabled())
 			{
 				if (LoadedEditorRegions.Num())
 				{
