Changed default arcball behavior to act more like other DCC's, and added a limit cvar to extend the size of the arcball, again so it works like other DCC's.

diff --git a/Engine/Source/Editor/UnrealEd/Private/UnrealWidget.cpp b/Engine/Source/Editor/UnrealEd/Private/UnrealWidget.cpp
--- a/Engine/Source/Editor/UnrealEd/Private/UnrealWidget.cpp
+++ b/Engine/Source/Editor/UnrealEd/Private/UnrealWidget.cpp
@@ -117,9 +117,44 @@ void FWidget::ConvertMouseToAxis_Translate(FVector2D DragDir, FVector& InOutDelt
 	OutDrag = CustomCoordSystem.TransformPosition(OutDrag);
 }
 
+//CVAR for the arcball size, so animators can adjust it
+static TAutoConsoleVariable<float> CVarArcballLimit(
+	TEXT("r.Editor.ArcballDragLimit"),
+	2.0,
+	TEXT("For how long the arcball rotates until it switches to a screens space rotate, default of 1.0 equals the size of the arcball"),
+	ECVF_RenderThreadSafe
+);
+//function used to find quat between two angles, and works much better, faster and less degenerate, then trying to use cross and dot product
+static FQuat FindQuatBetweenNormals(const FVector& A, const FVector& B)
+{
+	const FQuat::FReal Dot = FVector::DotProduct(A, B);
+	FQuat::FReal W = 1 + Dot;
+	FQuat Result;
+
+	if (W < SMALL_NUMBER)
+	{
+		// A and B point in opposite directions
+		W = 2 - W;
+		Result = FQuat(-A.Y * B.Z + A.Z * B.Y, -A.Z * B.X + A.X * B.Z, -A.X * B.Y + A.Y * B.X, W).GetNormalized();
+
+		const FVector Normal = FMath::Abs(A.X) > FMath::Abs(A.Y) ? FVector::YAxisVector : FVector::XAxisVector;
+		const FVector BiNormal = FVector::CrossProduct(A, Normal);
+		const FVector TauNormal = FVector::CrossProduct(A, BiNormal);
+		Result = Result * FQuat(TauNormal, PI);
+	}
+	else
+	{
+		//Axis = FVector::CrossProduct(A, B);
+		Result = FQuat(A.Y * B.Z - A.Z * B.Y, A.Z * B.X - A.X * B.Z, A.X * B.Y - A.Y * B.X, W);
+	}
+
+	Result.Normalize();
+	return Result;
+};
+
 void FWidget::ConvertMouseToAxis_Rotate(FVector2D TangentDir, FVector2D DragDir, FSceneView* InView,
-                                        FEditorViewportClient* InViewportClient, FVector& InOutDelta,
-                                        FRotator& OutRotation)
+	FEditorViewportClient* InViewportClient, FVector& InOutDelta,
+	FRotator& OutRotation)
 {
 	if (CurrentAxis == EAxisList::X)
 	{
@@ -153,7 +188,7 @@ void FWidget::ConvertMouseToAxis_Rotate(FVector2D TangentDir, FVector2D DragDir,
 		FSnappingUtils::SnapRotatorToGrid(Rotation);
 
 		CurrentDeltaRotation = Rotation.Pitch;
-		EffectiveDelta       = AxisDir * Rotation.Pitch;
+		EffectiveDelta = AxisDir * Rotation.Pitch;
 		// Adjust the input delta according to how much rotation was actually applied
 		InOutDelta = FVector(EffectiveDelta.X, -EffectiveDelta.Y, 0.0f);
 		// Need to get the delta rotation in the current coordinate space of the widget
@@ -169,7 +204,7 @@ void FWidget::ConvertMouseToAxis_Rotate(FVector2D TangentDir, FVector2D DragDir,
 		FSnappingUtils::SnapRotatorToGrid(Rotation);
 
 		CurrentDeltaRotation = Rotation.Yaw;
-		EffectiveDelta       = AxisDir * Rotation.Yaw;
+		EffectiveDelta = AxisDir * Rotation.Yaw;
 		// Adjust the input delta according to how much rotation was actually applied
 		InOutDelta = FVector(EffectiveDelta.X, -EffectiveDelta.Y, 0.0f);
 		// Need to get the delta rotation in the current coordinate space of the widget
@@ -177,40 +212,74 @@ void FWidget::ConvertMouseToAxis_Rotate(FVector2D TangentDir, FVector2D DragDir,
 	}
 	else if (CurrentAxis == EAxisList::XYZ) //arcball rotate
 	{
-		//To do this we need to calculate the rotation axis and rotation angle.
-		//The Axis is the cross product of the current ray from eye to pixel in world space with the previous ray.
-		//The Angle is angle amount we rotate from the object's location to the imaginary sphere that matches up with the difference between the current and previous ray
-		//From the Camera. Those rays form a triangle, which we can bisect with a common side.
+		//For Arball rotate we have three states we can be in
+		// If within the OUTER_AXIS_CIRCLE_RADIUS, we hit test the rays from current and old camera to the widget with that circle radius, and then find the angle
+		// between the vectors form the center of the widget to those intersections.
+		// If ouside the arcball, we still want it to rotate like an arcball since our is smaller than other DCC's and doesnt' feel right.
+		// We use a CVAR to specify how far out to do this extra pseudo arcball rotation, 2.0 feels right for animators.
+		// So in this case the Axis is the cross product of the current ray from eye to pixel in world space with the previous ray.
+		// The Angle is angle amount we rotate from the object's location to the imaginary sphere that matches up with the difference 
+		// between the current and previous ray from the Camera. Those rays form a triangle, which we can bisect with a common side, to find the angle
+		// This gives us an arcball like rotation beyond the normal arcball.
+		// Finally if outside the specified cvar distance we do a screen rotate like other DCC's
 		FVector2D MousePosition(InViewportClient->Viewport->GetMouseX(), InViewportClient->Viewport->GetMouseY());
 		FViewportCursorLocation OldMouseViewportRay(InView, InViewportClient, LastDragPos.X, LastDragPos.Y);
 		FViewportCursorLocation MouseViewportRay(InView, InViewportClient, MousePosition.X, MousePosition.Y);
-
-		LastDragPos               = MousePosition;
+		LastDragPos = MousePosition;
 		FVector DirectionToWidget = InViewportClient->GetWidgetLocation() - MouseViewportRay.GetOrigin();
-		float Length              = DirectionToWidget.Size();
-		if (!FMath::IsNearlyZero(Length))
+		float Length = DirectionToWidget.Size();
+		if (!FMath::IsNearlyZero(Length)) //degenerate check
 		{
-			//Calc Axis
+			//Find screen space distance to the arcball size
 			DirectionToWidget /= Length;
-			const FVector CameraToPixelDir    = MouseViewportRay.GetDirection();
+			const FVector CameraToPixelDir = MouseViewportRay.GetDirection();
 			const FVector OldCameraToPixelDir = OldMouseViewportRay.GetDirection();
-			FVector RotationAxis              = FVector::CrossProduct(OldCameraToPixelDir, CameraToPixelDir);
+			FVector RotationAxis = FVector::CrossProduct(CameraToPixelDir, OldCameraToPixelDir);
 			RotationAxis.Normalize();
-			float RotationAngle     = 0.0f;
+			float RotationAngle = 0.0f;
 			FVector4 ScreenLocation = InView->WorldToScreen(InViewportClient->GetWidgetLocation());
 			FVector2D PixelLocation;
 			InView->ScreenToPixel(ScreenLocation, PixelLocation);
-			float Distance      = FVector2D::Distance(PixelLocation, MousePosition);
-			const float MaxDiff = 2.0f * (INNER_AXIS_CIRCLE_RADIUS);
-			//If outside radius, do screen rotate instead, like other DCC's
-
-			if (Distance > MaxDiff)
+			float Distance = FVector2D::Distance(PixelLocation, MousePosition);
+			const float ExternalScale = EditorModeTools ? EditorModeTools->GetWidgetScale() : 1.0f;
+			const float CircleRadius = OUTER_AXIS_CIRCLE_RADIUS * ExternalScale;
+			const float ArcballLimit = CVarArcballLimit.GetValueOnGameThread() * 2.0f;
+			const float MaxDiff = ArcballLimit * CircleRadius;
+			//If within Arcball do rotate the ball based on angle difference on hit tested sphere 
+			if (Distance < CircleRadius)
+			{
+				const float ScaleInScreen = InView->WorldToScreen(InViewportClient->GetWidgetLocation()).W * (4.0f / InView->UnscaledViewRect.Width() / InView->ViewMatrices.GetProjectionMatrix().M[0][0]);
+				const float SphereRadius = CircleRadius * ScaleInScreen  + GetDefault<ULevelEditorViewportSettings>()->TransformWidgetSizeAdjustment;
+				FVector OldLocation, NewLocation;
+				FMath::SphereDistToLine(InViewportClient->GetWidgetLocation(), SphereRadius, MouseViewportRay.GetOrigin(), OldCameraToPixelDir, OldLocation);
+				FMath::SphereDistToLine(InViewportClient->GetWidgetLocation(), SphereRadius, MouseViewportRay.GetOrigin(), CameraToPixelDir, NewLocation);
+				FVector OldLineToCenter = OldLocation - InViewportClient->GetWidgetLocation();
+				OldLineToCenter.Normalize();
+				FVector NewLineToCenter = NewLocation - InViewportClient->GetWidgetLocation();
+				NewLineToCenter.Normalize();
+				const FQuat QuatRotation = FindQuatBetweenNormals(OldLineToCenter, NewLineToCenter);
+				OutRotation = FRotator(QuatRotation);
+			}
+			else if (Distance > CircleRadius && Distance < MaxDiff) //do 
+			{
+				const float Scale = ScreenLocation.W *
+					(4.0f / InView->UnscaledViewRect.Width() / InView->ViewMatrices.GetProjectionMatrix().M[0][0]);
+				const float Radius = (CircleRadius * Scale) +
+					GetDefault<ULevelEditorViewportSettings>()->TransformWidgetSizeAdjustment;
+				const float LengthOfAdjacent = Length - Radius;
+				RotationAngle = FMath::Acos(FVector::DotProduct(OldCameraToPixelDir, CameraToPixelDir));
+				const float OppositeSize = FMath::Tan(RotationAngle) * LengthOfAdjacent;
+				RotationAngle = FMath::Atan2(OppositeSize, Radius);
+				const FQuat QuatRotation(RotationAxis, RotationAngle);
+				OutRotation = FRotator(QuatRotation);
+			}
+			else if (Distance > MaxDiff) 	//If outside radius, do screen rotate instead, like other DCC's
 			{
 				FPlane Plane(InViewportClient->GetWidgetLocation(), DirectionToWidget);
 				FVector StartOnPlane =
-				    FMath::RayPlaneIntersection(MouseViewportRay.GetOrigin(), CameraToPixelDir, Plane);
+					FMath::RayPlaneIntersection(MouseViewportRay.GetOrigin(), CameraToPixelDir, Plane);
 				FVector OldOnPlane =
-				    FMath::RayPlaneIntersection(MouseViewportRay.GetOrigin(), OldCameraToPixelDir, Plane);
+					FMath::RayPlaneIntersection(MouseViewportRay.GetOrigin(), OldCameraToPixelDir, Plane);
 				StartOnPlane -= InViewportClient->GetWidgetLocation();
 				OldOnPlane -= InViewportClient->GetWidgetLocation();
 				StartOnPlane.Normalize();
@@ -223,22 +292,9 @@ void FWidget::ConvertMouseToAxis_Rotate(FVector2D TangentDir, FVector2D DragDir,
 					RotationAngle *= -1.0f;
 				}
 				RotationAxis = DirectionToWidget;
+				const FQuat QuatRotation(RotationAxis, RotationAngle);
+				OutRotation = FRotator(QuatRotation);
 			}
-			else
-			{
-				const float Scale = ScreenLocation.W *
-				    (4.0f / InView->UnscaledViewRect.Width() / InView->ViewMatrices.GetProjectionMatrix().M[0][0]);
-				const float InnerRadius = (INNER_AXIS_CIRCLE_RADIUS * Scale) +
-				    GetDefault<ULevelEditorViewportSettings>()->TransformWidgetSizeAdjustment;
-				const float LengthOfAdjacent = Length - InnerRadius;
-				RotationAngle                = FMath::Acos(FVector::DotProduct(OldCameraToPixelDir, CameraToPixelDir));
-				const float OppositeSize     = FMath::Tan(RotationAngle) * LengthOfAdjacent;
-				RotationAngle                = FMath::Atan(OppositeSize / InnerRadius);
-				RotationAngle                = RotationAngle < 0.0f ? RotationAngle : -RotationAngle;
-			}
-
-			const FQuat QuatRotation(RotationAxis, RotationAngle);
-			OutRotation = FRotator(QuatRotation);
 		}
 		return;
 	}
@@ -248,15 +304,15 @@ void FWidget::ConvertMouseToAxis_Rotate(FVector2D TangentDir, FVector2D DragDir,
 		FViewportCursorLocation OldMouseViewportRay(InView, InViewportClient, LastDragPos.X, LastDragPos.Y);
 		FViewportCursorLocation MouseViewportRay(InView, InViewportClient, MousePosition.X, MousePosition.Y);
 
-		LastDragPos               = MousePosition;
+		LastDragPos = MousePosition;
 		FVector DirectionToWidget = InViewportClient->GetWidgetLocation() - MouseViewportRay.GetOrigin();
-		float Length              = DirectionToWidget.Size();
+		float Length = DirectionToWidget.Size();
 
 		if (!FMath::IsNearlyZero(Length))
 		{
 			DirectionToWidget /= Length;
 
-			const FVector CameraToPixelDir    = MouseViewportRay.GetDirection();
+			const FVector CameraToPixelDir = MouseViewportRay.GetDirection();
 			const FVector OldCameraToPixelDir = OldMouseViewportRay.GetDirection();
 			FPlane Plane(InViewportClient->GetWidgetLocation(), DirectionToWidget);
 			FVector StartOnPlane = FMath::RayPlaneIntersection(MouseViewportRay.GetOrigin(), CameraToPixelDir, Plane);
@@ -266,7 +322,7 @@ void FWidget::ConvertMouseToAxis_Rotate(FVector2D TangentDir, FVector2D DragDir,
 			StartOnPlane.Normalize();
 			OldOnPlane.Normalize();
 			float RotationAngle = FMath::Acos(FVector::DotProduct(StartOnPlane, OldOnPlane));
-			FVector Cross       = FVector::CrossProduct(OldCameraToPixelDir, CameraToPixelDir);
+			FVector Cross = FVector::CrossProduct(OldCameraToPixelDir, CameraToPixelDir);
 			if (FVector::DotProduct(DirectionToWidget, Cross) < 0.0f)
 			{
 				RotationAngle *= -1.0f;
diff --git a/Engine/Source/Editor/UnrealEd/Private/UnrealWidgetRender.cpp b/Engine/Source/Editor/UnrealEd/Private/UnrealWidgetRender.cpp
--- a/Engine/Source/Editor/UnrealEd/Private/UnrealWidgetRender.cpp
+++ b/Engine/Source/Editor/UnrealEd/Private/UnrealWidgetRender.cpp
@@ -648,10 +648,11 @@ void FWidget::DrawColoredSphere(FPrimitiveDrawInterface* PDI, const FVector& Cen
 	MeshBuilder.Draw(PDI, FScaleMatrix(Radii) * FRotationMatrix(Orientation) * FTranslationMatrix(Center), MaterialRenderProxy, DepthPriority, bDisableBackfaceCulling);
 }
 //CVAR for the arcball size, so animators can adjust it
+//DEPRECRATREDin 5.1, since changing the size will cause hit test problems.
 static TAutoConsoleVariable<float> CVarArcballSize(
 	TEXT("r.Editor.ArcballSize"),
 	1.0,
-	TEXT("Custom Size of the Arcball, size of 1.0 means same size as gizmo)"),
+	TEXT("DEPRECRATED in 5.1)"),
 	ECVF_RenderThreadSafe
 );
 
@@ -702,7 +703,7 @@ void FWidget::Render_Rotate( const FSceneView* View,FPrimitiveDrawInterface* PDI
 		{
 			FVector Center = InLocation;
 			FRotator Orientation = FRotator::ZeroRotator;
-			const float ArcballSize = CVarArcballSize.GetValueOnGameThread();
+			const float ArcballSize = 1.0f; //always constant, otherwise can cause hit testing issues
 
 			const float InnerDistance = (INNER_AXIS_CIRCLE_RADIUS * ArcballSize * Scale) + GetDefault<ULevelEditorViewportSettings>()->TransformWidgetSizeAdjustment;
 			FVector Radii(InnerDistance, InnerDistance, InnerDistance);
