Sequencer: When DragAlwaysDoesMarquee setting is on it also always moves the camera when ALT is pressed.

diff --git a/Engine/Source/Editor/Sequencer/Private/SequencerEdMode.cpp b/Engine/Source/Editor/Sequencer/Private/SequencerEdMode.cpp
--- a/Engine/Source/Editor/Sequencer/Private/SequencerEdMode.cpp
+++ b/Engine/Source/Editor/Sequencer/Private/SequencerEdMode.cpp
@@ -41,6 +41,9 @@
 #include "SnappingUtils.h"
 #include "SequencerNodeTree.h"
 #include "SequencerSettings.h"
+#include "UnrealEdGlobals.h"
+#include "UnrealEdMisc.h"
+#include "Editor/UnrealEdEngine.h"
 
 const FEditorModeID FSequencerEdMode::EM_SequencerMode(TEXT("EM_SequencerMode"));
 
@@ -249,6 +252,14 @@ USequencerSettings* FSequencerEdMode:: GetSequencerSettings() const
 	return nullptr;
 }
 
+bool FSequencerEdMode::IsMovingCamera(FViewport* InViewport) const
+{
+	const bool LeftMouseButtonDown = InViewport->KeyState(EKeys::LeftMouseButton);
+	const bool bIsAltKeyDown = InViewport->KeyState(EKeys::LeftAlt) || InViewport->KeyState(EKeys::RightAlt);
+	const USequencerSettings* SequencerSettings = GetSequencerSettings();
+
+	return ((SequencerSettings ? SequencerSettings->GetLeftMouseDragDoesMarquee() : false) && LeftMouseButtonDown && bIsAltKeyDown);
+}
 bool FSequencerEdMode::IsDoingDrag(FViewport* InViewport) const
 {
 	const USequencerSettings* SequencerSettings = GetSequencerSettings();
@@ -282,6 +293,11 @@ bool FSequencerEdMode::StartTracking(FEditorViewportClient* InViewportClient, FV
 		}
 		return true;
 	}
+	else if (IsMovingCamera(InViewport))
+	{
+		InViewportClient->SetCurrentWidgetAxis(EAxisList::None);
+		return true;
+	}
 	else if (IsDoingDrag(InViewport))
 	{
 		DragToolHandler.MakeDragTool(InViewportClient);
@@ -295,6 +311,10 @@ bool FSequencerEdMode::EndTracking(FEditorViewportClient* InViewportClient, FVie
 	{
 		return true;
 	}
+	else if (IsMovingCamera(InViewport))
+	{
+		return true;
+	}
 	else if (DragToolHandler.EndTracking(InViewportClient, InViewport))
 	{
 		return true;
@@ -308,9 +328,22 @@ bool FSequencerEdMode::InputDelta(FEditorViewportClient* InViewportClient, FView
 	{
 		return true;
 	}
+	else if (IsMovingCamera(InViewport))
+	{
+		if (InDrag.IsNearlyZero() == false || InRot.IsNearlyZero() == false || InScale.IsNearlyZero() == false)
+		{
+			InViewportClient->SetCurrentWidgetAxis(EAxisList::None);
+			InViewportClient->PeformDefaultCameraMovement(InDrag, InRot, InScale);
+			GUnrealEd->UpdatePivotLocationForSelection();
+		}
+
+		return true;
+	}
 	else if (IsDoingDrag(InViewport))
 	{
 		return DragToolHandler.InputDelta(InViewportClient, InViewport, InDrag, InRot, InScale);
+		GUnrealEd->UpdatePivotLocationForSelection();
+
 	}
 	return FEdMode::InputDelta(InViewportClient, InViewport, InDrag, InRot, InScale);
 }
diff --git a/Engine/Source/Editor/Sequencer/Private/SequencerEdMode.h b/Engine/Source/Editor/Sequencer/Private/SequencerEdMode.h
--- a/Engine/Source/Editor/Sequencer/Private/SequencerEdMode.h
+++ b/Engine/Source/Editor/Sequencer/Private/SequencerEdMode.h
@@ -102,6 +102,7 @@ public:
 
 	bool IsPressingMoveTimeSlider(FViewport* InViewport) const;
 	bool IsDoingDrag(FViewport* InViewport) const;
+	bool IsMovingCamera(FViewport* InViewport) const;
 	void AddSequencer(TWeakPtr<FSequencer> InSequencer) { Sequencers.AddUnique(InSequencer); }
 	void RemoveSequencer(TWeakPtr<FSequencer> InSequencer) { Sequencers.Remove(InSequencer); }
 	USequencerSettings* GetSequencerSettings() const;
diff --git a/Engine/Source/Editor/Sequencer/Public/SequencerSettings.h b/Engine/Source/Editor/Sequencer/Public/SequencerSettings.h
--- a/Engine/Source/Editor/Sequencer/Public/SequencerSettings.h
+++ b/Engine/Source/Editor/Sequencer/Public/SequencerSettings.h
@@ -214,9 +214,9 @@ public:
 	void SetShowSelectedNodesOnly(bool Visible);
 	FOnShowSelectedNodesOnlyChanged& GetOnShowSelectedNodesOnlyChanged() { return OnShowSelectedNodesOnlyChangedEvent; }
 
-	/** Gets whether left mouse drag does marquee select instead of camera orbit */
+	/** Gets whether left mouse drag does marquee select instead of camera orbit and ALT always moves the camera */
 	bool GetLeftMouseDragDoesMarquee() const;
-	/** Sets whether left mouse drag does marquee select instead of camera orbit */
+	/** Sets whether left mouse drag does marquee select instead of camera orbit and ALT always moves the camera */
 	void SetLeftMouseDragDoesMarque(bool bDoMarque);
 
 	/** Get zoom in/out position (mouse position or current time). */
diff --git a/Engine/Source/Editor/UnrealEd/Private/EditorViewportClient.cpp b/Engine/Source/Editor/UnrealEd/Private/EditorViewportClient.cpp
--- a/Engine/Source/Editor/UnrealEd/Private/EditorViewportClient.cpp
+++ b/Engine/Source/Editor/UnrealEd/Private/EditorViewportClient.cpp
@@ -2265,48 +2265,7 @@ void FEditorViewportClient::UpdateMouseDelta()
 
 			if( !bInputHandledByGizmos )
 			{
-				if ( ShouldOrbitCamera() )
-				{
-					bool bHasMovement = !DragDelta.IsNearlyZero();
-
-					BeginCameraMovement(bHasMovement);
-
-					FVector TempDrag;
-					FRotator TempRot;
-					InputAxisForOrbit( Viewport, DragDelta, TempDrag, TempRot );
-				}
-				else
-				{
-					// Disable orbit camera
-					const bool bEnable=false;
-					ToggleOrbitCamera(bEnable);
-
-					if ( ShouldPanOrDollyCamera() )
-					{
-						bool bHasMovement = !Drag.IsNearlyZero() || !Rot.IsNearlyZero();
-
-						BeginCameraMovement(bHasMovement);
-
-						if( !IsOrtho())
-						{
-							const float CameraSpeed = GetCameraSpeed();
-							Drag *= CameraSpeed;
-						}
-						MoveViewportCamera( Drag, Rot );
-
-						if ( IsPerspective() && LeftMouseButtonDown && !MiddleMouseButtonDown && !RightMouseButtonDown )
-						{
-							FEditorViewportStats::Using(FEditorViewportStats::CAT_PERSPECTIVE_MOUSE_DOLLY);
-						}
-						else
-						{
-							if ( !Drag.IsZero() )
-							{
-								FEditorViewportStats::Using(IsPerspective() ? FEditorViewportStats::CAT_PERSPECTIVE_MOUSE_PAN : FEditorViewportStats::CAT_ORTHOGRAPHIC_MOUSE_PAN);
-							}
-						}
-					}
-				}
+				PeformDefaultCameraMovement(Drag, Rot, Scale);
 			}
 
 			// Clean up
@@ -2317,6 +2276,56 @@ void FEditorViewportClient::UpdateMouseDelta()
 	}
 }
 
+void FEditorViewportClient::PeformDefaultCameraMovement(FVector& Drag, FRotator& Rot, FVector& Scale)
+{
+	FVector DragDelta = MouseDeltaTracker->GetDelta();
+	if (ShouldOrbitCamera())
+	{
+		bool bHasMovement = !DragDelta.IsNearlyZero();
+
+		BeginCameraMovement(bHasMovement);
+
+		FVector TempDrag;
+		FRotator TempRot;
+		InputAxisForOrbit(Viewport, DragDelta, TempDrag, TempRot);
+	}
+	else
+	{
+		// Disable orbit camera
+		const bool bEnable = false;
+		ToggleOrbitCamera(bEnable);
+
+		if (ShouldPanOrDollyCamera())
+		{
+			bool bHasMovement = !Drag.IsNearlyZero() || !Rot.IsNearlyZero();
+
+			BeginCameraMovement(bHasMovement);
+
+			if (!IsOrtho())
+			{
+				const float CameraSpeed = GetCameraSpeed();
+				Drag *= CameraSpeed;
+			}
+			MoveViewportCamera(Drag, Rot);
+
+			const bool LeftMouseButtonDown = Viewport->KeyState(EKeys::LeftMouseButton);
+			const bool MiddleMouseButtonDown = Viewport->KeyState(EKeys::MiddleMouseButton);
+			const bool RightMouseButtonDown = Viewport->KeyState(EKeys::RightMouseButton);
+
+			if (IsPerspective() && LeftMouseButtonDown && !MiddleMouseButtonDown && !RightMouseButtonDown)
+			{
+				FEditorViewportStats::Using(FEditorViewportStats::CAT_PERSPECTIVE_MOUSE_DOLLY);
+			}
+			else
+			{
+				if (!Drag.IsZero())
+				{
+					FEditorViewportStats::Using(IsPerspective() ? FEditorViewportStats::CAT_PERSPECTIVE_MOUSE_PAN : FEditorViewportStats::CAT_ORTHOGRAPHIC_MOUSE_PAN);
+				}
+			}
+		}
+	}
+}
 
 static bool IsOrbitRotationMode( FViewport* Viewport )
 {
diff --git a/Engine/Source/Editor/UnrealEd/Public/EditorViewportClient.h b/Engine/Source/Editor/UnrealEd/Public/EditorViewportClient.h
--- a/Engine/Source/Editor/UnrealEd/Public/EditorViewportClient.h
+++ b/Engine/Source/Editor/UnrealEd/Public/EditorViewportClient.h
@@ -487,6 +487,8 @@ public:
 			ViewTransform.SetLocation(OrbitMatrix.GetOrigin());
 		}
 	}
+	/** Perform default camera movement that would normally happen* This can be used by systems to override how the camera is triggered*/
+	void PeformDefaultCameraMovement(FVector& Drag, FRotator& Rot, FVector& Scale);
 
 	/** Sets ortho zoom amount */
 	void SetOrthoZoom( float InOrthoZoom ) 
