Cherry pick fix for SActorEditorContext when the widget is closed.

diff --git a/Engine/Source/Editor/UnrealEd/Private/SActorEditorContext.cpp b/Engine/Source/Editor/UnrealEd/Private/SActorEditorContext.cpp
--- a/Engine/Source/Editor/UnrealEd/Private/SActorEditorContext.cpp
+++ b/Engine/Source/Editor/UnrealEd/Private/SActorEditorContext.cpp
@@ -26,12 +26,12 @@ void SActorEditorContext::Construct(const FArguments& InArgs)
 SActorEditorContext::~SActorEditorContext()
 {
 	FEditorDelegates::MapChange.RemoveAll(this);
-	if (GEditor)
+	if (GEditor && UObjectInitialized())
 	{
 		GEditor->GetEditorWorldContext().RemoveRef(World);
-		if (UObjectInitialized())
+		if (UActorEditorContextSubsystem* ActorEditorContextSubsystem = UActorEditorContextSubsystem::Get())
 		{
-			UActorEditorContextSubsystem::Get()->OnActorEditorContextSubsystemChanged().RemoveAll(this);
+			ActorEditorContextSubsystem->OnActorEditorContextSubsystemChanged().RemoveAll(this);
 		}
 	}
 }
