Fixed an issue where no Place Actors entry would be created for any actor factory that relied on the default implementation of CanCreateActorFrom.

Null asset data is passed when the PlacementModule generates a list of placeable actor classes and the default implementation of CanCreateActorFrom would return false in that case every time.

diff --git a/Engine/Source/Editor/UnrealEd/Private/Factories/ActorFactory.cpp b/Engine/Source/Editor/UnrealEd/Private/Factories/ActorFactory.cpp
--- a/Engine/Source/Editor/UnrealEd/Private/Factories/ActorFactory.cpp
+++ b/Engine/Source/Editor/UnrealEd/Private/Factories/ActorFactory.cpp
@@ -227,7 +227,7 @@ bool UActorFactory::CanCreateActorFrom( const FAssetData& AssetData, FText& OutE
 
 	if (!AssetData.IsValid())
 	{
-		return false;
+		return true;
 	}
 
 	UObject* DefaultActor = GetDefaultActor(AssetData);
@@ -235,8 +235,7 @@ bool UActorFactory::CanCreateActorFrom( const FAssetData& AssetData, FText& OutE
 	FSoftObjectPath DefaultActorClassPath(DefaultActor->GetClass());
 
 	// By Default we assume the factory can't work with existing asset data
-	return !AssetData.IsValid() || 
-		AssetData.GetSoftObjectPath() == DefaultActorPath || 
+	return AssetData.GetSoftObjectPath() == DefaultActorPath || 
 		AssetData.GetSoftObjectPath() == DefaultActorClassPath;
 }
 
