Interchange: Make sure alternate skinning and custom LODs import play well when interchange fbx is disabled

#rnx

diff --git a/Engine/Source/Editor/UnrealEd/Private/SkinWeightsUtilities.cpp b/Engine/Source/Editor/UnrealEd/Private/SkinWeightsUtilities.cpp
--- a/Engine/Source/Editor/UnrealEd/Private/SkinWeightsUtilities.cpp
+++ b/Engine/Source/Editor/UnrealEd/Private/SkinWeightsUtilities.cpp
@@ -73,6 +73,7 @@ bool FSkinWeightsUtilities::ImportAlternateSkinWeight(USkeletalMesh* SkeletalMes
 #if WITH_EDITOR
 	bUseInterchangeFramework = GetDefault<UEditorExperimentalSettings>()->bEnableInterchangeFramework;
 #endif
+	const FString FileExtension = FPaths::GetExtension(AbsoluteFilePath);
 
 	FAssetToolsModule& AssetToolsModule = FModuleManager::LoadModuleChecked<FAssetToolsModule>("AssetTools");
 	FString ImportAssetPath = TEXT("/Engine/TempEditor/SkeletalMeshTool");
@@ -104,7 +105,9 @@ bool FSkinWeightsUtilities::ImportAlternateSkinWeight(USkeletalMesh* SkeletalMes
 	UObject* ImportedObject = nullptr;
 	UnFbx::FBXImportOptions ImportOptions;
 
-	if (bUseInterchangeFramework)
+	//Only use interchange if the base skeletal mesh was imported with interchange
+	const UInterchangeAssetImportData* SelectedInterchangeAssetImportData = Cast<UInterchangeAssetImportData>(SkeletalMesh->GetAssetImportData());
+	if (bUseInterchangeFramework && SelectedInterchangeAssetImportData)
 	{
 		UE::Interchange::FScopedSourceData ScopedSourceData(AbsoluteFilePath);
 		const UInterchangeProjectSettings* InterchangeProjectSettings = GetDefault<UInterchangeProjectSettings>();
@@ -128,7 +131,7 @@ bool FSkinWeightsUtilities::ImportAlternateSkinWeight(USkeletalMesh* SkeletalMes
 			ImportedObject = ImportedSkeletalMesh;
 		}
 	}
-	else
+	else if(FileExtension.Equals(TEXT("fbx"), ESearchCase::IgnoreCase))
 	{
 		//Import the alternate fbx into a temporary skeletal mesh using the same import options
 		UFbxFactory* FbxFactory = NewObject<UFbxFactory>(UFbxFactory::StaticClass());
@@ -212,7 +215,7 @@ bool FSkinWeightsUtilities::ImportAlternateSkinWeight(USkeletalMesh* SkeletalMes
 	USkeletalMesh* TmpSkeletalMesh = Cast<USkeletalMesh>(ImportedObject);
 	if (TmpSkeletalMesh == nullptr || TmpSkeletalMesh->GetSkeleton() == nullptr)
 	{
-		UE_LOG(LogSkinWeightsUtilities, Error, TEXT("Failed to import Skin Weight Profile from provided FBX file (%s)."), *Path);
+		UE_LOG(LogSkinWeightsUtilities, Error, TEXT("Failed to import Skin Weight Profile from provided file (%s)."), *Path);
 		DeletePathAssets();
 		return false;
 	}
@@ -264,7 +267,7 @@ bool FSkinWeightsUtilities::ImportAlternateSkinWeight(USkeletalMesh* SkeletalMes
 				if (!bResult)
 				{
 					// Remove invalid profile data due to failed import
-					if (!bIsReimport)
+					if (!bIsReimportLocal)
 					{
 						TargetLODModel.SkinWeightProfiles.Remove(ProfileName);
 					}
@@ -276,7 +279,7 @@ bool FSkinWeightsUtilities::ImportAlternateSkinWeight(USkeletalMesh* SkeletalMes
 				}
 
 				// Only add if it is an initial import and it was successful 
-				if (!bIsReimport && bResult)
+				if (!bIsReimportLocal && bResult)
 				{
 					FSkinWeightProfileInfo SkeletalMeshProfile;
 					SkeletalMeshProfile.DefaultProfile = (SkeletalMesh->GetNumSkinWeightProfiles() == 0);
