Remove ensure that dates back 10 years, embrace renaming graphs out of the way to maintain caller invariants. Don't mark graphs as garbage (as that is still causing pointers to be null'd), this fixes graphpanel crash after removing an interface function via live coding

diff --git a/Engine/Source/Editor/UnrealEd/Private/Kismet2/BlueprintEditorUtils.cpp b/Engine/Source/Editor/UnrealEd/Private/Kismet2/BlueprintEditorUtils.cpp
--- a/Engine/Source/Editor/UnrealEd/Private/Kismet2/BlueprintEditorUtils.cpp
+++ b/Engine/Source/Editor/UnrealEd/Private/Kismet2/BlueprintEditorUtils.cpp
@@ -2371,9 +2371,11 @@ UEdGraph* FBlueprintEditorUtils::CreateNewGraph(UObject* ParentScope, const FNam
 	{
 		if (UObject* ExistingObject = FindObject<UObject>(ParentScope, *(GraphName.ToString())))
 		{
-			if (!ensureMsgf(!ExistingObject->IsA<UEdGraph>(), TEXT("Graph %s already exists: %s"), *GraphName.ToString(), *ExistingObject->GetFullName()))
+			if (ExistingObject->IsA<UEdGraph>())
 			{
-				// Rename the old graph out of the way; but we have already failed at this point
+				// Rename the old graph out of the way - this may confuse the user somewhat - and even
+				// break their logic. But name collisions are not avoidable e.g. someone can add
+				// a function to an interface that conflicts with something in a class hierarchy
 				ExistingObject->Rename(nullptr, ExistingObject->GetOuter(), REN_DoNotDirty | REN_ForceNoResetLoaders);
 			}
 			else if (ExistingObject->IsA<UObjectRedirector>())
@@ -7433,15 +7435,19 @@ static void ConformInterfaceByName(UBlueprint* Blueprint, FBPInterfaceDescriptio
 					CurrentGraph->GetSchema()->HandleGraphBeingDeleted(*CurrentGraph);
 
 					// rename to free up the graph's name.. which may be needed by an inherited function
+					// alternatively we could move this into the functions list?
 					CurrentGraph->Rename(
 						nullptr,
-						GetTransientPackage(),
+						CurrentGraph->GetOuter(),
 						(Blueprint->bIsRegeneratingOnLoad ? REN_ForceNoResetLoaders : 0) | REN_DoNotDirty | REN_DontCreateRedirectors);
-					// orphan the graph, we need to mark as garbage so that it's tab cleans up (FTabPayload_UObject)
 					// removing from root, standalone, and public is defensive to make sure it is not saved:
 					CurrentGraph->ClearFlags(RF_Standalone | RF_Public);
 					CurrentGraph->RemoveFromRoot();
-					CurrentGraph->MarkAsGarbage();
+					// MarkAsGarbage could be used here, which would nicely trigger tab manager cleanup, but atm
+					// use of MarkAsGarbage is causing SGraphPanel to have reference's nulled out (treated as weak)
+					// by the GC. For now, I'm just going to flag this as no longer editable. This isn't a bad
+					// out come as if the user has anything in the graph they might be able to copy it out
+					CurrentGraph->bEditable = false;
 				}
 				
 				CurrentInterfaceDesc.Graphs.RemoveAt(GraphIndex, 1);
