Make shader flush log only happen every 10 seconds.

diff --git a/Engine/Source/Editor/UnrealEd/Private/Commandlets/CookCommandlet.cpp b/Engine/Source/Editor/UnrealEd/Private/Commandlets/CookCommandlet.cpp
--- a/Engine/Source/Editor/UnrealEd/Private/Commandlets/CookCommandlet.cpp
+++ b/Engine/Source/Editor/UnrealEd/Private/Commandlets/CookCommandlet.cpp
@@ -942,7 +942,7 @@ void UCookCommandlet::ConditionalCollectGarbage(uint32 TickResults, UCookOnTheFl
 		GCReason = TEXT("Exceeded Max Memory");
 
 		int32 JobsToLogAt = GShaderCompilingManager->GetNumRemainingJobs();
-
+		double NextFlushMsgSeconds = FPlatformTime::Seconds();
 		UE_SCOPED_COOKTIMER(CookByTheBook_ShaderJobFlush);
 		UE_LOG(LogCookCommandlet, Display, TEXT("Detected max mem exceeded - forcing shader compilation flush"));
 		while (true)
@@ -956,7 +956,12 @@ void UCookCommandlet::ConditionalCollectGarbage(uint32 TickResults, UCookOnTheFl
 
 			if (NumRemainingJobs < JobsToLogAt)
 			{
-				UE_LOG(LogCookCommandlet, Display, TEXT("Flushing shader jobs, remaining jobs %d"), NumRemainingJobs);
+				double Now = FPlatformTime::Seconds();
+				if (NextFlushMsgSeconds <= Now)
+				{
+					UE_LOG(LogCookCommandlet, Display, TEXT("Flushing shader jobs, remaining jobs %d"), NumRemainingJobs);
+					NextFlushMsgSeconds = Now + 10;
+				}
 			}
 
 			GShaderCompilingManager->ProcessAsyncResults(false, false);
@@ -1028,4 +1033,4 @@ void UCookCommandlet::ConditionalCollectGarbage(uint32 TickResults, UCookOnTheFl
 	{
 		COTFS.EvaluateGarbageCollectionResults(NumObjectsBeforeGC, MemStatsBeforeGC);
 	} 
-}
\ No newline at end of file
+}
