Update diffassetbulkdata to handle all chunk types.

diff --git a/Engine/Source/Editor/UnrealEd/Private/Commandlets/DiffAssetBulkDataCommandlet.cpp b/Engine/Source/Editor/UnrealEd/Private/Commandlets/DiffAssetBulkDataCommandlet.cpp
--- a/Engine/Source/Editor/UnrealEd/Private/Commandlets/DiffAssetBulkDataCommandlet.cpp
+++ b/Engine/Source/Editor/UnrealEd/Private/Commandlets/DiffAssetBulkDataCommandlet.cpp
@@ -155,13 +155,12 @@ int32 UDiffAssetBulkDataCommandlet::Main(const FString& FullCommandLine)
 		}
 	}
 
-	// Now we need to see what changed. We current care about BulkData (or DerivedData when it gets turned on).
-	// OptionalBulkData and MemoryMappedBulkData not yet supported...
+	// Now we need to see what changed.
 	//
 	// This whole thing assumes that the index parameter of CreateIoChunkId is always 0. This is likely not going
 	// to be true with FDerivedData, once that gets turned on, but should be easy to update when the time comes.
 	//
-	TArray<FName> ChangedBulkData, NewBulkData, DeletedBulkData;
+	TArray<FName> ChangedChunksByType[(uint32)EIoChunkType::MAX], NewChunksByType[(uint32)EIoChunkType::MAX], DeletedChunksByType[(uint32)EIoChunkType::MAX];
 	for (const FName& SharedPackage : SharedPackages)
 	{
 		const FAssetPackageData* BasePackage = BasePackages[SharedPackage];
@@ -169,41 +168,42 @@ int32 UDiffAssetBulkDataCommandlet::Main(const FString& FullCommandLine)
 
 		for (const TPair<FIoChunkId, FIoHash>& ChunkHashPair : BasePackage->ChunkHashes)
 		{
-			if (ChunkHashPair.Key.GetChunkType() != EIoChunkType::BulkData)
-			{
-				continue;
-			}		
-
 			const FIoHash* CurrentHash = CurrentPackage->ChunkHashes.Find(ChunkHashPair.Key);
 			if (CurrentHash == nullptr)
 			{
-				check(DeletedBulkData.Contains(SharedPackage) == false); // Because only 0 chunk index
-				DeletedBulkData.Add(SharedPackage);
+				TArray<FName>& Deleted = DeletedChunksByType[(uint32)ChunkHashPair.Key.GetChunkType()];
+				check(Deleted.Contains(SharedPackage) == false); // Because only 0 chunk index
+				Deleted.Add(SharedPackage);
 				continue;
 			}
 
 			if (*CurrentHash != ChunkHashPair.Value)
 			{
-				check(ChangedBulkData.Contains(SharedPackage) == false); // Because only 0 chunk index
-				ChangedBulkData.Add(SharedPackage);
+				TArray<FName>& Changed = ChangedChunksByType[(uint32)ChunkHashPair.Key.GetChunkType()];
+				check(Changed.Contains(SharedPackage) == false); // Because only 0 chunk index
+				Changed.Add(SharedPackage);
 			}
 		}
 
 		for (const TPair<FIoChunkId, FIoHash>& ChunkHashPair : CurrentPackage->ChunkHashes)
 		{
-			if (ChunkHashPair.Key.GetChunkType() != EIoChunkType::BulkData)
-			{
-				continue;
-			}
-
 			if (BasePackage->ChunkHashes.Contains(ChunkHashPair.Key) == false)
 			{
-				check(NewBulkData.Contains(SharedPackage) == false); // Because only 0 chunk index
-				NewBulkData.Add(SharedPackage);
+				TArray<FName>& New = NewChunksByType[(uint32)ChunkHashPair.Key.GetChunkType()];
+				check(New.Contains(SharedPackage) == false); // Because only 0 chunk index
+				New.Add(SharedPackage);
 			}
 		}
 	}
 
+	// Get a unique list of changed packages.
+	TSet<FName> ChangedPackages;
+	for (uint32 ChunkTypeIndex = 0; ChunkTypeIndex < (uint32)EIoChunkType::MAX; ChunkTypeIndex++)
+	{
+		ChangedPackages.Append(ChangedChunksByType[ChunkTypeIndex]);
+	}
+
+
 	//
 	// We know what bulk datas *packages* changed. Try and see if any of the assets in the package have
 	// diff blame tags for us to determine cause. _usually_ there's one asset per package, but it's definitely
@@ -223,9 +223,9 @@ int32 UDiffAssetBulkDataCommandlet::Main(const FString& FullCommandLine)
 	TMap<FName /* TagName */, TMap<FTopLevelAssetPath /* AssetClass */, TArray<FDiffResult>>> Results;
 	TMap<FTopLevelAssetPath /* AssetClass */, TArray<FName /* PackageName */>> NoTagPackagesByAssumedClass;
 	TArray<FName /* PackageName */> PackagesWithUnassignableDiffsAndUntaggedAssets;
-	TArray<FName /* PackageName */> PackagesWithUnassignableDiffs;
+	TMap<FTopLevelAssetPath /* AssetClass */, TArray<FName /* PackageName */>> PackagesWithUnassignableDiffsByAssumedClass;
 	
-	for (const FName& ChangedPackageName : ChangedBulkData)
+	for (const FName& ChangedPackageName : ChangedPackages)
 	{
 		TConstArrayView<FAssetData const*> BaseAssetDatas = BaseState.GetAssetsByPackageName(ChangedPackageName);		
 		TConstArrayView<FAssetData const*> CurrentAssetDatas = CurrentState.GetAssetsByPackageName(ChangedPackageName);
@@ -359,23 +359,59 @@ int32 UDiffAssetBulkDataCommandlet::Main(const FString& FullCommandLine)
 			}
 			else
 			{
-				PackagesWithUnassignableDiffs.Add(ChangedPackageName);
+				FAssetData const* RepresentativeAsset = UE::AssetRegistry::GetMostImportantAsset(CurrentAssetDatas, true);
+				if (RepresentativeAsset == nullptr)
+				{
+					PackagesWithUnassignableDiffsByAssumedClass.FindOrAdd(FTopLevelAssetPath()).Add(ChangedPackageName);
+				}
+				else
+				{
+					PackagesWithUnassignableDiffsByAssumedClass.FindOrAdd(RepresentativeAsset->AssetClassPath).Add(ChangedPackageName);
+				}
 			}
 		}
 	}
 	
-	UE_LOG(LogDiffAssetBulk, Display, TEXT("Packages with Bulk Data: %d new, %d deleted, %d changed"), NewBulkData.Num(), DeletedBulkData.Num(), ChangedBulkData.Num());
-	if (ChangedBulkData.Num() == 0)
+	int32 TotalNewChunks = 0, TotalChangedChunks = 0, TotalDeletedChunks = 0;
+	UE_LOG(LogDiffAssetBulk, Display, TEXT("Modifications By IoStore Chunk:"));
+	UE_LOG(LogDiffAssetBulk, Display, TEXT(""));
+	UE_LOG(LogDiffAssetBulk, Display, TEXT("    ChunkType                   New    Deleted    Changed"));
+	for (uint32 ChunkTypeIndex = 0; ChunkTypeIndex < (uint8)EIoChunkType::MAX; ChunkTypeIndex++)
+	{
+		EIoChunkType ChunkType = (EIoChunkType)ChunkTypeIndex;
+		const TArray<FName>& NewChunksForType = NewChunksByType[ChunkTypeIndex];
+		const TArray<FName>& DeletedChunksForType = DeletedChunksByType[ChunkTypeIndex];
+		const TArray<FName>& ChangedChunksForType = ChangedChunksByType[ChunkTypeIndex];
+
+		TotalNewChunks += NewChunksForType.Num();
+		TotalChangedChunks += ChangedChunksForType.Num();
+		TotalDeletedChunks += DeletedChunksForType.Num();
+		UE_LOG(LogDiffAssetBulk, Display, TEXT("    %-20s %10d %10d %10d"), *LexToString(ChunkType), NewChunksForType.Num(), DeletedChunksForType.Num(), ChangedChunksForType.Num());
+	}
+	UE_LOG(LogDiffAssetBulk, Display, TEXT("    =================================================="));
+
+	UE_LOG(LogDiffAssetBulk, Display, TEXT("    %-20s %10d %10d %10d"), TEXT("Total"), TotalNewChunks, TotalDeletedChunks, TotalChangedChunks);
+
+	UE_LOG(LogDiffAssetBulk, Display, TEXT(""));
+
+	UE_LOG(LogDiffAssetBulk, Display, TEXT("    Packages Added:     %8d"), NewPackages.Num());
+	UE_LOG(LogDiffAssetBulk, Display, TEXT("    Packages Deleted:   %8d"), DeletedPackages.Num());
+	UE_LOG(LogDiffAssetBulk, Display, TEXT("    Packages Changed:   %8d"), ChangedPackages.Num());
+	UE_LOG(LogDiffAssetBulk, Display, TEXT("    Packages Unmodified:%8d"), SharedPackages.Num() - ChangedPackages.Num());
+
+	if (ChangedPackages.Num() == 0)
 	{
 		return 0;
 	}
 
+	UE_LOG(LogDiffAssetBulk, Display, TEXT(""));
+
 	TArray<FName>& CantDetermineAssetClassPackages = NoTagPackagesByAssumedClass.FindOrAdd(FTopLevelAssetPath());
 
 	UE_LOG(LogDiffAssetBulk, Display, TEXT("Changed package breakdown:"));
 	UE_LOG(LogDiffAssetBulk, Display, TEXT("    No blame information available:"));
 	UE_LOG(LogDiffAssetBulk, Display, TEXT("        Can't determine asset class   : %d"), CantDetermineAssetClassPackages.Num());
-	for (TPair<FTopLevelAssetPath, TArray<FName>> ClassPackages : NoTagPackagesByAssumedClass)
+	for (TPair<FTopLevelAssetPath, TArray<FName>>& ClassPackages : NoTagPackagesByAssumedClass)
 	{
 		if (ClassPackages.Key == FTopLevelAssetPath())
 		{
@@ -395,14 +431,24 @@ int32 UDiffAssetBulkDataCommandlet::Main(const FString& FullCommandLine)
 		}
 	}
 
-	if (PackagesWithUnassignableDiffs.Num())
+	if (PackagesWithUnassignableDiffsByAssumedClass.Num())
 	{
-		Algo::Sort(PackagesWithUnassignableDiffs, FNameLexicalLess());
+		int32 TotalUnassignablePackages = 0;
+		for (TPair<FTopLevelAssetPath, TArray<FName>>& ClassPackages : PackagesWithUnassignableDiffsByAssumedClass)
+		{
+			TotalUnassignablePackages += ClassPackages.Value.Num();
+		}
+		
 
-		UE_LOG(LogDiffAssetBulk, Warning, TEXT("    Can't determine blame:        : %-7d // Assets had blame tags but all matched - check determinism!"), PackagesWithUnassignableDiffs.Num());
-		for (const FName& PackageName : PackagesWithUnassignableDiffs)
+		UE_LOG(LogDiffAssetBulk, Warning, TEXT("    Can't determine blame:        : %-7d // Assets had blame tags but all matched - check determinism!"), TotalUnassignablePackages);
+		for (TPair<FTopLevelAssetPath, TArray<FName>>& ClassPackages : PackagesWithUnassignableDiffsByAssumedClass)
 		{
-			UE_LOG(LogDiffAssetBulk, Warning, TEXT("        %s"), *PackageName.ToString());
+			UE_LOG(LogDiffAssetBulk, Warning, TEXT("        %s : %d"), *ClassPackages.Key.ToString(), ClassPackages.Value.Num());
+			Algo::Sort(ClassPackages.Value, FNameLexicalLess());
+			for (const FName& PackageName : ClassPackages.Value)
+			{
+				UE_LOG(LogDiffAssetBulk, Warning, TEXT("            %s"), *PackageName.ToString());
+			}
 		}
 	}
 
