Fixed display of animation compression progress

diff --git a/Engine/Source/Editor/UnrealEd/Private/Animation/AsyncAnimCompressionUI.cpp b/Engine/Source/Editor/UnrealEd/Private/Animation/AsyncAnimCompressionUI.cpp
--- a/Engine/Source/Editor/UnrealEd/Private/Animation/AsyncAnimCompressionUI.cpp
+++ b/Engine/Source/Editor/UnrealEd/Private/Animation/AsyncAnimCompressionUI.cpp
@@ -24,7 +24,7 @@ protected:
 		{
 			FFormatNamedArguments Args;
 			Args.Add(TEXT("AnimsToCompress"), FText::AsNumber(RemainingJobs));
-			UpdateProgressMessage(FText::Format(NSLOCTEXT("AsyncAnimCompression", "AnimCompressionInProgressFormat", "Compressing Animations ({AnimsToCompress})"), Args));
+			UpdateProgressMessage(FText::Format(NSLOCTEXT("AsyncAnimCompression", "AnimPreparationInProgressFormat", "Preparing Animations ({AnimsToCompress})"), Args));
 
 		}
 
