Prevent duplicate addition of implemented interfaces in parent classes

Dont just check this classes interfaces, also check its parents

diff --git a/Engine/Source/Editor/UnrealEd/Private/Kismet2/BlueprintEditorUtils.cpp b/Engine/Source/Editor/UnrealEd/Private/Kismet2/BlueprintEditorUtils.cpp
--- a/Engine/Source/Editor/UnrealEd/Private/Kismet2/BlueprintEditorUtils.cpp
+++ b/Engine/Source/Editor/UnrealEd/Private/Kismet2/BlueprintEditorUtils.cpp
@@ -8880,10 +8880,11 @@ TSharedRef<SWidget> FBlueprintEditorUtils::ConstructBlueprintInterfaceClassPicke
 		}
 
 		// Do not allow adding interfaces that are already added to the Blueprint
-		for(TArray<FBPInterfaceDescription>::TConstIterator it(Blueprint->ImplementedInterfaces); it; ++it)
+		TArray<UClass*> InterfaceClasses;
+		FindImplementedInterfaces(Blueprint, true, InterfaceClasses);
+		for(UClass* InterfaceClass : InterfaceClasses)
 		{
-			const FBPInterfaceDescription& CurrentInterface = *it;
-			Filter->DisallowedClasses.Add(CurrentInterface.Interface);
+			Filter->DisallowedClasses.Add(InterfaceClass);
 		}
 
 		// Include a class viewer filter for imported namespaces if the class picker is being hosted in an editor context
