Use FPackageName::IsTempPackage

diff --git a/Engine/Source/Editor/UnrealEd/Private/LevelEditorViewport.cpp b/Engine/Source/Editor/UnrealEd/Private/LevelEditorViewport.cpp
--- a/Engine/Source/Editor/UnrealEd/Private/LevelEditorViewport.cpp
+++ b/Engine/Source/Editor/UnrealEd/Private/LevelEditorViewport.cpp
@@ -538,7 +538,7 @@ UObject* FLevelEditorViewportClient::GetOrCreateMaterialFromTexture(UTexture* Un
 		if (UWorld* EditorWorld = GEditor->GetEditorWorldContext().World())
 		{
 			const FString EditorWorldPackageName = EditorWorld->GetPackage()->GetName();
-			if (!EditorWorldPackageName.StartsWith(TEXT("/Temp/")))
+			if (!FPackageName::IsTempPackage(EditorWorldPackageName))
 			{
 				NewPackageFolder = FPackageName::GetLongPackagePath(EditorWorldPackageName);
 				if (!AssetTools.GetWritableFolderPermissionList()->PassesStartsWithFilter(NewPackageFolder))
