Fix PollPendingCookedPlatformDatas to also tick FTickableCookObject::TickObjects, since it can be called in tight loops that do not tick TickObjects normally.
#rnx

diff --git a/Engine/Source/Editor/UnrealEd/Private/CookOnTheFlyServer.cpp b/Engine/Source/Editor/UnrealEd/Private/CookOnTheFlyServer.cpp
--- a/Engine/Source/Editor/UnrealEd/Private/CookOnTheFlyServer.cpp
+++ b/Engine/Source/Editor/UnrealEd/Private/CookOnTheFlyServer.cpp
@@ -197,8 +197,6 @@ static FAutoConsoleVariableRef CVarCookDisplayWarnBusyTime(
 ///////////////////////////////////////////////////////////////
 UCookOnTheFlyServer* UCookOnTheFlyServer::ActiveCOTFS = nullptr;
 
-constexpr float TickCookableObjectsFrameTime = .100f;
-
 namespace UE
 {
 namespace Cook
@@ -3830,7 +3828,7 @@ void UCookOnTheFlyServer::ReleaseCookedPlatformData(UE::Cook::FPackageData& Pack
 
 void UCookOnTheFlyServer::TickCancels()
 {
-	PackageDatas->PollPendingCookedPlatformDatas(false);
+	PackageDatas->PollPendingCookedPlatformDatas(false, LastCookableObjectTickTime);
 }
 
 bool UCookOnTheFlyServer::LoadPackageForCooking(UE::Cook::FPackageData& PackageData, UPackage*& OutPackage,
@@ -4109,7 +4107,7 @@ void UCookOnTheFlyServer::PumpSaves(UE::Cook::FTickStackData& StackData, uint32
 
 		// Release any completed pending CookedPlatformDatas, so that slots in the per-class limits on calls to BeginCacheForCookedPlatformData are freed up for new objects to use
 		bool bForce = IsCookOnTheFlyMode() && !IsRealtimeMode();
-		PackageDatas->PollPendingCookedPlatformDatas(bForce);
+		PackageDatas->PollPendingCookedPlatformDatas(bForce, LastCookableObjectTickTime);
 
 		// If BeginCacheCookPlatformData is not ready then postpone the package, exit, or wait for it as appropriate
 		EPollStatus PrepareSaveStatus = PrepareSave(PackageData, StackData.Timer, false /* bPrecaching */);
@@ -4161,7 +4159,7 @@ void UCookOnTheFlyServer::PumpSaves(UE::Cook::FTickStackData& StackData, uint32
 					// sleep for a bit
 					FPlatformProcess::Sleep(0.0f);
 					// Poll the results again and check whether we are now done
-					PackageDatas->PollPendingCookedPlatformDatas(true);
+					PackageDatas->PollPendingCookedPlatformDatas(true, LastCookableObjectTickTime);
 					PrepareSaveStatus = PrepareSave(PackageData, StackData.Timer, false /* bPrecaching */);
 				} while (!StackData.Timer.IsTimeUp() && PrepareSaveStatus == EPollStatus::Incomplete);
 			}
diff --git a/Engine/Source/Editor/UnrealEd/Private/Cooker/CookPackageData.cpp b/Engine/Source/Editor/UnrealEd/Private/Cooker/CookPackageData.cpp
--- a/Engine/Source/Editor/UnrealEd/Private/Cooker/CookPackageData.cpp
+++ b/Engine/Source/Editor/UnrealEd/Private/Cooker/CookPackageData.cpp
@@ -2619,18 +2619,18 @@ TArray<FPendingCookedPlatformData>& FPackageDatas::GetPendingCookedPlatformDatas
 	return PendingCookedPlatformDatas;
 }
 
-void FPackageDatas::PollPendingCookedPlatformDatas(bool bForce)
+void FPackageDatas::PollPendingCookedPlatformDatas(bool bForce, double& LastCookableObjectTickTime)
 {
 	if (PendingCookedPlatformDatas.Num() == 0)
 	{
 		return;
 	}
 
+	double CurrentTime = FPlatformTime::Seconds();
 	if (!bForce)
 	{
 		// ProcessAsyncResults and IsCachedCookedPlatformDataLoaded can be expensive to call
 		// Cap the frequency at which we call them.
-		double CurrentTime = FPlatformTime::Seconds();
 		if (CurrentTime < LastPollAsyncTime + GPollAsyncPeriod)
 		{
 			return;
@@ -2641,6 +2641,12 @@ void FPackageDatas::PollPendingCookedPlatformDatas(bool bForce)
 	GShaderCompilingManager->ProcessAsyncResults(true /* bLimitExecutionTime */,
 		false /* bBlockOnGlobalShaderCompletion */);
 	FAssetCompilingManager::Get().ProcessAsyncTasks(true);
+	if (LastCookableObjectTickTime + TickCookableObjectsFrameTime <= CurrentTime)
+	{
+		UE_SCOPED_COOKTIMER(TickCookableObjects);
+		FTickableCookObject::TickObjects(CurrentTime - LastCookableObjectTickTime, false /* bTickComplete */);
+		LastCookableObjectTickTime = CurrentTime;
+	}
 
 	FPendingCookedPlatformData* Datas = PendingCookedPlatformDatas.GetData();
 	for (int Index = 0; Index < PendingCookedPlatformDatas.Num();)
diff --git a/Engine/Source/Editor/UnrealEd/Private/Cooker/CookPackageData.h b/Engine/Source/Editor/UnrealEd/Private/Cooker/CookPackageData.h
--- a/Engine/Source/Editor/UnrealEd/Private/Cooker/CookPackageData.h
+++ b/Engine/Source/Editor/UnrealEd/Private/Cooker/CookPackageData.h
@@ -1178,7 +1178,7 @@ public:
 	 * Iterate over all elements in PendingCookedPlatformDatas and check whether they have completed,
 	 * releasing their resources and pending count if so.
 	 */
-	void PollPendingCookedPlatformDatas(bool bForce);
+	void PollPendingCookedPlatformDatas(bool bForce, double& LastCookableObjectTickTime);
 
 	/** RangedFor methods for iterating over all FPackageData managed by this FPackageDatas */
 	TArray<FPackageData*>::RangedForIteratorType begin();
diff --git a/Engine/Source/Editor/UnrealEd/Private/Cooker/CookTypes.h b/Engine/Source/Editor/UnrealEd/Private/Cooker/CookTypes.h
--- a/Engine/Source/Editor/UnrealEd/Private/Cooker/CookTypes.h
+++ b/Engine/Source/Editor/UnrealEd/Private/Cooker/CookTypes.h
@@ -485,3 +485,5 @@ LLM_DECLARE_TAG(Cooker);
 constexpr uint32 ExpectedMaxNumPlatforms = 32;
 #define REMAPPED_PLUGINS TEXT("RemappedPlugins")
 extern float GCookProgressWarnBusyTime;
+
+constexpr float TickCookableObjectsFrameTime = .100f;
\ No newline at end of file
