Add a blueprint library function that make a UDIM virtual texture from a list of regular 2D textures

diff --git a/Engine/Source/Editor/UnrealEd/Classes/Factories/TextureFactory.h b/Engine/Source/Editor/UnrealEd/Classes/Factories/TextureFactory.h
--- a/Engine/Source/Editor/UnrealEd/Classes/Factories/TextureFactory.h
+++ b/Engine/Source/Editor/UnrealEd/Classes/Factories/TextureFactory.h
@@ -259,3 +259,21 @@ private:
 };
 
 ENUM_CLASS_FLAGS(UTextureFactory::EImageImportFlags);
+
+UCLASS()
+class UUDIMTextureFunctionLibrary : public UBlueprintFunctionLibrary
+{
+	GENERATED_UCLASS_BODY()
+
+	/**
+	* Make a UDIM virtual texture from a list of regular 2D textures
+	* @param OutputPathName			Path name of the UDIM texture (e.g. /Game/MyTexture)
+	* @param SourceTextures			List of regular 2D textures to be packed into the atlas
+	* @param BlockCoords			Coordinates of the corresponding texture in the atlas
+	* @param bKeepExistingSettings	Whether to keep existing settings if a texture with the same path name exists. Otherwise, settings will be copied from the first source texture
+	* @param bCheckOutAndSave		Whether to check out and save the UDIM texture
+	* @return UTexture2D*			Pointer to the UDIM texture or null if failed
+	*/
+	UFUNCTION(BlueprintCallable, Category = "Utilities", meta = (DispalyName = "Make UDIM Texture from Texture2Ds"))
+	static UTexture2D* MakeUDIMVirtualTextureFromTexture2Ds(FString OutputPathName, const TArray<UTexture2D*>& SourceTextures, const TArray<FIntPoint>& BlockCoords, bool bKeepExistingSettings = false, bool bCheckOutAndSave = false);
+};
diff --git a/Engine/Source/Editor/UnrealEd/Private/Factories/EditorFactories.cpp b/Engine/Source/Editor/UnrealEd/Private/Factories/EditorFactories.cpp
--- a/Engine/Source/Editor/UnrealEd/Private/Factories/EditorFactories.cpp
+++ b/Engine/Source/Editor/UnrealEd/Private/Factories/EditorFactories.cpp
@@ -7966,5 +7966,226 @@ UObject* UPreviewMeshCollectionFactory::FactoryCreateNew(UClass* Class, UObject*
 	return NewCollection;
 }
 
+UUDIMTextureFunctionLibrary::UUDIMTextureFunctionLibrary(const FObjectInitializer& ObjectInitializer)
+	: Super(ObjectInitializer)
+{}
+
+struct FScopedTextureSourceExtracts
+{
+	const TArray<UTexture2D*>& Textures;
+	TArray<const uint8*> DataArray;
+	TArray<FTextureSourceBlock> Blocks;
+	ETextureSourceFormat Format;
+
+	FScopedTextureSourceExtracts(const TArray<UTexture2D*>& InTextures, const TArray<FIntPoint>& BlockCoords)
+		: Textures(InTextures)
+		, Format(TSF_Invalid)
+	{
+		bool bMismatchedFormat = false;
+
+		DataArray.Reserve(Textures.Num());
+		Blocks.Reserve(Textures.Num());
+
+		for (int32 TextureIndex = 0; TextureIndex < Textures.Num(); ++TextureIndex)
+		{
+			UTexture2D* Texture = Textures[TextureIndex];
+			const FIntPoint& BlockXY = BlockCoords[TextureIndex];
+
+			if (Texture)
+			{
+				check(Texture->Source.GetNumLayers() == 1 && Texture->Source.GetNumBlocks() == 1);
+
+				if (Format == TSF_Invalid)
+				{
+					Format = Texture->Source.GetFormat();
+				}
+
+				if (Format != Texture->Source.GetFormat())
+				{
+					bMismatchedFormat = true;
+				}
+
+				const uint8* DataPtr = Texture->Source.LockMipReadOnly(0);
+				if (DataPtr)
+				{
+					DataArray.Add(DataPtr);
+					Blocks.AddDefaulted();
+
+					FTextureSourceBlock& Block = Blocks.Last();
+					Block.BlockX = BlockXY.X;
+					Block.BlockY = BlockXY.Y;
+					Block.SizeX = Texture->Source.GetSizeX();
+					Block.SizeY = Texture->Source.GetSizeY();
+					Block.NumSlices = 1;
+					Block.NumMips = Texture->Source.GetNumMips();
+				}
+			}
+		}
+
+		if (bMismatchedFormat)
+		{
+			UE_LOG(LogBlueprintUserMessages, Error, TEXT("Cannot create a UDIM texture from Texture2Ds with different formats."));
+			DataArray.Empty();
+			Blocks.Empty();
+		}
+	}
+
+	~FScopedTextureSourceExtracts()
+	{
+		for (int32 TextureIndex = 0; TextureIndex < Textures.Num(); ++TextureIndex)
+		{
+			UTexture2D* Texture = Textures[TextureIndex];
+			if (Texture)
+			{
+				Texture->Source.UnlockMip(0);
+			}
+		}
+	}
+
+	bool IsValid() const
+	{
+		return DataArray.Num() > 0 && Blocks.Num() > 0;
+	}
+};
+
+UTexture2D* UUDIMTextureFunctionLibrary::MakeUDIMVirtualTextureFromTexture2Ds(FString OutputPathName, const TArray<UTexture2D*>& SourceTextures, const TArray<FIntPoint>& BlockCoords, bool bKeepExistingSettings, bool bCheckOutAndSave)
+{
+	static const TConsoleVariableData<int32>* CVarVirtualTexturesEnabled = IConsoleManager::Get().FindTConsoleVariableDataInt(TEXT("r.VirtualTextures"));
+	check(CVarVirtualTexturesEnabled);
+
+	if (!CVarVirtualTexturesEnabled || !CVarVirtualTexturesEnabled->GetValueOnAnyThread())
+	{
+		UE_LOG(LogBlueprintUserMessages, Error, TEXT("Cannot create virtual textures when virtual texture support (r.VirtualTextures) is disabled."));
+		return nullptr;
+	}
+	if (SourceTextures.Num() <= 1)
+	{
+		UE_LOG(LogBlueprintUserMessages, Error, TEXT("Not enough source textures to create a UDIM texture."));
+		return nullptr;
+	}
+	if (SourceTextures.Num() != BlockCoords.Num())
+	{
+		UE_LOG(LogBlueprintUserMessages, Error, TEXT("The numbers of source textures and block coordinates don't match."));
+		return nullptr;
+	}
+#if DO_CHECK
+	for (int32 Index = 0; Index < SourceTextures.Num(); ++Index)
+	{
+		UTexture2D* Texture = SourceTextures[Index];
+		const FIntPoint& Coords = BlockCoords[Index];
+		if (Texture && (Coords.X < 0 || Coords.Y < 0))
+		{
+			UE_LOG(LogBlueprintUserMessages, Error, TEXT("BlockCoords[%d] contains negative coordinate(s)."), Index);
+			return nullptr;
+		}
+	}
+#endif
+
+	OutputPathName = UPackageTools::SanitizePackageName(OutputPathName);
+
+	FName TextureName;
+	{
+		int32 SlashIndex;
+		int32 BackSlashIndex;
+		OutputPathName.FindLastChar(TEXT('/'), SlashIndex);
+		OutputPathName.FindLastChar(TEXT('\\'), BackSlashIndex);
+		const int32 TextureNameStart = FMath::Max(SlashIndex, BackSlashIndex) + 1;
+		TextureName = *OutputPathName.RightChop(TextureNameStart);
+	}
+
+	UTexture2D* ExistingTexture = LoadObject<UTexture2D>(nullptr, *OutputPathName, nullptr, LOAD_Quiet | LOAD_NoWarn);
+	UPackage* TexturePackage;
+
+	if (!ExistingTexture)
+	{
+		TexturePackage = CreatePackage(*OutputPathName);
+	}
+	else
+	{
+		TexturePackage = ExistingTexture->GetOutermost();
+	}
+
+	FScopedTextureSourceExtracts SourceExtracts(SourceTextures, BlockCoords);
+	UTextureFactory* TextureFactory = NewObject<UTextureFactory>();
+
+	if (!SourceExtracts.IsValid() || !TextureFactory)
+	{
+		return nullptr;
+	}
+
+	TextureFactory->AddToRoot();
+	TextureFactory->SuppressImportOverwriteDialog();
+
+	UTexture2D* NewTexture = TextureFactory->CreateTexture2D(TexturePackage, TextureName, RF_Public | RF_Standalone);
+
+	TextureFactory->RemoveFromRoot();
+
+	if (!NewTexture)
+	{
+		return nullptr;
+	}
+
+	ETextureSourceFormat Format = SourceExtracts.Format;
+	NewTexture->Source.InitBlocked(&Format, SourceExtracts.Blocks.GetData(), 1, SourceExtracts.Blocks.Num(), SourceExtracts.DataArray.GetData());
+
+	UTexture2D* FirstSourceTexture = SourceTextures[0];
+	UTexture2D* SettingsSource = ExistingTexture && bKeepExistingSettings ? ExistingTexture : FirstSourceTexture;
+
+	NewTexture->AddressX = SettingsSource->AddressX;
+	NewTexture->AddressY = SettingsSource->AddressY;
+	NewTexture->Filter = SettingsSource->Filter;
+	NewTexture->LODGroup = SettingsSource->LODGroup;
+	NewTexture->CompressionSettings = SettingsSource->CompressionSettings;
+	NewTexture->LODBias = SettingsSource->LODBias;
+	NewTexture->NumCinematicMipLevels = SettingsSource->NumCinematicMipLevels;
+	NewTexture->NeverStream = SettingsSource->NeverStream;
+	NewTexture->SRGB = SettingsSource->SRGB;
+	NewTexture->bPreserveBorder = SettingsSource->bPreserveBorder;
+	NewTexture->CompressionNone = SettingsSource->CompressionNone;
+	NewTexture->CompressionNoAlpha = SettingsSource->CompressionNoAlpha;
+	NewTexture->DeferCompression = SettingsSource->DeferCompression;
+	NewTexture->bFlipGreenChannel = SettingsSource->bFlipGreenChannel;
+	NewTexture->bDoScaleMipsForAlphaCoverage = SettingsSource->bDoScaleMipsForAlphaCoverage;
+	NewTexture->AlphaCoverageThresholds = SettingsSource->AlphaCoverageThresholds;
+	NewTexture->bUseNewMipFilter = SettingsSource->bUseNewMipFilter;
+	NewTexture->AdjustBrightness = SettingsSource->AdjustBrightness;
+	NewTexture->AdjustBrightnessCurve = SettingsSource->AdjustBrightnessCurve;
+	NewTexture->AdjustVibrance = SettingsSource->AdjustVibrance;
+	NewTexture->AdjustSaturation = SettingsSource->AdjustSaturation;
+	NewTexture->AdjustRGBCurve = SettingsSource->AdjustRGBCurve;
+	NewTexture->AdjustHue = SettingsSource->AdjustHue;
+	NewTexture->AdjustMinAlpha = SettingsSource->AdjustMinAlpha;
+	NewTexture->AdjustMaxAlpha = SettingsSource->AdjustMaxAlpha;
+	NewTexture->MipGenSettings = SettingsSource->MipGenSettings;
+	NewTexture->VirtualTextureStreaming = true;
+
+	if (ExistingTexture)
+	{
+		// Disable streaming
+		ExistingTexture->UpdateResource();
+		// Wait for InitRHI() to complete before the FTextureReferenceReplacer calls ReleaseRHI() to follow the workflow.
+		// Static texture needs to avoid having pending InitRHI() before enqueuing ReleaseRHI() to safely track access of the PlatformData on the renderthread.
+		ExistingTexture->WaitForPendingInitOrStreaming();
+	}
+
+	FTextureReferenceReplacer(ExistingTexture).Replace(NewTexture);
+
+	// Invalidate any materials using the new texture. (occurs if overwriting an existing texture)
+	NewTexture->PostEditChange();
+
+	FAssetRegistryModule::AssetCreated(NewTexture);
+
+	TexturePackage->MarkPackageDirty();
+
+	if (bCheckOutAndSave)
+	{
+		TArray<UPackage*> PackagesToSave;
+		PackagesToSave.Add(TexturePackage);
+		UEditorLoadingAndSavingUtils::SavePackages(PackagesToSave, true);
+	}
+
+	return NewTexture;
+}
+
 #undef LOCTEXT_NAMESPACE
 
