Fix non unity compile errors

diff --git a/Engine/Source/Editor/UnrealEd/Classes/Factories/TextureFactory.h b/Engine/Source/Editor/UnrealEd/Classes/Factories/TextureFactory.h
--- a/Engine/Source/Editor/UnrealEd/Classes/Factories/TextureFactory.h
+++ b/Engine/Source/Editor/UnrealEd/Classes/Factories/TextureFactory.h
@@ -8,6 +8,7 @@
 #include "Factories/Factory.h"
 #include "Engine/Texture.h"
 #include "ImportSettings.h"
+#include "Kismet/BlueprintFunctionLibrary.h"
 #include "TextureFactory.generated.h"
 
 struct FImportImage
diff --git a/Engine/Source/Editor/UnrealEd/Private/Factories/EditorFactories.cpp b/Engine/Source/Editor/UnrealEd/Private/Factories/EditorFactories.cpp
--- a/Engine/Source/Editor/UnrealEd/Private/Factories/EditorFactories.cpp
+++ b/Engine/Source/Editor/UnrealEd/Private/Factories/EditorFactories.cpp
@@ -273,6 +273,7 @@
 #include "SkinWeightsUtilities.h"
 #include "TextureImportUtils.h"
 #include "UDIMUtilities.h"
+#include "FileHelpers.h"
 
 DEFINE_LOG_CATEGORY(LogEditorFactories);
 
