Don't store PreviewShaderPlatform in EditorSettings because it can change depending on setup

diff --git a/Engine/Source/Editor/UnrealEd/Classes/Editor/EditorPerProjectUserSettings.h b/Engine/Source/Editor/UnrealEd/Classes/Editor/EditorPerProjectUserSettings.h
--- a/Engine/Source/Editor/UnrealEd/Classes/Editor/EditorPerProjectUserSettings.h
+++ b/Engine/Source/Editor/UnrealEd/Classes/Editor/EditorPerProjectUserSettings.h
@@ -162,10 +162,6 @@ public:
 	UPROPERTY(config)
 	int32 PreviewFeatureLevel;
 
-	/** The Preview EShaderPlatform used for Preview */
-	UPROPERTY(config)
-	int32 PreviewShaderPlatform;
-
 	/** The platform name to preview, or NAME_None if there is no preview platform */
 	UPROPERTY(config)
 	FName PreviewPlatformName;
diff --git a/Engine/Source/Editor/UnrealEd/Private/EditorEngine.cpp b/Engine/Source/Editor/UnrealEd/Private/EditorEngine.cpp
--- a/Engine/Source/Editor/UnrealEd/Private/EditorEngine.cpp
+++ b/Engine/Source/Editor/UnrealEd/Private/EditorEngine.cpp
@@ -7795,7 +7795,10 @@ ERHIFeatureLevel::Type UEditorEngine::GetActiveFeatureLevelPreviewType() const
 void UEditorEngine::LoadEditorFeatureLevel()
 {
 	auto* Settings = GetMutableDefault<UEditorPerProjectUserSettings>();
-	if (Settings->PreviewFeatureLevel >= 0 && Settings->PreviewFeatureLevel < (int32)ERHIFeatureLevel::Num && Settings->PreviewShaderPlatform < (int32)EShaderPlatform::SP_NumPlatforms)
+
+	EShaderPlatform ShaderPlatformToPreview = FDataDrivenShaderPlatformInfo::GetShaderPlatformFromName(Settings->PreviewShaderPlatformName);
+
+	if (Settings->PreviewFeatureLevel >= 0 && Settings->PreviewFeatureLevel < (int32)ERHIFeatureLevel::Num && ShaderPlatformToPreview < (int32)EShaderPlatform::SP_NumPlatforms)
 	{
 		// Try to map a saved ShaderFormatName to the PreviewPlatformName using ITargetPlatform if we don't have one. 
 		// We now store the PreviewPlatformName explicitly to support preview for platforms we don't have an ITargetPlatform of.
@@ -7814,7 +7817,7 @@ void UEditorEngine::LoadEditorFeatureLevel()
 		{
 			PreviewShaderPlatformName = *(LexToString(GShaderPlatformForFeatureLevel[Settings->PreviewFeatureLevel]));
 		}
-		SetPreviewPlatform(FPreviewPlatformInfo((ERHIFeatureLevel::Type)Settings->PreviewFeatureLevel, (EShaderPlatform)Settings->PreviewShaderPlatform, Settings->PreviewPlatformName, Settings->PreviewShaderFormatName, Settings->PreviewDeviceProfileName, Settings->bPreviewFeatureLevelActive, PreviewShaderPlatformName), false);
+		SetPreviewPlatform(FPreviewPlatformInfo((ERHIFeatureLevel::Type)Settings->PreviewFeatureLevel, ShaderPlatformToPreview, Settings->PreviewPlatformName, Settings->PreviewShaderFormatName, Settings->PreviewDeviceProfileName, Settings->bPreviewFeatureLevelActive, PreviewShaderPlatformName), false);
 	}
 }
 
@@ -7822,7 +7825,6 @@ void UEditorEngine::SaveEditorFeatureLevel()
 {
 	auto* Settings = GetMutableDefault<UEditorPerProjectUserSettings>();
 	Settings->PreviewFeatureLevel = (int32)PreviewPlatform.PreviewFeatureLevel;
-	Settings->PreviewShaderPlatform = (int32)PreviewPlatform.ShaderPlatform;
 	Settings->PreviewPlatformName = PreviewPlatform.PreviewPlatformName;
 	Settings->PreviewShaderFormatName = PreviewPlatform.PreviewShaderFormatName;
 	Settings->bPreviewFeatureLevelActive = PreviewPlatform.bPreviewFeatureLevelActive;
diff --git a/Engine/Source/Editor/UnrealEd/Private/Settings/EditorPerProjectUserSettings.cpp b/Engine/Source/Editor/UnrealEd/Private/Settings/EditorPerProjectUserSettings.cpp
--- a/Engine/Source/Editor/UnrealEd/Private/Settings/EditorPerProjectUserSettings.cpp
+++ b/Engine/Source/Editor/UnrealEd/Private/Settings/EditorPerProjectUserSettings.cpp
@@ -33,7 +33,6 @@ UEditorPerProjectUserSettings::UEditorPerProjectUserSettings(const FObjectInitia
 	bPreviewFeatureLevelActive = false;
 	bPreviewFeatureLevelWasDefault = true;
 	PreviewDeviceProfileName = NAME_None;
-	PreviewShaderPlatform = EShaderPlatform::SP_NumPlatforms;
 }
 
 void UEditorPerProjectUserSettings::PostInitProperties()
diff --git a/Engine/Source/Runtime/RHI/Private/RHI.cpp b/Engine/Source/Runtime/RHI/Private/RHI.cpp
--- a/Engine/Source/Runtime/RHI/Private/RHI.cpp
+++ b/Engine/Source/Runtime/RHI/Private/RHI.cpp
@@ -2668,6 +2668,26 @@ void FGenericDataDrivenShaderPlatformInfo::UpdatePreviewPlatforms()
 	}
 }
 
+const EShaderPlatform FGenericDataDrivenShaderPlatformInfo::GetShaderPlatformFromName(const FName ShaderPlatformName)
+{
+	EShaderPlatform ShaderPlatformToPreview = EShaderPlatform::SP_NumPlatforms;
+	for (int32 i = 0; i < SP_NumPlatforms; ++i)
+	{
+		const EShaderPlatform Platform = static_cast<EShaderPlatform>(i);
+		if (!Infos[Platform].bContainsValidPlatformInfo)
+		{
+			continue;
+		}
+
+		if (Infos[Platform].Name == ShaderPlatformName)
+		{
+			ShaderPlatformToPreview = Platform;
+			break;
+		}
+	}
+	return ShaderPlatformToPreview;
+}
+
 //
 //	MSAA sample offsets.
 //
diff --git a/Engine/Source/Runtime/RHI/Public/RHIDefinitions.h b/Engine/Source/Runtime/RHI/Public/RHIDefinitions.h
--- a/Engine/Source/Runtime/RHI/Public/RHIDefinitions.h
+++ b/Engine/Source/Runtime/RHI/Public/RHIDefinitions.h
@@ -449,6 +449,7 @@ public:
 	static void Initialize();
 	static void UpdatePreviewPlatforms();
 	static void ParseDataDrivenShaderInfo(const FConfigSection& Section, FGenericDataDrivenShaderPlatformInfo& Info);
+	static const EShaderPlatform GetShaderPlatformFromName(const FName ShaderPlatformName);
 
 	static FORCEINLINE_DEBUGGABLE const FName GetName(const FStaticShaderPlatform Platform)
 	{
