MultiprocessCook - Add ConstructReachabilityGraph helper function for LoadBalanceCookBurden.
#rnx

diff --git a/Engine/Source/Editor/UnrealEd/Private/Cooker/Algo/GraphConvert.cpp b/Engine/Source/Editor/UnrealEd/Private/Cooker/Algo/GraphConvert.cpp
new file mode 100644
--- /dev/null
+++ b/Engine/Source/Editor/UnrealEd/Private/Cooker/Algo/GraphConvert.cpp
@@ -0,0 +1,878 @@
+// Copyright Epic Games, Inc. All Rights Reserved.
+
+#include "GraphConvert.h"
+
+#include "Misc/AutomationTest.h"
+
+#if WITH_DEV_AUTOMATION_TESTS
+#include "Algo/Compare.h"
+#endif
+
+namespace Algo::Graph
+{
+
+void ConstructTransposeGraph(TConstArrayView<TConstArrayView<FVertex>> Graph, TArray64<FVertex>& OutTransposeGraphBuffer,
+	TArray<TConstArrayView<FVertex>>& OutTransposeGraph, EConvertToGraphOptions Options)
+{
+	int32 NumVertices = Graph.Num();
+
+	TArray<TArray<FVertex>> TransposeGraph;
+	TransposeGraph.SetNum(NumVertices);
+
+	for (FVertex Vertex = 0; Vertex < NumVertices; ++Vertex)
+	{
+		for (FVertex Edge : Graph[Vertex])
+		{
+			// Normalize Step 1: Remove edges to self
+			if (Edge == Vertex)
+			{
+				continue;
+			}
+
+			TransposeGraph[Edge].Add(Vertex);
+		}
+	}
+
+	// Finish normalizing; we have to do this before converting to the buffer format
+	// because we need the duplicates removed to get an accurate count
+	int32 NumTotalEdges = 0;
+	for (FVertex Vertex = 0; Vertex < NumVertices; ++Vertex)
+	{
+		TArray<FVertex>& Edges = TransposeGraph[Vertex];
+		// Normalize Step 2: Sort edges
+		Algo::Sort(Edges);
+		// Normalize Step 3: Remove duplicates
+		Edges.SetNum(Algo::Unique(Edges), false /* bAllowShrinking */);
+		NumTotalEdges += Edges.Num();
+	}
+
+	// Convert the Array of Arrays into the output VertexBuffer and Array of Arrayviews
+	if (EnumHasAnyFlags(Options, EConvertToGraphOptions::Shrink))
+	{
+		OutTransposeGraphBuffer.Empty(NumTotalEdges);
+		OutTransposeGraph.Empty(NumVertices);
+	}
+	else
+	{
+		OutTransposeGraphBuffer.Reset(NumTotalEdges);
+		OutTransposeGraph.Reset(NumVertices);
+	}
+
+	FVertex* BufferData = OutTransposeGraphBuffer.GetData();
+	for (FVertex Vertex = 0; Vertex < NumVertices; ++Vertex)
+	{
+		TArray<FVertex>& Edges = TransposeGraph[Vertex];
+		OutTransposeGraph.Emplace(BufferData + OutTransposeGraphBuffer.Num(), Edges.Num());
+		OutTransposeGraphBuffer.Append(Edges);
+	}
+	check(OutTransposeGraphBuffer.Num() == NumTotalEdges);
+	check(BufferData == OutTransposeGraphBuffer.GetData()); // We have TArrayViews into it
+}
+
+bool TryConstructCondensationGraph(TConstArrayView<TConstArrayView<FVertex>> Graph, TArray64<FVertex>& OutGraphBuffer,
+	TArray<TConstArrayView<FVertex>>& OutGraph, TArray64<FVertex>* OutOutVertexToInVerticesBuffer,
+	TArray<TConstArrayView<FVertex>>* OutOutVertexToInVertices, TArray<FVertex>* OutInVertexToOutVertex,
+	EConvertToGraphOptions Options)
+{
+	// Kosaraju-Sharir's algorithm - order the vertices from first out in a DFS search to last out. Any vertex that is
+	// referenced by a referencer vertex that exited the dfs earlier is in a strongly connected component with the
+	// referencer vertex.
+
+	// Provide local buffers for any optional outputs that were not provided
+	checkf((OutOutVertexToInVerticesBuffer != nullptr) == (OutOutVertexToInVertices != nullptr),
+		TEXT("OutOutVertexToInVerticesBuffer and OutOutVertexToInVertices must either both be null or both be non-null."));
+	TArray64<FVertex> LocalBufferOutVertexToInVerticesBuffer;
+	TArray<TConstArrayView<FVertex>> LocalBufferOutVertexToInVertices;
+	TArray<FVertex> LocalBufferInVertexToOutVertex;
+
+	TArray64<FVertex>& OutVertsBuffer(OutOutVertexToInVerticesBuffer ?
+		*OutOutVertexToInVerticesBuffer : LocalBufferOutVertexToInVerticesBuffer);
+	TArray<TConstArrayView<FVertex>>& OutVertexToInVertices(OutOutVertexToInVertices ?
+		*OutOutVertexToInVertices : LocalBufferOutVertexToInVertices);
+	TArray<FVertex>& InVertexToOutVertex(OutInVertexToOutVertex ?
+		*OutInVertexToOutVertex : LocalBufferInVertexToOutVertex);
+
+
+	int32 NumVertices = Graph.Num();
+	enum class EVisitStatus : uint8
+	{
+		NotVisited,
+		InProgress,
+		Visited,
+	};
+	struct FVisitData
+	{
+		FVertex Vertex;
+		int32 NextEdge;
+	};
+
+	TArray<FVertex> DFSOutFirstToLast;
+	bool bHasCycle = false;
+	TArray<FVisitData> Stack;
+	Stack.Reserve(NumVertices);
+	int32 NumEdges = 0;
+
+	// Create the DFSOutFirstToLast order. DepthFirstSearch the input graph and record each vertex after its edges.
+	// Record whether we find any cycles during the search.
+	{
+		TArray<EVisitStatus> VisitStatus;
+		Stack.Reset();
+		VisitStatus.SetNumUninitialized(NumVertices);
+		DFSOutFirstToLast.Reserve(NumVertices);
+
+		for (FVertex Vertex = 0; Vertex < NumVertices; ++Vertex)
+		{
+			VisitStatus[Vertex] = EVisitStatus::NotVisited;
+		}
+		for (FVertex Root = 0; Root < NumVertices; ++Root)
+		{
+			if (VisitStatus[Root] != EVisitStatus::NotVisited)
+			{
+				continue;
+			}
+			VisitStatus[Root] = EVisitStatus::InProgress;
+			Stack.Add({ Root, 0 });
+			while (!Stack.IsEmpty())
+			{
+				FVisitData& VisitData = Stack.Last();
+				FVertex Vertex = VisitData.Vertex;
+				TConstArrayView<FVertex> Edges = Graph[Vertex];
+				bool bPushed = false;
+				while (VisitData.NextEdge < Edges.Num())
+				{
+					FVertex EdgeVertex = Edges[VisitData.NextEdge++];
+					EVisitStatus& EdgeVisitStatus = VisitStatus[EdgeVertex];
+					if (EdgeVisitStatus == EVisitStatus::Visited)
+					{
+					}
+					else if (EdgeVisitStatus == EVisitStatus::NotVisited)
+					{
+						EdgeVisitStatus = EVisitStatus::InProgress;
+						Stack.Add({ EdgeVertex, 0 });
+						bPushed = true;
+						break;
+					}
+					else // EVisitStatus::InProgress
+					{
+						// Don't count edges to self as a cycle, they are removed in normalization of the input graph
+						if (EdgeVertex != Vertex)
+						{
+							bHasCycle = true;
+						}
+					}
+				}
+				if (!bPushed)
+				{
+					NumEdges += Edges.Num();
+					DFSOutFirstToLast.Add(Vertex);
+					VisitStatus[Vertex] = EVisitStatus::Visited;
+					Stack.Pop();
+				}
+			}
+		}
+	}
+
+	// Allocate the OutVertexToInVertices output values, using NumVertices as a conservative estimate of NumOutVertices
+	bool bShrink = EnumHasAnyFlags(Options, EConvertToGraphOptions::Shrink);
+	if (bHasCycle || OutOutVertexToInVertices || OutOutVertexToInVertices)
+	{
+		if (bShrink)
+		{
+			OutVertexToInVertices.Empty(NumVertices);
+			OutVertsBuffer.Empty(NumVertices);
+		}
+		else
+		{
+			OutVertexToInVertices.Reset(NumVertices);
+			OutVertsBuffer.Reset(NumVertices);
+		}
+	}
+
+	// If we did not find any cycles, then the CondensationGraph is equal to the input Graph, so return false.
+	if (!bHasCycle)
+	{
+		OutGraphBuffer.Reset();
+		OutGraph.Reset();
+		if (OutOutVertexToInVertices || OutInVertexToOutVertex)
+		{
+			if (bShrink)
+			{
+				InVertexToOutVertex.Empty(NumVertices);
+			}
+			else
+			{
+				InVertexToOutVertex.Reset(NumVertices);
+			}
+
+			// Copy reversed DFSOutFirstToLast into OutVertexToInVertices, and inverted into InVertexToOutVertex
+			OutVertexToInVertices.SetNum(NumVertices, false /* bAllowShrinking */);
+			OutVertsBuffer.SetNum(NumVertices, false /* bAllowShrinking */);
+			InVertexToOutVertex.SetNum(NumVertices, false /* bAllowShrinking */);
+			for (FVertex OutVertex = 0; OutVertex < NumVertices; ++OutVertex)
+			{
+				FVertex InVertex = DFSOutFirstToLast[NumVertices - 1 - OutVertex];
+				OutVertsBuffer[OutVertex] = InVertex;
+				OutVertexToInVertices[OutVertex] = TConstArrayView<FVertex>(OutVertsBuffer.GetData() + OutVertex, 1);
+				InVertexToOutVertex[InVertex] = OutVertex;
+			}
+		}
+		return false;
+	}
+
+	// Get the TransposeGraph
+	TArray64<FVertex> ReferenceGraphBuffer;
+	TArray<TConstArrayView<FVertex>> ReferenceGraph;
+	ConstructTransposeGraph(Graph, ReferenceGraphBuffer, ReferenceGraph);
+
+	TBitArray<> Assigned;
+
+	// Iterate vertices from last DFSOut to first DFSOut. Each one that is not already assigned creates its own
+	// strongly connected component. Any that are not involved in a cycle will have a strongly connected component
+	// that consists of a single vertex. When we find a vertex that is not already assigned, visit all unassigned
+	// vertices in its reference graph and assign them all to its component.
+	{
+		Stack.Reset();
+		Assigned.Init(false, NumVertices);
+		FVertex* OutVertsData = OutVertsBuffer.GetData();
+
+		for (int32 IndexOfRoot = NumVertices - 1; IndexOfRoot >= 0; --IndexOfRoot)
+		{
+			FVertex Root = DFSOutFirstToLast[IndexOfRoot];
+			if (Assigned[Root])
+			{
+				continue;
+			}
+			Assigned[Root] = true;
+			int64 InitialOffset = OutVertsBuffer.Num();
+			OutVertsBuffer.Add(Root);
+			Stack.Add({ Root, 0 });
+			while (!Stack.IsEmpty())
+			{
+				FVisitData& VisitData = Stack.Last();
+				TConstArrayView<FVertex> Edges = ReferenceGraph[VisitData.Vertex];
+				bool bPushed = false;
+				while (VisitData.NextEdge < Edges.Num())
+				{
+					FVertex EdgeVertex = Edges[VisitData.NextEdge++];
+					if (Assigned[EdgeVertex])
+					{
+					}
+					else
+					{
+						Assigned[EdgeVertex] = true;
+						OutVertsBuffer.Add(EdgeVertex);
+						Stack.Add({ EdgeVertex, 0 });
+						bPushed = true;
+						break;
+					}
+				}
+				if (!bPushed)
+				{
+					Stack.Pop();
+				}
+			}
+			OutVertexToInVertices.Emplace(OutVertsData + InitialOffset, OutVertsBuffer.Num() - InitialOffset);
+		}
+		// We have TArrayViews into OutVertsBuffer; assert its data buffer did not change
+		check(OutVertsBuffer.GetData() == OutVertsData);
+	}
+	int32 NumOutVertices = OutVertexToInVertices.Num();
+
+	// Allocate the output variables that depend on NumOutVertices
+	if (bShrink)
+	{
+		// Shrink OutVertexToInVertices and OutVertexToInVerticesBuffer if requested
+		if (OutOutVertexToInVertices)
+		{
+			TArray64<FVertex> CopyBuffer = OutVertsBuffer;
+			FVertex* OutVertsData = OutVertsBuffer.GetData();
+			FVertex* CopyData = CopyBuffer.GetData();
+			OutVertexToInVertices.Shrink();
+			for (TConstArrayView<FVertex>& InVertices : OutVertexToInVertices)
+			{
+				InVertices = TConstArrayView<FVertex>(CopyData + (InVertices.GetData() - OutVertsData), InVertices.Num());
+			}
+			OutVertsBuffer = MoveTemp(CopyBuffer);
+		}
+		InVertexToOutVertex.Empty(NumVertices);
+		OutGraph.Empty(NumOutVertices);
+	}
+	else
+	{
+		InVertexToOutVertex.Reset(NumVertices);
+		OutGraph.Reset(NumOutVertices);
+	}
+	// We don't know NumEdges in the CondensationGraph yet; NumEdges from the original graph is a conservative estimate
+	OutGraphBuffer.Reset(NumEdges);
+
+	// Invert OutVertexToInVertices to create InVertexToOutVertex
+	InVertexToOutVertex.SetNumUninitialized(NumVertices);
+	for (FVertex OutVertex = 0; OutVertex < NumOutVertices; ++OutVertex)
+	{
+		for (FVertex InVertex : OutVertexToInVertices[OutVertex])
+		{
+			InVertexToOutVertex[InVertex] = OutVertex;
+		}
+	}
+
+	// Create the edges of the condensationgraph by converting each source vertex edge vertex into the condensation
+	// edge vertex and unioning all transformed edge vertices of all source vertices for the condensation root vertex
+	{
+		Assigned.SetRange(0, NumVertices, false);
+		FVertex* OutGraphData = OutGraphBuffer.GetData();
+		for (FVertex OutVertex = 0; OutVertex < NumOutVertices; ++OutVertex)
+		{
+			// Uniquely add all of the unioned edge vertices for this OutVertex onto a range in OutGraphBuffer
+			int64 InitialOffset = OutGraphBuffer.Num();
+			for (FVertex InVertex : OutVertexToInVertices[OutVertex])
+			{
+				for (FVertex SourceEdgeVertex : Graph[InVertex])
+				{
+					FVertex OutEdgeVertex = InVertexToOutVertex[SourceEdgeVertex];
+					if (!Assigned[OutEdgeVertex]) // Note this condition also removes edges to self
+					{
+						Assigned[OutEdgeVertex] = true;
+						OutGraphBuffer.Add(OutEdgeVertex);
+					}
+				}
+			}
+			// Record that range of edges in our OutGraph
+			TConstArrayView<FVertex>& OutEdges = OutGraph.Emplace_GetRef(OutGraphData + InitialOffset,
+				OutGraphBuffer.Num() - InitialOffset);
+
+			// Clear the values in Assigned that we set; this is faster than clearing the whole array
+			for (FVertex OutEdgeVertex : OutEdges)
+			{
+				Assigned[OutEdgeVertex] = false;
+			}
+		}
+		// We have TArrayViews into OutVertexToInVertexBuffer; assert its data buffer did not change
+		check(OutGraphBuffer.GetData() == OutGraphData);
+
+		// Shrink OutGraphBuffer if requested
+		if (bShrink)
+		{
+			TArray64<FVertex> CopyBuffer = OutGraphBuffer;
+			FVertex* CopyData = CopyBuffer.GetData();
+			for (TConstArrayView<FVertex>& Edges : OutGraph)
+			{
+				Edges = TConstArrayView<FVertex>(CopyData + (Edges.GetData() - OutGraphData), Edges.Num());
+			}
+			OutGraphBuffer = MoveTemp(CopyBuffer);
+			check(OutGraphBuffer.GetData() == CopyData);
+		}
+	}
+	return true;
+}
+
+void ConstructPartialTransposeGraph(TConstArrayView<TConstArrayView<FVertex>> Graph,
+	TArrayView<FVertex> InVertices, int64 MaxOutGraphEdges,
+	TArray64<FVertex>& OutTransposeGraphBuffer, TArray<TConstArrayView<FVertex>>& OutTransposeGraph,
+	TArray<FVertex>& OutInVerticesInOutGraph)
+{
+	OutInVerticesInOutGraph.Reset(InVertices.Num());
+	Algo::Sort(InVertices, [Graph](FVertex A, FVertex B)\
+		{
+			int32 NumA = Graph[A].Num();
+			int32 NumB = Graph[B].Num();
+			return NumA > NumB || (NumA == NumB && A < B);
+		});
+
+	int32 NumVertices = Graph.Num();
+	TArray<TArray<FVertex>> TransposeGraph;
+	TransposeGraph.SetNum(NumVertices);
+	int32 NumEdges = 0;
+	for (FVertex InVertex : InVertices)
+	{
+		NumEdges += Graph[InVertex].Num();
+		if (NumEdges > MaxOutGraphEdges)
+		{
+			NumEdges -= Graph[InVertex].Num();
+			break;
+		}
+		OutInVerticesInOutGraph.Add(InVertex);
+		for (FVertex TargetVertex : Graph[InVertex])
+		{
+			TransposeGraph[TargetVertex].Add(InVertex);
+		}
+	}
+	ConvertToSingleBufferGraph(TransposeGraph, OutTransposeGraphBuffer, OutTransposeGraph);
+}
+
+}
+
+#if WITH_DEV_AUTOMATION_TESTS
+IMPLEMENT_SIMPLE_AUTOMATION_TEST(FCondensationGraphTest, "System.Core.Algo.CondensationGraph", EAutomationTestFlags::ApplicationContextMask | EAutomationTestFlags::EngineFilter);
+
+bool FCondensationGraphTest::RunTest(const FString& Parameters)
+{
+	using namespace Algo::Graph;
+
+	TArray<TArray<FVertex>> Buffer;
+	TArray<TConstArrayView<FVertex>> Graph;
+	TArray64<FVertex> CondensationBuffer;
+	TArray<TConstArrayView<FVertex>> CondensationGraph;
+	TArray64<FVertex> OutVertexToInVerticesBuffer;
+	TArray<TConstArrayView<FVertex>> OutVertexToInVertices;
+	TArray<FVertex> InVertexToOutVertex;
+	TArray64<TArray<FVertex>> ExpectedOutVertexToInVertices;
+	TArray<FVertex> ExpectedInVertexToOutVertex;
+	TArray<FVertex> ActualInVertexToOutVertex;
+	TArray<FVertex> ExpectedComponentVerticesScratch;
+	TArray<FVertex> ActualComponentVerticesScratch;
+	TArray<FVertex> ExpectedOutEdges;
+	TArray<FVertex> ActualOutEdges;
+	TArray<FVertex> Reachable;
+	TBitArray<> Assigned;
+
+	auto AddVertex = [&Buffer, &Graph](FVertex Vertex, TConstArrayView<FVertex> Edges)
+	{
+		if (Graph.Num() <= Vertex)
+		{
+			Graph.SetNum(Vertex + 1);
+			Buffer.SetNum(Vertex + 1);
+		}
+		Buffer[Vertex] = Edges;
+		Graph[Vertex] = Buffer[Vertex];
+	};
+	auto AddExpectedComponent = [&ExpectedOutVertexToInVertices](TConstArrayView<FVertex> InputVertices)
+	{
+		ExpectedOutVertexToInVertices.Emplace_GetRef().Append(InputVertices);
+	};
+	auto Clear = [&Buffer, &Graph, &ExpectedOutVertexToInVertices]()
+	{
+		Buffer.Reset();
+		Graph.Reset();
+		ExpectedOutVertexToInVertices.Reset();
+	};
+	auto WriteArrayToString = [](TConstArrayView<FVertex> A)
+	{
+		TStringBuilder<256> Writer;
+		Writer << TEXT("[");
+		Writer.Join(A, TEXT(','));
+		Writer << TEXT("]");
+		return FString(Writer);
+	};
+	auto ConfirmResults =
+		[this, &Graph, &CondensationBuffer, &CondensationGraph, &OutVertexToInVerticesBuffer, &OutVertexToInVertices,
+		&InVertexToOutVertex , &ExpectedOutVertexToInVertices, &ExpectedInVertexToOutVertex, &Assigned,
+		&ActualInVertexToOutVertex, &ExpectedComponentVerticesScratch, &ActualComponentVerticesScratch,
+		&WriteArrayToString, &ExpectedOutEdges, &ActualOutEdges, &Reachable]
+		(const TCHAR* TestCaseName)
+	{
+		bool bResult = TryConstructCondensationGraph(Graph, CondensationBuffer, CondensationGraph,
+			&OutVertexToInVerticesBuffer, &OutVertexToInVertices, &InVertexToOutVertex, EConvertToGraphOptions::Shrink);
+		bool bExpectedResult = !ExpectedOutVertexToInVertices.IsEmpty();
+		if (bResult != bExpectedResult)
+		{
+			AddError(FString::Printf(TEXT("Failed for case \"%s\": expected ConstructCondensationGraph to return %s but it returned %s."),
+				TestCaseName, bExpectedResult ? TEXT("true") : TEXT("false"), bResult ? TEXT("true") : TEXT("false")));
+			return;
+		}
+
+		int32 NumVertices = Graph.Num();
+		int32 NumOutVertices = CondensationGraph.Num();
+		Assigned.SetNumUninitialized(NumVertices);
+
+		// Verify InVertexToOutVertex and OutVertexToInVertices match, cover all vertices, and do not put the same
+		// invertex into two outvertices
+		ActualInVertexToOutVertex.SetNumUninitialized(NumVertices, false /* bAllowShrinking */);
+		Assigned.SetRange(0, NumVertices, false);
+		int32 ComponentIndex = 0;
+		for (TConstArrayView<FVertex> ActualComponentVertices : OutVertexToInVertices)
+		{
+			for (FVertex InVertex : ActualComponentVertices)
+			{
+				if (Assigned[InVertex])
+				{
+					AddError(FString::Printf(TEXT("Failed for case \"%s\": OutVertexToInVertices has the same invertex %d in both outvertex %d and outvertex %d."),
+						TestCaseName, (int32)InVertex, (int32)ActualInVertexToOutVertex[InVertex], (int32)ComponentIndex));
+					return;
+				}
+				Assigned[InVertex] = true;
+				ActualInVertexToOutVertex[InVertex] = ComponentIndex;
+			}
+			++ComponentIndex;
+		}
+		for (FVertex InVertex = 0; InVertex < NumVertices; ++InVertex)
+		{
+			if (!Assigned[InVertex])
+			{
+				AddError(FString::Printf(TEXT("Failed for case \"%s\": OutVertexToInVertices is missing invertex %d; it did not assign that vertex to any outvertex."),
+					TestCaseName, (int32) InVertex));
+				return;
+			}
+			if (ActualInVertexToOutVertex[InVertex] != InVertexToOutVertex[InVertex])
+			{
+				AddError(FString::Printf(TEXT("Failed for case \"%s\": OutVertexToInVertices has outvertex %d <- invertex %d, but InVertexToOutVertex has invertex %d -> outvertex %d."),
+					TestCaseName, (int32)ActualInVertexToOutVertex[InVertex], (int32)InVertex, (int32)InVertex,
+					(int32)InVertexToOutVertex[InVertex]));
+			}
+		}
+
+		// Finish the construction of ExpectedOutVertexToInVertices; it has the contract that any unmentioned vertex is in a
+		// component by itself. Also construct its inverse - ExpectedInVertexToOutVertex.
+		ExpectedInVertexToOutVertex.SetNumUninitialized(NumVertices, false /* bAllowShrinking */);
+		Assigned.SetRange(0, NumVertices, false);
+		ComponentIndex = 0;
+		for (TArray<FVertex>& ExpectedComponentVertices : ExpectedOutVertexToInVertices)
+		{
+			for (FVertex InVertex : ExpectedComponentVertices)
+			{
+				Assigned[InVertex] = true;
+				ExpectedInVertexToOutVertex[InVertex] = ComponentIndex;
+			}
+			++ComponentIndex;
+		}
+		for (FVertex InVertex = 0; InVertex < NumVertices; ++InVertex)
+		{
+			if (!Assigned[InVertex])
+			{
+				ExpectedInVertexToOutVertex[InVertex] = ComponentIndex++;
+				ExpectedOutVertexToInVertices.Add({ InVertex });
+			}
+		}
+
+		// Verify that the actual components place every InVertex in a component with the same list of Inputvertices
+		// as the expected components
+		Assigned.SetRange(0, NumVertices, false);
+		for (FVertex InVertex = 0; InVertex < NumVertices; ++InVertex)
+		{
+			if (Assigned[InVertex])
+			{
+				continue;
+			}
+			FVertex ExpectedOutVertex = ExpectedInVertexToOutVertex[InVertex];
+			FVertex ActualOutVertex = ActualInVertexToOutVertex[InVertex];
+			ExpectedComponentVerticesScratch.Reset();
+			ExpectedComponentVerticesScratch.Append(ExpectedOutVertexToInVertices[ExpectedOutVertex]);
+			ActualComponentVerticesScratch.Reset();
+			ActualComponentVerticesScratch.Append(OutVertexToInVertices[ActualOutVertex]);
+			Algo::Sort(ExpectedComponentVerticesScratch);
+			Algo::Sort(ActualComponentVerticesScratch);
+			if (!Algo::Compare(ExpectedComponentVerticesScratch, ActualComponentVerticesScratch))
+			{
+				AddError(FString::Printf(TEXT("Failed for case \"%s\": Input vertex %d was placed into a component with invertices %s, but it was expected in a component with invertices %s."),
+					TestCaseName, (int32)InVertex, *WriteArrayToString(ActualComponentVerticesScratch),
+					*WriteArrayToString(ExpectedComponentVerticesScratch)));
+				return;
+			}
+		}
+
+		// Verify that every edge in the input graph is a component edge in the output graph, or is inside the component
+		if (bResult)
+		{
+			for (FVertex InSourceVertex = 0; InSourceVertex < NumVertices; ++InSourceVertex)
+			{
+				FVertex OutSourceVertex = ActualInVertexToOutVertex[InSourceVertex];
+				for (FVertex InEdgeVertex : Graph[InSourceVertex])
+				{
+					FVertex ExpectedOutEdgeVertex = ActualInVertexToOutVertex[InEdgeVertex];
+					if (ExpectedOutEdgeVertex == OutSourceVertex)
+					{
+						continue;
+					}
+					if (!CondensationGraph[OutSourceVertex].Contains(ExpectedOutEdgeVertex))
+					{
+						AddError(FString::Printf(TEXT("Failed for case \"%s\": Missing edge in the condensation graph: ")
+							TEXT("expected edge from vertex %d in component %d to vertex %d in component %d, but this edge is missing."),
+							TestCaseName, (int32)InSourceVertex, (int32)OutSourceVertex, (int32)InEdgeVertex,
+							(int32)ExpectedOutEdgeVertex));
+						return;
+					}
+				}
+			}
+
+			// Verify that no component edge exists in the output graph unless there is a witness edge for it in the input graph
+			// Create the expected edges of the condensationgraph by converting each source vertex edge vertex into the condensation
+			// edge vertex and unioning all transformed edge vertices of all source vertices for the condensation root vertex
+			for (FVertex OutVertex = 0; OutVertex < NumOutVertices; ++OutVertex)
+			{
+				Assigned.SetRange(0, NumOutVertices, false);
+				ExpectedOutEdges.Reset();
+
+				for (FVertex InVertex : OutVertexToInVertices[OutVertex])
+				{
+					for (FVertex SourceEdgeVertex : Graph[InVertex])
+					{
+						FVertex OutEdgeVertex = ActualInVertexToOutVertex[SourceEdgeVertex];
+						if (!Assigned[OutEdgeVertex])
+						{
+							Assigned[OutEdgeVertex] = true;
+							ExpectedOutEdges.Add(OutEdgeVertex);
+						}
+					}
+				}
+
+				ActualOutEdges.Reset();
+				ActualOutEdges.Append(CondensationGraph[OutVertex]);
+				Algo::Sort(ExpectedOutEdges);
+				Algo::Sort(ActualOutEdges);
+				if (!Algo::Compare(ActualOutEdges, ExpectedOutEdges))
+				{
+					AddError(FString::Printf(TEXT("Failed for case \"%s\": Expected component %d to have edges %s, but it had %s."),
+						TestCaseName, (int32)OutVertex, *WriteArrayToString(ActualOutEdges),
+						*WriteArrayToString(ExpectedOutEdges)));
+					return;
+				}
+			}
+
+			// Verify the output is sorted from root to leaf
+			for (FVertex Root = 0; Root < NumOutVertices; ++Root)
+			{
+				Assigned.SetRange(0, NumOutVertices, false);
+				Reachable.Reset(NumOutVertices);
+				Assigned[Root] = true;
+				Reachable.Add(Root);
+				int32 NextIndex = 0;
+				while (NextIndex < Reachable.Num())
+				{
+					FVertex Vertex = Reachable.Pop(false /* bAllowShrinking */);
+					for (FVertex EdgeVertex : CondensationGraph[Vertex])
+					{
+						if (!Assigned[EdgeVertex])
+						{
+							Assigned[EdgeVertex] = true;
+							Reachable.Add(EdgeVertex);
+						}
+					}
+				}
+				for (FVertex ReachableVertex : Reachable)
+				{
+					if (ReachableVertex < Root)
+					{
+						AddError(FString::Printf(TEXT("Failed for case \"%s\": Expected outgraph to be topologically sorted from root to leaf, but vertex %d is reachable from vertex %d."),
+							TestCaseName, (int32)ReachableVertex, (int32)Root));
+						return;
+					}
+				}
+			}
+		}
+		else // !if (bResult)
+		{
+			// Should have cleared the graph arguments
+			if (!CondensationGraph.IsEmpty() || !CondensationBuffer.IsEmpty())
+			{
+				AddError(FString::Printf(TEXT("Failed for case \"%s\": Expected outgraph to be cleared since result was false, but it is non-empty."),
+					TestCaseName));
+				return;
+			}
+		}
+	};
+
+	{
+		Clear();
+		AddVertex(0, { });
+		AddVertex(1, { 0 });
+		AddVertex(2, { 1 });
+		ConfirmResults(TEXT("Each node depends on the previous one"));
+	}
+	{
+		Clear();
+		AddVertex(0, { 1 });
+		AddVertex(1, { 2 });
+		AddVertex(2, { });
+		ConfirmResults(TEXT("Each node depends on the next one"));
+	}
+	{
+		Clear();
+		AddVertex(0, { 0 });
+		AddVertex(1, { 0,1 });
+		AddVertex(2, { 1,2 });
+		ConfirmResults(TEXT("SelfReferences"));
+	}
+	{
+		//              6
+		//             / \
+		//            5   7
+		//           / \   \
+		//          0   1   8
+		//           \ / \   \
+		//            3   4   |
+		//             \   \ /
+		//              \   9
+		//               \ /
+		//                2       
+		Clear();
+		AddVertex(6, { 5,7 });
+		AddVertex(5, { 0,1 });
+		AddVertex(7, { 8 });
+		AddVertex(0, { 3 });
+		AddVertex(1, { 3,4 });
+		AddVertex(8, { 9 });
+		AddVertex(3, { 2 });
+		AddVertex(4, { 9 });
+		AddVertex(9, { 2 });
+		AddVertex(2, { });
+		ConfirmResults(TEXT("SketchedOutExample1"));
+	}
+	{
+		Clear();
+		AddVertex(0, { 0,1 });
+		AddVertex(1, { 0,1 });
+		AddExpectedComponent({ 0, 1 });
+		ConfirmResults(TEXT("Simple cycle"));
+	}
+	{
+		Clear();
+		AddVertex(0, { 1 });
+		AddVertex(1, { 0,1,2 });
+		AddVertex(2, { 2,0 });
+		AddExpectedComponent({ 0, 1, 2 });
+		ConfirmResults(TEXT("Short cycle in a long cycle"));
+	}
+	{
+		Clear();
+		AddVertex(0, { 1 });
+		AddVertex(1, { 0,2 });
+		AddVertex(2, { 3 });
+		AddVertex(3, { });
+		AddExpectedComponent({ 0, 1 });
+		ConfirmResults(TEXT("Cycle in the root and with the root cycle depending on a chain of non-cycle verts"));
+	}
+	{
+		Clear();
+		AddVertex(0, { });
+		AddVertex(1, { 0 });
+		AddVertex(2, { 1,3 });
+		AddVertex(3, { 2 });
+		AddExpectedComponent({ 2, 3 });
+		ConfirmResults(TEXT("Cycle in the root and with the root cycle depending on a chain of non-cycle verts, submitted in reverse"));
+	}
+	{
+		Clear();
+		AddVertex(0, { 1 });
+		AddVertex(1, { 2 });
+		AddVertex(2, { 3 });
+		AddVertex(3, { 2 });
+		AddExpectedComponent({ 2,3 });
+		ConfirmResults(TEXT("Cycle at a leaf and a chain from the root depending on that cycle"));
+	}
+	{
+		Clear();
+		AddVertex(0, { 1 });
+		AddVertex(1, { 0 });
+		AddVertex(2, { 1 });
+		AddVertex(3, { 2 });
+		AddExpectedComponent({ 0, 1 });
+		ConfirmResults(TEXT("Cycle in the root and with the root cycle depending on a chain of non-cycle verts"));
+	}
+	{
+		Clear();
+		AddVertex(0, { 1 });
+		AddVertex(1, { 2,3 });
+		AddVertex(2, { 1,3 });
+		AddVertex(3, { 1,2 });
+		AddExpectedComponent({ 1,2,3 });
+		ConfirmResults(TEXT("Vertex dependent upon a cycle"));
+	}
+	{
+		// 0 -> 1 -> 2 -> 3 -> 1
+		//           |
+		//           v
+		//      4 -> 5 -> 6 -> 4
+		Clear();
+		AddVertex(0, { 1 });
+		AddVertex(1, { 2 });
+		AddVertex(2, { 3,5 });
+		AddVertex(3, { 1 });
+		AddVertex(4, { 5 });
+		AddVertex(5, { 6 });
+		AddVertex(6, { 4 });
+		AddExpectedComponent({ 1,2,3 });
+		AddExpectedComponent({ 4,5,6 });
+		ConfirmResults(TEXT("One cycle dependent upon another"));
+	}
+	{
+		// 5 -> 0 -> (1 -> 2 -> 1)
+		//      |          |
+		//      |          v
+		//      |          5
+		//      v
+		//     (3 -> 4 -> 3)
+		Clear();
+		AddVertex(0, { 1,3 });
+		AddVertex(1, { 2 });
+		AddVertex(2, { 1,5 });
+		AddVertex(3, { 4 });
+		AddVertex(4, { 3 });
+		AddVertex(5, { 0 });
+		AddExpectedComponent({ 0,1,2,5 });
+		AddExpectedComponent({ 3,4 });
+		ConfirmResults(TEXT("MutuallyReachableSet Problem1"));
+	}
+	{
+		Clear();
+		TArray<FVertex> AllButItself;
+		TArray<FVertex> All;
+		for (FVertex Vertex = 0; Vertex < 6; ++Vertex)
+		{
+			AllButItself.Add(Vertex);
+			All.Add(Vertex);
+		}
+		for (FVertex Vertex = 0; Vertex < 6; ++Vertex)
+		{
+			AllButItself.Remove(Vertex);
+			AddVertex(Vertex, AllButItself);
+			AllButItself.Add(Vertex);
+		}
+		AddExpectedComponent(All);
+		ConfirmResults(TEXT("FullyConnectedGraph"));
+	}
+	{
+		Clear();
+		// 0 - 1 - 2 - 0
+		// |
+		// 3 - 4 - 5 - 3
+		//         |
+		// 6 - 7 - 8 - 6
+		//     |
+		//     1
+		AddVertex(0, { 1,3 });
+		AddVertex(1, { 2 });
+		AddVertex(2, { 0 });
+		AddVertex(3, { 4 });
+		AddVertex(4, { 5 });
+		AddVertex(5, { 3,8 });
+		AddVertex(6, { 7 });
+		AddVertex(7, { 1,8 });
+		AddVertex(8, { 6 });
+		AddExpectedComponent({ 0,1,2,3,4,5,6,7,8 });
+		ConfirmResults(TEXT("Cycle of cycles"));
+	}
+	{
+		Clear();
+		// 0 ----- 1 ----- 2 ----- 0
+		// |       |       |
+		// |       0       |
+		// |               |
+		// 5 - 4 - 3 - 5   6 - 7 - 6
+		//     |\              |
+		//     | 5             |
+		//     |               |
+		//     8 - 9 - 10 - 9  11
+		//     |   |   |
+		//     10  8   8
+		AddVertex(0, { 1,5 });
+		AddVertex(1, { 0,2 });
+		AddVertex(2, { 0,6 });
+		AddVertex(3, { 5 });
+		AddVertex(4, { 3,5,8 });
+		AddVertex(5, { 4 });
+		AddVertex(6, { 7 });
+		AddVertex(7, { 6,11 });
+		AddVertex(8, { 9,10 });
+		AddVertex(9, { 8,10 });
+		AddVertex(10, { 8,9 });
+		AddVertex(11, { });
+		AddExpectedComponent({ 0,1,2 });
+		AddExpectedComponent({ 3,4,5 });
+		AddExpectedComponent({ 6,7 });
+		AddExpectedComponent({ 8,9,10 });
+		ConfirmResults(TEXT("Tree of cycles"));
+	}
+
+	return true;
+}
+
+#endif
diff --git a/Engine/Source/Editor/UnrealEd/Private/Cooker/Algo/GraphConvert.h b/Engine/Source/Editor/UnrealEd/Private/Cooker/Algo/GraphConvert.h
new file mode 100644
--- /dev/null
+++ b/Engine/Source/Editor/UnrealEd/Private/Cooker/Algo/GraphConvert.h
@@ -0,0 +1,262 @@
+// Copyright Epic Games, Inc. All Rights Reserved.
+
+#pragma once
+
+#include "Algo/Sort.h"
+#include "Algo/Unique.h"
+#include "Containers/Array.h"
+#include "Containers/ArrayView.h"
+#include "Containers/Map.h"
+#include "HAL/Platform.h"
+#include "Misc/AssertionMacros.h"
+#include "Misc/EnumClassFlags.h"
+#include "Templates/IdentityFunctor.h"
+#include "Templates/Invoke.h"
+#include "Traits/ElementType.h"
+
+#include <utility>
+
+namespace Algo::Graph
+{
+
+typedef int32 FVertex;
+constexpr FVertex InvalidVertex = static_cast<FVertex>(INDEX_NONE);
+
+enum class EConvertToGraphOptions
+{
+	None = 0,
+	/** The output buffer will be shrunk to reduce memory to minimum required. */
+	Shrink = 1 << 0,
+};
+ENUM_CLASS_FLAGS(EConvertToGraphOptions);
+
+/**
+ * Convert an array of Keys and a function that returns the directed edges from each Key into the Normalized Graph
+ * Form. Normalized Graph Form is an implicit set of vertices identified as integers from 0 to N - 1, and an array of
+ * arrayviews of edges from each vertex. The Normalized Graph form has two structures: The GraphBuffer which contains
+ * the memory storage for the Graph and the Graph itself which is the array of arrayviews of edges. The length of the
+ * array in the graph defines the number of vertices.
+ * 
+ * @param UniqueKeys A range with element type KeyType. KeyType must support GetTypeHash and copy+move constructors.
+ *        KeyType being pointertype is recommended. Keys in the range must be unique. The ith element of the range
+ *        will correspond to as vertex i in the OutGraph. UniqueKeys must have length < MAX_int32.
+ * @param GetKeyEdges A callable with prototype that is one of
+ *            RangeType<KeyType> GetKeyEdges(const KeyType& Key);
+ *            const RangeType<KeyType>& GetKeyEdges(const KeyType& Key);
+ *        It must return the TargetKeys that are pointed to from the directed edges from Key. All returned TargetKeys
+ *        must be elements of UniqueKeys. RangeType must support ranged-for (begin() and end()).
+ * @param OutGraphBuffer Output value that holds the memory used by the graph. The array must remain allocated and
+ *        unmodified until OutGraph is no longer referenced.
+ * @param OutGraph Output value that holds the graph. OutGraph[i] is a TConstArrayView of edges from Vertex i, which
+ *        corresponds to UniqueKeys[i]. Each element of OutGraph[i] is a an FVertex which when cast to an integer is
+ *        the index into UniqueKeys that corresponds to the vertex. Elements of OutGraph[i] are sorted by FVertex
+ *        value, aka the index into UniqueKeys.
+ * @param Options for the conversion, @see EConvertToGraphOptions.
+ */
+template <typename RangeType, typename GetKeyEdgesType>
+inline void ConvertToGraph(const RangeType& UniqueKeys, GetKeyEdgesType GetKeyEdges,
+	TArray64<FVertex>& OutGraphBuffer, TArray<TConstArrayView<FVertex>>& OutGraph,
+	EConvertToGraphOptions Options = EConvertToGraphOptions::None)
+{
+	using KeyType = TElementType_T<RangeType>;
+
+	// Map Elements to vertices in VertexOfKey
+	int32 NumVertices = GetNum(UniqueKeys);
+	TMap<KeyType, FVertex> VertexOfKey;
+	VertexOfKey.Reserve(NumVertices);
+	FVertex Vertex = 0;
+	for (const KeyType& Key : UniqueKeys)
+	{
+		FVertex& ExistingHandle = VertexOfKey.FindOrAdd(Key, InvalidVertex);
+		check(ExistingHandle == InvalidVertex);
+		if (ExistingHandle == InvalidVertex)
+		{
+			ExistingHandle = Vertex++;
+		}
+	}
+
+	TArray<TPair<int64, int32>> EdgeOffsets;
+	EdgeOffsets.Reserve(NumVertices);
+	OutGraphBuffer.Reset();
+
+	// Call GetKeyEdges on each Element, map its returnvalue to Vertices, normalize them, and store them in EdgeOffsets
+	Vertex = 0;
+	for (const KeyType& Element : UniqueKeys)
+	{
+		int64 InitialOffset = OutGraphBuffer.Num();
+		for (const KeyType& Dependency : Invoke(GetKeyEdges, Element))
+		{
+			FVertex* TargetVertex = VertexOfKey.Find(Dependency);
+			checkf(TargetVertex, TEXT("ConvertToGraph: GetKeyEdges returned a TargetKey that is not a member of UniqueKeys"));
+
+			// Normalize Step 1: remove edges to self
+			if (*TargetVertex == Vertex)
+			{
+				continue;
+			}
+			OutGraphBuffer.Add(*TargetVertex);
+		}
+
+		// Normalize Step 2: Sort
+		TArrayView<FVertex> VertexEdges = TArrayView<FVertex>(OutGraphBuffer).Mid(InitialOffset);
+		Algo::Sort(VertexEdges);
+
+		// Normalize Step 3: Remove duplicates
+		VertexEdges = VertexEdges.Left(Algo::Unique(VertexEdges));
+		OutGraphBuffer.SetNum(InitialOffset + VertexEdges.Num(), false /* bAllowShrinking */);
+
+		// Store the vertex's offset into GraphBuffer, for later conversion to TArrayView
+		EdgeOffsets.Emplace_GetRef(InitialOffset, VertexEdges.Num());
+		++Vertex;
+	}
+	if (EnumHasAnyFlags(Options, EConvertToGraphOptions::Shrink))
+	{
+		OutGraphBuffer.Shrink();
+		OutGraph.Empty(NumVertices);
+	}
+	else
+	{
+		OutGraph.Reset(NumVertices);
+	}
+
+	// Convert offsets to TArrayView
+	for (const TPair<int64, int32>& EdgeOffset : EdgeOffsets)
+	{
+		OutGraph.Add(TConstArrayView<FVertex>(OutGraphBuffer.GetData() + EdgeOffset.Key, EdgeOffset.Value));
+	}
+}
+
+/**
+ * Convert an array of separately allocated edge ranges into a single buffer shared by all vertices and an edge graph
+ * of arrayviews into that buffer. No overlapping is done; every input edge range is reproduced in the output buffer.
+ * 
+ * @param Graph ArrayView of N vertices; each element is interpreted by the ProjectionType to return a range of edges
+ *        for that vertex.
+ * @param OutGraphBuffer Output value that holds the memory used by the graph. This array must remain allocated and
+ *        unmodified until OutGraph is no longer referenced.
+ * @param OutGraph Output value that holds the graph. @see ConvertToGraph for a description of the format.
+ * @param Proj Returns the edges for Graph[n]. The prototype is
+ *            // Input is either const& or by value. Output can be const&, &&, or by value.
+ *            EdgeRangeType Proj(MultiBufferRangeType VertexEdges)
+ *        EdgeRangeType must support GetData and GetNum.
+ * @param Options for the conversion, @see EConvertToGraphOptions.
+ */
+template <typename RangeType, typename ProjectionType>
+inline void ConvertToSingleBufferGraph(RangeType&& Graph, TArray64<FVertex>& OutGraphBuffer,
+	TArray<TConstArrayView<FVertex>>& OutGraph, ProjectionType Proj,
+	EConvertToGraphOptions Options = EConvertToGraphOptions::None)
+{
+	typedef TElementType_T<RangeType> InEdgeRangeType;
+	typedef decltype(Proj(std::declval<const InEdgeRangeType&>())) ProjectedEdgeRangeType;
+
+	int32 NumVertices = Graph.Num();
+	int64 NumEdges = 0;
+	for (const InEdgeRangeType& VertexEdges : Graph)
+	{
+		NumEdges += GetNum(Proj(VertexEdges));
+	}
+
+	if (EnumHasAnyFlags(Options, EConvertToGraphOptions::Shrink))
+	{
+		OutGraphBuffer.Empty(NumEdges);
+		OutGraph.Empty(NumVertices);
+	}
+	else
+	{
+		OutGraphBuffer.Reset(NumEdges);
+		OutGraph.Reset(NumVertices);
+	}
+
+	FVertex* OutBufferData = OutGraphBuffer.GetData();
+	for (const InEdgeRangeType& InVertexEdges : Graph)
+	{
+		ProjectedEdgeRangeType VertexEdges = Proj(InVertexEdges);
+		int32 NumVertexEdges = GetNum(VertexEdges);
+		OutGraph.Emplace(OutBufferData + OutGraphBuffer.Num(), NumVertexEdges);
+		OutGraphBuffer.Append(GetData(VertexEdges), NumVertexEdges);
+	}
+	check(OutGraphBuffer.GetData() == OutBufferData); // We have arrayviews into OutBufferData
+}
+
+/**
+ * Convert an array of separately allocated edge ranges into a single buffer shared by all vertices and an edge graph
+ * of arrayviews into that buffer. No overlapping is done; every input edge range is reproduced in the output buffer.
+ * This version requires MultiBufferRange type to be a range type; see other prototype to provide a Projection.
+ * 
+ * @param Graph ArrayView of N vertices; each element is interpreted by the ProjectionType to return a range of edges
+ *        for that vertex.
+ * @param OutBuffer Output value that holds the memory used by the graph. This array must remain allocated and
+ *        unmodified until OutGraph is no longer referenced.
+ * @param OutGraph Output value that holds the graph. @see ConvertToGraph for a description of the format.
+ * @param Options for the conversion, @see EConvertToGraphOptions.
+ */
+template <typename RangeType>
+inline void ConvertToSingleBufferGraph(RangeType Graph, TArray64<FVertex>& OutBuffer,
+	TArray<TConstArrayView<FVertex>>& OutGraph, EConvertToGraphOptions Options = EConvertToGraphOptions::None)
+{
+	ConvertToSingleBufferGraph(Graph, OutBuffer, OutGraph, FIdentityFunctor(), Options);
+}
+
+/**
+ * Return a new buffer and graph with the same vertices but with each edge reversed. If and only if (i,j) is an edge
+ * in the input Graph, (j,i) is an edge in the output TransposeGraph.
+ * 
+ * @param Graph A standard GraphConvert-format Graph. @see ConvertToGraph for a description of the format.
+ * @param OutTransposeGraphBuffer Output value that holds the memory used by the graph. This array must remain
+ *        allocated and unmodified until OutTransposeGraph is no longer referenced.
+ * @param OutTransposeGraph Output value that holds the graph. @see ConvertToGraph for a description of the format.
+ * @param Options for the conversion, @see EConvertToGraphOptions.
+ */
+void ConstructTransposeGraph(TConstArrayView<TConstArrayView<FVertex>> Graph,
+	TArray64<FVertex>& OutTransposeGraphBuffer, TArray<TConstArrayView<FVertex>>& OutTransposeGraph,
+	EConvertToGraphOptions Options = EConvertToGraphOptions::None);
+
+/**
+ * Return a new graph with cycles in the input graph replaced by a single vertex. The new graph is topologically sorted
+ * from root to leaf. If the input graph has no cycles, the function returns false and the OutGraph is reset to empty.
+ *
+ * @param Graph The input graph, may contain cycles. @see ConvertToGraph for a description of the format.
+ * @param OutGraphBuffer Output value that holds the memory used by OutGraph. The memory held by the array must remain
+ *        allocated and unmodified until OutGraph is no longer referenced.
+ * @param OutGraph The edges of the condensed graph, in a standard GraphConvert-format Graph. @see ConvertToGraph for a
+ *        description of the format. It has no cycles and the vertices are topologically sorted from root to leaf.
+ * @param OutOutVertexToInVerticesBuffer Output value, optional, can be null. If null, OutOutVertexToInVertices must
+ *        also be null. If non-null, it is populated with the memory used by *OutOutVertexToInVertices. It must remain
+ *        remain allocated and unmodified until *OutOutVertexToSourceVertex is no longer referenced. This out value
+ *        is set even if the function returns false.
+ * @param OutOutVertexToInVertices Output value, optional, can be null. If null, OutOutVertexToInVerticesBuffer must
+ *        also be null. If non-null, it is populated with an array of arrayviews. The length of the array is equal to
+ *        the length of OutGraph, which is the number of vertices in the condensed graph. THe nth element of the array
+ *        contains the list of (not necessarily sorted) vertices from the input graph that are contained within vertex
+ *        n of the condensed graph. Vertices in the outgraph that correspond to an input vertex that was not in a cycle
+ *        will have an array of only a single element. This out value is set even if the function returns false.
+ * @param OutInVertexToOutVertex Output value, optional, can be null. If non-null, it is populated with an array. The
+ *        length of the array is the number of vertices in the input graph. The nth element of the array is the vertex of
+ *        the condensed graph that contains the nth vertex of the input graph. This out value is set even if the function
+ *        returns false.
+ * @param Options for the conversion, @see EConvertToGraphOptions.
+ *
+ * @return True if the input graph had and cycles, false otherwise. If false is returned, OutGraphBuffer and OutGraph
+ *         are reset to empty, but the other output variables are still populated if present.
+ */
+bool TryConstructCondensationGraph(TConstArrayView<TConstArrayView<FVertex>> Graph,
+	TArray64<FVertex>& OutGraphBuffer, TArray<TConstArrayView<FVertex>>& OutGraph,
+	TArray64<FVertex>* OutOutVertexToInVerticesBuffer, TArray<TConstArrayView<FVertex>>* OutOutVertexToInVertices,
+	TArray<FVertex>* OutInVertexToOutVertex,
+	EConvertToGraphOptions Options = EConvertToGraphOptions::None);
+
+/**
+ * Construct the TransposeGraph (see @ConstructTransposeGraph), but only include edges in the original graph
+ * from the given vertices (in the transpose graph, edges will be present to the given vertices).
+ * Additionally, limit the number of edges in the returned graph, and report which InVertices were added.
+ * This is used to create a partial ReachedBy graph from a large Reachability graph.
+ * 
+ * @param InVertices Vertices in the input graph, edges from which will be included in the output.
+ *        This list will be modified: vertices will be sorted from largest number of edges to smallest.
+ */
+void ConstructPartialTransposeGraph(TConstArrayView<TConstArrayView<FVertex>> Graph,
+	TArrayView<FVertex> InVertices, int64 MaxOutGraphEdges,
+	TArray64<FVertex>& OutTransposeGraphBuffer, TArray<TConstArrayView<FVertex>>& OutTransposeGraph,
+	TArray<FVertex>& OutInVerticesInOutGraph);
+
+}
\ No newline at end of file
diff --git a/Engine/Source/Editor/UnrealEd/Private/Cooker/Algo/GraphReachability.cpp b/Engine/Source/Editor/UnrealEd/Private/Cooker/Algo/GraphReachability.cpp
new file mode 100644
--- /dev/null
+++ b/Engine/Source/Editor/UnrealEd/Private/Cooker/Algo/GraphReachability.cpp
@@ -0,0 +1,956 @@
+// Copyright Epic Games, Inc. All Rights Reserved.
+
+#include "GraphReachability.h"
+
+#include "Algo/BinarySearch.h"
+#include "Algo/IsSorted.h"
+#include "Algo/Sort.h"
+#include "Algo/Unique.h"
+#include "Containers/BinaryHeap.h"
+#include "Math/NumericLimits.h"
+#include "Misc/AssertionMacros.h"
+#include "Misc/AutomationTest.h"
+#include "Misc/EnumClassFlags.h"
+#include "SortSpecialCases.h"
+
+#if WITH_DEV_AUTOMATION_TESTS
+#include "Containers/UnrealString.h"
+#include "Misc/StringBuilder.h"
+#endif
+
+template <typename ArrayType, typename NumToAddSizeType>
+void ReserveGrowth(ArrayType& Array, NumToAddSizeType NumToAdd)
+{
+	typename ArrayType::SizeType OriginalCount = Array.Num();
+	Array.AddUninitialized(NumToAdd); // this will grow the array geometrically using CalculateSlackGrow
+	Array.SetNum(OriginalCount, false /* bAllowShrinking */);
+}
+
+namespace Algo::Graph
+{
+
+/** Implements ConstructReachabilityGraph. */
+class FReachabilityBuilder
+{
+public:
+	FReachabilityBuilder(TConstArrayView<TConstArrayView<FVertex>> InGraph, TArray64<FVertex>& OutReachabilityGraphBuffer,
+		TArray<TConstArrayView<FVertex>>& OutGraph);
+
+	/** Calculate ReachabilityGraph for the current Graph, which might be cyclic or acyclic. */
+	void Build();
+
+private:
+	/** Status for vertices in the graphsearch. InProgress should never be encountered since our graph is acyclic. */
+	enum class EVertexStatus : uint8
+	{
+		NotStarted,
+		InProgress,
+		Done,
+	};
+
+	/**
+	 * A view into an array. Similar to an arrayview, but it is based on a pointer to the array and an index rather
+	 * than a pointer to the allocation. This prevents us from having to update it when the array reallocates.
+	 */
+	struct FSliceOfArray
+	{
+		TArray64<FVertex>* Buffer = nullptr;
+		int64 Offset = INDEX_NONE;
+		int32 Num = 0;
+		TConstArrayView<FVertex> GetView() const
+		{
+			return TConstArrayView<FVertex>(Buffer->GetData() + Offset, Num);
+		}
+		void Set(TArray64<FVertex>* InBuffer, int64 InOffset, int32 InNum)
+		{
+			Buffer = InBuffer;
+			Offset = InOffset;
+			Num = InNum;
+		}
+	};
+	/**
+	 * Data about each vertex that lasts beyond the current graph search stack. ReachablesSlice will be volatile during the 
+	 * search that first finds the vertex but will end up as a slice of the ReachabilityBuffer.
+	 */
+	struct FVertexData
+	{
+		FSliceOfArray ReachablesSlice;
+		EVertexStatus Status = EVertexStatus::NotStarted;
+	};
+	/**
+	 * A buffer that holds the ReachablesSlice for one or more vertices encountered during a graph search stack. It is
+	 * formed from a vertex that cannot be embedded into its referencer's reachability slice. When the stack is
+	 * emptied, the Buffer is copied onto the reachabilitybuffer after the slices embedded into the reachabilitybuffer
+	 * during the search, and the EmbeddedSlices are updated to point to the copy on the reachabilitybuffer.
+	 */
+	struct FSeparateBuffer
+	{
+		TArray64<FVertex> Buffer;
+		TArray<FSliceOfArray*> EmbeddedSlices;
+	};
+	/**
+	 * Information about the reachables of direct edges of a vertex on the stack. The reachables of these edges are
+	 * embedded into the vertex's reachables during FinishVertex, if possible.
+	 */
+	struct FStackEdgeData
+	{
+		FVertex EdgeVertex = InvalidVertex;
+		FSeparateBuffer EdgeBuffer;
+		bool bEmbeddable = false;
+	};
+	/**
+	 * Information about the reachables of a vertex on the stack, including its direct edge's reachables that need
+	 * to be embedded into it. Buffer and EmbeddedSlices are provided by the referencer of the vertex, for it to
+	 * store its reachables into so that it can be embedded into its referencer when its turn to be embedded comes.
+	 */
+	struct FStackData
+	{
+		TArray64<FVertex>* Buffer;
+		TArray<FSliceOfArray*>* EmbeddedSlices;
+		TArray<FStackEdgeData> StackEdgeDatas;
+		FVertex Vertex = InvalidVertex;
+		int32 NextEdge = 0;
+	};
+
+private:
+	/** Calculate ReachabilityGraph for the current Graph. Assumes Graph is acyclic. */
+	void BuildAcyclic(bool bShrink);
+	/**
+	 * Push a vertex onto the graph search stack and initialize the data used to calculate its reachables.
+	 * Its reachables (and the list of ReachabilitySlices for all vertices using the buffer) will be stored in
+	 * Buffer and EmbeddedSlices.
+	 */
+	void StartVertex(FVertex Vertex, TArray64<FVertex>* Buffer, TArray<FSliceOfArray*>* EmbeddedSlices);
+	/**
+	 * Called after all directedges have recursively calculated their reachables (or were previously calculated.
+	 * Creates the list of reachables for the vertex, embeds the reachables of all possible direct edges into it,
+	 * stores the list in the Buffer that was passed into StartVertex, and records its VertexData's ReachabilitySlice
+	 * as (currently) pointing to that Buffer. The ReachabilitySlice will be updated as the Buffer moves around.
+	 * Also pops the vertex from the graph search stack.
+	 */
+	void FinishVertex();
+
+	/** Called when the input graph is already acyclic. Calculate RootToLeafOrder required by BuildAcyclic. */
+	void ReadRootToLeafOrderFromInputVertexToGraphVertex();
+	/**
+	 * Called when the input graph has cycles and so BuildAcyclic was called on the CondensationGraph.
+	 * Transform the results from the CondensationGraph back into the results for the input graph.
+	 */
+	void TransformCondensationReachabilityToInputReachability();
+
+	/**
+	 * Append the given SourceBuffer onto the TargetBuffer, and update the EmbeddedSlices that were pointing to the
+	 * SourceBuffer to point to their new copy in the TargetBuffer.
+	 */
+	static void RelocateViews(TArray64<FVertex>& InOutTargetBuffer, TArray64<FVertex>& SourceBuffer,
+		TArrayView<FSliceOfArray*> EmbeddedSlices);
+
+private:
+
+	/**
+	 * The graph for which we calculate reachability. Possibly used during the graph search, but we will instead
+	 * operate on its CondensationGraph if it has cycles
+	 */
+	TConstArrayView<TConstArrayView<FVertex>> InputGraph;
+	/** Output buffer for the edges of the reachability graph. We write to it during the graph search. */
+	TArray64<FVertex>& ReachabilityBuffer;
+	/**
+	 * Output list of arrayviews for the edges of the reachabilitygraph. To save memory, these arrayviews will
+	 * overlap each other. We calculate it at the end from intermediate data
+	 */
+	TArray<TConstArrayView<FVertex>>& ReachabilityGraph;
+
+
+	/** The graph of edges used in the graph search. Either the InputGraph or its CondensationGraph. */
+	TConstArrayView<TConstArrayView<FVertex>> Graph;
+
+	/** Status and reachables Data about each vertex. Reachables are recorded as a slice into a changeable buffer. */
+	TArray<FVertexData> VertexDatas;
+	/** Information to keep place and calculate reachables for vertices encountered during a search from a root. */
+	TArray<FStackData> Stack;
+	/**
+	 * List of reachables buffers for vertices that were not embeddable during the graph search and need to be added
+	 * to reachabilitybuffer afterwards.
+	 */
+	TArray<FSeparateBuffer> SeparateBuffers;
+
+	/** RootToLeafOrder of the input vertices. Calculated at the same time as we check for cycles and create the
+	 * CondensationGraph. Our algorithm for embedding reachables relies on acyclic graph and RootToLleaf order.
+	 */
+	TArray<FVertex> RootToLeafOrder;
+	/** Buffer for edges for the CondensationGraph. Empty and unused if InputGraph is acyclic. */
+	TArray64<FVertex> CondensationGraphBuffer;
+	/** ArrayViews of edges of the CondensationGraph. Empty and unused if InputGraph is acyclic. */
+	TArray<TConstArrayView<FVertex>> GraphEdges;
+	/** Buffer for input graph vertices for each condensationgraph vertex. unused if InputGraph is acyclic. */
+	TArray64<FVertex> CondensationVertexToInputVerticesBuffer;
+	/** List of input graph vertices for each condensationgraph vertex unused if InputGraph is acyclic. */
+	TArray<TConstArrayView<FVertex>> GraphVertexToInputVertices;
+	/**
+	 * Map from input vertex to condensationgraph vertex; multiple input vertices can map to the same condensation
+	 * vertex. Unused (after creating RootToLeafOrder) if InputGraph is acyclic
+	 */
+	TArray<FVertex> InputVertexToGraphVertex;
+
+	/** Scratch-space buffer for unioning sets of vertices during FinishVertex. NumVertices array of bools. */
+	TBitArray<> IsReachable;
+};
+
+void ConstructReachabilityGraph(TConstArrayView<TConstArrayView<FVertex>> InGraph, TArray64<FVertex>& OutReachabilityGraphBuffer,
+	TArray<TConstArrayView<FVertex>>& OutGraph)
+{
+	FReachabilityBuilder Builder(InGraph, OutReachabilityGraphBuffer, OutGraph);
+	Builder.Build();
+}
+
+FReachabilityBuilder::FReachabilityBuilder(TConstArrayView<TConstArrayView<FVertex>> InGraph, TArray64<FVertex>& OutReachabilityGraphBuffer,
+	TArray<TConstArrayView<FVertex>>& OutGraph)
+	: InputGraph(InGraph)
+	, ReachabilityBuffer(OutReachabilityGraphBuffer)
+	, ReachabilityGraph(OutGraph)
+{
+}
+void FReachabilityBuilder::Build()
+{
+	// Our compaction algorithm requires that the graph is acyclic, so create the condensation graph,
+	// run the reachability graph on the condensation graph, and then transform the reachability of the
+	// condensation graph into reachability for the input graph
+	if (TryConstructCondensationGraph(InputGraph, CondensationGraphBuffer, GraphEdges, &CondensationVertexToInputVerticesBuffer,
+		&GraphVertexToInputVertices, &InputVertexToGraphVertex))
+	{
+		Graph = GraphEdges;
+		RootToLeafOrder = Algo::RangeArray<TArray<FVertex>>(0, Graph.Num());
+		BuildAcyclic(false /* bShrink */);
+		TransformCondensationReachabilityToInputReachability();
+	}
+	else
+	{
+		// If the input graph is acyclic, no transforming is needed, but we do still need a root to leaf
+		// order of the input graph. TryConstructCondensationGraph left it for us in InputVertexToGraphVertex
+		Graph = InputGraph;
+		ReadRootToLeafOrderFromInputVertexToGraphVertex();
+		BuildAcyclic(true /* bShrink */);
+	}
+}
+
+void FReachabilityBuilder::BuildAcyclic(bool bShrink)
+{
+	// About compaction:
+	// The reachability graph for N vertices could be as large as NxN, and in practice will often
+	// have size N*N/2. When N is 1 million this is too big to fit into memory in our current machines
+	// which have on the order of 100GB. (Half of million*million = 1TB*BytesPerElement/2).
+	// But reachability graphs have a lot of redundancy and can be compacted; multiple vertices will overlap
+	// each other in the buffer of reachable vertices.
+	// As we search through the graph, we embed the reachables of direct edges into the reachables of the current
+	// vertex when possible.
+	int32 NumVertices = Graph.Num();
+	Stack.Reset(NumVertices);
+	SeparateBuffers.Reset(NumVertices);
+	VertexDatas.Reset(NumVertices);
+	VertexDatas.SetNum(NumVertices, false /* bAllowShrinking */);
+	ReachabilityBuffer.Reset();
+	IsReachable.Init(false, NumVertices);
+
+	// Start searching at root vertices so we can maximize compaction
+	for (FVertex RootVertex : RootToLeafOrder)
+	{
+		EVertexStatus RootStatus = VertexDatas[RootVertex].Status;
+		check(RootStatus == EVertexStatus::NotStarted || RootStatus == EVertexStatus::Done);
+		if (RootStatus == EVertexStatus::Done)
+		{
+			continue;
+		}
+
+		// For each new root vertex, graph search from it to calculate its direct edges' reachables, union those
+		// reachables into its own complete list of reachables, and embed the direct edges' reachables into its
+		// bigger list.
+		check(Stack.IsEmpty());
+		StartVertex(RootVertex, &ReachabilityBuffer, nullptr);
+		while (!Stack.IsEmpty())
+		{
+			FStackData& StackData = Stack.Last();
+			FVertex Vertex = StackData.Vertex;
+			TConstArrayView<FVertex> Edges = Graph[Vertex];
+			bool bPushed = false;
+			while (!bPushed && StackData.NextEdge < Edges.Num())
+			{
+				int32 NextEdge = StackData.NextEdge++;
+				FVertex EdgeVertex = Edges[NextEdge];
+				if (EdgeVertex == Vertex)
+				{
+					continue; // Ignore edges to self
+				}
+				FVertexData& EdgeData = VertexDatas[EdgeVertex];
+				if (EdgeData.Status == EVertexStatus::NotStarted)
+				{
+					FStackEdgeData& StackEdgeData = StackData.StackEdgeDatas[NextEdge];
+					StartVertex(EdgeVertex, &StackEdgeData.EdgeBuffer.Buffer, &StackEdgeData.EdgeBuffer.EmbeddedSlices);
+					bPushed = true;
+					break;
+				}
+				else
+				{
+					// Our graph is acyclic so we should never encounter an in-progress vertex
+					check(EdgeData.Status == EVertexStatus::Done);
+				}
+			}
+			if (!bPushed)
+			{
+				FinishVertex();
+			}
+		}
+
+		// Whenever we are unable to embed a vertex into its parent, we add it as a SeparateBuffer.
+		// Now that the graph search from the root is complete, add all these SeparateBuffers onto
+		// the reachabilitybuffer after the root's reachables.
+		for (FSeparateBuffer& Buffer : SeparateBuffers)
+		{
+			RelocateViews(ReachabilityBuffer, Buffer.Buffer, Buffer.EmbeddedSlices);
+		}
+		SeparateBuffers.Reset();
+	}
+	check(SeparateBuffers.IsEmpty());
+
+	// We're done creating ReachabilityBuffer. All of the vertices have FSliceOfArrays into it to record
+	// their reachables. Shrink it now to save memory, unless we're going to reallocate it anyway later.
+	if (bShrink)
+	{
+		ReachabilityBuffer.Shrink();
+	}
+
+	// Convert the FSliceOfArrays into TArrayViews for the output ReachabilityGraph
+	FVertex* ReachabilityData = ReachabilityBuffer.GetData();
+	ReachabilityGraph.Reset(NumVertices);
+	for (FVertex Vertex = 0; Vertex < NumVertices; ++Vertex)
+	{
+		FVertexData& VertexData = VertexDatas[Vertex];
+		check(VertexData.Status == EVertexStatus::Done);
+		FSliceOfArray& ReachablesSlice = VertexData.ReachablesSlice;
+		check(ReachablesSlice.Buffer == &ReachabilityBuffer);
+		ReachabilityGraph.Emplace(ReachablesSlice.GetView());
+	}
+
+	// Clear memory for the structures we used during the search
+	Stack.Empty();
+	VertexDatas.Empty();
+	SeparateBuffers.Empty();
+	IsReachable.Empty();
+}
+
+void FReachabilityBuilder::StartVertex(FVertex Vertex, TArray64<FVertex>* Buffer, TArray<FSliceOfArray*>* EmbeddedSlices)
+{
+	FStackData* StackDataPointer = Stack.GetData();
+	FStackData& StackData = Stack.Emplace_GetRef();
+	check(Stack.GetData() == StackDataPointer); // StackDatas have pointers to TArrays higher on the stack so we cannot let it reallocate
+	StackData.Buffer = Buffer;
+	StackData.EmbeddedSlices = EmbeddedSlices;
+	StackData.Vertex = Vertex;
+	StackData.NextEdge = 0;
+	TConstArrayView<FVertex> Edges = Graph[Vertex];
+	int32 NumEdges = Edges.Num();
+	StackData.StackEdgeDatas.Reset(NumEdges);
+	for (int32 EdgeIndex = 0; EdgeIndex < NumEdges; ++EdgeIndex)
+	{
+		FStackEdgeData& StackEdgeData = StackData.StackEdgeDatas.Emplace_GetRef();
+		StackEdgeData.EdgeVertex = Edges[EdgeIndex];
+	}
+	FVertexData& VertexData = VertexDatas[Vertex];
+	VertexData.Status = EVertexStatus::InProgress;
+};
+
+void FReachabilityBuilder::FinishVertex()
+{
+	FStackData& StackData = Stack.Last();
+	FVertex StackVertex = StackData.Vertex;
+	int32 NumVertices = Graph.Num();
+	int32 NumEdges = StackData.StackEdgeDatas.Num();
+
+	// Sort the edges from largest to smallest reachables count; push already allocated to the end
+	// DO NOT move the FStackEdgeDatas around. Some vertices have FSliceOfArrays pointing at TArray
+	// fields on the FStackEdgeData, and these pointers will be left pointing to an array with the 
+	// wrong allocation if we move-assign the arrays' allocations.
+	TArray<FStackEdgeData*> StackEdgeDatas;
+	StackEdgeDatas.Reset(NumEdges);
+	for (FStackEdgeData& StackEdgeData : StackData.StackEdgeDatas)
+	{
+		StackEdgeDatas.Add(&StackEdgeData);
+	}
+	Algo::Sort(StackEdgeDatas, [this, StackVertex](const FStackEdgeData* A, const FStackEdgeData* B)
+		{
+			bool bAIsEmpty = A->EdgeBuffer.Buffer.IsEmpty();
+			if (bAIsEmpty != B->EdgeBuffer.Buffer.IsEmpty())
+			{
+				return !bAIsEmpty;
+			}
+			int32 NumA = VertexDatas[A->EdgeVertex].ReachablesSlice.Num;
+			int32 NumB = VertexDatas[B->EdgeVertex].ReachablesSlice.Num;
+			if (NumA != NumB)
+			{
+				return NumA > NumB;
+			}
+			return A->EdgeVertex < B->EdgeVertex;
+		});
+
+	// Ignoring edges that were already complete before we entered this vertex, try to embed all of the direct
+	// edges' reachables into this vertex's reachables. We can only embed a set of reachables if it does not
+	// overlap with reachables we have already collected. (So in a diamond graph A->B->D, A->C->D, we will only
+	// be able to embed one of B or C into A). Iterate through direct edges from largest to smallest and greedily
+	// embed the largest direct edge possible until no more are possible.
+	
+	// This is a heavy set of set operations for large vertices; we use the IsReachable array of bools to execute
+	// those operations. An array of bools for all vertices has the down-side of being expensive to clear. To mitigate
+	// that we initialze IsReachable to false for all vertices at the start of BuildAcyclic, and we incrementally clear
+	// it back to false when we're done with it here.
+
+	// List of leftover vertices to add on to this vertex's reachables from direct edges that we could not embed.
+	TArray<FVertex> NonEmbeddedReachables;
+
+	int32 NumReachables = 1; // Start at 1 to count this vertex as part of its own reachables
+	for (FStackEdgeData* StackEdgeData : StackEdgeDatas)
+	{
+		FVertex EdgeVertex = StackEdgeData->EdgeVertex;
+		FVertexData& EdgeData = VertexDatas[EdgeVertex];
+		if (StackEdgeData->EdgeBuffer.Buffer.IsEmpty())
+		{
+			// The edgevertex was already complete before we entered this vertex, so we can't embed it
+			StackEdgeData->bEmbeddable = false;
+		}
+		else
+		{
+			StackEdgeData->bEmbeddable = true;
+			for (FVertex Vertex : EdgeData.ReachablesSlice.GetView())
+			{
+				if (IsReachable[Vertex])
+				{
+					// The edgevertex's reachables overlap with a reachable we already added; we can't embed it
+					StackEdgeData->bEmbeddable = false;
+					break;
+				}
+			}
+			if (StackEdgeData->bEmbeddable)
+			{
+				// The edgevertex is embeddable, and its the biggest, so commit to the embedding now
+				for (FVertex Vertex : EdgeData.ReachablesSlice.GetView())
+				{
+					IsReachable[Vertex] = true;
+				}
+				NumReachables += EdgeData.ReachablesSlice.Num;
+			}
+		}
+	}
+
+	// Gather all the unique vertices we don't already have from the non-embeddable edgevertices
+	TArray<FVertex> OverlappedReachables;
+	for (FStackEdgeData* StackEdgeData : StackEdgeDatas)
+	{
+		if (!StackEdgeData->bEmbeddable)
+		{
+			if (StackEdgeData->EdgeVertex == StackVertex)
+			{
+				continue; // Ignore edges to self
+			}
+			FVertex EdgeVertex = StackEdgeData->EdgeVertex;
+			FVertexData& EdgeData = VertexDatas[EdgeVertex];
+			for (FVertex Vertex : EdgeData.ReachablesSlice.GetView())
+			{
+				if (!IsReachable[Vertex])
+				{
+					++NumReachables;
+					IsReachable[Vertex] = true;
+					NonEmbeddedReachables.Add(Vertex);
+				}
+			}
+		}
+	}
+
+	// Create our ReachablesSlice to point to the buffer we're about to populate
+	int64 InitialBufferNum = StackData.Buffer->Num();
+	FVertexData& VertexData = VertexDatas[StackVertex];
+	VertexData.ReachablesSlice.Set(StackData.Buffer, InitialBufferNum, NumReachables);
+	if (StackData.EmbeddedSlices)
+	{
+		StackData.EmbeddedSlices->Add(&VertexData.ReachablesSlice);
+	}
+
+	// Add all of our reachables into the buffer, copying the reachables for each embedded
+	// edgevertex into a contiguous range so its ReachablesSlice can overlap our own
+	ReserveGrowth(*StackData.Buffer, NumReachables);
+	StackData.Buffer->Add(StackVertex);
+	for (FStackEdgeData* StackEdgeData : StackEdgeDatas)
+	{
+		FVertex EdgeVertex = StackEdgeData->EdgeVertex;
+		FVertexData& EdgeData = VertexDatas[EdgeVertex];
+		if (StackEdgeData->bEmbeddable)
+		{
+			// Copies from the edgevertex's buffer into ours
+			RelocateViews(*StackData.Buffer, StackEdgeData->EdgeBuffer.Buffer, 
+				StackEdgeData->EdgeBuffer.EmbeddedSlices);
+			if (StackData.EmbeddedSlices)
+			{
+				StackData.EmbeddedSlices->Append(StackEdgeData->EdgeBuffer.EmbeddedSlices);
+			}
+		}
+		else if (!StackEdgeData->EdgeBuffer.Buffer.IsEmpty())
+		{
+			// For all of the edgevertices we recursively calculated but cannot embed, add them as separatebuffers to
+			// copy onto reachabilitybuffer after the graph search stack empties.
+			FSeparateBuffer* SeparateBuffersData = SeparateBuffers.GetData();
+			SeparateBuffers.Add(MoveTemp(StackEdgeData->EdgeBuffer));
+			check(SeparateBuffers.GetData() == SeparateBuffersData);
+			FSeparateBuffer& SeparateBuffer = SeparateBuffers.Last();
+			for (FSliceOfArray* EmbeddedSlice : SeparateBuffer.EmbeddedSlices)
+			{
+				EmbeddedSlice->Buffer = &SeparateBuffer.Buffer;
+			}
+		}
+	}
+	StackData.Buffer->Append(NonEmbeddedReachables);
+	check(StackData.Buffer->Num() == InitialBufferNum + NumReachables);
+
+	// Fast-clear IsReachable, by individually clearing all of the values we set to true
+	constexpr int32 FastClearFraction = 10;
+	if (((int64)NumReachables) * FastClearFraction > NumVertices)
+	{
+		// In this case we touched so much of IsReachable that it's faster to bulk-clear all of it
+		IsReachable.Init(false, NumVertices);
+	}
+	else
+	{
+		for (FVertex Vertex : VertexData.ReachablesSlice.GetView())
+		{
+			IsReachable[Vertex] = false;
+		}
+	}
+
+	// Mark our status complete and pop the stack
+	VertexData.Status = EVertexStatus::Done;
+	FStackData* StackDataPointer = Stack.GetData();
+	Stack.Pop(false /* bAllowShrinking */);
+	// StackDatas have pointers to TArrays higher on the stack so we cannot let it reallocate
+	check(Stack.GetData() == StackDataPointer);
+}
+
+void FReachabilityBuilder::TransformCondensationReachabilityToInputReachability()
+{
+	// We walk through all CondensationGraph vertices in the ReachabilityBuffer, and replace them with their single or
+	// multiple InputGraph vertices from GraphVertexToInputVertices.
+	// Whenever we insert vertices in the ReachabilityGraph instead of merely replacing them, we have to update all of
+	// the affected TArrayViews: any TArrayViews that started after that spot have to be shifted higher up the buffer,
+	// and any TArrayViews that overlapped the insertion spot have to have their Num increased.
+	// We handle the shifting by sorting the arrayviews by start point and keeping track of the cumulative increase
+	// as we reach the startpoint of each arrayview. We handle the num growing by keeping a heap of arrayviews that
+	// have started but not yet finished at our current position (heap-sorted by their ending position so we can pop
+	// them off as we move along the buffer).
+
+	// Sort the ArrayViews into ViewsSortedByStart
+	int32 NumVertices = ReachabilityGraph.Num();
+	TArray<TConstArrayView<FVertex>*> ViewsSortedByStart;
+	ViewsSortedByStart.Reserve(NumVertices);
+	for (FVertex Vertex = 0; Vertex < NumVertices; ++Vertex)
+	{
+		ViewsSortedByStart.Add(&ReachabilityGraph[Vertex]);
+	}
+	Algo::Sort(ViewsSortedByStart, [](TConstArrayView<FVertex>* A, TConstArrayView<FVertex>* B)
+		{
+			if (A->GetData() != B->GetData())
+			{
+				return A->GetData() < B->GetData();
+			}
+			return A->Num() > B->Num();
+		});
+	int64 AddedSize = 0;
+	for (FVertex Vertex : ReachabilityBuffer)
+	{
+		AddedSize += GraphVertexToInputVertices[Vertex].Num() - 1;
+	}
+
+	// The new version of ReachabilityBuffer with all of the CondensationGraph verts converted to InputGraph verts
+	TArray64<FVertex> CopyBuffer;
+	CopyBuffer.Reserve(AddedSize + ReachabilityBuffer.Num());
+
+	// Start and end of the source and copy buffers' allocations
+	FVertex* SourceData = ReachabilityBuffer.GetData();
+	FVertex* SourceDataEnd = SourceData + ReachabilityBuffer.Num();
+	FVertex* CopyData = CopyBuffer.GetData();
+
+	// The next ArrayView from ViewsSortedByStart that we need to add when reaching its start position
+	int32 NextIndexInSortedViews = 0;
+
+	// The heap of ArrayViews that overlap the current position. The arrayview that ends next is at top of heap
+	FBinaryHeap<const FVertex*, uint32> ActiveViews;
+
+	// An array NumVertices in size that stores the CumulativeGrowth that was present when we reached the start of the
+	// vertex's ArrayView. Element n corresponds to the vertex in ViewsSortedByStat[n]. We use the delta between
+	// CumulativeGrowth from start to end of the ArrayView into the CondensationGraph's buffer to to detect how much
+	// the ArrayView's Num needs to grow for the InputGraph's version of the ArrayView.
+	TArray<int64> CumulativeGrowthWhenAdded;
+	CumulativeGrowthWhenAdded.SetNumUninitialized(NumVertices);
+	int64 CumulativeGrowth = 0;
+
+	for (FVertex* NextVertex = SourceData; ; ++NextVertex)
+	{
+		// Remove any active ArrayViews that reached the end of their view in the source data, and update their Num
+		while (!ActiveViews.IsEmpty())
+		{
+			uint32 IndexInSortedViewsUnsigned = ActiveViews.Top();
+			const FVertex* SourceViewEnd = ActiveViews.GetKey(IndexInSortedViewsUnsigned);
+			if (SourceViewEnd != NextVertex)
+			{
+				check(SourceViewEnd > NextVertex);
+				break;
+			}
+			int32 IndexInSortedViews = static_cast<int32>(IndexInSortedViewsUnsigned);
+			TConstArrayView<FVertex>& View = *ViewsSortedByStart[IndexInSortedViews];
+			View = TConstArrayView<FVertex>(View.GetData(),
+				View.Num() + CumulativeGrowth - CumulativeGrowthWhenAdded[IndexInSortedViews]);
+			ActiveViews.Pop();
+		}
+		// write the for loop termination condition here so we can break after finalizing arrayviews at end of buffer
+		if (NextVertex == SourceDataEnd)
+		{
+			break;
+		}
+
+		// Add any new ArrayViews from ViewsSortedByStart that start at the current position in SourceData
+		while (NextIndexInSortedViews < ViewsSortedByStart.Num() &&
+			ViewsSortedByStart[NextIndexInSortedViews]->GetData() == NextVertex)
+		{
+			TConstArrayView<FVertex>& View = *ViewsSortedByStart[NextIndexInSortedViews];
+			CumulativeGrowthWhenAdded[NextIndexInSortedViews] = CumulativeGrowth;
+			const FVertex* ViewSourceData = View.GetData();
+			int32 ViewNum = View.Num();
+			const FVertex* ViewEndData = ViewSourceData + ViewNum;
+			View = TConstArrayView<FVertex>(CopyData + (ViewSourceData - SourceData) + CumulativeGrowth, ViewNum);
+			ActiveViews.Add(ViewEndData, static_cast<uint32>(NextIndexInSortedViews));
+			++NextIndexInSortedViews;
+		}
+		check(NextIndexInSortedViews == ViewsSortedByStart.Num() ||
+			ViewsSortedByStart[NextIndexInSortedViews]->GetData() > SourceData);
+
+		// Convert the CondensationGraph vertex to 1 or more Input vertices and copy them onto the CopyBuffer
+		TConstArrayView<FVertex>& OutVertices = GraphVertexToInputVertices[*NextVertex];
+		CopyBuffer.Append(OutVertices);
+		CumulativeGrowth += OutVertices.Num() - 1;
+	}
+	// Swap the new buffer into our output variable
+	ReachabilityBuffer = MoveTemp(CopyBuffer);
+	check(ReachabilityBuffer.GetData() == CopyData);
+
+	// Phase two: we've updated the data in the buffer, but now we need to update which vertices point to which
+	// ranges of it. This phase is much easier since we just need to copy ArrayViews without modifying them.
+	TArray<TConstArrayView<FVertex>> InputReachabilityGraph;
+	int32 NumInputVertices = InputGraph.Num();
+	InputReachabilityGraph.Reserve(NumInputVertices);
+	for (FVertex InputVertex = 0; InputVertex < NumInputVertices; ++InputVertex)
+	{
+		FVertex GraphVertex = InputVertexToGraphVertex[InputVertex];
+		InputReachabilityGraph.Add(ReachabilityGraph[GraphVertex]);
+	}
+
+	// Swap the new graph into our output variable
+	ReachabilityGraph = MoveTemp(InputReachabilityGraph);
+}
+
+void FReachabilityBuilder::ReadRootToLeafOrderFromInputVertexToGraphVertex()
+{
+	int32 NumVertices = Graph.Num();
+	RootToLeafOrder.SetNumUninitialized(NumVertices);
+	for (FVertex& Vertex : RootToLeafOrder)
+	{
+		Vertex = InvalidVertex;
+	}
+	check(InputVertexToGraphVertex.Num() == NumVertices);
+	for (FVertex UnsortedVertex = 0; UnsortedVertex < NumVertices; ++UnsortedVertex)
+	{
+		int32 SortedIndex = (int32)InputVertexToGraphVertex[UnsortedVertex];
+		check(SortedIndex < NumVertices && RootToLeafOrder[SortedIndex] == InvalidVertex);
+		RootToLeafOrder[SortedIndex] = UnsortedVertex;
+	}
+}
+
+void FReachabilityBuilder::RelocateViews(TArray64<FVertex>& InOutTargetBuffer, TArray64<FVertex>& SourceBuffer,
+	TArrayView<FSliceOfArray*> EmbeddedSlices)
+{
+	int64 InitialOffset = InOutTargetBuffer.Num();
+	InOutTargetBuffer.Append(SourceBuffer);
+	for (FSliceOfArray* EmbeddedSlize : EmbeddedSlices)
+	{
+		check(EmbeddedSlize->Buffer == &SourceBuffer);
+		EmbeddedSlize->Buffer = &InOutTargetBuffer;
+		EmbeddedSlize->Offset += InitialOffset;
+	}
+}
+
+} // namespace Algo::Graph
+
+
+
+#if WITH_DEV_AUTOMATION_TESTS
+IMPLEMENT_SIMPLE_AUTOMATION_TEST(FReachabilityTest, "System.Core.Algo.ReachabilityGraph", EAutomationTestFlags::ApplicationContextMask | EAutomationTestFlags::EngineFilter);
+
+bool FReachabilityTest::RunTest(const FString& Parameters)
+{
+	using namespace Algo::Graph;
+
+	TArray<TArray<FVertex>> Buffer;
+	TArray<TConstArrayView<FVertex>> Graph;
+	TArray64<FVertex> ReachabilityBuffer;
+	TArray<TConstArrayView<FVertex>> Reachability;
+	TArray<TArray<FVertex>> ExpectedReachability;
+	auto AddVertex = [&Buffer, &Graph, &ExpectedReachability](FVertex Vertex, TConstArrayView<FVertex> Edges, TConstArrayView<FVertex> ExpectedVertexReachability)
+	{
+		if (Graph.Num() <= Vertex)
+		{
+			Graph.SetNum(Vertex + 1);
+			Buffer.SetNum(Vertex + 1);
+			ExpectedReachability.SetNum(Vertex + 1);
+		}
+		Buffer[Vertex] = Edges;
+		Graph[Vertex] = Buffer[Vertex];
+		ExpectedReachability[Vertex] = ExpectedVertexReachability;
+	};
+	auto Clear = [&Buffer, &Graph, &ExpectedReachability]()
+	{
+		Buffer.Reset();
+		Graph.Reset();
+		ExpectedReachability.Reset();
+	};
+	auto WriteArrayToString = [](TConstArrayView<FVertex> A)
+	{
+		TStringBuilder<256> Writer;
+		Writer << TEXT("[");
+		Writer.Join(A, TEXT(','));
+		Writer << TEXT("]");
+		return FString(Writer);
+	};
+	auto ConfirmResults = [&Reachability, &ExpectedReachability, &WriteArrayToString](FString& OutError)
+	{
+		int32 NumVertices = Reachability.Num();
+		if (NumVertices != ExpectedReachability.Num())
+		{
+			OutError = FString::Printf(TEXT("reachability graph has incorrect number of vertices, Expected=%d, Actual=%d"),
+				ExpectedReachability.Num(), NumVertices);
+			return false;
+		}
+		for (int32 Vertex = 0; Vertex < NumVertices; ++Vertex)
+		{
+			TArray<FVertex> Expected = ExpectedReachability[Vertex];
+			TArray<FVertex> Actual(Reachability[Vertex]);
+			Algo::Sort(Expected);
+			Algo::Sort(Actual);
+			if (Expected != Actual)
+			{
+				OutError = FString::Printf(TEXT("reachability for vertex %d does not match. Expected=%s, Actual=%s"),
+					Vertex, *WriteArrayToString(Expected), *WriteArrayToString(Actual));
+				return false;
+			}
+		}
+		return true;
+	};
+	auto RunTrialsForCase = [&ConfirmResults, &Graph, &ReachabilityBuffer, &Reachability, this](const TCHAR* TestCaseName)
+	{
+		FString Error;
+		{
+			FReachabilityBuilder Builder(Graph, ReachabilityBuffer, Reachability);
+			Builder.Build();
+		}
+		if (!ConfirmResults(Error))
+		{
+			AddError(FString::Printf(TEXT("Search failed for case \"%s\": %s."), TestCaseName, *Error));
+		}
+	};
+
+	{
+		Clear();
+		AddVertex(0, { }, { 0 });
+		AddVertex(1, { 0 }, { 0,1 });
+		AddVertex(2, { 1 }, { 0,1,2 });
+		RunTrialsForCase(TEXT("Each node depends on the previous one"));
+	}
+	{
+		Clear();
+		AddVertex(0, { 1 }, {0,1,2});
+		AddVertex(1, { 2 }, { 1,2 });
+		AddVertex(2, { }, { 2 });
+		RunTrialsForCase(TEXT("Each node depends on the next one"));
+	}
+	{
+		Clear();
+		AddVertex(0, { 0 }, { 0 });
+		AddVertex(1, { 0,1 }, { 0,1 });
+		AddVertex(2, { 1,2 }, { 0,1,2 });
+		RunTrialsForCase(TEXT("SelfReferences"));
+	}
+	{
+		Clear();
+		AddVertex(0, { 1 }, { 0,1 });
+		AddVertex(1, { 0 }, { 0,1 });
+		RunTrialsForCase(TEXT("Simple cycle"));
+	}
+	{
+		Clear();
+		for (FVertex Vertex = 0; Vertex < 10; ++Vertex)
+		{
+			if (Vertex != 5)
+			{
+				AddVertex(Vertex, { 5 }, { Vertex, 5 });
+			}
+		}
+		AddVertex(5, { }, { 5 });
+		RunTrialsForCase(TEXT("Every node has a dependency on a single node"));
+	}
+	{
+		//              6
+		//             / \
+		//            5   7
+		//           / \   \
+		//          0   1   8
+		//           \ / \   \
+		//            3   4   |
+		//             \   \ /
+		//              \   9
+		//               \ /
+		//                2       
+		Clear();
+		AddVertex(6, { 5,7 }, { 0,1,2,3,4,5,6,7,8,9 });
+		AddVertex(5, { 0,1 }, { 5,0,1,3,4,9,2 });
+		AddVertex(7, { 8 }, { 7,8,9,2 });
+		AddVertex(0, { 3 }, { 0,3,2 });
+		AddVertex(1, { 3,4 }, { 1,3,4,9,2 });
+		AddVertex(8, { 9 }, { 8,9,2 });
+		AddVertex(3, { 2 }, { 3,2 });
+		AddVertex(4, { 9 }, { 4,9,2 });
+		AddVertex(9, { 2 }, { 2,9 });
+		AddVertex(2, { }, { 2 });
+		RunTrialsForCase(TEXT("SketchedOutExample1"));
+	}
+	{
+		//            0
+		//      1          2
+		//   3     4     5     6
+		//  7  8  9 10 11 12 13 14
+		Clear();
+		AddVertex(0, { 1,2 }, { 0,1,2,3,4,5,6,7,8,9,10,11,12,13,14 });
+		AddVertex(1, { 3,4 }, { 1,3,4,7,8,9,10 });
+		AddVertex(2, { 5,6 }, { 2,5,6,11,12,13,14 });
+		AddVertex(3, { 7,8 }, { 3,7,8 });
+		AddVertex(4, { 9,10 }, { 4,9,10 });
+		AddVertex(5, { 11,12 }, { 5,11,12 });
+		AddVertex(6, { 13,14 }, { 6,13,14 });
+		for (FVertex Vertex = 7; Vertex <= 14; ++Vertex)
+		{
+			AddVertex(Vertex, { }, { Vertex });
+		}
+		RunTrialsForCase(TEXT("BinaryTree"));
+	}
+	{
+		//            0
+		//      1          2
+		//   3     4     5     6    // 3 through 6 have edges to every vertex in 7 through 14
+		//  7  8  9 10 11 12 13 14
+		Clear();
+		AddVertex(0, { 1,2 }, { 0,1,2,3,4,5,6,7,8,9,10,11,12,13,14 });
+		AddVertex(1, { 3,4 }, { 1,3,4,7,8,9,10,11,12,13,14 });
+		AddVertex(2, { 5,6 }, { 2,5,6,7,8,9,10,11,12,13,14 });
+		AddVertex(3, { 7,8,9,10,11,12,13,14 }, { 3,7,8,9,10,11,12,13,14 });
+		AddVertex(4, { 7,8,9,10,11,12,13,14 }, { 4,7,8,9,10,11,12,13,14 });
+		AddVertex(5, { 7,8,9,10,11,12,13,14 }, { 5,7,8,9,10,11,12,13,14 });
+		AddVertex(6, { 7,8,9,10,11,12,13,14 }, { 6,7,8,9,10,11,12,13,14 });
+		for (FVertex Vertex = 7; Vertex <= 14; ++Vertex)
+		{
+			AddVertex(Vertex, { }, { Vertex });
+		}
+		RunTrialsForCase(TEXT("BinaryTreeWithEverythingFromRow2ToRow3"));
+	}
+	{
+		Clear();
+		AddVertex(0, { 1 }, { 0,1,2,3 });
+		AddVertex(1, { 0,2 }, { 0,1,2,3 });
+		AddVertex(2, { 3 }, { 2,3 });
+		AddVertex(3, { }, { 3 });
+		RunTrialsForCase(TEXT("Cycle in the root and with the root cycle depending on a chain of non-cycle verts"));
+	}
+	{
+		Clear();
+		AddVertex(0, { }, { 0 });
+		AddVertex(1, { 0 }, { 0,1 });
+		AddVertex(2, { 1,3 }, { 0,1,2,3 });
+		AddVertex(3, { 2 }, { 0,1,2,3 });
+		RunTrialsForCase(TEXT("Cycle in the root and with the root cycle depending on a chain of non-cycle verts, submitted in reverse"));
+	}
+	{
+		Clear();
+		AddVertex(0, { 1 }, { 0,1,2,3 });
+		AddVertex(1, { 2 }, { 1,2,3 });
+		AddVertex(2, { 3 }, { 2,3 });
+		AddVertex(3, { 2 }, { 2,3 });
+		RunTrialsForCase(TEXT("Cycle at a leaf and a chain from the root depending on that cycle"));
+	}
+	{
+		Clear();
+		AddVertex(0, { 1 }, { 0,1 });
+		AddVertex(1, { 0 }, { 0,1 });
+		AddVertex(2, { 1 }, { 2,1,0 });
+		AddVertex(3, { 2 }, { 3,2,1,0 });
+		RunTrialsForCase(TEXT("Cycle in the root and with the root cycle depending on a chain of non-cycle verts"));
+	}
+	{
+		Clear();
+		AddVertex(0, { 1 }, { 0,1,2,3 });
+		AddVertex(1, { 2,3 }, { 1,2,3 });
+		AddVertex(2, { 1,3 }, { 1,2,3 });
+		AddVertex(3, { 1,2 }, { 1,2,3 });
+		RunTrialsForCase(TEXT("Vertex dependent upon a cycle"));
+	}
+	{
+		// 0 -> 1 -> 2 -> 3 -> 1
+		//           |
+		//           v
+		//      4 -> 5 -> 6 -> 4
+		Clear();
+		AddVertex(0, { 1 }, { 0,1,2,3,4,5,6 });
+		AddVertex(1, { 2 }, { 1,2,3,4,5,6 });
+		AddVertex(2, { 3,5 }, { 1,2,3,4,5,6 });
+		AddVertex(3, { 1 }, { 1,2,3,4,5,6 });
+		AddVertex(4, { 5 }, { 4,5,6 });
+		AddVertex(5, { 6 }, { 4,5,6 });
+		AddVertex(6, { 4 }, { 4,5,6 });
+		RunTrialsForCase(TEXT("One cycle dependent upon another"));
+	}
+	{
+		// 5 -> 0 -> (1 -> 2 -> 1)
+		//      |          |
+		//      |          v
+		//      |          5
+		//      v
+		//     (3 -> 4 -> 3)
+		Clear();
+		AddVertex(0, { 1,3 }, { 0,1,2,3,4,5 });
+		AddVertex(1, { 2 }, { 0,1,2,3,4,5 });
+		AddVertex(2, { 1,5 }, { 0,1,2,3,4,5 });
+		AddVertex(3, { 4 }, { 3,4 });
+		AddVertex(4, { 3 }, { 3,4 });
+		AddVertex(5, { 0 }, { 0,1,2,3,4,5 });
+		RunTrialsForCase(TEXT("MutuallyReachableSet Problem1"));
+	}
+	{
+		Clear();
+		TArray<FVertex> All;
+		TArray<FVertex> AllButItself;
+		for (FVertex Vertex = 0; Vertex < 6; ++Vertex)
+		{
+			All.Add(Vertex);
+			AllButItself.Add(Vertex);
+		}
+		for (FVertex Vertex = 0; Vertex < 6; ++Vertex)
+		{
+			AllButItself.Remove(Vertex);
+			AddVertex(Vertex, AllButItself, All);
+			AllButItself.Add(Vertex);
+		}
+		RunTrialsForCase(TEXT("FullyConnectedGraph"));
+	}
+
+	return true;
+}
+#endif
\ No newline at end of file
diff --git a/Engine/Source/Editor/UnrealEd/Private/Cooker/Algo/GraphReachability.h b/Engine/Source/Editor/UnrealEd/Private/Cooker/Algo/GraphReachability.h
new file mode 100644
--- /dev/null
+++ b/Engine/Source/Editor/UnrealEd/Private/Cooker/Algo/GraphReachability.h
@@ -0,0 +1,24 @@
+// Copyright Epic Games, Inc. All Rights Reserved.
+
+#pragma once
+
+#include "GraphConvert.h"
+
+namespace Algo::Graph
+{
+
+/**
+ * Calculates the reachability graph for the given input graph: an edge exists in the reachability graph from vertex A
+ * to vertex B if and only if B is reachable from A by following edges in the input graph.
+ * 
+ * @param InGraph The input graph
+ * @param OutReachabilityGraphBuffer Buffer that holds memory for the output graph. It must not be deallocated or
+ *        modified until OutReachabilityGraph is no longer referenced.
+ * @param OutReachabilityGraph The reachability graph of the input graph. Note that it is in Graph Form but is
+ *        unnormalized: edges for each vertex are arbitrarily ordered according to internal calculations to 
+ *        reduce memory use.
+ */
+void ConstructReachabilityGraph(TConstArrayView<TConstArrayView<FVertex>> InGraph, TArray64<FVertex>& OutReachabilityGraphBuffer,
+	TArray<TConstArrayView<FVertex>>& OutGraph);
+
+}
\ No newline at end of file
diff --git a/Engine/Source/Editor/UnrealEd/Private/Cooker/Algo/SortSpecialCases.h b/Engine/Source/Editor/UnrealEd/Private/Cooker/Algo/SortSpecialCases.h
new file mode 100644
--- /dev/null
+++ b/Engine/Source/Editor/UnrealEd/Private/Cooker/Algo/SortSpecialCases.h
@@ -0,0 +1,138 @@
+// Copyright Epic Games, Inc. All Rights Reserved.
+
+#pragma once
+
+#include "Algo/BinarySearch.h"
+#include "Containers/Array.h"
+#include "Containers/ArrayView.h"
+#include "HAL/Platform.h"
+#include "Misc/AssertionMacros.h"
+#include "Templates/Less.h"
+#include "Templates/UnrealTemplate.h"
+#include "Traits/ElementType.h"
+
+namespace Algo
+{
+
+/** Given two ranges each of which is sorted and unique, calculate the sorted and unique Union */
+template <typename RangeTypeA, typename RangeTypeB, typename RangeTypeMerged, typename PredicateType>
+inline void MergeUniqueSortedRanges(RangeTypeA&& A, RangeTypeB&& B, RangeTypeMerged& OutMerged,
+	PredicateType Pred)
+{
+	typedef TElementType_T<RangeTypeA> ElementType;
+	typedef decltype(GetNum(A)) IndexType;
+	
+	IndexType ANum = GetNum(A);
+	IndexType BNum = GetNum(B);
+	OutMerged.Reset(ANum >= BNum ? BNum : ANum);
+	const ElementType* AData = GetData(A);
+	const ElementType* BData = GetData(B);
+	const ElementType* AEnd = AData + ANum;
+	const ElementType* BEnd = BData + BNum;
+	while (AData < AEnd && BData < BEnd)
+	{
+		if (Pred(*AData, *BData))
+		{
+			OutMerged.Add(*(AData++));
+		}
+		else if (Pred(*BData, *AData))
+		{
+			OutMerged.Add(*(BData++));
+		}
+		else
+		{
+			OutMerged.Add(*(AData++));
+			++BData;
+		}
+	}
+	while (AData < AEnd)
+	{
+		OutMerged.Add(*(AData++));
+	}
+	while (BData < BEnd)
+	{
+		OutMerged.Add(*(BData++));
+	}
+};
+
+/** Given two ranges each of which is sorted and unique, calculate the sorted and unique Union */
+template <typename RangeTypeA, typename RangeTypeB, typename RangeTypeMerged>
+inline void MergeUniqueSortedRanges(RangeTypeA&& A, RangeTypeB&& B, RangeTypeMerged& OutMerged)
+{
+	MergeUniqueSortedRanges(Forward<RangeTypeA>(A), Forward<RangeTypeB>(B), OutMerged,
+		TLess<TElementType_T<RangeTypeA>>());
+}
+
+/** Given two ranges each of which is sorted and unique, calculate the sorted and unique Union */
+template <typename RangeTypeMerged, typename RangeTypeA, typename RangeTypeB, typename PredicateType>
+inline RangeTypeMerged MergeUniqueSortedRanges(RangeTypeA&& A, RangeTypeB&& B, PredicateType Pred)
+{
+	RangeTypeMerged OutMerged;
+	MergeUniqueSortedRanges(Forward<RangeTypeA>(A), Forward<RangeTypeB>(B), OutMerged, MoveTemp(Pred));
+	return OutMerged;
+}
+
+/** Given two ranges each of which is sorted and unique, calculate the sorted and unique Union */
+template <typename RangeTypeMerged, typename RangeTypeA, typename RangeTypeB>
+inline RangeTypeMerged MergeUniqueSortedRanges(RangeTypeA&& A, RangeTypeB&& B)
+{
+	RangeTypeMerged OutMerged;
+	MergeUniqueSortedRanges(Forward<RangeTypeA>(A), Forward<RangeTypeB>(B), OutMerged,
+		TLess<TElementType_T<RangeTypeA>>());
+	return OutMerged;
+}
+
+/**
+ * Given an array that is sorted except for the element at the ModifiedIndex, shift the
+ * modified element into its correct sorted location.
+ */
+template <typename RangeType, typename ModifiedIndexType, typename PredicateType>
+inline void RestoreSort(RangeType&& Range, ModifiedIndexType ModifiedIndex, PredicateType Pred)
+{
+	typedef TElementType_T<RangeType> ElementType;
+	typedef decltype(GetNum(Range)) IndexType;
+	typedef TArrayView<ElementType, IndexType> ArrayViewType;
+
+	ElementType* RangeData = GetData(Range);
+	IndexType RangeNum = GetNum(Range);
+	ArrayViewType RangeArray(RangeData, RangeNum);
+	ElementType& Modified = RangeData[ModifiedIndex];
+	if (ModifiedIndex > 0 && Pred(Modified, RangeData[ModifiedIndex - 1]))
+	{
+		ArrayViewType LowerRange = RangeArray.Left(ModifiedIndex - 1);
+		IndexType DesiredIndex = Algo::UpperBound(LowerRange, Modified, Pred);
+		ElementType Swap = MoveTemp(Modified);
+		for (IndexType SwapFromIndex = ModifiedIndex - 1; SwapFromIndex >= DesiredIndex; --SwapFromIndex)
+		{
+			RangeData[SwapFromIndex + 1] = MoveTemp(RangeData[SwapFromIndex]);
+		}
+		RangeData[DesiredIndex] = MoveTemp(Swap);
+	}
+	else if (ModifiedIndex < RangeNum - 1 && Pred(RangeData[ModifiedIndex + 1], Modified))
+	{
+		ArrayViewType UpperRange = RangeArray.RightChop(ModifiedIndex + 1);
+		IndexType DesiredIndex = ModifiedIndex + 1 + Algo::LowerBound(UpperRange, Modified, Pred);
+		ElementType Swap = MoveTemp(Modified);
+		for (IndexType SwapFromIndex = ModifiedIndex + 1; SwapFromIndex <= DesiredIndex; ++SwapFromIndex)
+		{
+			RangeData[SwapFromIndex - 1] = MoveTemp(RangeData[SwapFromIndex]);
+		}
+		RangeData[DesiredIndex] = MoveTemp(Swap);
+	}
+}
+
+/** Construct and return a range containing consecutive integers in the range [Min, MaxPlusOne) */
+template <typename RangeType, typename IntType>
+RangeType RangeArray(IntType Min, IntType MaxPlusOne)
+{
+	check(MaxPlusOne >= Min);
+	RangeType Result;
+	Result.Reserve(MaxPlusOne - Min);
+	for (IntType Value = Min; Value < MaxPlusOne; ++Value)
+	{
+		Result.Add(Value);
+	}
+	return Result;
+}
+
+}
