StaticLighting: Check UnrealLightmass version against editor to ensure it is up-to-date and prompt for recompilation when needed

diff --git a/Engine/Source/Editor/UnrealEd/Private/StaticLightingSystem/StaticLightingPrivate.h b/Engine/Source/Editor/UnrealEd/Private/StaticLightingSystem/StaticLightingPrivate.h
--- a/Engine/Source/Editor/UnrealEd/Private/StaticLightingSystem/StaticLightingPrivate.h
+++ b/Engine/Source/Editor/UnrealEd/Private/StaticLightingSystem/StaticLightingPrivate.h
@@ -269,6 +269,8 @@ public:
 	 */
 	FStaticLightingSystem(const FLightingBuildOptions& InOptions, UWorld* InWorld, ULevel* InLightingScenario);
 	~FStaticLightingSystem();
+
+	bool CheckLightmassExecutableVersion();
 	
 	/** Kicks off the lightmass processing, and, if successful, starts the asynchronous task */
 	bool BeginLightmassProcess();
diff --git a/Engine/Source/Editor/UnrealEd/Private/StaticLightingSystem/StaticLightingSystem.cpp b/Engine/Source/Editor/UnrealEd/Private/StaticLightingSystem/StaticLightingSystem.cpp
--- a/Engine/Source/Editor/UnrealEd/Private/StaticLightingSystem/StaticLightingSystem.cpp
+++ b/Engine/Source/Editor/UnrealEd/Private/StaticLightingSystem/StaticLightingSystem.cpp
@@ -82,6 +82,9 @@ DEFINE_LOG_CATEGORY(LogStaticLightingSystem);
 #include "Misc/UObjectToken.h"
 #include "Subsystems/AssetEditorSubsystem.h"
 #include "Rendering/StaticLightingSystemInterface.h"
+#include "BuildSettings.h"
+#include "Misc/EngineBuildSettings.h"
+#include "TargetReceipt.h"
 
 #define LOCTEXT_NAMESPACE "StaticLightingSystem"
 
@@ -308,35 +311,42 @@ void FStaticLightingManager::CreateStaticLightingSystem(const FLightingBuildOpti
 
 		ActiveStaticLightingSystem = StaticLightingSystems[0].Get();
 
-		bool bSuccess = ActiveStaticLightingSystem->BeginLightmassProcess();
-
-		if (bSuccess)
+		if (ActiveStaticLightingSystem->CheckLightmassExecutableVersion())
 		{
-			SendProgressNotification();
+			if (ActiveStaticLightingSystem->BeginLightmassProcess())
+			{
+				SendProgressNotification();
+			}
+			else
+			{
+				FStaticLightingManager::Get()->FailLightingBuild();
+			}
 		}
 		else
 		{
-			// BeginLightmassProcess returns false if there are errors or no precomputed lighting is allowed. Handle both cases.
-			static const auto AllowStaticLightingVar = IConsoleManager::Get().FindTConsoleVariableDataInt(TEXT("r.AllowStaticLighting"));
-			const bool bAllowStaticLighting = (!AllowStaticLightingVar || AllowStaticLightingVar->GetValueOnGameThread() != 0);
-			const bool bForceNoPrecomputedLighting = World->GetWorldSettings()->bForceNoPrecomputedLighting || !bAllowStaticLighting;
-
-
-			if (bForceNoPrecomputedLighting)
+			if (FEngineBuildSettings::IsSourceDistribution())
 			{
-				DestroyStaticLightingSystems();
+				FStaticLightingManager::Get()->FailLightingBuild(LOCTEXT("LightmassExecutableOutdatedMessage", "Unreal Lightmass executable is outdated. Recompile UnrealLightmass project in Visual Studio."));
 			}
 			else
 			{
-				FStaticLightingManager::Get()->FailLightingBuild();
+				// Lightmass should never be outdated in a launcher binary build.
+				FStaticLightingManager::Get()->FailLightingBuild(LOCTEXT("LauncherBuildNeedsVerificationMessage", "Unreal Lightmass executable is damaged. Try verifying your engine installation in Epic Games Launcher."));
 			}
 		}
 	}
 	else
 	{
 		// Tell the user that they must close their current build first.
-		FStaticLightingManager::Get()->FailLightingBuild(
-			LOCTEXT("LightBuildInProgressWarning", "A lighting build is already in progress! Please cancel it before triggering a new build."));
+		FNotificationInfo Info( LOCTEXT("LightBuildInProgressWarning", "A lighting build is already in progress! Please cancel it before triggering a new build.") );
+		Info.ExpireDuration = 5.0f;
+		TSharedPtr<SNotificationItem> Notification = FSlateNotificationManager::Get().AddNotification(Info);
+		if (Notification.IsValid())
+		{
+			Notification->SetCompletionState(SNotificationItem::CS_Fail);
+		}
+
+		FEditorDelegates::OnLightingBuildFailed.Broadcast();
 	}
 }
 
@@ -390,35 +400,42 @@ void FStaticLightingManager::UpdateBuildLighting()
 
 void FStaticLightingManager::FailLightingBuild( FText ErrorText)
 {
+	static const auto AllowStaticLightingVar = IConsoleManager::Get().FindTConsoleVariableDataInt(TEXT("r.AllowStaticLighting"));
+	const bool bAllowStaticLighting = (!AllowStaticLightingVar || AllowStaticLightingVar->GetValueOnGameThread() != 0);
+	const bool bForceNoPrecomputedLighting = GWorld->GetWorldSettings()->bForceNoPrecomputedLighting || !bAllowStaticLighting;
+
 	FStaticLightingManager::Get()->ClearCurrentNotification();
 	
-	if (GEditor->GetMapBuildCancelled())
-	{
-		ErrorText = LOCTEXT("LightBuildCanceledMessage", "Lighting build canceled.");
-	}
-	else
+	if (!bForceNoPrecomputedLighting)
 	{
-		// Override failure message if one provided
-		if (ErrorText.IsEmpty())
+		if (GEditor->GetMapBuildCancelled())
 		{
-			ErrorText = LOCTEXT("LightBuildFailedMessage", "Lighting build failed.");
+			ErrorText = LOCTEXT("LightBuildCanceledMessage", "Lighting build canceled.");
+		}
+		else
+		{
+			// Override failure message if one provided
+			if (ErrorText.IsEmpty())
+			{
+				ErrorText = LOCTEXT("LightBuildFailedMessage", "Lighting build failed.");
+			}
 		}
-	}
 
-	FNotificationInfo Info( ErrorText );
-	Info.ExpireDuration = 4.f;
+		FNotificationInfo Info( ErrorText );
+		Info.ExpireDuration = 4.f;
 	
-	FEditorDelegates::OnLightingBuildFailed.Broadcast();
+		FEditorDelegates::OnLightingBuildFailed.Broadcast();
 
-	LightBuildNotification = FSlateNotificationManager::Get().AddNotification(Info);
-	if (LightBuildNotification.IsValid())
-	{
-		LightBuildNotification.Pin()->SetCompletionState(SNotificationItem::CS_Fail);
-	}
+		LightBuildNotification = FSlateNotificationManager::Get().AddNotification(Info);
+		if (LightBuildNotification.IsValid())
+		{
+			LightBuildNotification.Pin()->SetCompletionState(SNotificationItem::CS_Fail);
+		}
 
-	UE_LOG(LogStaticLightingSystem, Warning, TEXT("Failed to build lighting!!! %s"),*ErrorText.ToString());
+		UE_LOG(LogStaticLightingSystem, Warning, TEXT("Failed to build lighting!!! %s"),*ErrorText.ToString());
 
-	FMessageLog("LightingResults").Open();
+		FMessageLog("LightingResults").Open();
+	}
 
 	DestroyStaticLightingSystems();
 }
@@ -478,6 +495,18 @@ FStaticLightingSystem::~FStaticLightingSystem()
 	}
 }
 
+bool FStaticLightingSystem::CheckLightmassExecutableVersion()
+{
+	FTargetReceipt LightmassReceipt;
+	
+	if (!LightmassReceipt.Read(FTargetReceipt::GetDefaultPath(*FPaths::EngineDir(), TEXT("UnrealLightmass"), FPlatformProcess::GetBinariesSubdirectory(), FApp::GetBuildConfiguration(), nullptr)))
+	{
+		return false;
+	}
+
+	return BuildSettings::GetCurrentChangelist() == LightmassReceipt.Version.Changelist;
+}
+
 bool FStaticLightingSystem::BeginLightmassProcess()
 {
 	StartTime = FPlatformTime::Seconds();
diff --git a/Engine/Source/Editor/UnrealEd/UnrealEd.Build.cs b/Engine/Source/Editor/UnrealEd/UnrealEd.Build.cs
--- a/Engine/Source/Editor/UnrealEd/UnrealEd.Build.cs
+++ b/Engine/Source/Editor/UnrealEd/UnrealEd.Build.cs
@@ -217,6 +217,7 @@ public class UnrealEd : ModuleRules
 				"CookOnTheFly",
 				"CookOnTheFlyNetServer",
 				"Zen",
+				"BuildSettings",
 			}
 		);
 
