Added nullptr check for Timeline array entry to avoid crash when opening blueprint with diff tool.

diff --git a/Engine/Source/Editor/UnrealEd/Private/Kismet2/BlueprintEditorUtils.cpp b/Engine/Source/Editor/UnrealEd/Private/Kismet2/BlueprintEditorUtils.cpp
--- a/Engine/Source/Editor/UnrealEd/Private/Kismet2/BlueprintEditorUtils.cpp
+++ b/Engine/Source/Editor/UnrealEd/Private/Kismet2/BlueprintEditorUtils.cpp
@@ -3760,7 +3760,7 @@ int32 FBlueprintEditorUtils::FindTimelineIndex(const UBlueprint* Blueprint, cons
 	const FName TimelineTemplateName = *UTimelineTemplate::TimelineVariableNameToTemplateName(InName);
 	for(int32 i=0; i<Blueprint->Timelines.Num(); i++)
 	{
-		if(Blueprint->Timelines[i]->GetFName() == TimelineTemplateName)
+		if(Blueprint->Timelines[i] && Blueprint->Timelines[i]->GetFName() == TimelineTemplateName)
 		{
 			return i;
 		}
