MultiprocessCook: Algo::Graph::ConstructLoadBalance - fixes from code review.
#rnx

diff --git a/Engine/Source/Editor/UnrealEd/Private/Cooker/Algo/GraphLoadBalance.cpp b/Engine/Source/Editor/UnrealEd/Private/Cooker/Algo/GraphLoadBalance.cpp
--- a/Engine/Source/Editor/UnrealEd/Private/Cooker/Algo/GraphLoadBalance.cpp
+++ b/Engine/Source/Editor/UnrealEd/Private/Cooker/Algo/GraphLoadBalance.cpp
@@ -462,7 +462,7 @@ void FLoadBalanceBuilder::PackClusters(TArray<FCluster>&& InDisjointSubgraphs, T
 
 		TArray<FCluster> NewClusters = SplitCluster(MoveTemp(ClusterToSplit), SplitSize, DesiredSpread);
 
-		// Remove the SplitCluster from whereever it is in the middle of Clusters, and shift down by one all of our
+		// Remove the SplitCluster from wherever it is in the middle of Clusters, and shift down by one all of our
 		// indices in data that persists between loops that were pointing to clusters at a higher index.
 		Clusters.RemoveAt(IndexToSplit);
 		ClusterCosts.RemoveAt(IndexToSplit);
diff --git a/Engine/Source/Editor/UnrealEd/Private/Cooker/Algo/GraphLoadBalance.h b/Engine/Source/Editor/UnrealEd/Private/Cooker/Algo/GraphLoadBalance.h
--- a/Engine/Source/Editor/UnrealEd/Private/Cooker/Algo/GraphLoadBalance.h
+++ b/Engine/Source/Editor/UnrealEd/Private/Cooker/Algo/GraphLoadBalance.h
@@ -15,7 +15,7 @@ namespace Algo::Graph
  * A transitively-closed set of vertices means that if a vertex is present in a bucket then all vertices reachable
  * from it are also present.
  * 
- * The definition of longpoleandspreadminimized is:
+ * The definition of longpoleandspread-minimized is:
  * 1) The size of the largest bucket is minimized.
  * 2) If multiple assignments achieve that minimum, the difference in size between the largest bucket and smallest
  *    bucket is minimized.
@@ -31,6 +31,10 @@ namespace Algo::Graph
 void ConstructLoadBalance(TConstArrayView<TConstArrayView<FVertex>> Graph, int32 NumBuckets, 
 	TArray<TArray<FVertex>>& OutAssignments);
 
+/**
+ * Expanded arguments for ConstructLoadBalance for callers that want to provide as input or receive as output some
+ * of the intermediate data.
+ */
 struct FLoadBalanceContext
 {
 	/**
@@ -62,10 +66,7 @@ struct FLoadBalanceContext
 	int32 NumBuckets = 1;
 };
 
-/**
- * Expanded arguments for ConstructLoadBalance for callers that want to provide as input or receive as output some
- * of the intermediate data. @see FLoadBalanceContext for argument descriptions.
- */
+/** ConstructLoadBalance that takes extra arguments. @see FLoadBalanceContext for argument descriptions. */
 void ConstructLoadBalance(FLoadBalanceContext& Context);
 
 
diff --git a/Engine/Source/Editor/UnrealEd/Private/Cooker/Algo/Packing.cpp b/Engine/Source/Editor/UnrealEd/Private/Cooker/Algo/Packing.cpp
--- a/Engine/Source/Editor/UnrealEd/Private/Cooker/Algo/Packing.cpp
+++ b/Engine/Source/Editor/UnrealEd/Private/Cooker/Algo/Packing.cpp
@@ -48,7 +48,7 @@ void ScheduleValues(TConstArrayView<int32> Values, int32 NumBuckets, TConstArray
 		}
 	}
 
-	// Initialize remaining values to the entire list of values, and sort them nto the order in which we should
+	// Initialize remaining values to the entire list of values, and sort them into the order in which we should
 	// try to place them.
 	TArray<int32> RemainingValueIndices = RangeArray<TArray<int32>>(0, NumValues);
 
@@ -71,9 +71,8 @@ void ScheduleValues(TConstArrayView<int32> Values, int32 NumBuckets, TConstArray
 	};
 	Algo::Sort(RemainingValueIndices, ValueIndexSelectionOrder);
 
-	// After every piece we place we need to resort the modified smallest bucket back
-	// into the list and find the new smallest.
-	// FBucketData is our sortable information about a bucket
+	// After every piece we place we need to resort the modified smallest bucket back into the list and find the new
+	// smallest. FBucketData is our sortable information about a bucket
 	struct FBucketData
 	{
 		TArray<int32>* Bucket = nullptr;
@@ -95,14 +94,13 @@ void ScheduleValues(TConstArrayView<int32> Values, int32 NumBuckets, TConstArray
 		{
 			return A.Size > B.Size;
 		}
-		// When buckets are equal, prefer to add to the earlier buckets, by pushing them
-		// to the back of the list
+		// When buckets are equal, prefer to add to the earlier buckets, by pushing them to the back of the list
 		return A.Bucket > B.Bucket;
 	};
 	Algo::Sort(SortedBuckets, SortBucketDataDescending);
 
-	// Main loop: iterate over piece and put it into the current smallest bucket
-	// (or the smallest valid bucket if the piece is in an exclusion group)
+	// Main loop: iterate over each piece and put it into the current smallest bucket (or the smallest valid bucket
+	// if the piece is in an exclusion group)
 	for (int32 ValueIndex : RemainingValueIndices)
 	{
 		int32* ExclusionGroupIndex = ExclusionGroupOfValueIndex.Find(ValueIndex);
@@ -123,9 +121,7 @@ void ScheduleValues(TConstArrayView<int32> Values, int32 NumBuckets, TConstArray
 					break;
 				}
 			}
-			// All exclusiongroups are <= NumBuckets in size and we only add each cluster (and therefore each member of
-			// an exclusiongroup) to a single bucket, so we should always have a free bucket for the last member of an
-			// exclusiongroup
+			// exclusiongroups are <= NumBuckets in size so we always have a free bucket for an exclusiongroup member
 			check(BestBucketIndex != INDEX_NONE);
 		}
 
diff --git a/Engine/Source/Editor/UnrealEd/Private/Cooker/Algo/Packing.h b/Engine/Source/Editor/UnrealEd/Private/Cooker/Algo/Packing.h
--- a/Engine/Source/Editor/UnrealEd/Private/Cooker/Algo/Packing.h
+++ b/Engine/Source/Editor/UnrealEd/Private/Cooker/Algo/Packing.h
@@ -10,13 +10,13 @@ namespace Algo
 {
 
 /**
- * Assign labeled values to N buckets such that the sum of the values in the maximum bucket is minimized
+ * Assign labeled values to N buckets such that the sum of the values in the maximum bucket is minimized.
  * 
- * @param Values Range of values to schedule
- * @param NumBuckets How many buckets to divide the values into
+ * @param Values Range of values to schedule.
+ * @param NumBuckets How many buckets to divide the values into.
  * @param ExclusionGroups Each element is a range of indexes into Values. Each of the Values in an ExclusionGroup will
  *        be sorted into a different bucket. ExclusionGroups must not overlap, and the size of each ExclusionGroup
- *        must be <= N (the number of buckets).
+ *        must be <= the number of buckets.
  * @param OutBucketsIndices Filled with the output buckets; each element is a range of unsorted indices into Values.
  */
 void ScheduleValues(TConstArrayView<int32> Values, int32 NumBuckets, TConstArrayView<TArray<int32>> ExclusionGroups,
