DerivedDataCacheCommandlet: Respect the contract with IsCachedCookedPlatformDataLoaded that it will not be called again after returning true.
#rnx

diff --git a/Engine/Source/Editor/UnrealEd/Classes/Commandlets/DerivedDataCacheCommandlet.h b/Engine/Source/Editor/UnrealEd/Classes/Commandlets/DerivedDataCacheCommandlet.h
--- a/Engine/Source/Editor/UnrealEd/Classes/Commandlets/DerivedDataCacheCommandlet.h
+++ b/Engine/Source/Editor/UnrealEd/Classes/Commandlets/DerivedDataCacheCommandlet.h
@@ -6,11 +6,22 @@
 
 #pragma once
 
-#include "CoreMinimal.h"
-#include "UObject/ObjectMacros.h"
 #include "Commandlets/Commandlet.h"
+#include "Containers/Array.h"
+#include "Containers/BitArray.h"
+#include "Containers/Map.h"
+#include "Containers/Set.h"
+#include "Containers/UnrealString.h"
+#include "Templates/UniquePtr.h"
+#include "UObject/NameTypes.h"
+#include "UObject/ObjectMacros.h"
+
 #include "DerivedDataCacheCommandlet.generated.h"
 
+class ITargetPlatform;
+class UPackage;
+class UWorld;
+
 UCLASS()
 class UDerivedDataCacheCommandlet : public UCommandlet
 {
@@ -22,11 +33,21 @@ class UDerivedDataCacheCommandlet : public UCommandlet
 	virtual int32 Main(const FString& Params) override;
 	//~ End UCommandlet Interface
 
-	// Objects currently being cached along with when was the last time in seconds we verified if they were still compiling
-	// The last time is used to do some throttling on the IsCachedCookedPlatformDataLoaded which can be quite expensive on some objects
-	TMap<UObject*, double> CachingObjects;
+	struct FCachingData
+	{
+		/** Element n == IsCachedCookedPlatformDataLoaded has returned true for Commandlet->Platforms[n] */
+		TBitArray<> PlatformIsComplete;
+		/**
+		 * The last time in seconds we tested whether the object is still compiling and/or called
+		 * IsCachedCookedPlatformDataLoaded. Used to throttle IsCachedCookedPlatformDataLoaded which can be quite
+		 * expensive on some objects
+		 */
+		double LastTimeTested = 0.;
+	};
+	TMap<UObject*, FCachingData> CachingObjects;
 	TSet<FName>    ProcessedPackages;
 	TSet<FName>    PackagesToProcess;
+	TArray<const ITargetPlatform*> Platforms;
 	double FinishCacheTime = 0.0;
 	double BeginCacheTime = 0.0;
 
@@ -38,10 +59,10 @@ class UDerivedDataCacheCommandlet : public UCommandlet
 
 	void MaybeMarkPackageAsAlreadyLoaded(UPackage *Package);
 
-	void CacheLoadedPackages(UPackage* CurrentPackage, uint8 PackageFilter, const TArray<ITargetPlatform*>& Platforms, TSet<FName>& OutNewProcessedPackages);
-	void CacheWorldPackages(UWorld* World, uint8 PackageFilter, const TArray<ITargetPlatform*>& Platforms, TSet<FName>& OutNewProcessedPackages);
-	bool ProcessCachingObjects(const TArray<ITargetPlatform*>& Platforms);
-	void FinishCachingObjects(const TArray<ITargetPlatform*>& Platforms);
+	void CacheLoadedPackages(UPackage* CurrentPackage, uint8 PackageFilter, TSet<FName>& OutNewProcessedPackages);
+	void CacheWorldPackages(UWorld* World, uint8 PackageFilter, TSet<FName>& OutNewProcessedPackages);
+	bool ProcessCachingObjects();
+	void FinishCachingObjects();
 };
 
 
diff --git a/Engine/Source/Editor/UnrealEd/Private/Commandlets/DerivedDataCacheCommandlet.cpp b/Engine/Source/Editor/UnrealEd/Private/Commandlets/DerivedDataCacheCommandlet.cpp
--- a/Engine/Source/Editor/UnrealEd/Private/Commandlets/DerivedDataCacheCommandlet.cpp
+++ b/Engine/Source/Editor/UnrealEd/Private/Commandlets/DerivedDataCacheCommandlet.cpp
@@ -4,44 +4,45 @@
 DerivedDataCacheCommandlet.cpp: Commandlet for DDC maintenence
 =============================================================================*/
 #include "Commandlets/DerivedDataCacheCommandlet.h"
-#include "UObject/UObjectHash.h"
-#include "UObject/UObjectIterator.h"
-#include "UObject/Package.h"
-#include "Misc/ConfigCacheIni.h"
-#include "Misc/PackageAccessTrackingOps.h"
-#include "Misc/PackageName.h"
-#include "PackageHelperFunctions.h"
+
+#include "Algo/RemoveIf.h"
+#include "Algo/StableSort.h"
+#include "Algo/Transform.h"
+#include "AssetCompilingManager.h"
+#include "CollectionManagerModule.h"
+#include "CollectionManagerTypes.h"
+#include "CookOnTheSide/CookOnTheFlyServer.h"
 #include "DerivedDataCacheInterface.h"
+#include "DistanceFieldAtlas.h"
+#include "Editor.h"
+#include "EditorWorldUtils.h"
+#include "Engine/Texture.h"
 #include "GlobalShader.h"
+#include "ICollectionManager.h"
 #include "Interfaces/ITargetPlatform.h"
 #include "Interfaces/ITargetPlatformManagerModule.h"
-#include "ShaderCompiler.h"
-#include "DistanceFieldAtlas.h"
+#include "LevelInstance/LevelInstanceSubsystem.h"
 #include "MeshCardRepresentation.h"
+#include "Misc/ConfigCacheIni.h"
+#include "Misc/PackageAccessTrackingOps.h"
+#include "Misc/PackageName.h"
 #include "Misc/RedirectCollector.h"
-#include "Engine/Texture.h"
-#include "CookOnTheSide/CookOnTheFlyServer.h"
-#include "Algo/RemoveIf.h"
-#include "Algo/Transform.h"
-#include "Algo/StableSort.h"
+#include "PackageHelperFunctions.h"
 #include "Settings/ProjectPackagingSettings.h"
-#include "Editor.h"
-#include "EditorWorldUtils.h"
-#include "AssetCompilingManager.h"
+#include "ShaderCompiler.h"
+#include "UObject/Package.h"
+#include "UObject/UObjectHash.h"
+#include "UObject/UObjectIterator.h"
 #include "WorldPartition/WorldPartition.h"
 #include "WorldPartition/WorldPartitionHelpers.h"
 #include "WorldPartition/WorldPartitionSubsystem.h"
-#include "LevelInstance/LevelInstanceSubsystem.h"
-#include "CollectionManagerModule.h"
-#include "ICollectionManager.h"
-#include "CollectionManagerTypes.h"
 
 DEFINE_LOG_CATEGORY_STATIC(LogDerivedDataCacheCommandlet, Log, All);
 
 class UDerivedDataCacheCommandlet::FObjectReferencer : public FGCObject
 {
 public:
-	FObjectReferencer(TMap<UObject*, double>& InReferencedObjects)
+	FObjectReferencer(TMap<UObject*, FCachingData>& InReferencedObjects)
 		: ReferencedObjects(InReferencedObjects)
 	{
 	}
@@ -60,7 +61,7 @@
 	}
 
 	FString ReferencerName;
-	TMap<UObject*, double>& ReferencedObjects;
+	TMap<UObject*, FCachingData>& ReferencedObjects;
 };
 
 class UDerivedDataCacheCommandlet::FPackageListener : public FUObjectArray::FUObjectCreateListener, public FUObjectArray::FUObjectDeleteListener
@@ -207,7 +208,7 @@
 
 } // UE::Private::DerivedDataCacheCommandlet
 
-void UDerivedDataCacheCommandlet::CacheLoadedPackages(UPackage* CurrentPackage, uint8 PackageFilter, const TArray<ITargetPlatform*>& Platforms, TSet<FName>& OutNewProcessedPackages)
+void UDerivedDataCacheCommandlet::CacheLoadedPackages(UPackage* CurrentPackage, uint8 PackageFilter, TSet<FName>& OutNewProcessedPackages)
 {
 	TRACE_CPUPROFILER_EVENT_SCOPE(UDerivedDataCacheCommandlet::CacheLoadedPackages);
 
@@ -249,7 +250,8 @@
 						Object->BeginCacheForCookedPlatformData(Platform);
 					}
 
-					CachingObjects.Add(Object);
+					FCachingData& CachingData = CachingObjects.FindOrAdd(Object);
+					CachingData.PlatformIsComplete.Init(false, Platforms.Num());
 				}
 			}
 		}
@@ -261,10 +263,10 @@
 
 	BeginCacheTime += FPlatformTime::Seconds() - BeginCacheTimeStart;
 
-	ProcessCachingObjects(Platforms);
+	ProcessCachingObjects();
 }
 
-bool UDerivedDataCacheCommandlet::ProcessCachingObjects(const TArray<ITargetPlatform*>& Platforms)
+bool UDerivedDataCacheCommandlet::ProcessCachingObjects()
 {
 	TRACE_CPUPROFILER_EVENT_SCOPE(UDerivedDataCacheCommandlet::ProcessCachingObjects);
 
@@ -277,7 +279,8 @@
 		for (auto It = CachingObjects.CreateIterator(); It; ++It)
 		{
 			// Call IsCachedCookedPlatformDataLoaded once a second per object since it can be quite expensive
-			if (CurrentTime - It->Value > 1.0)
+			FCachingData& CachingData = It->Value;
+			if (CurrentTime - CachingData.LastTimeTested > 1.0)
 			{
 				UObject* Object = It->Key;
 				bool bIsFinished = true;
@@ -289,11 +292,25 @@
 
 				{
 					UE_TRACK_REFERENCING_PACKAGE_SCOPED(Object->GetPackage(), PackageAccessTrackingOps::NAME_CookerBuildObject);
-					for (auto Platform : Platforms)
+					for (int32 PlatformIndex = 0; PlatformIndex < Platforms.Num(); ++PlatformIndex)
 					{
 						// IsCachedCookedPlatformDataLoaded can be quite slow for some objects
 						// Do not call it if bIsFinished is already false
-						bIsFinished = bIsFinished && Object->IsCachedCookedPlatformDataLoaded(Platform);
+						// Also, by the contract of IsCachedCookedPlatformDataLoaded, we are not allowed to call it
+						// again once it returns true
+						if (bIsFinished && !CachingData.PlatformIsComplete[PlatformIndex])
+						{
+							const ITargetPlatform* Platform = Platforms[PlatformIndex];
+							if (Object->IsCachedCookedPlatformDataLoaded(Platform))
+							{
+								CachingData.PlatformIsComplete[PlatformIndex] = true;
+								bHadActivity = true;
+							}
+							else
+							{
+								bIsFinished = false;
+							}
+						}
 					}
 				}
 
@@ -306,7 +323,7 @@
 				}
 				else
 				{
-					It->Value = CurrentTime;
+					CachingData.LastTimeTested = CurrentTime;
 				}
 			}
 		}
@@ -315,7 +332,7 @@
 	return bHadActivity;
 }
 
-void UDerivedDataCacheCommandlet::FinishCachingObjects(const TArray<ITargetPlatform*>& Platforms)
+void UDerivedDataCacheCommandlet::FinishCachingObjects()
 {
 	// Timing variables
 	double DDCCommandletMaxWaitSeconds = 60. * 10.;
@@ -328,7 +345,7 @@
 	{
 		UE::Private::DerivedDataCacheCommandlet::TickCookObjects();
 
-		bool bHadActivity = ProcessCachingObjects(Platforms);
+		bool bHadActivity = ProcessCachingObjects();
 
 		double CurrentTime = FPlatformTime::Seconds();
 		if (!bHadActivity)
@@ -359,7 +376,7 @@
 	FinishCacheTime += FPlatformTime::Seconds() - FinishCacheTimeStart;
 }
 
-void UDerivedDataCacheCommandlet::CacheWorldPackages(UWorld* World, uint8 PackageFilter, const TArray<ITargetPlatform*>& Platforms, TSet<FName>& OutNewProcessedPackages)
+void UDerivedDataCacheCommandlet::CacheWorldPackages(UWorld* World, uint8 PackageFilter, TSet<FName>& OutNewProcessedPackages)
 {
 	// Setup the world
 	UWorld::InitializationValues IVS;
@@ -380,12 +397,12 @@
 		UWorldPartition* WorldPartition = World->GetWorldPartition();
 		check(WorldPartition);
 
-		FWorldPartitionHelpers::ForEachActorWithLoading(WorldPartition, [this, PackageFilter, &Platforms, &OutNewProcessedPackages](const FWorldPartitionActorDesc* ActorDesc)
+		FWorldPartitionHelpers::ForEachActorWithLoading(WorldPartition, [this, PackageFilter, &OutNewProcessedPackages](const FWorldPartitionActorDesc* ActorDesc)
 		{
 			if (AActor* Actor = ActorDesc->GetActor())
 			{
 				UE_LOG(LogDerivedDataCacheCommandlet, Display, TEXT("Loaded actor %s"), *Actor->GetName());
-				CacheLoadedPackages(Actor->GetPackage(), PackageFilter, Platforms, OutNewProcessedPackages);
+				CacheLoadedPackages(Actor->GetPackage(), PackageFilter, OutNewProcessedPackages);
 			}
 			return true;
 		});
@@ -417,6 +434,9 @@
 	double GCTime = 0.0;
 	FinishCacheTime = 0.;
 	BeginCacheTime = 0.;
+	ITargetPlatformManagerModule* TPM = GetTargetPlatformManager();
+	Platforms.Reset();
+	Platforms.Append(TPM->GetActiveTargetPlatforms());
 
 	if (!bStartupOnly && bFillCache)
 	{
@@ -631,9 +651,6 @@
 			}
 		}
 
-		ITargetPlatformManagerModule* TPM = GetTargetPlatformManager();
-		const TArray<ITargetPlatform*>& Platforms = TPM->GetActiveTargetPlatforms();
-
 		for (int32 Index = 0; Index < Platforms.Num(); Index++)
 		{
 			TArray<FName> DesiredShaderFormats;
@@ -763,14 +780,14 @@
 
 			// Find any new packages and cache all the objects in each package
 			TSet<FName> NewProcessedPackages;
-			CacheLoadedPackages(Package, PackageFilter, Platforms, NewProcessedPackages);
+			CacheLoadedPackages(Package, PackageFilter, NewProcessedPackages);
 
 			// Ensure we load maps to process all their referenced packages in case they are using world partition.
 			if (bLastPackageWasMap)
 			{
 				if (UWorld* World = UWorld::FindWorldInPackage(Package))
 				{
-					CacheWorldPackages(World, PackageFilter, Platforms, NewProcessedPackages);
+					CacheWorldPackages(World, PackageFilter, NewProcessedPackages);
 				}
 			}
 
@@ -831,7 +848,7 @@
 		}
 	}
 
-	FinishCachingObjects(GetTargetPlatformManager()->GetActiveTargetPlatforms());
+	FinishCachingObjects();
 
 	UE::Private::DerivedDataCacheCommandlet::TickCookObjects(/*bCookComplete*/ true);
 
