Fix various issues with the CompileShadersTestBed commandlet by making it behave more like the cooker.

- Commandlet shouldn't be configured for rendering.  This was originally done to force shader compilation to happen in PostLoad.  This caused textures resource to be created, which is not needed and was causing crashes.
- Properly wait on potential DDC requests.
- Wait on shader compilation by calling ProcessAsyncResults, which is exactly what the cooker calls.
- While we have jobs to compile we process the results.
- Be sure to flush rendering commands so we cleanup deferred delete objects like Shader Maps.
- Added logging to commandlet.
- Enable LogMaterial for this commandlet so we can see DDC puts w/ DDC key.

todo:

- I need to properly poll for DDC results and keep an active list of materials.  Similar to the DDC commandlet.

diff --git a/Engine/Source/Editor/UnrealEd/Private/Commandlets/CompileShadersTestBedCommandlet.cpp b/Engine/Source/Editor/UnrealEd/Private/Commandlets/CompileShadersTestBedCommandlet.cpp
--- a/Engine/Source/Editor/UnrealEd/Private/Commandlets/CompileShadersTestBedCommandlet.cpp
+++ b/Engine/Source/Editor/UnrealEd/Private/Commandlets/CompileShadersTestBedCommandlet.cpp
@@ -25,6 +25,8 @@ int32 UCompileShadersTestBedCommandlet::Main(const FString& Params)
 {
 	TRACE_CPUPROFILER_EVENT_SCOPE(UCompileShadersTestBedCommandlet::Main);
 
+	StaticExec(nullptr, TEXT("log LogMaterial Log"));
+
 	TArray<FString> Tokens;
 	TArray<FString> Switches;
 	TMap<FString, FString> ParamVals;
@@ -40,8 +42,6 @@ int32 UCompileShadersTestBedCommandlet::Main(const FString& Params)
 		return 0;
 	}
 
-	PRIVATE_GAllowCommandletRendering = true;
-
 	IAssetRegistry& AssetRegistry = FModuleManager::Get().LoadModuleChecked<FAssetRegistryModule>("AssetRegistry").Get();
 	AssetRegistry.SearchAllAssets(true);
 
@@ -101,6 +101,8 @@ PRAGMA_ENABLE_DEPRECATION_WARNINGS
 		AssetRegistry.GetAssets(Filter, MaterialList);
 	}
 
+	static constexpr bool bLimitExecutationTime = false;
+
 	// For all active platforms
 	ITargetPlatformManagerModule* TPM = GetTargetPlatformManager();
 	const TArray<ITargetPlatform*>& Platforms = TPM->GetActiveTargetPlatforms();
@@ -111,9 +113,18 @@ PRAGMA_ENABLE_DEPRECATION_WARNINGS
 		{
 			TRACE_CPUPROFILER_EVENT_SCOPE(DefaultMaterials);
 
+			UE_LOG(LogCompileShadersTestBedCommandlet, Log, TEXT("Compile default materials"));
+
 			for (int32 Domain = 0; Domain < MD_MAX; ++Domain)
 			{
-				UMaterial::GetDefaultMaterial(static_cast<EMaterialDomain>(Domain))->BeginCacheForCookedPlatformData(Platform);
+				UMaterialInterface* DefaultMaterial = UMaterial::GetDefaultMaterial(static_cast<EMaterialDomain>(Domain));
+
+				DefaultMaterial->BeginCacheForCookedPlatformData(Platform);
+				while (!DefaultMaterial->IsCachedCookedPlatformDataLoaded(Platform))
+				{
+					GShaderCompilingManager->ProcessAsyncResults(bLimitExecutationTime, false /* bBlockOnGlobalShaderCompilation */);
+				}
+				DefaultMaterial->ClearCachedCookedPlatformData(Platform);
 			}
 		}
 
@@ -124,53 +135,91 @@ PRAGMA_ENABLE_DEPRECATION_WARNINGS
 			TArray<FName> DesiredShaderFormats;
 			Platform->GetAllTargetedShaderFormats(DesiredShaderFormats);
 
+			UE_LOG(LogCompileShadersTestBedCommandlet, Log, TEXT("Compile global shaders"));
+
 			for (int32 FormatIndex = 0; FormatIndex < DesiredShaderFormats.Num(); FormatIndex++)
 			{
 				const EShaderPlatform ShaderPlatform = ShaderFormatToLegacyShaderPlatform(DesiredShaderFormats[FormatIndex]);
 				CompileGlobalShaderMap(ShaderPlatform, Platform, false);
 			}
+
+			GShaderCompilingManager->ProcessAsyncResults(bLimitExecutationTime, true /* bBlockOnGlobalShaderCompilation */);
 		}
 
-		// Compile material shaders specified on the command line
+		// Begin Material Compiles
 		{
-			TRACE_CPUPROFILER_EVENT_SCOPE(MaterialShaders);
+			TRACE_CPUPROFILER_EVENT_SCOPE(BeginCacheForCookedPlatformData);
 
 			// Sort the material lists by name so the order is stable.
 			Algo::SortBy(MaterialList, [](const FAssetData& AssetData) { return AssetData.GetSoftObjectPath(); }, [](const FSoftObjectPath& A, const FSoftObjectPath& B) { return A.LexicalLess(B); });
 
+			UE_LOG(LogCompileShadersTestBedCommandlet, Log, TEXT("Begin Cache For Cooked PlatformData"));
+
 			for (const FAssetData& AssetData : MaterialList)
 			{
-				if (UMaterial* Material = Cast<UMaterial>(AssetData.GetAsset()))
+				if (UMaterialInterface* MaterialInterface = Cast<UMaterialInterface>(AssetData.GetAsset()))
 				{
-					Material->BeginCacheForCookedPlatformData(Platform);
+					UE_LOG(LogCompileShadersTestBedCommandlet, Log, TEXT("BeginCache for %s"), *MaterialInterface->GetFullName());
+					MaterialInterface->BeginCacheForCookedPlatformData(Platform);
 				}
-				else if (UMaterialInstance* MaterialInstance = Cast<UMaterialInstance>(AssetData.GetAsset()))
+			}
+		}
+
+		bool bAllDone = false;
+
+		// Submit all the jobs.
+		{
+			TRACE_CPUPROFILER_EVENT_SCOPE(SubmitJobs);
+
+			UE_LOG(LogCompileShadersTestBedCommandlet, Log, TEXT("Submit Jobs"));
+
+			for (const FAssetData& AssetData : MaterialList)
+			{
+				if (UMaterialInterface* MaterialInterface = Cast<UMaterialInterface>(AssetData.GetAsset()))
+				{
+					const bool bMaterialDone = MaterialInterface->IsCachedCookedPlatformDataLoaded(Platform);
+					if (bMaterialDone)
+					{
+						UE_LOG(LogCompileShadersTestBedCommandlet, Log, TEXT("Finished cache for %s"), *MaterialInterface->GetFullName());
+					}
+
+					bAllDone &= bMaterialDone;
+				}
+
+				if (bAllDone)
 				{
-					MaterialInstance->BeginCacheForCookedPlatformData(Platform);
+					break;
 				}
 			}
 		}
 
-		// Block on all the jobs submitted above.
+		// Process the shader maps and save to the DDC.
 		{
-			TRACE_CPUPROFILER_EVENT_SCOPE(BlockOnShaderCompiles);
+			TRACE_CPUPROFILER_EVENT_SCOPE(ProcessShaderCompileResults);
+
+			UE_LOG(LogCompileShadersTestBedCommandlet, Log, TEXT("ProcessAsyncResults"));
 
-			GShaderCompilingManager->FinishAllCompilation();
+			while (GShaderCompilingManager->IsCompiling())
+			{
+				GShaderCompilingManager->ProcessAsyncResults(bLimitExecutationTime, false /* bBlockOnGlobalShaderCompilation */);
+				
+				// Flush rendering commands to release any RHI resources (shaders and shader maps).
+				// Delete any FPendingCleanupObjects (shader maps).
+				FlushRenderingCommands();
+			}
 		}
 
 		// Perform cleanup and clear cached data for cooking.
 		{
 			TRACE_CPUPROFILER_EVENT_SCOPE(ClearCachedCookedPlatformData);
 
+			UE_LOG(LogCompileShadersTestBedCommandlet, Log, TEXT("Clear Cached Cooked Platform Data"));
+
 			for (const FAssetData& AssetData : MaterialList)
 			{
-				if (UMaterial* Material = Cast<UMaterial>(AssetData.GetAsset()))
-				{
-					Material->ClearAllCachedCookedPlatformData();
-				}
-				else if (UMaterialInstance* MaterialInstance = Cast<UMaterialInstance>(AssetData.GetAsset()))
+				if (UMaterialInterface* MaterialInterface = Cast<UMaterialInterface>(AssetData.GetAsset()))
 				{
-					MaterialInstance->ClearAllCachedCookedPlatformData();
+					MaterialInterface->ClearAllCachedCookedPlatformData();
 				}
 			}
 		}
