Cooker - Add BlockOnAssetRegistryTime to the Cook Profile logs printed at the end of the cook.
#rnx

diff --git a/Engine/Source/Editor/UnrealEd/Private/Commandlets/CookCommandlet.cpp b/Engine/Source/Editor/UnrealEd/Private/Commandlets/CookCommandlet.cpp
--- a/Engine/Source/Editor/UnrealEd/Private/Commandlets/CookCommandlet.cpp
+++ b/Engine/Source/Editor/UnrealEd/Private/Commandlets/CookCommandlet.cpp
@@ -73,6 +73,7 @@ extern double TickCookOnTheSideLoadPackagesTimeSec;
 extern double TickCookOnTheSideResolveRedirectorsTimeSec;
 extern double TickCookOnTheSideSaveCookedPackageTimeSec;
 extern double TickCookOnTheSidePrepareSaveTimeSec;
+extern double BlockOnAssetRegistryTimeSec;
 extern double GameCookModificationDelegateTimeSec;
 double TickLoopGCTimeSec = 0.0;
 double TickLoopRecompileShaderRequestsTimeSec = 0.0;
@@ -94,7 +95,8 @@ FCookStatsManager::FAutoRegisterCallback RegisterCookStats([](FCookStatsManager:
 	ADD_COOK_STAT_FLT(" 0. 0", StartupWallTimeSec);
 	ADD_COOK_STAT_FLT(" 0. 1", CookByTheBookTimeSec);
 	ADD_COOK_STAT_FLT(" 0. 1. 0", StartCookByTheBookTimeSec);
-	ADD_COOK_STAT_FLT(" 0. 1. 0. 0", GameCookModificationDelegateTimeSec);
+	ADD_COOK_STAT_FLT(" 0. 1. 0. 0", BlockOnAssetRegistryTimeSec);
+	ADD_COOK_STAT_FLT(" 0. 1. 0. 1", GameCookModificationDelegateTimeSec);
 	ADD_COOK_STAT_FLT(" 0. 1. 1", TickCookOnTheSideTimeSec);
 	ADD_COOK_STAT_FLT(" 0. 1. 1. 0", TickCookOnTheSideLoadPackagesTimeSec);
 	ADD_COOK_STAT_FLT(" 0. 1. 1. 1", TickCookOnTheSideSaveCookedPackageTimeSec);
diff --git a/Engine/Source/Editor/UnrealEd/Private/CookOnTheFlyServer.cpp b/Engine/Source/Editor/UnrealEd/Private/CookOnTheFlyServer.cpp
--- a/Engine/Source/Editor/UnrealEd/Private/CookOnTheFlyServer.cpp
+++ b/Engine/Source/Editor/UnrealEd/Private/CookOnTheFlyServer.cpp
@@ -7107,6 +7107,8 @@ void UCookOnTheFlyServer::BlockOnAssetRegistry()
 		return;
 	}
 	TRACE_CPUPROFILER_EVENT_SCOPE(UCookOnTheFlyServer::BlockOnAssetRegistry);
+	COOK_STAT(FScopedDurationTimer TickTimer(DetailedCookStats::BlockOnAssetRegistryTimeSec));
+
 	UE_LOG(LogCook, Display, TEXT("Waiting for Asset Registry"));
 	// Blocking on the AssetRegistry has to be done on the game thread since some AssetManager functions require it
 	check(IsInGameThread());
diff --git a/Engine/Source/Editor/UnrealEd/Private/Cooker/CookProfiling.cpp b/Engine/Source/Editor/UnrealEd/Private/Cooker/CookProfiling.cpp
--- a/Engine/Source/Editor/UnrealEd/Private/Cooker/CookProfiling.cpp
+++ b/Engine/Source/Editor/UnrealEd/Private/Cooker/CookProfiling.cpp
@@ -178,6 +178,7 @@ namespace DetailedCookStats
 	double TickCookOnTheSideResolveRedirectorsTimeSec = 0.0;
 	double TickCookOnTheSideSaveCookedPackageTimeSec = 0.0;
 	double TickCookOnTheSidePrepareSaveTimeSec = 0.0;
+	double BlockOnAssetRegistryTimeSec = 0.0;
 	double GameCookModificationDelegateTimeSec = 0.0;
 
 	// Stats tracked through FAutoRegisterCallback
diff --git a/Engine/Source/Editor/UnrealEd/Private/Cooker/CookProfiling.h b/Engine/Source/Editor/UnrealEd/Private/Cooker/CookProfiling.h
--- a/Engine/Source/Editor/UnrealEd/Private/Cooker/CookProfiling.h
+++ b/Engine/Source/Editor/UnrealEd/Private/Cooker/CookProfiling.h
@@ -122,6 +122,7 @@ namespace DetailedCookStats
 	extern double TickCookOnTheSideResolveRedirectorsTimeSec;
 	extern double TickCookOnTheSideSaveCookedPackageTimeSec;
 	extern double TickCookOnTheSidePrepareSaveTimeSec;
+	extern double BlockOnAssetRegistryTimeSec;
 	extern double GameCookModificationDelegateTimeSec;
 	extern double TickLoopShaderProcessAsyncResultsTimeSec;
 	extern double TickLoopProcessDeferredCommandsTimeSec;
