Added support for static mesh previews to moderation artifact generator commandlet.

diff --git a/Engine/Source/Editor/UnrealEd/Classes/Commandlets/BaseIteratePackagesCommandlet.h b/Engine/Source/Editor/UnrealEd/Classes/Commandlets/BaseIteratePackagesCommandlet.h
--- a/Engine/Source/Editor/UnrealEd/Classes/Commandlets/BaseIteratePackagesCommandlet.h
+++ b/Engine/Source/Editor/UnrealEd/Classes/Commandlets/BaseIteratePackagesCommandlet.h
@@ -224,6 +224,8 @@ protected:
 	*/
 	virtual void PerformAdditionalOperations(class UWorld* World, bool& bSavePackage) { }
 
+	virtual void PostPerformAdditionalOperations(class UPackage* Package) { }
+
 	/**
 	 * Allows the commandlet to do something after all the packages have been processed
 	 */
diff --git a/Engine/Source/Editor/UnrealEd/Classes/Commandlets/GenerateModerationArtifactsCommandlet.h b/Engine/Source/Editor/UnrealEd/Classes/Commandlets/GenerateModerationArtifactsCommandlet.h
deleted file mode 100644
diff --git a/Engine/Source/Editor/UnrealEd/Private/Commandlets/BaseIteratePackagesCommandlet.cpp b/Engine/Source/Editor/UnrealEd/Private/Commandlets/BaseIteratePackagesCommandlet.cpp
--- a/Engine/Source/Editor/UnrealEd/Private/Commandlets/BaseIteratePackagesCommandlet.cpp
+++ b/Engine/Source/Editor/UnrealEd/Private/Commandlets/BaseIteratePackagesCommandlet.cpp
@@ -665,6 +665,8 @@ void UBaseIteratePackagesCommandlet::LoadAndSaveOnePackage(const FString& Filena
 				PerformAdditionalOperations(Object, bSavePackage);
 			});
 			
+			PostPerformAdditionalOperations(Package);
+
 			VerboseMessage(TEXT("Post PerformAdditionalOperations Loop"));
 
 			if (bStripEditorOnlyContent)
diff --git a/Engine/Source/Editor/UnrealEd/Private/Commandlets/GeneratemoderationArtifactsCommandlet.cpp b/Engine/Source/Editor/UnrealEd/Private/Commandlets/GeneratemoderationArtifactsCommandlet.cpp
deleted file mode 100644
