Fixed crash when display loading errors for OFPA packages
- When getting a load error during the project selection screen, the content drawer doesn't exist so don't try to open it.
- Don't try to search for assets inside map data packages or packages flagged as not containing assets.
- Removed a CastChecked that would fail for packages that we can't load.
- Changed OnGotoAsset so that it doesn't load any assets and uses FAssetData instead.

diff --git a/Engine/Source/Editor/ContentBrowser/Private/ContentBrowserSingleton.cpp b/Engine/Source/Editor/ContentBrowser/Private/ContentBrowserSingleton.cpp
--- a/Engine/Source/Editor/ContentBrowser/Private/ContentBrowserSingleton.cpp
+++ b/Engine/Source/Editor/ContentBrowser/Private/ContentBrowserSingleton.cpp
@@ -705,7 +705,7 @@ void FContentBrowserSingleton::ChooseNewPrimaryBrowser()
 	}
 
 	// Set the content browser drawer as the primary browser
-	if (!PrimaryContentBrowser.IsValid())
+	if (!PrimaryContentBrowser.IsValid() && ContentBrowserDrawer.IsValid())
 	{
 		SetPrimaryContentBrowser(ContentBrowserDrawer.Pin().ToSharedRef());
 	}
diff --git a/Engine/Source/Editor/UnrealEd/Private/UnrealEdMisc.cpp b/Engine/Source/Editor/UnrealEd/Private/UnrealEdMisc.cpp
--- a/Engine/Source/Editor/UnrealEd/Private/UnrealEdMisc.cpp
+++ b/Engine/Source/Editor/UnrealEd/Private/UnrealEdMisc.cpp
@@ -1637,20 +1637,14 @@ void FUnrealEdMisc::OnGotoAsset(const FString& InAssetPath) const
 		TArray<FAssetData> AssetDataToSync;
 
 		// if its a package, sync the browser to the assets inside the package
-		if(AssetData.GetClass() == UPackage::StaticClass())
+		if(AssetData.GetClass() == UPackage::StaticClass() && !AssetData.HasAnyPackageFlags(PKG_ContainsMapData | PKG_ContainsNoAsset))
 		{
-			TArray<UPackage*> Packages;
-			Packages.Add(CastChecked<UPackage>(AssetData.GetAsset()));
-			TArray<UObject*> ObjectsInPackages;
-			UPackageTools::GetObjectsInPackages(&Packages, ObjectsInPackages);
-
-			for(auto It(ObjectsInPackages.CreateConstIterator()); It; ++It)
+			TArray<FAssetData> AssetsInPackage;
+			if (AssetRegistry.GetAssetsByPackageName(AssetData.PackageName, AssetsInPackage))
 			{
-				UObject* ObjectInPackage = *It;
-				if(ObjectInPackage->IsAsset())
+				for (const FAssetData& SubAssetData : AssetsInPackage)
 				{
-					FAssetData SubAssetData(ObjectInPackage);
-					if(SubAssetData.IsValid())
+					if (SubAssetData.IsValid() && SubAssetData.IsUAsset())
 					{
 						AssetDataToSync.Add(SubAssetData);
 					}
