Adapting FScreenshotRequest OnScreenshotCaptured to broadcast even when not in automation testing mode.

This change allows to take editor viewport screenshots and pass them in-memory without having to load screenshots from files

diff --git a/Engine/Source/Editor/UnrealEd/Private/EditorViewportClient.cpp b/Engine/Source/Editor/UnrealEd/Private/EditorViewportClient.cpp
--- a/Engine/Source/Editor/UnrealEd/Private/EditorViewportClient.cpp
+++ b/Engine/Source/Editor/UnrealEd/Private/EditorViewportClient.cpp
@@ -5877,18 +5877,18 @@ bool FEditorViewportClient::ProcessScreenShots(FViewport* InViewport)
 					BitmapSize = FIntPoint(NewWidth, NewHeight);
 				}
 			}
-
-			if (GIsAutomationTesting)
+			
+			if (FScreenshotRequest::OnScreenshotCaptured().IsBound())
 			{
-				// Under automation test, the screenshot is highjacked and sent to be compared
-				TArray<FColor> BitmapForCompare(Bitmap);
+				TArray<FColor> BitmapForBroadcast(Bitmap);
+
 				if (!bWriteAlpha)
 				{
 					// Set full alpha on the bitmap
-					for (FColor& Pixel : BitmapForCompare) { Pixel.A = 255; }
+					for (FColor& Pixel : BitmapForBroadcast) { Pixel.A = 255; }
 				}
 
-				FScreenshotRequest::OnScreenshotCaptured().Broadcast(BitmapSize.X, BitmapSize.Y, MoveTemp(BitmapForCompare));
+				FScreenshotRequest::OnScreenshotCaptured().Broadcast(BitmapSize.X, BitmapSize.Y, MoveTemp(BitmapForBroadcast));
 			}
 
 			bool bSuppressWritingToFile = false;
