MultiprocessCook - Modify memory and corelimit settings when cooking multiprocess.
#rnx

diff --git a/Engine/Source/Editor/UnrealEd/Private/Cooker/CookDirector.cpp b/Engine/Source/Editor/UnrealEd/Private/Cooker/CookDirector.cpp
--- a/Engine/Source/Editor/UnrealEd/Private/Cooker/CookDirector.cpp
+++ b/Engine/Source/Editor/UnrealEd/Private/Cooker/CookDirector.cpp
@@ -2,6 +2,7 @@
 
 #include "CookDirector.h"
 
+#include "Async/Fundamental/Scheduler.h"
 #include "CompactBinaryTCP.h"
 #include "CookMPCollector.h"
 #include "CookPackageData.h"
@@ -9,15 +10,19 @@
 #include "CookOnTheSide/CookOnTheFlyServer.h"
 #include "CoreGlobals.h"
 #include "LoadBalanceCookBurden.h"
+#include "HAL/PlatformMisc.h"
 #include "Math/NumericLimits.h"
 #include "Misc/CommandLine.h"
 #include "Misc/ConfigCacheIni.h"
 #include "Serialization/CompactBinary.h"
 #include "Serialization/CompactBinaryWriter.h"
+#include "ShaderCompiler.h"
 #include "Sockets.h"
 #include "SocketSubsystem.h"
 #include "String/ParseTokens.h"
 
+CORE_API int32 GNumForegroundWorkers; // TaskGraph.cpp
+
 namespace UE::Cook
 {
 
@@ -120,6 +125,7 @@ FCookDirector::~FCookDirector()
 void FCookDirector::AssignRequests(TArrayView<UE::Cook::FPackageData*> Requests, TArray<FWorkerId>& OutAssignments,
 	TMap<FPackageData*, TArray<FPackageData*>>&& RequestGraph)
 {
+	ActivateMachineResourceReduction();
 	InitializeWorkers();
 
 	int32 NumRemoteWorkers = RemoteWorkers.Num();
@@ -442,6 +448,62 @@ void FCookDirector::InitializeWorkers()
 	}
 }
 
+void FCookDirector::ActivateMachineResourceReduction()
+{
+	if (bHasReducedMachineResources)
+	{
+		return;
+	}
+	bHasReducedMachineResources = true;
+
+	// Add MemoryInFree if it's not already set
+	if (COTFS.MemoryMinFreeVirtual > 0 || COTFS.MemoryMinFreePhysical > 0)
+	{
+		// When running a multiprocess cook, we remove the MemoryMaxUsed triggers and allow GCing based solely on MemoryMinFree
+		COTFS.MemoryMaxUsedVirtual = 0;
+		COTFS.MemoryMaxUsedPhysical = 0;
+	}
+	else
+	{
+		// If MemoryMinFree is not set, then keep MemoryMaxUsed but reduce it by the number of CookWorkers.
+		constexpr float FixedOverheadFraction = 0.10f;
+		float TargetRatio = (FixedOverheadFraction + (1 - FixedOverheadFraction) / RequestedCookWorkerCount);
+		COTFS.MemoryMaxUsedPhysical = TargetRatio * COTFS.MemoryMaxUsedPhysical;
+		COTFS.MemoryMaxUsedVirtual = TargetRatio * COTFS.MemoryMaxUsedVirtual;
+	}
+
+	UE_LOG(LogCook, Display, TEXT("CookMultiprocess changed CookSettings for Memory: MemoryMaxUsedVirtual %dMiB, MemoryMaxUsedPhysical %dMiB,")
+		TEXT("MemoryMinFreeVirtual % dMiB, MemoryMinFreePhysical % dMiB"),
+		COTFS.MemoryMaxUsedVirtual / 1024 / 1024, COTFS.MemoryMaxUsedPhysical / 1024 / 1024,
+		COTFS.MemoryMinFreeVirtual / 1024 / 1024, COTFS.MemoryMinFreePhysical / 1024 / 1024);
+
+	// Set CoreLimit for updating workerthreads in this process and passing to the commandline for workers
+	// CoreLimit changes both NumberOfCoresIncludingHyperthreads() and NumberOfCores to return the max 
+	int32 NumProcesses = RequestedCookWorkerCount + 1;
+	int32 NumberOfCores = FPlatformMisc::NumberOfCores();
+	int32 HyperThreadCount = FPlatformMisc::NumberOfCoresIncludingHyperthreads();
+	int32 NumberOfHyperThreadsPerCore = HyperThreadCount / NumberOfCores;
+	CoreLimit = FMath::Max(NumberOfCores / NumProcesses, 1);
+	int32 NumberOfWorkers = FMath::Max(CoreLimit - 1, 1) * NumberOfHyperThreadsPerCore;
+
+	// Update the number of Cores and WorkerThreads for this process
+	check(IsInGameThread());
+	// MPCOOKTODO: UnComment when SetCoreLimit is implemented
+	//FPlatformMisc::SetCoreLimit(CoreLimit);
+	int32 NumBackgroundWorkers = FMath::Max(1, NumberOfWorkers - FMath::Min<int32>(GNumForegroundWorkers, NumberOfWorkers));
+	int32 NumForegroundWorkers = FMath::Max(1, NumberOfWorkers - NumBackgroundWorkers);
+	LowLevelTasks::FScheduler::Get().RestartWorkers(NumForegroundWorkers, NumBackgroundWorkers);
+
+	// Update the number of ShaderCompilerWorkers that can be launched
+	// MPCOOKTODO: UnComment once OnMachineResourcesChanges is implemented
+	// GShaderCompilingManager->OnMachineResourcesChanged();
+
+	UE_LOG(LogCook, Display, TEXT("CookMultiprocess changed number of cores from %d to %d."),
+		NumberOfCores, FPlatformMisc::NumberOfCores());
+	UE_LOG(LogCook, Display, TEXT("CookMultiprocess changed number of hyperthreads from %d to %d."),
+		HyperThreadCount, FPlatformMisc::NumberOfCoresIncludingHyperthreads());
+}
+
 void FCookDirector::TickWorkerConnects()
 {
 	using namespace UE::CompactBinaryTCP;
@@ -577,7 +639,10 @@ FString FCookDirector::GetWorkerCommandLine(FWorkerId WorkerId)
 				Token.StartsWith(TEXT("-CookDirectorCount=")) ||
 				Token.StartsWith(TEXT("-CookDirectorHost=")) ||
 				Token.StartsWith(TEXT("-CookWorkerId=")) ||
-				Token.StartsWith(TEXT("-ShowCookWorker"))
+				Token.StartsWith(TEXT("-ShowCookWorker")) ||
+				Token.StartsWith(TEXT("-CoreLimit")) ||
+				Token.StartsWith(TEXT("-CoreLimitCores")) ||
+				Token.StartsWith(TEXT("-CoreLimitHyperThreads"))
 				)
 			{
 				return;
@@ -593,6 +658,11 @@ FString FCookDirector::GetWorkerCommandLine(FWorkerId WorkerId)
 	check(!WorkerConnectAuthority.IsEmpty()); // This should have been constructed in TryCreateWorkerConnectSocket before any CookWorkerServers could exist to call GetWorkerCommandLine
 	Tokens.Add(FString::Printf(TEXT("-CookDirectorHost=%s"), *WorkerConnectAuthority));
 	Tokens.Add(FString::Printf(TEXT("-CookWorkerId=%d"), WorkerId.GetRemoteIndex()));
+	if (CoreLimit > 0)
+	{
+		// MPCOOKTODO: Replace CoreLimit with CoreLimitCores once it is implemented
+		Tokens.Add(FString::Printf(TEXT("-CoreLimit=%d"), CoreLimit));
+	}
 
 	return FString::Join(Tokens, TEXT(" "));
 }
diff --git a/Engine/Source/Editor/UnrealEd/Private/Cooker/CookDirector.h b/Engine/Source/Editor/UnrealEd/Private/Cooker/CookDirector.h
--- a/Engine/Source/Editor/UnrealEd/Private/Cooker/CookDirector.h
+++ b/Engine/Source/Editor/UnrealEd/Private/Cooker/CookDirector.h
@@ -88,6 +88,8 @@ private:
 	bool TryCreateWorkerConnectSocket();
 	/** Initialization helper: add the local Server for a remote worker, worker process not yet created. */
 	void InitializeWorkers();
+	/** Reduce memory settings, cpusettings, and anything else that needs to be shared with CookWorkers. */
+	void ActivateMachineResourceReduction();
 	/** Tick helper: tick any workers that have not yet finished initialization. */
 	void TickWorkerConnects();
 	/** Tick helper: tick any workers that are shutting down. */
@@ -124,8 +126,10 @@ private:
 	int32 RequestedCookWorkerCount = 0;
 	int32 CookWorkerCount = 0;
 	int32 WorkerConnectPort = 0;
+	int32 CoreLimit = 0;
 	bool bWorkersStalled = false;
 	bool bIsFirstAssignment = true;
+	bool bHasReducedMachineResources = false;
 	EShowWorker ShowWorkerOption = EShowWorker::CombinedLogs;
 	ELoadBalanceAlgorithm LoadBalanceAlgorithm = ELoadBalanceAlgorithm::CookBurden;
 
diff --git a/Engine/Source/Runtime/Core/Private/Async/TaskGraph.cpp b/Engine/Source/Runtime/Core/Private/Async/TaskGraph.cpp
--- a/Engine/Source/Runtime/Core/Private/Async/TaskGraph.cpp
+++ b/Engine/Source/Runtime/Core/Private/Async/TaskGraph.cpp
@@ -109,7 +109,7 @@ static FAutoConsoleVariableRef CVarTaskDynamicPrioritization(
 );
 
 CORE_API int32 GUseNewTaskBackend = 1;
-int32 GNumForegroundWorkers = 2;
+CORE_API int32 GNumForegroundWorkers = 2;
 static FAutoConsoleVariableRef CVarNumForegroundWorkers(
 	TEXT("TaskGraph.NumForegroundWorkers"),
 	GNumForegroundWorkers,
