Actors can no longer be attached to a Level Instance not being edited.

diff --git a/Engine/Source/Editor/UnrealEd/Private/EditorEngine.cpp b/Engine/Source/Editor/UnrealEd/Private/EditorEngine.cpp
--- a/Engine/Source/Editor/UnrealEd/Private/EditorEngine.cpp
+++ b/Engine/Source/Editor/UnrealEd/Private/EditorEngine.cpp
@@ -4195,8 +4195,8 @@ bool UEditorEngine::CanParentActors( const AActor* ParentActor, const AActor* Ch
 		return false;
 	}
 
-	const ABrush* ParentBrush = Cast<const  ABrush >( ParentActor );
-	const ABrush* ChildBrush = Cast<const  ABrush >( ChildActor );
+	const ABrush* ParentBrush = Cast<const ABrush>( ParentActor );
+	const ABrush* ChildBrush = Cast<const ABrush>( ChildActor );
 	if( (ParentBrush && !ParentBrush->IsVolumeBrush() ) || ( ChildBrush && !ChildBrush->IsVolumeBrush() ) )
 	{
 		if (ReasonText)
@@ -4208,7 +4208,7 @@ bool UEditorEngine::CanParentActors( const AActor* ParentActor, const AActor* Ch
 
 	{
 		FText Reason;
-		if (!ChildActor->EditorCanAttachTo(ParentActor, Reason))
+		if (!ParentActor->EditorCanAttachFrom(ChildActor, Reason) || !ChildActor->EditorCanAttachTo(ParentActor, Reason))
 		{
 			if (ReasonText)
 			{
diff --git a/Engine/Source/Runtime/Engine/Classes/GameFramework/Actor.h b/Engine/Source/Runtime/Engine/Classes/GameFramework/Actor.h
--- a/Engine/Source/Runtime/Engine/Classes/GameFramework/Actor.h
+++ b/Engine/Source/Runtime/Engine/Classes/GameFramework/Actor.h
@@ -2310,6 +2310,9 @@ public:
 	/** Returns true if this actor is allowed to be attached to the given actor */
 	virtual bool EditorCanAttachTo(const AActor* InParent, FText& OutReason) const;
 
+	/** Returns true if this actor is allowed to be attached from the given actor */
+	virtual bool EditorCanAttachFrom(const AActor* InChild, FText& OutReason) const;
+
 	/** Returns the actor attachement parent that should be used in editor */
 	virtual AActor* GetSceneOutlinerParent() const;
 
diff --git a/Engine/Source/Runtime/Engine/Private/ActorEditor.cpp b/Engine/Source/Runtime/Engine/Private/ActorEditor.cpp
--- a/Engine/Source/Runtime/Engine/Private/ActorEditor.cpp
+++ b/Engine/Source/Runtime/Engine/Private/ActorEditor.cpp
@@ -973,6 +973,11 @@ bool AActor::EditorCanAttachTo(const AActor* InParent, FText& OutReason) const
 	return true;
 }
 
+bool AActor::EditorCanAttachFrom(const AActor* InChild, FText& OutReason) const
+{
+	return true;
+}
+
 AActor* AActor::GetSceneOutlinerParent() const
 {
 	return GetAttachParentActor();
diff --git a/Engine/Source/Runtime/Engine/Private/LevelInstance/LevelInstanceActor.cpp b/Engine/Source/Runtime/Engine/Private/LevelInstance/LevelInstanceActor.cpp
--- a/Engine/Source/Runtime/Engine/Private/LevelInstance/LevelInstanceActor.cpp
+++ b/Engine/Source/Runtime/Engine/Private/LevelInstance/LevelInstanceActor.cpp
@@ -271,6 +271,16 @@ bool ALevelInstance::OpenAssetEditor()
 	return CanEnterEdit() ? EnterEdit() : false;
 }
 
+bool ALevelInstance::EditorCanAttachFrom(const AActor* InChild, FText& OutReason) const
+{
+	if (IsEditing())
+	{
+		return true;
+	}
+
+	return false;
+}
+
 FBox ALevelInstance::GetComponentsBoundingBox(bool bNonColliding, bool bIncludeFromChildActors) const
 {
 	FBox Box = Super::GetComponentsBoundingBox(bNonColliding, bIncludeFromChildActors);
diff --git a/Engine/Source/Runtime/Engine/Public/LevelInstance/LevelInstanceActor.h b/Engine/Source/Runtime/Engine/Public/LevelInstance/LevelInstanceActor.h
--- a/Engine/Source/Runtime/Engine/Public/LevelInstance/LevelInstanceActor.h
+++ b/Engine/Source/Runtime/Engine/Public/LevelInstance/LevelInstanceActor.h
@@ -97,6 +97,7 @@ public:
 	virtual bool GetReferencedContentObjects(TArray<UObject*>& Objects) const override;
 	virtual bool GetSoftReferencedContentObjects(TArray<FSoftObjectPath>& SoftObjects) const override;
 	virtual bool OpenAssetEditor() override;
+	virtual bool EditorCanAttachFrom(const AActor* InChild, FText& OutReason) const override;
 	// End of AActor interface
 
 	DECLARE_MULTICAST_DELEGATE_OneParam(FOnLevelInstanceActorPostLoad, ALevelInstance*);
