Unified autosubmit code for World Partition Builders
* Added auto submit support to existing builders
* Added auto submit tags support to allow scripts to provide robomerge tags for a CL

diff --git a/Engine/Source/Editor/UnrealEd/Private/WorldPartition/WorldPartitionBuilder.cpp b/Engine/Source/Editor/UnrealEd/Private/WorldPartition/WorldPartitionBuilder.cpp
--- a/Engine/Source/Editor/UnrealEd/Private/WorldPartition/WorldPartitionBuilder.cpp
+++ b/Engine/Source/Editor/UnrealEd/Private/WorldPartition/WorldPartitionBuilder.cpp
@@ -20,6 +20,10 @@
 #include "UObject/SavePackage.h"
 #include "Algo/Transform.h"
 
+#include "ISourceControlModule.h"
+#include "ISourceControlProvider.h"
+#include "SourceControlOperations.h"
+
 DEFINE_LOG_CATEGORY_STATIC(LogWorldPartitionBuilder, All, All);
 
 namespace FWorldPartitionBuilderLog
@@ -49,6 +53,10 @@ UWorldPartitionBuilder::UWorldPartitionBuilder(const FObjectInitializer& ObjectI
 	: Super(ObjectInitializer)
 {
 	bAutoSubmit = FParse::Param(FCommandLine::Get(), TEXT("AutoSubmit"));
+	if (bAutoSubmit)
+	{
+		FParse::Value(FCommandLine::Get(), TEXT("AutoSubmitTags="), AutoSubmitTags);
+	}
 }
 
 bool UWorldPartitionBuilder::RunBuilder(UWorld* World)
@@ -390,4 +398,44 @@ bool UWorldPartitionBuilder::DeletePackages(const TArray<FString>& PackageNames,
 	}
 
 	return true;
-}
\ No newline at end of file
+}
+
+bool UWorldPartitionBuilder::AutoSubmitPackages(const TArray<UPackage*>& InModifiedPackages, const FString& InChangelistDescription) const
+{
+	TArray<FString> ModifiedFiles;
+	Algo::Transform(InModifiedPackages, ModifiedFiles, [](const UPackage* InPackage) { return USourceControlHelpers::PackageFilename(InPackage); });
+	return AutoSubmitFiles(ModifiedFiles, InChangelistDescription);
+}
+
+bool UWorldPartitionBuilder::AutoSubmitFiles(const TArray<FString>& InModifiedFiles, const FString& InChangelistDescription) const
+{
+	bool bSucceeded = true;
+
+	if (bAutoSubmit)
+	{
+		UE_LOG(LogWorldPartitionBuilder, Display, TEXT("Submitting changes to source control..."));
+
+		if (!InModifiedFiles.IsEmpty())
+		{
+			FText ChangelistDescription = FText::FromString(FString::Printf(TEXT("%s\nBased on CL %d\n%s"), *InChangelistDescription, FEngineVersion::Current().GetChangelist(), *AutoSubmitTags));
+
+			TSharedRef<FCheckIn, ESPMode::ThreadSafe> CheckInOperation = ISourceControlOperation::Create<FCheckIn>();
+			CheckInOperation->SetDescription(ChangelistDescription);
+			if (ISourceControlModule::Get().GetProvider().Execute(CheckInOperation, InModifiedFiles) != ECommandResult::Succeeded)
+			{
+				UE_LOG(LogWorldPartitionBuilder, Error, TEXT("Failed to submit changes to source control."));
+				bSucceeded = false;
+			}
+			else
+			{
+				UE_LOG(LogWorldPartitionBuilder, Display, TEXT("Submitted changes to source control"));
+			}
+		}
+		else
+		{
+			UE_LOG(LogWorldPartitionBuilder, Display, TEXT("No files to submit!"));
+		}
+	}
+
+	return bSucceeded;
+}
diff --git a/Engine/Source/Editor/UnrealEd/Private/WorldPartition/WorldPartitionHLODsBuilder.cpp b/Engine/Source/Editor/UnrealEd/Private/WorldPartition/WorldPartitionHLODsBuilder.cpp
--- a/Engine/Source/Editor/UnrealEd/Private/WorldPartition/WorldPartitionHLODsBuilder.cpp
+++ b/Engine/Source/Editor/UnrealEd/Private/WorldPartition/WorldPartitionHLODsBuilder.cpp
@@ -13,10 +13,7 @@
 #include "ActorFolder.h"
 #include "EngineUtils.h"
 #include "SourceControlHelpers.h"
-
 #include "ISourceControlModule.h"
-#include "ISourceControlProvider.h"
-#include "SourceControlOperations.h"
 
 #include "AssetRegistry/AssetRegistryModule.h"
 #include "DerivedDataCacheInterface.h"
@@ -573,36 +570,12 @@ bool UWorldPartitionHLODsBuilder::DeleteHLODActors()
 
 bool UWorldPartitionHLODsBuilder::SubmitHLODActors()
 {
-	bool bRet = true;
-
 	// Ensure all files modified by the source control helper are taken into account
 	ModifiedFiles.Append(SourceControlHelper->GetModifiedFiles());
 
-	TArray<FString> FilesToSubmit = ModifiedFiles.GetAllFiles();
-
 	// Check in all modified files
-	if (FilesToSubmit.Num() > 0)
-	{
-		FText ChangelistDescription = FText::FromString(FString::Printf(TEXT("Rebuilt HLODs for \"%s\" at %s"), *WorldPartition->GetWorld()->GetName(), *FEngineVersion::Current().ToString()));
-
-		TSharedRef<FCheckIn, ESPMode::ThreadSafe> CheckInOperation = ISourceControlOperation::Create<FCheckIn>();
-		CheckInOperation->SetDescription(ChangelistDescription);
-		bRet = ISourceControlModule::Get().GetProvider().Execute(CheckInOperation, FilesToSubmit) == ECommandResult::Succeeded;
-		if (!bRet)
-		{
-			UE_LOG(LogWorldPartitionHLODsBuilder, Error, TEXT("Failed to submit %d files to source control."), FilesToSubmit.Num());
-		}
-		else
-		{
-			UE_LOG(LogWorldPartitionHLODsBuilder, Display, TEXT("#### Submitted %d files to source control ####"), FilesToSubmit.Num());
-		}
-	}
-	else
-	{
-		UE_LOG(LogWorldPartitionHLODsBuilder, Display, TEXT("#### No files to submit ####"));
-	}
-
-	return bRet;
+	const FString ChangeDescription = FString::Printf(TEXT("Rebuilt HLODs for %s"), *WorldPartition->GetWorld()->GetName());
+	return AutoSubmitFiles(ModifiedFiles.GetAllFiles(), ChangeDescription);
 }
 
 bool UWorldPartitionHLODsBuilder::DumpStats()
diff --git a/Engine/Source/Editor/UnrealEd/Private/WorldPartition/WorldPartitionLandscapeSplineMeshesBuilder.cpp b/Engine/Source/Editor/UnrealEd/Private/WorldPartition/WorldPartitionLandscapeSplineMeshesBuilder.cpp
--- a/Engine/Source/Editor/UnrealEd/Private/WorldPartition/WorldPartitionLandscapeSplineMeshesBuilder.cpp
+++ b/Engine/Source/Editor/UnrealEd/Private/WorldPartition/WorldPartitionLandscapeSplineMeshesBuilder.cpp
@@ -8,6 +8,7 @@
 #include "UObject/SavePackage.h"
 #include "UObject/ScriptInterface.h"
 #include "UObject/GCObjectScopeGuard.h"
+#include "Algo/Transform.h"
 #include "Engine/World.h"
 #include "Engine/StaticMesh.h"
 #include "ActorFolder.h"
@@ -508,6 +509,15 @@ bool UWorldPartitionLandscapeSplineMeshesBuilder::RunInternal(UWorld* InWorld, c
 		{
 			return false;
 		}
+
+		TArray<FString> ModifiedFiles = MoveTemp(PackagesToDelete);
+		Algo::Transform(PackagesToSave, ModifiedFiles, [](const UPackage* InPackage) { return USourceControlHelpers::PackageFilename(InPackage); });
+
+		const FString ChangeDescription = FString::Printf(TEXT("Rebuilt landscape splines for %s"), *InWorld->GetName());
+		if (!AutoSubmitFiles(ModifiedFiles, ChangeDescription))
+		{
+			return false;
+		}
 	}
 
 	return true;
diff --git a/Engine/Source/Editor/UnrealEd/Private/WorldPartition/WorldPartitionMiniMapBuilder.cpp b/Engine/Source/Editor/UnrealEd/Private/WorldPartition/WorldPartitionMiniMapBuilder.cpp
--- a/Engine/Source/Editor/UnrealEd/Private/WorldPartition/WorldPartitionMiniMapBuilder.cpp
+++ b/Engine/Source/Editor/UnrealEd/Private/WorldPartition/WorldPartitionMiniMapBuilder.cpp
@@ -16,10 +16,6 @@
 #include "UObject/SavePackage.h"
 #include "UObject/StrongObjectPtr.h"
 
-#include "ISourceControlModule.h"
-#include "ISourceControlProvider.h"
-#include "SourceControlOperations.h"
-
 #include "WorldPartition/ActorDescContainer.h"
 #include "WorldPartition/DataLayer/ActorDataLayer.h"
 #include "WorldPartition/DataLayer/DataLayerInstance.h"
@@ -232,22 +228,6 @@ bool UWorldPartitionMiniMapBuilder::PostRun(UWorld* World, FPackageSourceControl
 		return false;
 	}
 
-	if (bAutoSubmit)
-	{
-		FText ChangelistDescription = FText::FromString(FString::Printf(TEXT("Rebuilt minimap for \"%s\" at %s"), *World->GetName(), *FEngineVersion::Current().ToString()));
-
-		TSharedRef<FCheckIn, ESPMode::ThreadSafe> CheckInOperation = ISourceControlOperation::Create<FCheckIn>();
-		CheckInOperation->SetDescription(ChangelistDescription);
-		if (ISourceControlModule::Get().GetProvider().Execute(CheckInOperation, PackageFileName) != ECommandResult::Succeeded)
-		{
-			UE_LOG(LogWorldPartitionMiniMapBuilder, Error, TEXT("Failed to submit minimap (%s) to source control."), *PackageFileName);
-			return false;
-		}
-		else
-		{
-			UE_LOG(LogWorldPartitionMiniMapBuilder, Display, TEXT("#### Submitted minimap (%s) to source control ####"), *PackageFileName);
-		}
-	}
-
-	return true;
+	const FString ChangeDescription = FString::Printf(TEXT("Rebuilt minimap for %s"), *World->GetName());
+	return AutoSubmitFiles({ PackageFileName }, ChangeDescription);
 }
diff --git a/Engine/Source/Editor/UnrealEd/Public/WorldPartition/WorldPartitionBuilder.h b/Engine/Source/Editor/UnrealEd/Public/WorldPartition/WorldPartitionBuilder.h
--- a/Engine/Source/Editor/UnrealEd/Public/WorldPartition/WorldPartitionBuilder.h
+++ b/Engine/Source/Editor/UnrealEd/Public/WorldPartition/WorldPartitionBuilder.h
@@ -84,6 +84,9 @@ protected:
 	 */
 	virtual bool ShouldSkipCell(const FWorldBuilderCellCoord& CellCoord) const { return false; }
 
+	bool AutoSubmitFiles(const TArray<FString>& InModifiedFiles, const FString& InChangelistDescription) const;
+	bool AutoSubmitPackages(const TArray<UPackage*>& InModifiedPackages, const FString& InChangelistDescription) const;
+
 	int32 IterativeCellSize = 102400;
 	int32 IterativeCellOverlapSize = 0;
 	FBox  IterativeWorldBounds;
@@ -94,4 +97,5 @@ protected:
 	bool bLoadInitiallyActiveDataLayers = true;
 
 	bool bAutoSubmit = false;
+	FString AutoSubmitTags;
 };
\ No newline at end of file
diff --git a/Engine/Source/Programs/AutomationTool/Scripts/WorldPartitionBuilder.Automation.cs b/Engine/Source/Programs/AutomationTool/Scripts/WorldPartitionBuilder.Automation.cs
--- a/Engine/Source/Programs/AutomationTool/Scripts/WorldPartitionBuilder.Automation.cs
+++ b/Engine/Source/Programs/AutomationTool/Scripts/WorldPartitionBuilder.Automation.cs
@@ -21,6 +21,7 @@ namespace AutomationScripts.Automation
 			string CommandletArgs = ParseOptionalStringParam("CommandletArgs");
 
 			bool bSubmitResult = ParseParam("Submit");
+			string SubmitTags = ParseOptionalStringParam("SubmitTags");
 
 			string ShelveUser = ParseOptionalStringParam("ShelveUser");
 			string ShelveWorkspace = ParseOptionalStringParam("ShelveWorkspace");
@@ -37,6 +38,11 @@ namespace AutomationScripts.Automation
 			if (bSubmitResult)
 			{
 				CommandletArgs += " -AutoSubmit";
+
+				if (SubmitTags.Length != 0)
+				{
+					CommandletArgs += " -AutoSubmitTags=\"" + SubmitTags + "\"";
+				}
 			}
 
 			string EditorExe = "UnrealEditor-Cmd.exe";
