Dissallow selection of actors or components that aren't in the current editing context

diff --git a/Engine/Source/Editor/UnrealEd/Private/EditorSelectUtils.cpp b/Engine/Source/Editor/UnrealEd/Private/EditorSelectUtils.cpp
--- a/Engine/Source/Editor/UnrealEd/Private/EditorSelectUtils.cpp
+++ b/Engine/Source/Editor/UnrealEd/Private/EditorSelectUtils.cpp
@@ -490,7 +490,7 @@ bool UUnrealEdEngine::CanSelectActor(AActor* Actor, bool bInSelected, bool bSele
 
 void UUnrealEdEngine::SelectActor(AActor* Actor, bool bInSelected, bool bNotify, bool bSelectEvenIfHidden, bool bForceRefresh)
 {
-	if (!Actor)
+	if (!Actor || Actor->GetRootSelectionParent() != nullptr)
 	{
 		return;
 	}
@@ -531,7 +531,7 @@ void UUnrealEdEngine::SelectActor(AActor* Actor, bool bInSelected, bool bNotify,
 
 void UUnrealEdEngine::SelectComponent(UActorComponent* Component, bool bInSelected, bool bNotify, bool bSelectEvenIfHidden)
 {
-	if (!Component)
+	if (!Component || (Component->GetOwner() && Component->GetOwner()->GetRootSelectionParent() != nullptr))
 	{
 		return;
 	}
