World Partition: replace popup by on-screen message when streaming is disabled but the world is large enough to justify enabling it.

fyi jeansebastien.guay
#rnx

diff --git a/Engine/Source/Editor/LevelEditor/Private/SLevelViewport.cpp b/Engine/Source/Editor/LevelEditor/Private/SLevelViewport.cpp
--- a/Engine/Source/Editor/LevelEditor/Private/SLevelViewport.cpp
+++ b/Engine/Source/Editor/LevelEditor/Private/SLevelViewport.cpp
@@ -390,11 +390,11 @@ void SLevelViewport::ConstructViewportOverlayContent()
 			SNew(SBorder)
 			.BorderImage(FAppStyle::Get().GetBrush("FloatingBorder"))
 			[
-				SNew(SWorldPartitionViewportWidget)
+				SAssignNew(WorldPartitionViewportWidget, SWorldPartitionViewportWidget)
 				.Clickable(true)
 				.Visibility_Lambda([this]() 
 				{ 
-					return (!ActiveViewport.IsValid() || IsPlayInEditorViewportActive()) ? EVisibility::Collapsed : SWorldPartitionViewportWidget::GetVisibility(GetWorld());
+					return (!ActiveViewport.IsValid() || IsPlayInEditorViewportActive()) ? EVisibility::Collapsed : WorldPartitionViewportWidget->GetVisibility(GetWorld());
 				})
 			]
 		]
diff --git a/Engine/Source/Editor/LevelEditor/Public/SLevelViewport.h b/Engine/Source/Editor/LevelEditor/Public/SLevelViewport.h
--- a/Engine/Source/Editor/LevelEditor/Public/SLevelViewport.h
+++ b/Engine/Source/Editor/LevelEditor/Public/SLevelViewport.h
@@ -11,6 +11,7 @@
 #include "Widgets/SWidget.h"
 #include "Widgets/DeclarativeSyntaxSupport.h"
 #include "Framework/Commands/UICommandInfo.h"
+#include "SWorldPartitionViewportWidget.h"
 #include "EditorViewportClient.h"
 #include "Widgets/Layout/SScaleBox.h"
 #include "Widgets/SWindow.h"
@@ -936,6 +937,9 @@ private:
 	bool bIsInViewportMenuInitialized;
 	TSharedPtr<class SInViewportDetails> InViewportMenu;
 	static bool bInViewportMenuEnabled;
+
+	TSharedPtr<SWorldPartitionViewportWidget> WorldPartitionViewportWidget;
+
 protected:
 	void LockActorInternal(AActor* NewActorToLock);
 
diff --git a/Engine/Source/Editor/UnrealEd/Private/Commandlets/WorldPartitionConvertCommandlet.cpp b/Engine/Source/Editor/UnrealEd/Private/Commandlets/WorldPartitionConvertCommandlet.cpp
--- a/Engine/Source/Editor/UnrealEd/Private/Commandlets/WorldPartitionConvertCommandlet.cpp
+++ b/Engine/Source/Editor/UnrealEd/Private/Commandlets/WorldPartitionConvertCommandlet.cpp
@@ -227,7 +227,6 @@ UWorldPartition* UWorldPartitionConvertCommandlet::CreateWorldPartition(AWorldSe
 	if (bDisableStreaming)
 	{
 		WorldPartition->bEnableStreaming = false;
-		WorldPartition->bStreamingWasEnabled = false;
 	}
 		
 	// Read the conversion config file
diff --git a/Engine/Source/Editor/UnrealEd/Private/SWorldPartitionViewportWidget.cpp b/Engine/Source/Editor/UnrealEd/Private/SWorldPartitionViewportWidget.cpp
--- a/Engine/Source/Editor/UnrealEd/Private/SWorldPartitionViewportWidget.cpp
+++ b/Engine/Source/Editor/UnrealEd/Private/SWorldPartitionViewportWidget.cpp
@@ -25,14 +25,14 @@ void SWorldPartitionViewportWidget::Construct(const FArguments& InArgs)
 		.ButtonStyle(bClickable? &FCoreStyle::Get().GetWidgetStyle<FButtonStyle>("HoverHintOnly") : &FCoreStyle::Get().GetWidgetStyle<FButtonStyle>("NoBorder"))
 		.Cursor(EMouseCursor::Default)
 		.IsEnabled_Lambda([this]() { return bClickable; })
-		.OnClicked_Lambda([]
+		.OnClicked_Lambda([this]()
 		{
 			FLevelEditorModule& LevelEditorModule = FModuleManager::GetModuleChecked<FLevelEditorModule>(TEXT("LevelEditor"));
 			TSharedPtr<FTabManager> LevelEditorTabManager = LevelEditorModule.GetLevelEditorTabManager();
-			LevelEditorTabManager->TryInvokeTab(FName("WorldBrowserPartitionEditor"));
+			LevelEditorModule.GetLevelEditorTabManager()->TryInvokeTab(InvokeTab);
 			return FReply::Handled();
 		})
-		.ToolTipText(LOCTEXT("NoRegionsLoadedTooltip", "To load a region, drag select an area in the World Partition map and choose 'Load Region From Selection' from the context menu."))
+		.ToolTipText_Lambda([this]() { return Tooltip; })
 		.ContentPadding(0.f)
 		[
 			SNew(SVerticalBox)
@@ -56,28 +56,41 @@ void SWorldPartitionViewportWidget::Construct(const FArguments& InArgs)
 				.FillWidth(1)
 				[
 					SNew(STextBlock)
-					.Text(LOCTEXT("NoRegionsLoadedText","No regions loaded"))
+					.Text_Lambda([this]() { return Message; })
 					.ColorAndOpacity(FAppStyle::Get().GetSlateColor("Colors.Warning"))
 				]
 			]
 		]
 	];
-
-
 }
 
 SWorldPartitionViewportWidget::~SWorldPartitionViewportWidget()
-{
-}
+{}
 
 EVisibility SWorldPartitionViewportWidget::GetVisibility(UWorld* InWorld)
 {
 	if (InWorld && InWorld->IsPartitionedWorld())
 	{
 		UWorldPartition* WorldPartition = InWorld->GetWorldPartition();
-		if (WorldPartition->IsStreamingEnabled() && !WorldPartition->HasLoadedUserCreatedRegions())
+		if (WorldPartition->IsStreamingEnabled())
+		{
+			if (!WorldPartition->HasLoadedUserCreatedRegions())
+			{
+				Message = LOCTEXT("NoRegionsLoadedText","No regions loaded");
+				Tooltip = LOCTEXT("NoRegionsLoadedToolTip", "To load a region, drag select an area in the World Partition map and choose 'Load Region From Selection' from the context menu.");
+				InvokeTab = TEXT("WorldBrowserPartitionEditor");
+				return EVisibility::Visible;
+			}
+		}
+		else
 		{
-			return EVisibility::Visible;
+			if (WorldPartition->IsEnablingStreamingJustified())
+			{
+				Message = LOCTEXT("EnablingStreamingShouldBeEnabled","Streaming Disabled");
+				Tooltip = LOCTEXT("EnablingStreamingShouldBeEnabled", "The size of your world has grown enough to justify enabling streaming.");
+				InvokeTab = TEXT("WorldSettingsTab");
+				return EVisibility::Visible;
+			}
 		}
 	}
 
diff --git a/Engine/Source/Editor/UnrealEd/Public/SWorldPartitionViewportWidget.h b/Engine/Source/Editor/UnrealEd/Public/SWorldPartitionViewportWidget.h
--- a/Engine/Source/Editor/UnrealEd/Public/SWorldPartitionViewportWidget.h
+++ b/Engine/Source/Editor/UnrealEd/Public/SWorldPartitionViewportWidget.h
@@ -14,8 +14,11 @@ public:
 	void Construct(const FArguments& InArgs);
 	~SWorldPartitionViewportWidget();
 
-	static EVisibility GetVisibility(UWorld* InWorld);
-private:
+	EVisibility GetVisibility(UWorld* InWorld);
 
+private:
 	bool bClickable;
+	FText Message;
+	FText Tooltip;
+	FName InvokeTab;
 };
\ No newline at end of file
diff --git a/Engine/Source/Editor/WorldPartitionEditor/Private/WorldPartition/SWorldPartitionEditorGrid2D.cpp b/Engine/Source/Editor/WorldPartitionEditor/Private/WorldPartition/SWorldPartitionEditorGrid2D.cpp
--- a/Engine/Source/Editor/WorldPartitionEditor/Private/WorldPartition/SWorldPartitionEditorGrid2D.cpp
+++ b/Engine/Source/Editor/WorldPartitionEditor/Private/WorldPartition/SWorldPartitionEditorGrid2D.cpp
@@ -390,9 +390,9 @@ void SWorldPartitionEditorGrid2D::Construct(const FArguments& InArgs)
 		.HAlign(HAlign_Left)
 		.Padding(10.f, 0.f, 0.f, 10.f)
 		[
-			SNew(SWorldPartitionViewportWidget)
+			SAssignNew(ViewportWidget, SWorldPartitionViewportWidget)
 			.Clickable(false)
-			.Visibility_Lambda([this]() { return SWorldPartitionViewportWidget::GetVisibility(World); })
+			.Visibility_Lambda([this]() { return ViewportWidget->GetVisibility(World); })
 		]
 	];
 
@@ -569,7 +569,7 @@ void SWorldPartitionEditorGrid2D::MoveCameraHere()
 
 	for (FLevelEditorViewportClient* LevelVC : GEditor->GetLevelViewportClients())
 	{
-		WorldLocation.Z = bHitResultValid ? HitResult.ImpactPoint.Z + 250.0f : LevelVC->GetViewLocation().Z;
+		WorldLocation.Z = bHitResultValid ? HitResult.ImpactPoint.Z + 1000.0f : LevelVC->GetViewLocation().Z;
 
 		LevelVC->SetViewLocation(WorldLocation);
 		LevelVC->Invalidate();
diff --git a/Engine/Source/Editor/WorldPartitionEditor/Public/WorldPartition/SWorldPartitionEditorGrid2D.h b/Engine/Source/Editor/WorldPartitionEditor/Public/WorldPartition/SWorldPartitionEditorGrid2D.h
--- a/Engine/Source/Editor/WorldPartitionEditor/Public/WorldPartition/SWorldPartitionEditorGrid2D.h
+++ b/Engine/Source/Editor/WorldPartitionEditor/Public/WorldPartition/SWorldPartitionEditorGrid2D.h
@@ -8,6 +8,7 @@
 #include "Framework/Commands/UICommandList.h"
 #include "Layout/ArrangedChildren.h"
 #include "SWorldPartitionEditorGrid.h"
+#include "SWorldPartitionViewportWidget.h"
 #include "WorldPartition/WorldPartitionActorLoaderInterface.h"
 
 class SWorldPartitionEditorGrid2D : public SWorldPartitionEditorGrid
@@ -160,4 +161,6 @@ protected:
 	// Profiling
 	mutable double TickTime;
 	mutable double PaintTime;
+
+	TSharedPtr<SWorldPartitionViewportWidget> ViewportWidget;
 };
\ No newline at end of file
diff --git a/Engine/Source/Runtime/Engine/Private/WorldPartition/WorldPartition.cpp b/Engine/Source/Runtime/Engine/Private/WorldPartition/WorldPartition.cpp
--- a/Engine/Source/Runtime/Engine/Private/WorldPartition/WorldPartition.cpp
+++ b/Engine/Source/Runtime/Engine/Private/WorldPartition/WorldPartition.cpp
@@ -232,11 +232,11 @@ UWorldPartition::UWorldPartition(const FObjectInitializer& ObjectInitializer)
 	, PinnedActors(nullptr)
 	, WorldPartitionEditor(nullptr)
 	, bStreamingWasEnabled(true)
-	, bShouldEnableStreamingWarned(false)
 	, bShouldCheckEnableStreamingWarning(false)
 	, bCanBeUsedByLevelInstance(false)
 	, bForceGarbageCollection(false)
 	, bForceGarbageCollectionPurge(false)
+	, bEnablingStreamingJustified(false)
 	, bIsPIE(false)
 	, NumUserCreatedLoadedRegions(0)
 #endif
@@ -1226,8 +1226,10 @@ void UWorldPartition::Tick(float DeltaSeconds)
 	{
 		bShouldCheckEnableStreamingWarning = false;
 
-		if (!IsStreamingEnabled() && SupportsStreaming() && !bStreamingWasEnabled && !bShouldEnableStreamingWarned)
+		if (!IsStreamingEnabled() && SupportsStreaming())
 		{
+			bEnablingStreamingJustified = false;
+
 			FBox AllActorsBounds(ForceInit);
 			for (FActorDescContainerCollection::TConstIterator<> ActorDescIterator(this); ActorDescIterator; ++ActorDescIterator)
 			{
@@ -1238,18 +1240,7 @@ void UWorldPartition::Tick(float DeltaSeconds)
 					// Warn the user if the world becomes larger that 4km in any axis
 					if (AllActorsBounds.GetSize().GetMax() >= 400000.0f)
 					{
-						bShouldEnableStreamingWarned = true;
-
-						if (FMessageDialog::Open(EAppMsgType::YesNo, LOCTEXT("WorldPartitionShouldEnableStreaming", "The size of your world has grown enough to justify enabling streaming. Enable streaming now?")) == EAppReturnType::Yes)
-						{
-							SetEnableStreaming(true);
-							bStreamingWasEnabled = true;
-							if (FMessageDialog::Open(EAppMsgType::YesNo, LOCTEXT("WorldPartitionEnableStreamingDialolg", "Please refer to our documentation for how to set up streaming.\n\nWould you like to open it now? ")) == EAppReturnType::Yes)
-							{
-								IDocumentation::Get()->Open(TEXT("world-partition-in-unreal-engine"), FDocumentationSourceInfo(TEXT("worldpartition")));
-							}
-						}
-
+						bEnablingStreamingJustified = true;
 						break;
 					}
 				}
diff --git a/Engine/Source/Runtime/Engine/Public/WorldPartition/WorldPartition.h b/Engine/Source/Runtime/Engine/Public/WorldPartition/WorldPartition.h
--- a/Engine/Source/Runtime/Engine/Public/WorldPartition/WorldPartition.h
+++ b/Engine/Source/Runtime/Engine/Public/WorldPartition/WorldPartition.h
@@ -200,6 +200,8 @@ public:
 	bool HasLoadedUserCreatedRegions() const { return !!NumUserCreatedLoadedRegions; }
 	void OnUserCreatedRegionLoaded() { NumUserCreatedLoadedRegions++; }
 	void OnUserCreatedRegionUnloaded() { check(HasLoadedUserCreatedRegions()); NumUserCreatedLoadedRegions--; }
+
+	bool IsEnablingStreamingJustified() const { return bEnablingStreamingJustified; }
 #endif
 
 public:
@@ -253,10 +255,6 @@ private:
 	UPROPERTY()
 	bool bStreamingWasEnabled;
 
-	/** Used to know if the user has already been warned about that it should enable streaming based on world size. */
-	UPROPERTY()
-	bool bShouldEnableStreamingWarned;
-
 	/** Used to know if we need to recheck if the user should enable streaming based on world size. */
 	bool bShouldCheckEnableStreamingWarning;
 
@@ -285,6 +283,7 @@ private:
 #if WITH_EDITOR
 	bool bForceGarbageCollection;
 	bool bForceGarbageCollectionPurge;
+	bool bEnablingStreamingJustified;
 	bool bIsPIE;
 	int32 NumUserCreatedLoadedRegions;
 #endif
