Minimal fix for CDO creation during skeleton reinstancing, further hardening to come

diff --git a/Engine/Source/Editor/UnrealEd/Private/Kismet2/KismetReinstanceUtilities.cpp b/Engine/Source/Editor/UnrealEd/Private/Kismet2/KismetReinstanceUtilities.cpp
--- a/Engine/Source/Editor/UnrealEd/Private/Kismet2/KismetReinstanceUtilities.cpp
+++ b/Engine/Source/Editor/UnrealEd/Private/Kismet2/KismetReinstanceUtilities.cpp
@@ -1645,7 +1645,7 @@ void FBlueprintCompileReinstancer::MoveDependentSkelToReinst(UClass* OwnerClass,
 		if (OldCDO)
 		{
 			OwnerClass->ClassDefaultObject = nullptr;
-			OldCDO->Rename(nullptr, ReinstClass->GetOuter(), REN_DoNotDirty | REN_DontCreateRedirectors | REN_ForceNoResetLoaders);
+			OldCDO->Rename(nullptr, ReinstClass->GetOuter(), REN_DoNotDirty | REN_DontCreateRedirectors | REN_ForceNoResetLoaders | REN_NonTransactional);
 			ReinstClass->ClassDefaultObject = OldCDO;
 			OldCDO->SetClass(ReinstClass);
 		}
