Correct reference to extern variable to fix duplicate symbol errors

diff --git a/Engine/Source/Editor/UnrealEd/Private/Cooker/CookDirector.cpp b/Engine/Source/Editor/UnrealEd/Private/Cooker/CookDirector.cpp
--- a/Engine/Source/Editor/UnrealEd/Private/Cooker/CookDirector.cpp
+++ b/Engine/Source/Editor/UnrealEd/Private/Cooker/CookDirector.cpp
@@ -21,7 +21,7 @@
 #include "SocketSubsystem.h"
 #include "String/ParseTokens.h"
 
-CORE_API int32 GNumForegroundWorkers; // TaskGraph.cpp
+extern CORE_API int32 GNumForegroundWorkers; // TaskGraph.cpp
 
 namespace UE::Cook
 {
