Fix localization warnings.

#rnx

diff --git a/Engine/Source/Editor/UnrealEd/Private/SWorldPartitionViewportWidget.cpp b/Engine/Source/Editor/UnrealEd/Private/SWorldPartitionViewportWidget.cpp
--- a/Engine/Source/Editor/UnrealEd/Private/SWorldPartitionViewportWidget.cpp
+++ b/Engine/Source/Editor/UnrealEd/Private/SWorldPartitionViewportWidget.cpp
@@ -76,7 +76,7 @@ EVisibility SWorldPartitionViewportWidget::GetVisibility(UWorld* InWorld)
 		{
 			if (!WorldPartition->HasLoadedUserCreatedRegions())
 			{
-				Message = LOCTEXT("NoRegionsLoadedText","No regions loaded");
+				Message = LOCTEXT("NoRegionsLoadedText","No Regions Loaded");
 				Tooltip = LOCTEXT("NoRegionsLoadedToolTip", "To load a region, drag select an area in the World Partition map and choose 'Load Region From Selection' from the context menu.");
 				InvokeTab = TEXT("WorldBrowserPartitionEditor");
 				return EVisibility::Visible;
@@ -86,8 +86,8 @@ EVisibility SWorldPartitionViewportWidget::GetVisibility(UWorld* InWorld)
 		{
 			if (WorldPartition->IsEnablingStreamingJustified())
 			{
-				Message = LOCTEXT("EnablingStreamingShouldBeEnabled","Streaming Disabled");
-				Tooltip = LOCTEXT("EnablingStreamingShouldBeEnabled", "The size of your world has grown enough to justify enabling streaming.");
+				Message = LOCTEXT("StreamingDisabledText","Streaming Disabled");
+				Tooltip = LOCTEXT("StreamingDisabledToolTip", "The size of your world has grown enough to justify enabling streaming.");
 				InvokeTab = TEXT("WorldSettingsTab");
 				return EVisibility::Visible;
 			}
