FBX LOD Import: Fixed crash when restoring collision info during a LOD reimport.
- In that case, the original mesh and the reimported mesh are pointing to the same UBodySetup and we were trying to transfer the collision shapes from the original body setup to the reimported body setup.

diff --git a/Engine/Source/Editor/UnrealEd/Private/ImportUtils/StaticMeshImportUtils.cpp b/Engine/Source/Editor/UnrealEd/Private/ImportUtils/StaticMeshImportUtils.cpp
--- a/Engine/Source/Editor/UnrealEd/Private/ImportUtils/StaticMeshImportUtils.cpp
+++ b/Engine/Source/Editor/UnrealEd/Private/ImportUtils/StaticMeshImportUtils.cpp
@@ -1279,7 +1279,7 @@ void StaticMeshImportUtils::RestoreExistingMeshData(const TSharedPtr<const FExis
 		{
 			NewMesh->SetBodySetup(ExistingMeshDataPtr->ExistingBodySetup);
 		}
-		else
+		else if (NewMesh->GetBodySetup() != ExistingMeshDataPtr->ExistingBodySetup)
 		{
 			// New collision geometry, but we still want the original settings and the generated collisions
 			NewMesh->GetBodySetup()->CopyBodySetupProperty(ExistingMeshDataPtr->ExistingBodySetup);
