EditorFactories fix crash importing jpegs that fail to decode
make it check decompression will succeed at import time, even if using retain-compressed option
fix some other null derefs on failed imports

diff --git a/Engine/Source/Editor/UnrealEd/Private/Factories/EditorFactories.cpp b/Engine/Source/Editor/UnrealEd/Private/Factories/EditorFactories.cpp
--- a/Engine/Source/Editor/UnrealEd/Private/Factories/EditorFactories.cpp
+++ b/Engine/Source/Editor/UnrealEd/Private/Factories/EditorFactories.cpp
@@ -2878,52 +2878,6 @@ bool UTextureFactory::ImportImage(const uint8* Buffer, int64 Length, FFeedbackCo
 	// Use the magic bytes when possible to avoid calling inefficient code to check if the image is of the right format
 	EImageFormat ImageFormat = ImageWrapperModule.DetectImageFormat(Buffer, Length);
 
-	//
-	// JPEG
-	//
-	if (ImageFormat == EImageFormat::JPEG)
-	{
-		// unusual loader, retains jpeg
-		bool bRetainJpegFormat = false;
-		if (EnumHasAllFlags(Flags, EImageImportFlags::AllowReturnOfCompressedData))
-		{
-			// For now this option is opt in via the config files once there is no technical risk this will become the default path.
-			GConfig->GetBool(TEXT("TextureImporter"), TEXT("RetainJpegFormat"), bRetainJpegFormat, GEditorIni);
-		}
-		if ( bRetainJpegFormat)
-		{
-			// does not decode jpeg, just to get width/height :
-			TSharedPtr<IImageWrapper> JpegImageWrapper = ImageWrapperModule.CreateImageWrapper(EImageFormat::JPEG);
-			if (JpegImageWrapper.IsValid() && JpegImageWrapper->SetCompressed(Buffer, Length))
-			{
-				if (!IsImportResolutionValid(JpegImageWrapper->GetWidth(), JpegImageWrapper->GetHeight(), bAllowNonPowerOfTwo, Warn))
-				{
-					return false;
-				}
-			
-				ERawImageFormat::Type RawFormat = JpegImageWrapper->GetClosestRawImageFormat();
-				check( RawFormat != ERawImageFormat::Invalid );
-				ETextureSourceFormat TextureFormat = FImageCoreUtils::ConvertToTextureSourceFormat(RawFormat);
-				bool bSRGB = JpegImageWrapper->GetSRGB();
-
-				OutImage.Init2DWithParams(
-					JpegImageWrapper->GetWidth(),
-					JpegImageWrapper->GetHeight(),
-					TextureFormat,
-					bSRGB
-				);
-
-				OutImage.RawData.Append(Buffer, Length);
-				OutImage.RawDataCompressionFormat = ETextureSourceCompressionFormat::TSCF_JPEG;
-
-				UE_LOG(LogEditorFactories,Display,TEXT("JPEG imported and retained as JPEG in uasset."));
-
-				return true;
-			}
-		}
-		// if not bRetainJpegFormat , jpeg will use generic loader
-	}
-
 	if ( ImageFormat != EImageFormat::Invalid )
 	{
 		// Generic ImageWrapper loader :
@@ -2936,6 +2890,50 @@ bool UTextureFactory::ImportImage(const uint8* Buffer, int64 Length, FFeedbackCo
 				return false;
 			}
 		
+			// branch for JPEG, if retaining the jpeg compressed data
+			// this is inside the DecompressImage branch even though we don't use the LoadedImage at all
+			//	 just to ensure that the jpeg will decode successfully
+			if (ImageFormat == EImageFormat::JPEG)
+			{
+				// unusual loader, retains jpeg
+				bool bRetainJpegFormat = false;
+				if (EnumHasAllFlags(Flags, EImageImportFlags::AllowReturnOfCompressedData))
+				{
+					// For now this option is opt in via the config files once there is no technical risk this will become the default path.
+					GConfig->GetBool(TEXT("TextureImporter"), TEXT("RetainJpegFormat"), bRetainJpegFormat, GEditorIni);
+				}
+				if ( bRetainJpegFormat)
+				{
+					// does not decode jpeg, just to get width/height :
+					TSharedPtr<IImageWrapper> JpegImageWrapper = ImageWrapperModule.CreateImageWrapper(EImageFormat::JPEG);
+					if (JpegImageWrapper.IsValid() && JpegImageWrapper->SetCompressed(Buffer, Length))
+					{
+						check( JpegImageWrapper->GetWidth() == LoadedImage.SizeX );
+						check( JpegImageWrapper->GetHeight() == LoadedImage.SizeY );
+
+						ERawImageFormat::Type RawFormat = JpegImageWrapper->GetClosestRawImageFormat();
+						check( RawFormat != ERawImageFormat::Invalid );
+						ETextureSourceFormat TextureFormat = FImageCoreUtils::ConvertToTextureSourceFormat(RawFormat);
+						bool bSRGB = JpegImageWrapper->GetSRGB();
+
+						OutImage.Init2DWithParams(
+							JpegImageWrapper->GetWidth(),
+							JpegImageWrapper->GetHeight(),
+							TextureFormat,
+							bSRGB
+						);
+
+						OutImage.RawData.Append(Buffer, Length);
+						OutImage.RawDataCompressionFormat = ETextureSourceCompressionFormat::TSCF_JPEG;
+
+						UE_LOG(LogEditorFactories,Display,TEXT("JPEG imported and retained as JPEG in uasset."));
+
+						return true;
+					}
+				}
+				// if not bRetainJpegFormat , jpeg will continue with generic loader
+			}
+
 			if ( UE::TextureUtilitiesCommon::AutoDetectAndChangeGrayScale(LoadedImage) )
 			{
 				UE_LOG(LogEditorFactories,Display,TEXT("Auto-detected grayscale, image changed to G8"));
diff --git a/Engine/Source/Runtime/Engine/Private/Texture.cpp b/Engine/Source/Runtime/Engine/Private/Texture.cpp
--- a/Engine/Source/Runtime/Engine/Private/Texture.cpp
+++ b/Engine/Source/Runtime/Engine/Private/Texture.cpp
@@ -2370,7 +2370,7 @@ FSharedBuffer FTextureSource::TryDecompressData() const
 			
 			UE_LOG(LogTexture, Display, TEXT("TryDecompressData failed: LayerSize = %lld PayloadSize = %lld"),LayerSize,PayloadSize);
 			UE_LOG(LogTexture, Display, TEXT("TryDecompressData failed: LoadedMainStreamObjectVersion = %d, TextureSourceVirtualization = %d, VolumetricCloudReflectionSampleCountDefaultUpdate = %d"),
-				Owner->LoadedMainStreamObjectVersion,
+				Owner ? Owner->LoadedMainStreamObjectVersion : 0,
 				FUE5MainStreamObjectVersion::TextureSourceVirtualization,
 				FUE5MainStreamObjectVersion::VolumetricCloudReflectionSampleCountDefaultUpdate);
 
diff --git a/Engine/Source/Runtime/TextureUtilitiesCommon/Private/NormalMapIdentification.cpp b/Engine/Source/Runtime/TextureUtilitiesCommon/Private/NormalMapIdentification.cpp
--- a/Engine/Source/Runtime/TextureUtilitiesCommon/Private/NormalMapIdentification.cpp
+++ b/Engine/Source/Runtime/TextureUtilitiesCommon/Private/NormalMapIdentification.cpp
@@ -68,18 +68,19 @@ public:
 
 	~NormalMapSamplerBase()
 	{
-		if ( SourceTexture != NULL )
+		if ( SourceTextureData != NULL )
 		{
 			SourceTexture->Source.UnlockMip(0);
 		}
 	}
 
-	void SetSourceTexture( UTexture* Texture )
+	bool SetSourceTexture( UTexture* Texture )
 	{
 		SourceTexture = Texture;
 		TextureSizeX = Texture->Source.GetSizeX();
 		TextureSizeY = Texture->Source.GetSizeY();
 		SourceTextureData = Texture->Source.LockMipReadOnly(0);
+		return SourceTextureData != nullptr;
 	}
 
 	UTexture* SourceTexture;
@@ -285,7 +286,10 @@ public:
 		int32 NumTilesX = FMath::Min( TextureSizeX / SampleTileEdgeLength, MaxTilesPerAxis );
 		int32 NumTilesY = FMath::Min( TextureSizeY / SampleTileEdgeLength, MaxTilesPerAxis );
 
-		Sampler.SetSourceTexture( Texture );
+		if ( ! Sampler.SetSourceTexture( Texture ) )
+		{
+			return false;
+		}
 
 		if (( NumTilesX > 0 ) &&
 			( NumTilesY > 0 ))
