Add EXR support for screenshots

Add various functions to use FLinearColor in addition to FColor
Add support for EXR on platforms supporting HDR, but not openexr. This is driven by WITH_UNREALEXR_MINIMAL: the screenshot will not be compressed, only writes are supported, only RGBA float32 is supported as input (FLinearColor), and output will be RGB float16.
The convention used is ScRGB, not BT2100-PQ to store the HDR data, since the viewers I tested with were expecting ScRGB .

tested (both png and exr):
Editor d3d11/d3d12: F9 / highresshot 2 / High Resolution Screenshot Tool w/ and w/o Custom Depth Mask, w/ and w/o Buffer attributes / Trace.Screenshot
Game: shot / shot -showui / highresshot 2 / BugScreenShot / BugScreenShotWithHUDInfo (but HandleBugScreenShot is empty) / Trace.Screenshot / Trace.Screenshot w/ ui

What doesn't work : ScreenshotCapturedDelegate.Broadcast in HDR

diff --git a/Engine/Source/Editor/UnrealEd/Private/EditorViewportClient.cpp b/Engine/Source/Editor/UnrealEd/Private/EditorViewportClient.cpp
--- a/Engine/Source/Editor/UnrealEd/Private/EditorViewportClient.cpp
+++ b/Engine/Source/Editor/UnrealEd/Private/EditorViewportClient.cpp
@@ -65,6 +65,7 @@
 #include "GenericPlatform/GenericPlatformInputDeviceMapper.h"
 #include "CustomEditorStaticScreenPercentage.h"
 #include "IImageWrapperModule.h"
+#include "HDRHelper.h"
 
 ICustomEditorStaticScreenPercentage* GCustomEditorStaticScreenPercentage = nullptr;
 
@@ -5663,8 +5664,11 @@ void FEditorViewportClient::TakeScreenshot(FViewport* InViewport, bool bInValida
 
 	TUniquePtr<FImageWriteTask> ImageTask = MakeUnique<FImageWriteTask>();
 
+	// Read the contents of the viewport into an array.
+	bool bHdrEnabled = InViewport->GetSceneHDREnabled();
+	
+	if (!bHdrEnabled)
 	{
-		// Read the contents of the viewport into an array.
 		TArray<FColor> RawPixels;
 		RawPixels.SetNum(InViewport->GetSizeXY().X * InViewport->GetSizeXY().Y);
 		if( !InViewport->ReadPixels(RawPixels) )
@@ -5678,10 +5682,31 @@ void FEditorViewportClient::TakeScreenshot(FViewport* InViewport, bool bInValida
 
 		check(PixelData->IsDataWellFormed());
 		ImageTask->PixelData = MoveTemp(PixelData);
+
+		// Ensure the alpha channel is full alpha (this happens on the background thread)
+		ImageTask->PixelPreProcessors.Add(TAsyncAlphaWrite<FColor>(255));
 	}
+	else
+	{
+		TArray<FLinearColor> RawPixels;
+		RawPixels.SetNum(InViewport->GetSizeXY().X * InViewport->GetSizeXY().Y);
+		if (!InViewport->ReadLinearColorPixels(RawPixels))
+		{
+			// Failed to read the image from the viewport
+			SaveMessagePtr->SetText(NSLOCTEXT("UnrealEd", "ScreenshotFailedViewport", "Screenshot failed, unable to read image from viewport"));
+			return;
+	}
+
+		ConvertPixelDataToSCRGB(RawPixels, InViewport->GetDisplayOutputFormat());
+
+		TUniquePtr<TImagePixelData<FLinearColor>> PixelData = MakeUnique<TImagePixelData<FLinearColor>>(InViewport->GetSizeXY(), TArray64<FLinearColor>(MoveTemp(RawPixels)));
+
+		check(PixelData->IsDataWellFormed());
+		ImageTask->PixelData = MoveTemp(PixelData);
 
 	// Ensure the alpha channel is full alpha (this happens on the background thread)
-	ImageTask->PixelPreProcessors.Add(TAsyncAlphaWrite<FColor>(255));
+		ImageTask->PixelPreProcessors.Add(TAsyncAlphaWrite<FLinearColor>(1.0f));
+	}
 
 	// Create screenshot folder if not already present.
 	const FString& Directory = GetDefault<ULevelEditorMiscSettings>()->EditorScreenshotSaveDirectory.Path;
@@ -5694,13 +5719,15 @@ void FEditorViewportClient::TakeScreenshot(FViewport* InViewport, bool bInValida
 	}
 
 	// Save the contents of the array to a bitmap file.
-	HighResScreenshotConfig.SetHDRCapture(false);
+	HighResScreenshotConfig.SetHDRCapture(bHdrEnabled);
 
 	// Set the image task parameters
-	ImageTask->Format = EImageFormat::PNG;
+	ImageTask->Format = bHdrEnabled ? EImageFormat::EXR : EImageFormat::PNG;
 	ImageTask->CompressionQuality = (int32)EImageCompressionQuality::Default;
 
-	bool bGeneratedFilename = FFileHelper::GenerateNextBitmapFilename(Directory / TEXT("ScreenShot"), TEXT("png"), ImageTask->Filename);
+	const TCHAR* FileExtension = bHdrEnabled ? TEXT("exr") : TEXT("png");
+
+	bool bGeneratedFilename = FFileHelper::GenerateNextBitmapFilename(Directory / TEXT("ScreenShot"), FileExtension, ImageTask->Filename);
 	if (!bGeneratedFilename)
 	{
 		SaveMessagePtr->SetText(NSLOCTEXT( "UnrealEd", "ScreenshotFailed_TooManyScreenshots", "Screenshot failed, too many screenshots in output directory" ));
@@ -5783,6 +5810,109 @@ void FEditorViewportClient::TakeHighResScreenShot()
 	}
 }
 
+template<class FColorType>
+void ClipBitmapDataScreenshotDataEditor(bool& bWriteAlpha, TArray<FColorType>& Bitmap, FIntPoint& BitmapSize, const FIntRect& CaptureRect, bool bCaptureAreaValid)
+{
+	FHighResScreenshotConfig& HighResScreenshotConfig = GetHighResScreenshotConfig();
+
+	// Determine which region of the captured data we want to save out. If the highres screenshot capture region
+	// is not valid, we want to save out everything in the viewrect that we just grabbed.
+	FIntRect SourceRect = FIntRect(0, 0, 0, 0);
+	if (GIsHighResScreenshot && bCaptureAreaValid)
+	{
+		// Highres screenshot capture region is valid, so use that
+		SourceRect = HighResScreenshotConfig.CaptureRegion;
+	}
+
+	bWriteAlpha = false;
+
+	// If this is a high resolution screenshot and we are using the masking feature,
+	// Get the results of the mask rendering pass and insert into the alpha channel of the screenshot.
+	if (GIsHighResScreenshot && HighResScreenshotConfig.bMaskEnabled)
+	{
+		bWriteAlpha = HighResScreenshotConfig.MergeMaskIntoAlpha(Bitmap, CaptureRect);
+	}
+
+	// Clip the bitmap to just the capture region if valid
+	if (!SourceRect.IsEmpty())
+	{
+		const int32 OldWidth = BitmapSize.X;
+		const int32 OldHeight = BitmapSize.Y;
+
+		//clamp in bounds:
+		int CaptureMinX = FMath::Clamp(SourceRect.Min.X, 0, OldWidth);
+		int CaptureMinY = FMath::Clamp(SourceRect.Min.Y, 0, OldHeight);
+
+		int CaptureMaxX = FMath::Clamp(SourceRect.Max.X, 0, OldWidth);
+		int CaptureMaxY = FMath::Clamp(SourceRect.Max.Y, 0, OldHeight);
+
+		int32 NewWidth = CaptureMaxX - CaptureMinX;
+		int32 NewHeight = CaptureMaxY - CaptureMinY;
+
+		if (NewWidth > 0 && NewHeight > 0 && NewWidth != OldWidth && NewHeight != OldHeight)
+		{
+			FColorType* const Data = Bitmap.GetData();
+
+			for (int32 Row = 0; Row < NewHeight; Row++)
+			{
+				FMemory::Memmove(Data + Row * NewWidth, Data + (Row + CaptureMinY) * OldWidth + CaptureMinX, NewWidth * sizeof(*Data));
+			}
+
+			Bitmap.RemoveAt(NewWidth * NewHeight, OldWidth * OldHeight - NewWidth * NewHeight, false);
+			BitmapSize = FIntPoint(NewWidth, NewHeight);
+		}
+	}
+}
+
+template<class FColorType, typename TChannelType>
+bool RequestSaveScreenshot(bool bWriteAlpha, TArray<FColorType>& Bitmap, FIntPoint& BitmapSize, TChannelType OpaqueAlphaValue)
+{
+	FHighResScreenshotConfig& HighResScreenshotConfig = GetHighResScreenshotConfig();
+	bool bIsScreenshotSaved = false;
+	bool bSuppressWritingToFile = false;
+	if (SHOULD_TRACE_SCREENSHOT())
+	{
+		bSuppressWritingToFile = FTraceScreenshot::ShouldSuppressWritingToFile();
+		FTraceScreenshot::TraceScreenshot(BitmapSize.X, BitmapSize.Y, Bitmap, FScreenshotRequest::GetFilename());
+	}
+
+	if (!bSuppressWritingToFile)
+	{
+		TUniquePtr<FImageWriteTask> ImageTask = MakeUnique<FImageWriteTask>();
+		ImageTask->PixelData = MakeUnique<TImagePixelData<FColorType>>(BitmapSize, TArray64<FColorType>(MoveTemp(Bitmap)));
+
+		// Set full alpha on the bitmap
+		if (!bWriteAlpha)
+		{
+			ImageTask->PixelPreProcessors.Add(TAsyncAlphaWrite<FColorType>(OpaqueAlphaValue));
+		}
+
+		HighResScreenshotConfig.PopulateImageTaskParams(*ImageTask);
+		ImageTask->Filename = FScreenshotRequest::GetFilename();
+
+		{
+			// if not high dynamic range, get format from filename :
+			IImageWrapperModule& ImageWrapperModule = FModuleManager::LoadModuleChecked<IImageWrapperModule>(FName("ImageWrapper"));
+			EImageFormat ImageFormat = ImageWrapperModule.GetImageFormatFromExtension(*ImageTask->Filename);
+			if (ImageFormat != EImageFormat::Invalid)
+			{
+				ImageTask->Format = ImageFormat;
+			}
+		}
+
+		// Save the bitmap to disk
+		TFuture<bool> CompletionFuture = HighResScreenshotConfig.ImageWriteQueue->Enqueue(MoveTemp(ImageTask));
+		if (CompletionFuture.IsValid())
+		{
+			// this queues it then immediately waits? what's the point of ImageWriteQueue then?
+			// just use FImageUtils::Save
+			bIsScreenshotSaved = CompletionFuture.Get();
+		}
+	}
+
+	return bIsScreenshotSaved;
+}
+
 bool FEditorViewportClient::ProcessScreenShots(FViewport* InViewport)
 {
 	bool bIsScreenshotSaved = false;
@@ -5824,117 +5954,42 @@ bool FEditorViewportClient::ProcessScreenShots(FViewport* InViewport)
 			CaptureRect = View->UnscaledViewRect;
 		}
 
-		TArray<FColor> Bitmap;
-		if (GetViewportScreenShot(InViewport, Bitmap, CaptureRect))
-		{
-			// Determine the size of the captured viewport data.
-			FIntPoint BitmapSize = CaptureRect.Area() > 0 ? CaptureRect.Size() : InViewport->GetSizeXY();
-
-			// Determine which region of the captured data we want to save out. If the highres screenshot capture region
-			// is not valid, we want to save out everything in the viewrect that we just grabbed.
-			FIntRect SourceRect = FIntRect(0, 0, 0, 0);
-			if (GIsHighResScreenshot && bCaptureAreaValid)
-			{
-				// Highres screenshot capture region is valid, so use that
-				SourceRect = HighResScreenshotConfig.CaptureRegion;
-			}
-
-			bool bWriteAlpha = false;
+		bool bHdrEnabled = InViewport->GetSceneHDREnabled();
 
-			// If this is a high resolution screenshot and we are using the masking feature,
-			// Get the results of the mask rendering pass and insert into the alpha channel of the screenshot.
-			if (GIsHighResScreenshot && HighResScreenshotConfig.bMaskEnabled)
-			{
-				bWriteAlpha = HighResScreenshotConfig.MergeMaskIntoAlpha(Bitmap, CaptureRect);
-			}
-
-			// Clip the bitmap to just the capture region if valid
-			if (!SourceRect.IsEmpty())
+		// Determine the size of the captured viewport data.
+		FIntPoint BitmapSize = CaptureRect.Area() > 0 ? CaptureRect.Size() : InViewport->GetSizeXY();
+		if (!bHdrEnabled)
+		{
+			TArray<FColor> Bitmap;
+			if (GetViewportScreenShot(InViewport, Bitmap, CaptureRect))
 			{
-				const int32 OldWidth = BitmapSize.X;
-				const int32 OldHeight = BitmapSize.Y;
+				bool bWriteAlpha = false;
+				ClipBitmapDataScreenshotDataEditor(bWriteAlpha, Bitmap, BitmapSize, CaptureRect, bCaptureAreaValid);
 
-				//clamp in bounds:
-				int CaptureMinX = FMath::Clamp(SourceRect.Min.X,0,OldWidth);
-				int CaptureMinY = FMath::Clamp(SourceRect.Min.Y,0,OldHeight);
-				
-				int CaptureMaxX = FMath::Clamp(SourceRect.Max.X,0,OldWidth);
-				int CaptureMaxY = FMath::Clamp(SourceRect.Max.Y,0,OldHeight);
-				
-				int32 NewWidth  = CaptureMaxX - CaptureMinX;
-				int32 NewHeight = CaptureMaxY - CaptureMinY;
- 
-				if ( NewWidth > 0 && NewHeight > 0 )
+				if (FScreenshotRequest::OnScreenshotCaptured().IsBound())
 				{
-					FColor* const Data = Bitmap.GetData();
+					TArray<FColor> BitmapForBroadcast(Bitmap);
 
-					for (int32 Row = 0; Row < NewHeight; Row++)
+					if (!bWriteAlpha)
 					{
-						FMemory::Memmove(Data + Row * NewWidth, Data + (Row + CaptureMinY) * OldWidth + CaptureMinX, NewWidth * sizeof(*Data));
+						// Set full alpha on the bitmap
+						for (FColor& Pixel : BitmapForBroadcast) { Pixel.A = 255; }
 					}
 
-					Bitmap.RemoveAt(NewWidth * NewHeight, OldWidth * OldHeight - NewWidth * NewHeight, false);
-					BitmapSize = FIntPoint(NewWidth, NewHeight);
+					FScreenshotRequest::OnScreenshotCaptured().Broadcast(BitmapSize.X, BitmapSize.Y, MoveTemp(BitmapForBroadcast));
 				}
-			}
-			
-			if (FScreenshotRequest::OnScreenshotCaptured().IsBound())
-			{
-				TArray<FColor> BitmapForBroadcast(Bitmap);
 
-				if (!bWriteAlpha)
-				{
-					// Set full alpha on the bitmap
-					for (FColor& Pixel : BitmapForBroadcast) { Pixel.A = 255; }
-				}
-
-				FScreenshotRequest::OnScreenshotCaptured().Broadcast(BitmapSize.X, BitmapSize.Y, MoveTemp(BitmapForBroadcast));
+				bIsScreenshotSaved = RequestSaveScreenshot(bWriteAlpha, Bitmap, BitmapSize, 255);
 			}
-
-			bool bSuppressWritingToFile = false;
-			if (SHOULD_TRACE_SCREENSHOT())
-			{
-				bSuppressWritingToFile = FTraceScreenshot::ShouldSuppressWritingToFile();
-				FTraceScreenshot::TraceScreenshot(BitmapSize.X, BitmapSize.Y, Bitmap, FScreenshotRequest::GetFilename());
-			}
-
-			if (!bSuppressWritingToFile)
+		}
+		else
+		{
+			TArray<FLinearColor> BitmapHDR;
+			if (GetViewportScreenShotHDR(InViewport, BitmapHDR, CaptureRect))
 			{
-				TUniquePtr<FImageWriteTask> ImageTask = MakeUnique<FImageWriteTask>();
-				ImageTask->PixelData = MakeUnique<TImagePixelData<FColor>>(BitmapSize, TArray64<FColor>(MoveTemp(Bitmap)));
-
-				// Set full alpha on the bitmap
-				if (!bWriteAlpha)
-				{
-					ImageTask->PixelPreProcessors.Add(TAsyncAlphaWrite<FColor>(255));
-				}
-
-				HighResScreenshotConfig.PopulateImageTaskParams(*ImageTask);
-				ImageTask->Filename = FScreenshotRequest::GetFilename();
-
-				// Filename can be PNG but EImageFormat can be EXR
-
-				// PopulateImageTaskParams sets Format, ignores extension on file name
-				// ignore EXR from PopulateImageTaskParams, we are always 8-bit here
-				//if ( ImageTask->Format != EImageFormat::EXR )
-				{
-					// if not high dynamic range, get format from filename :
-					IImageWrapperModule& ImageWrapperModule = FModuleManager::LoadModuleChecked<IImageWrapperModule>(FName("ImageWrapper"));
-					EImageFormat ImageFormat = ImageWrapperModule.GetImageFormatFromExtension(*ImageTask->Filename);
-					if (ImageFormat != EImageFormat::Invalid)
-					{
-						ImageTask->Format = ImageFormat;
-					}
-				}
-
-				// Save the bitmap to disk
-				TFuture<bool> CompletionFuture = HighResScreenshotConfig.ImageWriteQueue->Enqueue(MoveTemp(ImageTask));
-				if (CompletionFuture.IsValid())
-				{
-					// this queues it then immediately waits? what's the point of ImageWriteQueue then?
-					// just use FImageUtils::Save
-					bIsScreenshotSaved = CompletionFuture.Get();
-				}
+				bool bWriteAlpha = false;
+				ClipBitmapDataScreenshotDataEditor(bWriteAlpha, BitmapHDR, BitmapSize, CaptureRect, bCaptureAreaValid);
+				bIsScreenshotSaved = RequestSaveScreenshot(bWriteAlpha, BitmapHDR, BitmapSize, 1.0f);
 			}
 		}
 
diff --git a/Engine/Source/Runtime/Engine/Private/GameViewportClient.cpp b/Engine/Source/Runtime/Engine/Private/GameViewportClient.cpp
--- a/Engine/Source/Runtime/Engine/Private/GameViewportClient.cpp
+++ b/Engine/Source/Runtime/Engine/Private/GameViewportClient.cpp
@@ -69,6 +69,7 @@
 #include "CustomStaticScreenPercentage.h"
 #include "ObjectTrace.h"
 #include "DynamicResolutionState.h"
+#include "HDRHelper.h"
 
 #if WITH_EDITOR
 #include "Settings/LevelEditorPlaySettings.h"
@@ -1900,6 +1901,98 @@ void UGameViewportClient::Draw(FViewport* InViewport, FCanvas* SceneCanvas)
 	EndDrawDelegate.Broadcast();
 }
 
+template<class FColorType, typename TChannelType>
+bool ProcessScreenshotData(TArray<FColorType>& Bitmap, FIntVector Size, TChannelType OpaqueAlphaValue, bool bHdrEnabled, bool bIsUI, const TCHAR* ToExtension)
+{
+	bool bIsScreenshotSaved = false;
+	{
+		FString ScreenShotName = FScreenshotRequest::GetFilename();
+		if (GIsDumpingMovie && ScreenShotName.IsEmpty())
+		{
+			// Request a new screenshot with a formatted name
+			const bool bShowUI = false;
+			const bool bAddFilenameSuffix = true;
+			FScreenshotRequest::RequestScreenshot(FString(), bShowUI, bAddFilenameSuffix, bHdrEnabled);
+			ScreenShotName = FScreenshotRequest::GetFilename();
+		}
+
+		// If a screenshot is requested during PIE (via F9), it does a screenshot of the entire editor window, including UI.
+		// We need to ignore the high resolution screenshot alpha mask in that case, as the mask isn't relevant when taking a
+		// screenshot of the entire window (and it will trigger an assert).  We don't want to solve this by modifying the
+		// global variables associated with the screenshot feature, as the application may be taking its own screenshots.
+		if (GIsHighResScreenshot && !bIsUI)
+		{
+			GetHighResScreenshotConfig().MergeMaskIntoAlpha(Bitmap, FIntRect(0, 0, 0, 0));
+		}
+		else
+		{
+			// Ensure that all pixels' alpha is set to opaque for a regular screenshot regardless of UI settings
+			// (Regular screenshots return with 0 in the alpha channels)
+			for (auto& Color : Bitmap)
+			{
+				Color.A = OpaqueAlphaValue;
+			}
+		}
+
+		FIntRect SourceRect(0, 0, GScreenshotResolutionX, GScreenshotResolutionY);
+		if (GIsHighResScreenshot)
+		{
+			SourceRect = GetHighResScreenshotConfig().CaptureRegion;
+		}
+
+		// Clip the bitmap to just the capture region if valid
+		if (!SourceRect.IsEmpty())
+		{
+			const int32 OldWidth = Size.X;
+			const int32 OldHeight = Size.Y;
+
+			//clamp in bounds:
+			int CaptureMinX = FMath::Clamp(SourceRect.Min.X, 0, OldWidth);
+			int CaptureMinY = FMath::Clamp(SourceRect.Min.Y, 0, OldHeight);
+
+			int CaptureMaxX = FMath::Clamp(SourceRect.Max.X, 0, OldWidth);
+			int CaptureMaxY = FMath::Clamp(SourceRect.Max.Y, 0, OldHeight);
+
+			int32 NewWidth = CaptureMaxX - CaptureMinX;
+			int32 NewHeight = CaptureMaxY - CaptureMinY;
+
+			if (NewWidth > 0 && NewHeight > 0 && ((NewWidth != OldWidth) || (NewHeight != OldHeight)))
+			{
+				FColorType* const Data = Bitmap.GetData();
+
+				for (int32 Row = 0; Row < NewHeight; Row++)
+				{
+					FMemory::Memmove(Data + Row * NewWidth, Data + (Row + CaptureMinY) * OldWidth + CaptureMinX, NewWidth * sizeof(*Data));
+				}
+
+				Bitmap.RemoveAt(NewWidth * NewHeight, OldWidth * OldHeight - NewWidth * NewHeight, false);
+				Size = FIntVector(NewWidth, NewHeight, 0);
+			}
+		}
+
+		if (FPaths::GetExtension(ScreenShotName).IsEmpty())
+		{
+			ScreenShotName += ToExtension;
+		}
+
+		bool bSuppressWritingToFile = false;
+		if (SHOULD_TRACE_SCREENSHOT())
+		{
+			bSuppressWritingToFile = FTraceScreenshot::ShouldSuppressWritingToFile();
+			FTraceScreenshot::TraceScreenshot(Size.X, Size.Y, Bitmap, ScreenShotName);
+		}
+
+		// Save the contents of the array to a png file.
+		if (!bSuppressWritingToFile)
+		{
+			FImageView Image((const FColorType*)Bitmap.GetData(), Size.X, Size.Y);
+			bIsScreenshotSaved = FImageUtils::SaveImageByExtension(*ScreenShotName, Image);
+		}
+	}
+
+	return bIsScreenshotSaved;
+}
+
 bool UGameViewportClient::ProcessScreenShots(FViewport* InViewport)
 {
 	bool bIsScreenshotSaved = false;
@@ -1907,6 +2000,7 @@ bool UGameViewportClient::ProcessScreenShots(FViewport* InViewport)
 	if (GIsDumpingMovie || FScreenshotRequest::IsScreenshotRequested() || GIsHighResScreenshot)
 	{
 		TArray<FColor> Bitmap;
+		TArray<FLinearColor> BitmapHDR;
 
 		bool bShowUI = false;
 		TSharedPtr<SWindow> WindowPtr = GetWindow();
@@ -1918,14 +2012,31 @@ bool UGameViewportClient::ProcessScreenShots(FViewport* InViewport)
 		bool bScreenshotSuccessful = false;
 		bool bIsUI = false;
 		FIntVector Size(InViewport->GetSizeXY().X, InViewport->GetSizeXY().Y, 0);
+
+		EDisplayOutputFormat ViewportOutputFormat = InViewport->GetDisplayOutputFormat();
+		bool bHdrEnabled = InViewport->GetSceneHDREnabled();
+
 		if( bShowUI && FSlateApplication::IsInitialized() )
 		{
 			TSharedRef<SWidget> WindowRef = WindowPtr.ToSharedRef();
-			bScreenshotSuccessful = FSlateApplication::Get().TakeScreenshot( WindowRef, Bitmap, Size);
+			if (bHdrEnabled)
+			{
+				bScreenshotSuccessful = FSlateApplication::Get().TakeHDRScreenshot(WindowRef, BitmapHDR, Size);
+				ConvertPixelDataToSCRGB(BitmapHDR, ViewportOutputFormat);
+			}
+			else
+			{
+				bScreenshotSuccessful = FSlateApplication::Get().TakeScreenshot(WindowRef, Bitmap, Size);
+			}
 			GScreenshotResolutionX = Size.X;
 			GScreenshotResolutionY = Size.Y;
 			bIsUI = true;
 		}
+		else if (bHdrEnabled)
+		{
+			bScreenshotSuccessful = GetViewportScreenShotHDR(InViewport, BitmapHDR);
+			ConvertPixelDataToSCRGB(BitmapHDR, ViewportOutputFormat);
+		}
 		else
 		{
 			bScreenshotSuccessful = GetViewportScreenShot(InViewport, Bitmap);
@@ -1933,102 +2044,28 @@ bool UGameViewportClient::ProcessScreenShots(FViewport* InViewport)
 
 		if (bScreenshotSuccessful)
 		{
-			if (ScreenshotCapturedDelegate.IsBound() && CVarScreenshotDelegate.GetValueOnGameThread())
-			{
-				// Ensure that all pixels' alpha is set to 255
-				for (auto& Color : Bitmap)
-				{
-					Color.A = 255;
-				}
-
-				// If delegate subscribed, fire it instead of writing out a file to disk
-				ScreenshotCapturedDelegate.Broadcast(Size.X, Size.Y, Bitmap);
-			}
-			else
+			if (Bitmap.Num() > 0)
 			{
-				FString ScreenShotName = FScreenshotRequest::GetFilename();
-				if (GIsDumpingMovie && ScreenShotName.IsEmpty())
+				if (UGameViewportClient::ScreenshotCapturedDelegate.IsBound() && CVarScreenshotDelegate.GetValueOnGameThread())
 				{
-					// Request a new screenshot with a formatted name
-					bShowUI = false;
-					const bool bAddFilenameSuffix = true;
-					FScreenshotRequest::RequestScreenshot(FString(), bShowUI, bAddFilenameSuffix);
-					ScreenShotName = FScreenshotRequest::GetFilename();
-				}
-
-				// If a screenshot is requested during PIE (via F9), it does a screenshot of the entire editor window, including UI.
-				// We need to ignore the high resolution screenshot alpha mask in that case, as the mask isn't relevant when taking a
-				// screenshot of the entire window (and it will trigger an assert).  We don't want to solve this by modifying the
-				// global variables associated with the screenshot feature, as the application may be taking its own screenshots.
-				if (GIsHighResScreenshot && !bIsUI)
-				{
-					GetHighResScreenshotConfig().MergeMaskIntoAlpha(Bitmap, FIntRect(0,0,0,0));
-				}
-				else
-				{
-					// Ensure that all pixels' alpha is set to 255 for a regular screenshot regardless of UI settings
-					// (Regular screenshots return with 0 in the alpha channels)
+					// Ensure that all pixels' alpha is set to 255
 					for (auto& Color : Bitmap)
 					{
 						Color.A = 255;
 					}
-				}
 
-				FIntRect SourceRect(0, 0, GScreenshotResolutionX, GScreenshotResolutionY);
-				if (GIsHighResScreenshot)
-				{
-					SourceRect = GetHighResScreenshotConfig().CaptureRegion;
+					// If delegate subscribed, fire it instead of writing out a file to disk
+					UGameViewportClient::ScreenshotCapturedDelegate.Broadcast(Size.X, Size.Y, Bitmap);
 				}
-
-				// Clip the bitmap to just the capture region if valid
-				if (!SourceRect.IsEmpty())
-				{
-					const int32 OldWidth = Size.X;
-					const int32 OldHeight = Size.Y;
-					
-					//clamp in bounds:
-					int CaptureMinX = FMath::Clamp(SourceRect.Min.X,0,OldWidth);
-					int CaptureMinY = FMath::Clamp(SourceRect.Min.Y,0,OldHeight);
-				
-					int CaptureMaxX = FMath::Clamp(SourceRect.Max.X,0,OldWidth);
-					int CaptureMaxY = FMath::Clamp(SourceRect.Max.Y,0,OldHeight);
-				
-					int32 NewWidth  = CaptureMaxX - CaptureMinX;
-					int32 NewHeight = CaptureMaxY - CaptureMinY;
- 
-					if ( NewWidth > 0 && NewHeight > 0 )
-					{
-						FColor* const Data = Bitmap.GetData();
-
-						for (int32 Row = 0; Row < NewHeight; Row++)
-						{
-							FMemory::Memmove(Data + Row * NewWidth, Data + (Row + CaptureMinY) * OldWidth + CaptureMinX, NewWidth * sizeof(*Data));
-						}
-
-						Bitmap.RemoveAt(NewWidth * NewHeight, OldWidth * OldHeight - NewWidth * NewHeight, false);
-						Size = FIntVector(NewWidth, NewHeight, 0);
-					}
-				}
-
-				if ( FPaths::GetExtension(ScreenShotName).IsEmpty() )
-				{
-					ScreenShotName += TEXT(".png");
-				}
-
-				bool bSuppressWritingToFile = false;
-				if (SHOULD_TRACE_SCREENSHOT())
-				{
-					bSuppressWritingToFile = FTraceScreenshot::ShouldSuppressWritingToFile();
-					FTraceScreenshot::TraceScreenshot(Size.X, Size.Y, Bitmap, ScreenShotName);
-				}
-
-				// Save the contents of the array to a png file.
-				if (!bSuppressWritingToFile)
+				else
 				{
-					FImageView Image((const FColor*)Bitmap.GetData(), Size.X, Size.Y);
-					bIsScreenshotSaved = FImageUtils::SaveImageByExtension(*ScreenShotName, Image);
+					bIsScreenshotSaved = ProcessScreenshotData(Bitmap, Size, 255, bHdrEnabled, bIsUI, TEXT(".png"));
 				}
 			}
+			else
+			{
+				bIsScreenshotSaved = ProcessScreenshotData(BitmapHDR, Size, 1.0f, bHdrEnabled, bIsUI, TEXT(".exr"));
+			}
 		}
 
 		FScreenshotRequest::Reset();
@@ -3716,7 +3753,7 @@ bool UGameViewportClient::HandleScreenshotCommand( const TCHAR* Cmd, FOutputDevi
 			bAddFilenameSuffix = false;
 		}
 
-		FScreenshotRequest::RequestScreenshot(FileName, bShowUI, bAddFilenameSuffix );
+		FScreenshotRequest::RequestScreenshot(FileName, bShowUI, bAddFilenameSuffix, Viewport->GetSceneHDREnabled());
 
 		GScreenshotResolutionX = Viewport->GetSizeXY().X;
 		GScreenshotResolutionY = Viewport->GetSizeXY().Y;
@@ -4018,9 +4055,11 @@ bool UGameViewportClient::RequestBugScreenShot(const TCHAR* Cmd, bool bDisplayHU
 	FString FileName = Cmd;
 
 	// Handle just a plain console command (e.g. "BUGSCREENSHOT").
+	bool bHDREnabled = (Viewport != NULL && Viewport->GetSceneHDREnabled());
+	const TCHAR* ScreenshotExtension = bHDREnabled ? TEXT("exr") : TEXT("png");
 	if (FileName.Len() == 0)
 	{
-		FileName = TEXT("BugScreenShot.png");
+		FileName = bHDREnabled ? TEXT("BugScreenShot.exr") : TEXT("BugScreenShot.png");
 	}
 
 	// Handle a console command and name (e.g. BUGSCREENSHOT FOO)
@@ -4030,7 +4069,7 @@ bool UGameViewportClient::RequestBugScreenShot(const TCHAR* Cmd, bool bDisplayHU
 		const FString BaseFile = FString::Printf(TEXT("%s%s_"), *FPaths::BugItDir(), *FPaths::GetBaseFilename(FileName));
 
 		// find the next filename in the sequence, e.g <gamename>/bugit/<platform>/desc_00000.png
-		FFileHelper::GenerateNextBitmapFilename(BaseFile, TEXT("png"), FileName);
+		FFileHelper::GenerateNextBitmapFilename(BaseFile, ScreenshotExtension, FileName);
 	}
 
 	if (Viewport != NULL)
@@ -4050,7 +4089,7 @@ bool UGameViewportClient::RequestBugScreenShot(const TCHAR* Cmd, bool bDisplayHU
 
 		const bool bShowUI = true;
 		const bool bAddFilenameSuffix = false;
-		FScreenshotRequest::RequestScreenshot(FileName, true, bAddFilenameSuffix);
+		FScreenshotRequest::RequestScreenshot(FileName, true, bAddFilenameSuffix, bHDREnabled);
 	}
 
 	return true;
diff --git a/Engine/Source/Runtime/Engine/Private/HighResScreenshot.cpp b/Engine/Source/Runtime/Engine/Private/HighResScreenshot.cpp
--- a/Engine/Source/Runtime/Engine/Private/HighResScreenshot.cpp
+++ b/Engine/Source/Runtime/Engine/Private/HighResScreenshot.cpp
@@ -143,7 +143,8 @@ bool FHighResScreenshotConfig::ParseConsoleCommand(const FString& InCmd, FOutput
 	return false;
 }
 
-bool FHighResScreenshotConfig::MergeMaskIntoAlpha(TArray<FColor>& InBitmap, const FIntRect& ViewRect)
+template<class FColorType, typename TChannelType>
+bool MergeMaskIntoAlphaInternal(TArray<FColorType>& InBitmap, const FIntRect& ViewRect, bool bMaskEnabled, TChannelType AlphaMultipler)
 {
 	bool bWritten = false;
 
@@ -161,7 +162,7 @@ bool FHighResScreenshotConfig::MergeMaskIntoAlpha(TArray<FColor>& InBitmap, cons
 			// Exact match, copy verbatim
 			for (int32 i = 0; i < InBitmap.Num(); ++i)
 			{
-				InBitmap[i].A = (*MaskArray)[i].R;
+				InBitmap[i].A = TChannelType((*MaskArray)[i].R) * AlphaMultipler;
 			}
 		}
 		else
@@ -176,7 +177,7 @@ bool FHighResScreenshotConfig::MergeMaskIntoAlpha(TArray<FColor>& InBitmap, cons
 			{
 				for (int32 i = ViewRect.Min.X; i < ViewRect.Max.X; i++, OutputOffset++)
 				{
-					InBitmap[OutputOffset].A = (*MaskArray)[j * MaskStride + i].R;
+					InBitmap[OutputOffset].A = TChannelType((*MaskArray)[j * MaskStride + i].R) * AlphaMultipler;
 				}
 			}
 		}
@@ -187,13 +188,23 @@ bool FHighResScreenshotConfig::MergeMaskIntoAlpha(TArray<FColor>& InBitmap, cons
 		// Ensure that all pixels' alpha is set to 255
 		for (auto& Color : InBitmap)
 		{
-			Color.A = 255;
+			Color.A = TChannelType(255) * AlphaMultipler;
 		}
 	}
 
 	return bWritten;
 }
 
+bool FHighResScreenshotConfig::MergeMaskIntoAlpha(TArray<FColor>& InBitmap, const FIntRect& ViewRect)
+{
+	return MergeMaskIntoAlphaInternal<FColor, uint8>(InBitmap, ViewRect, bMaskEnabled, 1);
+}
+
+bool FHighResScreenshotConfig::MergeMaskIntoAlpha(TArray<FLinearColor>& InBitmap, const FIntRect& ViewRect)
+{
+	return MergeMaskIntoAlphaInternal<FLinearColor, float>(InBitmap, ViewRect, bMaskEnabled, 1.0f / 255.0f);
+}
+
 void FHighResScreenshotConfig::SetHDRCapture(bool bCaptureHDRIN)
 {
 	bCaptureHDR = bCaptureHDRIN;
diff --git a/Engine/Source/Runtime/Engine/Private/ProfilingDebugging/TraceScreenshot.cpp b/Engine/Source/Runtime/Engine/Private/ProfilingDebugging/TraceScreenshot.cpp
--- a/Engine/Source/Runtime/Engine/Private/ProfilingDebugging/TraceScreenshot.cpp
+++ b/Engine/Source/Runtime/Engine/Private/ProfilingDebugging/TraceScreenshot.cpp
@@ -59,7 +59,31 @@ void FTraceScreenshot::RequestScreenshot(FString Name, bool bShowUI, const FLogC
 	FScreenshotRequest::RequestScreenshot(Name, bShowUI, bAddUniqueSuffix);
 }
 
-void FTraceScreenshot::TraceScreenshot(int32 InSizeX, int32 InSizeY, const TArray<FColor>& InImageData, const FString& InScreenshotName, int32 DesiredX)
+void ImageUtilsImageResize(int32 SrcWidth, int32 SrcHeight, const TArray<FColor>& SrcData, int32 DstWidth, int32 DstHeight, TArray<FColor>& DstData, bool bLinearSpace)
+{
+	FImageUtils::ImageResize(SrcWidth, SrcHeight, SrcData, DstWidth, DstHeight, DstData, bLinearSpace);
+}
+
+void ImageUtilsCompressImageArrayWrapper(int32 ImageWidth, int32 ImageHeight, const TArrayView64<const FColor>& SrcData, TArray64<uint8>& DstData)
+{
+	FImageUtils::PNGCompressImageArray(ImageWidth, ImageHeight, SrcData, DstData);
+}
+
+void ImageUtilsImageResize(int32 SrcWidth, int32 SrcHeight, const TArray64<FLinearColor>& SrcData, int32 DstWidth, int32 DstHeight, TArray64<FLinearColor>& DstData, bool /*bUnused*/)
+{
+	FImageUtils::ImageResize(SrcWidth, SrcHeight, SrcData, DstWidth, DstHeight, DstData);
+}
+
+void ImageUtilsCompressImageArrayWrapper(int32 ImageWidth, int32 ImageHeight, const TArrayView64<const FLinearColor>& SrcData, TArray64<uint8>& DstData)
+{
+	FImageView TmpImageView(SrcData.GetData(), ImageWidth, ImageHeight);
+	FImageUtils::CompressImage(DstData, TEXT(".exr"), TmpImageView);
+}
+
+
+template <class FColorType, class FResizedColorArrayType, typename TChannelType>
+void TraceScreenshotInternal(int32 InSizeX, int32 InSizeY, const TArray<FColorType>& InImageData, const FString& InScreenshotName, int32 DesiredX, 
+							 TChannelType OpaqueAlphaValue)
 {
 	TRACE_CPUPROFILER_EVENT_SCOPE(ScreenshotTracing_Prepare);
 	FString ScreenshotName = FPaths::GetBaseFilename(InScreenshotName);
@@ -70,19 +94,19 @@ void FTraceScreenshot::TraceScreenshot(int32 InSizeX, int32 InSizeY, const TArra
 	}
 
 	UE_LOG(LogCore, Display, TEXT("Tracing Screenshot \"%s\" taken with size: %d x %d"), *ScreenshotName, InSizeX, InSizeY);
-	
+
 	uint64 Cycles = FPlatformTime::Cycles64();
-	TArray<FColor> *ImageCopy = new TArray<FColor>(InImageData);
-	
-	UE::Tasks::Launch(UE_SOURCE_LOCATION, 
-		[ImageCopy, Cycles, ScreenshotName, InSizeX, InSizeY, DesiredX]
+	FResizedColorArrayType* ImageCopy = new FResizedColorArrayType(InImageData);
+
+	UE::Tasks::Launch(UE_SOURCE_LOCATION,
+		[ImageCopy, Cycles, ScreenshotName, InSizeX, InSizeY, DesiredX, OpaqueAlphaValue]
 		{
 			TRACE_CPUPROFILER_EVENT_SCOPE(ScreenshotTracing_Execute);
 
 			// Set full alpha on the bitmap
-			for (FColor& Pixel : *ImageCopy)
+			for (FColorType& Pixel : *ImageCopy)
 			{
-				Pixel.A = 255;
+				Pixel.A = OpaqueAlphaValue;
 			}
 
 			TArray64<uint8> CompressedBitmap;
@@ -91,22 +115,33 @@ void FTraceScreenshot::TraceScreenshot(int32 InSizeX, int32 InSizeY, const TArra
 				int32 ResizedX = FMath::Min(640, InSizeX);
 				int32 ResizedY = (InSizeY * ResizedX) / InSizeX;
 
-				TArray<FColor> ResizedImage;
+				FResizedColorArrayType ResizedImage;
 				ResizedImage.SetNum(ResizedX * ResizedY);
-				FImageUtils::ImageResize(InSizeX, InSizeY, *ImageCopy, ResizedX, ResizedY, ResizedImage, false);
+				ImageUtilsImageResize(InSizeX, InSizeY, *ImageCopy, ResizedX, ResizedY, ResizedImage, false);
 
-				FImageUtils::PNGCompressImageArray(ResizedX, ResizedY, TArrayView64<const FColor>(ResizedImage.GetData(), ResizedImage.Num()), CompressedBitmap);
+				ImageUtilsCompressImageArrayWrapper(ResizedX, ResizedY, TArrayView64<const FColorType>(ResizedImage.GetData(), ResizedImage.Num()), CompressedBitmap);
 				TRACE_SCREENSHOT(*ScreenshotName, Cycles, ResizedX, ResizedY, CompressedBitmap);
 			}
 			else
 			{
-				FImageUtils::PNGCompressImageArray(InSizeX, InSizeY, TArrayView64<const FColor>(ImageCopy->GetData(), ImageCopy->Num()), CompressedBitmap);
+				ImageUtilsCompressImageArrayWrapper(InSizeX, InSizeY, TArrayView64<const FColorType>(ImageCopy->GetData(), ImageCopy->Num()), CompressedBitmap);
 				TRACE_SCREENSHOT(*ScreenshotName, Cycles, InSizeX, InSizeY, CompressedBitmap);
 			}
 
 			delete ImageCopy;
 		});
 
+}
+
+void FTraceScreenshot::TraceScreenshot(int32 InSizeX, int32 InSizeY, const TArray<FColor>& InImageData, const FString& InScreenshotName, int32 DesiredX)
+{
+	TraceScreenshotInternal<FColor, TArray<FColor>, uint8>(InSizeX, InSizeY, InImageData, InScreenshotName, DesiredX, 255);
+	Reset();
+}
+
+void FTraceScreenshot::TraceScreenshot(int32 InSizeX, int32 InSizeY, const TArray<FLinearColor>& InImageData, const FString& InScreenshotName, int32 DesiredX)
+{
+	TraceScreenshotInternal<FLinearColor, TArray64<FLinearColor>, float>(InSizeX, InSizeY, InImageData, InScreenshotName, DesiredX, 1.0f);
 	Reset();
 }
 
diff --git a/Engine/Source/Runtime/Engine/Private/UnrealClient.cpp b/Engine/Source/Runtime/Engine/Private/UnrealClient.cpp
--- a/Engine/Source/Runtime/Engine/Private/UnrealClient.cpp
+++ b/Engine/Source/Runtime/Engine/Private/UnrealClient.cpp
@@ -285,25 +285,27 @@ FUnorderedAccessViewRHIRef FRenderTarget::GetRenderTargetUAV() const
 void FScreenshotRequest::RequestScreenshot(bool bInShowUI)
 {
 	// empty string means we'll later pick the name
-	RequestScreenshot(TEXT(""), bInShowUI, true);
+	RequestScreenshot(TEXT(""), bInShowUI, /*bAddUniqueSuffix*/ true, /*bHdrScreenshot*/ false);
 }
 
-void FScreenshotRequest::RequestScreenshot(const FString& InFilename, bool bInShowUI, bool bAddUniqueSuffix)
+void FScreenshotRequest::RequestScreenshot(const FString& InFilename, bool bInShowUI, bool bAddUniqueSuffix, bool bHdrScreenshot)
 {
 	FString GeneratedFilename = InFilename;
 	CreateViewportScreenShotFilename(GeneratedFilename);
 
+	const TCHAR* ScreenshotExtension = bHdrScreenshot ? TEXT("exr") : TEXT("png");
+
 	if (bAddUniqueSuffix)
 	{
 		const bool bRemovePath = false;
 		GeneratedFilename = FPaths::GetBaseFilename(GeneratedFilename, bRemovePath);
 		if (GetHighResScreenshotConfig().bDateTimeBasedNaming)
 		{
-			FFileHelper::GenerateDateTimeBasedBitmapFilename(GeneratedFilename, TEXT("png"), Filename);
+			FFileHelper::GenerateDateTimeBasedBitmapFilename(GeneratedFilename, ScreenshotExtension, Filename);
 		}
 		else
 		{
-			FFileHelper::GenerateNextBitmapFilename(GeneratedFilename, TEXT("png"), Filename);
+			FFileHelper::GenerateNextBitmapFilename(GeneratedFilename, ScreenshotExtension, Filename);
 		}
 	}
 	else
@@ -311,7 +313,8 @@ void FScreenshotRequest::RequestScreenshot(const FString& InFilename, bool bInSh
 		Filename = GeneratedFilename;
 		if (FPaths::GetExtension(Filename).Len() == 0)
 		{
-			Filename += TEXT(".png");
+			Filename += TEXT(".");
+			Filename += ScreenshotExtension;
 		}
 	}
 
@@ -1329,6 +1332,7 @@ void FViewport::HighResScreenshot()
 	const FString CachedScreenshotName = FScreenshotRequest::GetFilename();
 
 	FDummyViewport* DummyViewport = new FDummyViewport(ViewportClient);
+	DummyViewport->SetupHDR(GetDisplayColorGamut(), GetDisplayOutputFormat(), GetSceneHDREnabled());
 
 	DummyViewport->SizeX = (GScreenshotResolutionX > 0) ? GScreenshotResolutionX : SizeX;
 	DummyViewport->SizeY = (GScreenshotResolutionY > 0) ? GScreenshotResolutionY : SizeY;
@@ -1593,7 +1597,8 @@ void FViewport::Draw( bool bShouldPresent /*= true */)
 			{
 				const bool bShowUI = false;
 				const bool bAddFilenameSuffix = GetHighResScreenshotConfig().FilenameOverride.IsEmpty();
-				FScreenshotRequest::RequestScreenshot( FString(), bShowUI, bAddFilenameSuffix );
+				const bool bHDRScreenshot = GetSceneHDREnabled();
+				FScreenshotRequest::RequestScreenshot( FString(), bShowUI, bAddFilenameSuffix, bHDRScreenshot);
 				HighResScreenshot();
 			}
 			else if(bAnyScreenshotsRequired && bBufferVisualizationDumpingRequired)
@@ -2259,6 +2264,18 @@ ENGINE_API bool GetViewportScreenShot(FViewport* Viewport, TArray<FColor>& Bitma
 	return false;
 }
 
+ENGINE_API bool GetViewportScreenShotHDR(FViewport* Viewport, TArray<FLinearColor>& Bitmap, const FIntRect& ViewRect /*= FIntRect()*/)
+{
+	// Read the contents of the viewport into an array.
+	if (Viewport->ReadLinearColorPixels(Bitmap, FReadSurfaceDataFlags(RCM_MinMax), ViewRect))
+	{
+		check(Bitmap.Num() == ViewRect.Area() || (Bitmap.Num() == Viewport->GetSizeXY().X * Viewport->GetSizeXY().Y));
+		return true;
+	}
+
+	return false;
+}
+
 extern bool ParseResolution( const TCHAR* InResolution, uint32& OutX, uint32& OutY, int32& WindowMode );
 
 ENGINE_API bool GetHighResScreenShotInput(const TCHAR* Cmd, FOutputDevice& Ar, uint32& OutXRes, uint32& OutYRes, float& OutResMult, FIntRect& OutCaptureRegion, bool& OutShouldEnableMask, bool& OutDumpBufferVisualizationTargets, bool& OutCaptureHDR, FString& OutFilenameOverride, bool& OutUseDateTimeAsFileName)
@@ -2421,12 +2438,28 @@ FDummyViewport::~FDummyViewport()
 
 void FDummyViewport::InitDynamicRHI()
 {
+	EPixelFormat DummyViewportFormat = bSceneHDREnabled ? GRHIHDRDisplayOutputFormat : PF_A2B10G10R10;
 	const FRHITextureCreateDesc Desc =
 		FRHITextureCreateDesc::Create2D(TEXT("FDummyViewport"))
 		.SetExtent(SizeX, SizeY)
-		.SetFormat(PF_A2B10G10R10)
+		.SetFormat(DummyViewportFormat)
 		.SetFlags(ETextureCreateFlags::RenderTargetable | ETextureCreateFlags::ShaderResource)
 		.SetInitialState(ERHIAccess::SRVMask);
 
 	RenderTargetTextureRHI = RHICreateTexture(Desc);
 }
+
+EDisplayColorGamut FDummyViewport::GetDisplayColorGamut() const
+{
+	return DisplayColorGamut;
+}
+
+EDisplayOutputFormat FDummyViewport::GetDisplayOutputFormat() const
+{
+	return DisplayOutputFormat;
+}
+
+bool FDummyViewport::GetSceneHDREnabled() const
+{
+	return bSceneHDREnabled;
+}
diff --git a/Engine/Source/Runtime/Engine/Public/HighResScreenshot.h b/Engine/Source/Runtime/Engine/Public/HighResScreenshot.h
--- a/Engine/Source/Runtime/Engine/Public/HighResScreenshot.h
+++ b/Engine/Source/Runtime/Engine/Public/HighResScreenshot.h
@@ -55,6 +55,7 @@ struct ENGINE_API FHighResScreenshotConfig
 	  * Returns true if the mask was written, and false otherwise.
 	**/
 	bool MergeMaskIntoAlpha(TArray<FColor>& InBitmap, const FIntRect& ViewRect);
+	bool MergeMaskIntoAlpha(TArray<FLinearColor>& InBitmap, const FIntRect& ViewRect);
 
 	/** Enable/disable HDR capable captures **/
 	void SetHDRCapture(bool bCaptureHDRIN);
diff --git a/Engine/Source/Runtime/Engine/Public/ProfilingDebugging/TraceScreenshot.h b/Engine/Source/Runtime/Engine/Public/ProfilingDebugging/TraceScreenshot.h
--- a/Engine/Source/Runtime/Engine/Public/ProfilingDebugging/TraceScreenshot.h
+++ b/Engine/Source/Runtime/Engine/Public/ProfilingDebugging/TraceScreenshot.h
@@ -30,6 +30,7 @@ public:
 	* @param DesiredX - Optionally resize the image to the desired width before tracing. Aspect ratio is preserved.
 	*/
 	static void TraceScreenshot(int32 InSizeX, int32 InSizeY, const TArray<FColor>& InImageData, const FString& InScreenshotName, int32 DesiredX = -1);
+	static void TraceScreenshot(int32 InSizeX, int32 InSizeY, const TArray<FLinearColor>& InImageData, const FString& InScreenshotName, int32 DesiredX = -1);
 
 	/**
 	* Returns true if the screenshot should not be writted to a file.
diff --git a/Engine/Source/Runtime/Engine/Public/UnrealClient.h b/Engine/Source/Runtime/Engine/Public/UnrealClient.h
--- a/Engine/Source/Runtime/Engine/Public/UnrealClient.h
+++ b/Engine/Source/Runtime/Engine/Public/UnrealClient.h
@@ -166,7 +166,7 @@ struct ENGINE_API FScreenshotRequest
 	 * @param bInShowUI				Whether or not to show Slate UI
 	 * @param bAddFilenameSuffix	Whether an auto-generated unique suffix should be added to the supplied filename
 	 */
-	static void RequestScreenshot(const FString& InFilename, bool bInShowUI, bool bAddFilenameSuffix);
+	static void RequestScreenshot(const FString& InFilename, bool bInShowUI, bool bAddFilenameSuffix, bool bHdrScreenshot=false);
 
 	/**
 	 * Resets a screenshot request
@@ -787,6 +787,7 @@ extern ENGINE_API bool IsShiftDown(FViewport* Viewport);
 extern ENGINE_API bool IsAltDown(FViewport* Viewport);
 
 extern ENGINE_API bool GetViewportScreenShot(FViewport* Viewport, TArray<FColor>& Bitmap, const FIntRect& ViewRect = FIntRect());
+extern ENGINE_API bool GetViewportScreenShotHDR(FViewport* Viewport, TArray<FLinearColor>& Bitmap, const FIntRect& ViewRect = FIntRect());
 extern ENGINE_API bool GetHighResScreenShotInput(const TCHAR* Cmd, FOutputDevice& Ar, uint32& OutXRes, uint32& OutYRes, float& OutResMult, FIntRect& OutCaptureRegion, bool& OutShouldEnableMask, bool& OutDumpBufferVisualizationTargets, bool& OutCaptureHDR, FString& OutFilenameOverride, bool& OutUseDateTimeAsFileName);
 
 struct FInputKeyEventArgs
@@ -1250,6 +1251,13 @@ public:
 		check( IsInRenderingThread() );
 	}
 
+	void SetupHDR(EDisplayColorGamut InDisplayColorGamut, EDisplayOutputFormat InDisplayOutputFormat, bool bInSceneHDREnabled)
+	{
+		DisplayColorGamut = InDisplayColorGamut;
+		DisplayOutputFormat = InDisplayOutputFormat;
+		bSceneHDREnabled = bInSceneHDREnabled;
+	}
+
 	virtual void*	GetWindow() override { return 0; }
 	virtual void	MoveWindow(int32 NewPosX, int32 NewPosY, int32 NewSizeX, int32 NewSizeY) override {}
 	virtual void	Destroy() override {}
@@ -1266,6 +1274,9 @@ public:
 	virtual void DeferInvalidateHitProxy() override { }
 	virtual FViewportFrame* GetViewportFrame() override { return 0; }
 	virtual FCanvas* GetDebugCanvas() override { return DebugCanvas; }
+	virtual EDisplayColorGamut GetDisplayColorGamut() const;
+	virtual EDisplayOutputFormat GetDisplayOutputFormat() const;
+	virtual bool GetSceneHDREnabled() const;
 	//~ End FViewport Interface
 
 	//~ Begin FRenderResource Interface
@@ -1280,4 +1291,7 @@ public:
 	//~ End FRenderResource Interface
 private:
 	FCanvas* DebugCanvas;
+	EDisplayColorGamut DisplayColorGamut = EDisplayColorGamut::sRGB_D65;
+	EDisplayOutputFormat DisplayOutputFormat = EDisplayOutputFormat::SDR_sRGB;
+	bool bSceneHDREnabled = false;
 };
diff --git a/Engine/Source/Runtime/ImageWrapper/ImageWrapper.Build.cs b/Engine/Source/Runtime/ImageWrapper/ImageWrapper.Build.cs
--- a/Engine/Source/Runtime/ImageWrapper/ImageWrapper.Build.cs
+++ b/Engine/Source/Runtime/ImageWrapper/ImageWrapper.Build.cs
@@ -43,6 +43,8 @@ public class ImageWrapper : ModuleRules
 			AddEngineThirdPartyPrivateStaticDependencies(Target, "UElibJPG");
 		}
 
+		bool bEnableMinimalEXR = false;
+
 		// Add openEXR lib for windows builds.
 		if ((Target.Platform.IsInGroup(UnrealPlatformGroup.Windows)) ||
 			(Target.Platform == UnrealTargetPlatform.Mac) ||
@@ -55,8 +57,16 @@ public class ImageWrapper : ModuleRules
 		else
 		{
 			PublicDefinitions.Add("WITH_UNREALEXR=0");
+
+			if (Target.Platform.IsInGroup(UnrealPlatformGroup.Sony)  ||
+				Target.Platform.IsInGroup(UnrealPlatformGroup.XboxCommon)) 
+			{
+				bEnableMinimalEXR = true;
+			}
 		}
 
+		PublicDefinitions.Add("WITH_UNREALEXR_MINIMAL=" + (bEnableMinimalEXR ? "1" : "0"));
+
 		// Enable exceptions to allow error handling
 		bEnableExceptions = true;
 	}
diff --git a/Engine/Source/Runtime/ImageWrapper/Private/Formats/ExrImageWrapper.cpp b/Engine/Source/Runtime/ImageWrapper/Private/Formats/ExrImageWrapper.cpp
--- a/Engine/Source/Runtime/ImageWrapper/Private/Formats/ExrImageWrapper.cpp
+++ b/Engine/Source/Runtime/ImageWrapper/Private/Formats/ExrImageWrapper.cpp
@@ -523,7 +523,191 @@ void FExrImageWrapper::Uncompress(const ERGBFormat InFormat, const int32 InBitDe
 	BitDepth = InBitDepth;
 }
 
+#elif WITH_UNREALEXR_MINIMAL
 
+FExrImageWrapper::FExrImageWrapper()
+	: FImageWrapperBase()
+{
+}
+
+bool FExrImageWrapper::SetRaw(const void* InRawData, int64 InRawSize, const int32 InWidth, const int32 InHeight, const ERGBFormat InFormat, const int32 InBitDepth, const int32 InBytesPerRow)
+{
+	check(InRawData);
+	check(InRawSize > 0);
+	check(InWidth > 0);
+	check(InHeight > 0);
+	check(InBytesPerRow >= 0);
+	if (InBitDepth != 32)
+	{
+		return false;
+	}
+
+	if (InFormat != ERGBFormat::RGBAF)
+	{
+		return false;
+	}
+
+	return FImageWrapperBase::SetRaw(InRawData, InRawSize, InWidth, InHeight, InFormat, InBitDepth, InBytesPerRow);
+}
+
+bool FExrImageWrapper::SetCompressed(const void* InCompressedData, int64 InCompressedSize)
+{
+	return false;
+}
+
+void FExrImageWrapper::Compress(int32 Quality)
+{
+	check(RawData.Num());
+
+	// Ensure we haven't already compressed the file.
+	if (CompressedData.Num())
+	{
+		return;
+	}
+
+	const double StartTime = FPlatformTime::Seconds();
+
+	check(BitDepth == 32);
+	check(Format == ERGBFormat::RGBAF);
+
+	constexpr uint32 EXRHeaderSize = 313;
 
+	// Indicates the offset, from the start of the file, to get access to the 'full data' of the row
+	const uint32 LineOffsetTableSize = 8 * Height; 
+	const uint64 TotalHeaderSize = EXRHeaderSize + LineOffsetTableSize;
+
+	// We force use half float as output
+	const uint32 PerChannelStride = Width * sizeof(FFloat16);
+
+	// We store RGB, not A
+	const uint32 PixelDataRowSize = PerChannelStride * 3; 
+
+	// Full Row is Row Y coordinate + size of pixel data + pixel data
+	const uint32 FullDataRowSize = 2 * sizeof(uint32) + PixelDataRowSize; 
+
+	// Final data allocation
+	uint32 CompressedDataReservedSize = TotalHeaderSize + Height * FullDataRowSize;
+	CompressedData.Reserve(CompressedDataReservedSize);
+
+	auto AddDataU32 = [&](uint32 Value) { CompressedData.Append((uint8*)&Value, sizeof(uint32)); };
+	// Based on example header at https://www.openexr.com/documentation/openexrfilelayout.pdf
+	{
+		auto AddHeaderString = [&](const char* Value) { int32 Len = FCStringAnsi::Strlen(Value); CompressedData.Append((uint8*)Value, Len + 1); };
+
+		// magic number,			  version, flags
+		AddDataU32(0x01312f76);	AddDataU32(0x00000002);
+
+		// Attribute name / type / size: 18 bytes per channel + 1 terminating byte
+		AddHeaderString("channels"); AddHeaderString("chlist"); AddDataU32(55);
+		//name                type (u8,f16,f32)		pLinear/reserved        xSampling          ySampling
+		AddHeaderString("B");  AddDataU32(1);     AddDataU32(0);       AddDataU32(1);  AddDataU32(1);
+		AddHeaderString("G");  AddDataU32(1);     AddDataU32(0);       AddDataU32(1);  AddDataU32(1);
+		AddHeaderString("R");  AddDataU32(1);     AddDataU32(0);       AddDataU32(1);  AddDataU32(1);
+		// Separator
+		CompressedData.Add(0);
+
+		// Attribute name / type / size: compression expects a single byte
+		AddHeaderString("compression"); AddHeaderString("compression"); AddDataU32(1);
+		// NO_COMPRESSION = 0, RLE = 1, ZIPS = 2, ZIP = 3, PIZ = 4, PXR24= 5, B44= 6, B44A= 7,
+		CompressedData.Add(0);
+
+		// Attribute name / type / size
+		AddHeaderString("dataWindow"); AddHeaderString("box2i"); AddDataU32(16);
+		// Top Left / Bottom Right
+		AddDataU32(0); AddDataU32(0); AddDataU32(Width - 1); AddDataU32(Height - 1);
+
+		// Attribute name / type / size
+		AddHeaderString("displayWindow"); AddHeaderString("box2i"); AddDataU32(16);
+		AddDataU32(0); AddDataU32(0); AddDataU32(Width - 1); AddDataU32(Height - 1);
+
+		// Attribute name / type / size
+		AddHeaderString("lineOrder"); AddHeaderString("lineOrder"); AddDataU32(1);
+		// INCREASING_Y = 0, DECREASING_Y = 1, RANDOM_Y = 2
+		CompressedData.Add(0);
+
+		// Attribute name / type / size
+		AddHeaderString("pixelAspectRatio"); AddHeaderString("float"); AddDataU32(4);
+		// 1.0f
+		AddDataU32(0x3f800000);
+
+		// Attribute name / type / size
+		AddHeaderString("screenWindowCenter"); AddHeaderString("v2f"); AddDataU32(8);
+		// 0.0f / 0.0f
+		AddDataU32(0); AddDataU32(0);
+
+		// Attribute name / type / size
+		AddHeaderString("screenWindowWidth"); AddHeaderString("float"); AddDataU32(4);
+		// 1.0f
+		AddDataU32(0x3f800000);
+
+		// end of header
+		CompressedData.Add(0);
+
+		// Sanity check to make sure header size pre-allocation is still up-to-date
+		check(CompressedData.Num() == EXRHeaderSize);
+	}
+
+	uint64 LineOffset = TotalHeaderSize;
+	// Line Offset table
+	for (int32 RowIndex = 0; RowIndex < Height; ++RowIndex)
+	{
+		CompressedData.Append((uint8*)&LineOffset, sizeof(uint64));
+		LineOffset += FullDataRowSize;
+	}
+	check(CompressedData.Num() == TotalHeaderSize);
+
+	// Raw Data output
+	const FLinearColor* SrcPixelData = (FLinearColor*)RawData.GetData();
+	for (int32 RowIndex = 0; RowIndex < Height; ++RowIndex, SrcPixelData += Width)
+	{
+		uint32 OutEXRFileDataRowBeginSize = CompressedData.Num();
+
+		AddDataU32(RowIndex);
+		AddDataU32(PixelDataRowSize);
+
+		// Layout is BBBB..BBGGGG..GGRRRR..RR
+		for (int32 ColIndex = 0; ColIndex < Width; ++ColIndex)
+		{
+			FFloat16 B16(SrcPixelData[ColIndex].B);
+			CompressedData.Append((uint8*)&B16.Encoded, sizeof(B16.Encoded));
+		}
+
+		for (int32 ColIndex = 0; ColIndex < Width; ++ColIndex)
+		{
+			FFloat16 G16(SrcPixelData[ColIndex].G);
+			CompressedData.Append((uint8*)&G16.Encoded, sizeof(G16.Encoded));
+		}
+
+		for (int32 ColIndex = 0; ColIndex < Width; ++ColIndex)
+		{
+			FFloat16 R16(SrcPixelData[ColIndex].R);
+			CompressedData.Append((uint8*)&R16.Encoded, sizeof(R16.Encoded));
+		}
+
+		check(CompressedData.Num() == OutEXRFileDataRowBeginSize + FullDataRowSize);
+	}
+
+	// make sure we have written all the data we reserved in the first place
+	check(CompressedData.Num() == CompressedDataReservedSize);
+
+	const double DeltaTime = FPlatformTime::Seconds() - StartTime;
+	UE_LOG(LogImageWrapper, Verbose, TEXT("written image in %.3f seconds"), DeltaTime);
+}
+
+void FExrImageWrapper::Uncompress(const ERGBFormat InFormat, const int32 InBitDepth)
+{
+	ensure(false);
+	UE_LOG(LogImageWrapper, Error, TEXT("FExrImageWrapper::Uncompress is not supported"));
+}
+
+bool FExrImageWrapper::CanSetRawFormat(const ERGBFormat InFormat, const int32 InBitDepth) const
+{
+	return (InFormat == ERGBFormat::RGBAF) && (InBitDepth == 32);
+}
+
+ERawImageFormat::Type FExrImageWrapper::GetSupportedRawFormat(const ERawImageFormat::Type InFormat) const
+{
+	return ERawImageFormat::RGBA32F;
+}
 
 #endif // WITH_UNREALEXR
diff --git a/Engine/Source/Runtime/ImageWrapper/Private/Formats/ExrImageWrapper.h b/Engine/Source/Runtime/ImageWrapper/Private/Formats/ExrImageWrapper.h
--- a/Engine/Source/Runtime/ImageWrapper/Private/Formats/ExrImageWrapper.h
+++ b/Engine/Source/Runtime/ImageWrapper/Private/Formats/ExrImageWrapper.h
@@ -6,8 +6,9 @@
 
 #include "ImageWrapperBase.h"
 
-#if WITH_UNREALEXR
+#if WITH_UNREALEXR || WITH_UNREALEXR_MINIMAL
 
+#if WITH_UNREALEXR
 THIRD_PARTY_INCLUDES_START
 	#include "Imath/ImathBox.h"
 	#include "OpenEXR/ImfArray.h"
@@ -20,7 +21,7 @@ THIRD_PARTY_INCLUDES_START
 	#include "OpenEXR/ImfStdIO.h"
 	#include "OpenEXR/ImfVersion.h"
 THIRD_PARTY_INCLUDES_END
-
+#endif
 
 /**
  * OpenEXR implementation of the helper class
diff --git a/Engine/Source/Runtime/ImageWrapper/Private/ImageWrapperModule.cpp b/Engine/Source/Runtime/ImageWrapper/Private/ImageWrapperModule.cpp
--- a/Engine/Source/Runtime/ImageWrapper/Private/ImageWrapperModule.cpp
+++ b/Engine/Source/Runtime/ImageWrapper/Private/ImageWrapperModule.cpp
@@ -107,7 +107,7 @@ public:
 			ImageWrapper = MakeShared<FIcoImageWrapper>();
 			break;
 
-#if WITH_UNREALEXR
+#if WITH_UNREALEXR || WITH_UNREALEXR_MINIMAL
 		case EImageFormat::EXR:
 			ImageWrapper = MakeShared<FExrImageWrapper>();
 			break;
diff --git a/Engine/Source/Runtime/RenderCore/Private/RenderCore.cpp b/Engine/Source/Runtime/RenderCore/Private/RenderCore.cpp
--- a/Engine/Source/Runtime/RenderCore/Private/RenderCore.cpp
+++ b/Engine/Source/Runtime/RenderCore/Private/RenderCore.cpp
@@ -954,4 +954,47 @@ void ConfigureACESTonemapParams(FACESTonemapParams& OutACESTonemapParams, float
 
 	OutACESTonemapParams.ACESSceneColorMultiplier = CVarHDRAcesColorMultiplier.GetValueOnAnyThread();
 
-}
\ No newline at end of file
+}
+
+// Converts PQ signal to linear values, see https://www.itu.int/rec/R-REC-BT.2124-0-201901-I/en, conversion 3 
+static inline float ST2084ToLinear(float pq)
+{
+	const float m1 = 0.1593017578125; // = 2610. / 4096. * .25;
+	const float m2 = 78.84375; // = 2523. / 4096. *  128;
+	const float c1 = 0.8359375; // = 2392. / 4096. * 32 - 2413./4096.*32 + 1;
+	const float c2 = 18.8515625; // = 2413. / 4096. * 32;
+	const float c3 = 18.6875; // = 2392. / 4096. * 32;
+	const float C = 10000.;
+
+	float Np = FMath::Pow(pq, 1. / m2);
+	float L = Np - c1;
+	L = FMath::Max(0.0f, L);
+	L = L / (c2 - c3 * Np);
+	L = FMath::Pow(L, 1. / m1);
+	float P = L * C;
+
+	return P;
+}
+
+
+RENDERCORE_API void ConvertPixelDataToSCRGB(TArray<FLinearColor>& InOutRawPixels, EDisplayOutputFormat Pixelformat)
+{
+	int32 PixelCount = InOutRawPixels.Num();
+	if (Pixelformat == EDisplayOutputFormat::HDR_ACES_1000nit_ST2084 || Pixelformat == EDisplayOutputFormat::HDR_ACES_2000nit_ST2084)
+	{
+		const FMatrix44f Rec2020_2_sRGB_MAT = (XYZToGamutMatrix(EDisplayColorGamut::sRGB_D65) * GamutToXYZMatrix(EDisplayColorGamut::Rec2020_D65)).GetTransposed();
+
+		for (int32 PixelIndex = 0; PixelIndex < PixelCount; ++PixelIndex)
+		{
+			FLinearColor& CurrentPixel = InOutRawPixels[PixelIndex];
+			FVector4f PixelFloat4(CurrentPixel.R, CurrentPixel.G, CurrentPixel.B, 0);
+			PixelFloat4.X = ST2084ToLinear(PixelFloat4.X) / 80;
+			PixelFloat4.Y = ST2084ToLinear(PixelFloat4.Y) / 80;
+			PixelFloat4.Z = ST2084ToLinear(PixelFloat4.Z) / 80;
+			PixelFloat4 = Rec2020_2_sRGB_MAT.TransformFVector4(PixelFloat4);
+			CurrentPixel.R = PixelFloat4.X;
+			CurrentPixel.G = PixelFloat4.Y;
+			CurrentPixel.B = PixelFloat4.Z;
+		}
+	}
+}
diff --git a/Engine/Source/Runtime/RenderCore/Public/HDRHelper.h b/Engine/Source/Runtime/RenderCore/Public/HDRHelper.h
--- a/Engine/Source/Runtime/RenderCore/Public/HDRHelper.h
+++ b/Engine/Source/Runtime/RenderCore/Public/HDRHelper.h
@@ -15,6 +15,7 @@ RENDERCORE_API void HDRAddCustomMetaData(void* OSWindow, EDisplayOutputFormat Di
 RENDERCORE_API void HDRRemoveCustomMetaData(void* OSWindow);
 RENDERCORE_API FMatrix44f GamutToXYZMatrix(EDisplayColorGamut ColorGamut);
 RENDERCORE_API FMatrix44f XYZToGamutMatrix(EDisplayColorGamut ColorGamut);
+RENDERCORE_API void ConvertPixelDataToSCRGB(TArray<FLinearColor>& InOutRawPixels, EDisplayOutputFormat Pixelformat);
 
 struct FDisplayInformation
 {
diff --git a/Engine/Source/Runtime/Slate/Private/Framework/Application/SlateApplication.cpp b/Engine/Source/Runtime/Slate/Private/Framework/Application/SlateApplication.cpp
--- a/Engine/Source/Runtime/Slate/Private/Framework/Application/SlateApplication.cpp
+++ b/Engine/Source/Runtime/Slate/Private/Framework/Application/SlateApplication.cpp
@@ -3942,17 +3942,14 @@ bool FSlateApplication::TakeScreenshot(const TSharedRef<SWidget>& Widget, TArray
 	return TakeScreenshot(Widget, FIntRect(), OutColorData, OutSize);
 }
 
-bool FSlateApplication::TakeScreenshot(const TSharedRef<SWidget>& Widget, const FIntRect& InnerWidgetArea, TArray<FColor>& OutColorData, FIntVector& OutSize)
+bool FSlateApplication::TakeHDRScreenshot(const TSharedRef<SWidget>& Widget, TArray<FLinearColor>& OutColorData, FIntVector& OutSize)
 {
-	// We can't screenshot the widget unless there's a valid window handle to draw it in.
-	TSharedPtr<SWindow> WidgetWindow = FSlateApplication::Get().FindWidgetWindow(Widget);
-	if ( !WidgetWindow.IsValid() )
-	{
-		return false;
-	}
+	return TakeHDRScreenshot(Widget, FIntRect(), OutColorData, OutSize);
+}
 
-	TSharedRef<SWindow> CurrentWindowRef = WidgetWindow.ToSharedRef();
 
+void TakeScreenshotCommon(const TSharedRef<SWidget>& Widget, const FIntRect& InnerWidgetArea, FIntRect& ScreenshotRect, SWindow* WidgetWindow)
+{
 	FWidgetPath WidgetPath;
 	FSlateApplication::Get().GeneratePathToWidgetChecked(Widget, WidgetPath);
 
@@ -3961,12 +3958,28 @@ bool FSlateApplication::TakeScreenshot(const TSharedRef<SWidget>& Widget, const
 	FVector2D Size = ArrangedWidget.Geometry.GetDrawSize();
 	FVector2D WindowPosition = WidgetWindow->GetPositionInScreen();
 
-	FIntRect ScreenshotRect = InnerWidgetArea.IsEmpty() ? FIntRect(0, 0, (int32)Size.X, (int32)Size.Y) : InnerWidgetArea;
+	ScreenshotRect = InnerWidgetArea.IsEmpty() ? FIntRect(0, 0, (int32)Size.X, (int32)Size.Y) : InnerWidgetArea;
 
 	ScreenshotRect.Min.X += ( Position.X - WindowPosition.X );
 	ScreenshotRect.Min.Y += ( Position.Y - WindowPosition.Y );
 	ScreenshotRect.Max.X += ( Position.X - WindowPosition.X );
 	ScreenshotRect.Max.Y += ( Position.Y - WindowPosition.Y );
+}
+
+
+bool FSlateApplication::TakeScreenshot(const TSharedRef<SWidget>& Widget, const FIntRect& InnerWidgetArea, TArray<FColor>& OutColorData, FIntVector& OutSize)
+{
+	// We can't screenshot the widget unless there's a valid window handle to draw it in.
+	TSharedPtr<SWindow> WidgetWindow = FSlateApplication::Get().FindWidgetWindow(Widget);
+	if (!WidgetWindow.IsValid())
+	{
+		return false;
+	}
+
+	TSharedRef<SWindow> CurrentWindowRef = WidgetWindow.ToSharedRef();
+
+	FIntRect ScreenshotRect;
+	TakeScreenshotCommon(Widget, InnerWidgetArea, ScreenshotRect, WidgetWindow.Get());
 
 	Renderer->PrepareToTakeScreenshot(ScreenshotRect, &OutColorData, WidgetWindow.Get());
 	PrivateDrawWindows(WidgetWindow);
@@ -3977,6 +3990,30 @@ bool FSlateApplication::TakeScreenshot(const TSharedRef<SWidget>& Widget, const
 	return (OutSize.X != 0 && OutSize.Y != 0 && OutColorData.Num() >= OutSize.X * OutSize.Y);
 }
 
+bool FSlateApplication::TakeHDRScreenshot(const TSharedRef<SWidget>& Widget, const FIntRect& InnerWidgetArea, TArray<FLinearColor>& OutColorData, FIntVector& OutSize)
+{
+	// We can't screenshot the widget unless there's a valid window handle to draw it in.
+	TSharedPtr<SWindow> WidgetWindow = FSlateApplication::Get().FindWidgetWindow(Widget);
+	if (!WidgetWindow.IsValid())
+	{
+		return false;
+	}
+
+	TSharedRef<SWindow> CurrentWindowRef = WidgetWindow.ToSharedRef();
+
+	FIntRect ScreenshotRect;
+	TakeScreenshotCommon(Widget, InnerWidgetArea, ScreenshotRect, WidgetWindow.Get());
+
+	Renderer->PrepareToTakeHDRScreenshot(ScreenshotRect, &OutColorData, WidgetWindow.Get());
+	PrivateDrawWindows(WidgetWindow);
+
+	OutSize.X = ScreenshotRect.Size().X;
+	OutSize.Y = ScreenshotRect.Size().Y;
+
+	return (OutSize.X != 0 && OutSize.Y != 0 && OutColorData.Num() >= OutSize.X * OutSize.Y);
+}
+
+
 TSharedPtr<FSlateUser> FSlateApplication::GetUser(FPlatformUserId PlatformUser)
 {
 	int32 InternalId = 0;
@@ -7347,4 +7384,4 @@ void VerifyWidgetLayerId(const TSharedRef<SWindow>& WindowToDraw)
 }
 
 } //namespace
-#endif //WITH_SLATE_DEBUGGING
+#endif //WITH_SLATE_DEBUGGING
\ No newline at end of file
diff --git a/Engine/Source/Runtime/Slate/Public/Framework/Application/SlateApplication.h b/Engine/Source/Runtime/Slate/Public/Framework/Application/SlateApplication.h
--- a/Engine/Source/Runtime/Slate/Public/Framework/Application/SlateApplication.h
+++ b/Engine/Source/Runtime/Slate/Public/Framework/Application/SlateApplication.h
@@ -928,7 +928,15 @@ public:
 	 * 
 	 * @return true if taking the screenshot was successful.
 	 */
-	bool TakeScreenshot(const TSharedRef<SWidget>& Widget, TArray<FColor>&OutColorData, FIntVector& OutSize);
+	bool TakeScreenshot(const TSharedRef<SWidget>& Widget, TArray<FColor>& OutColorData, FIntVector& OutSize);
+	
+	/**
+	 * Takes a screenshot of the widget writing the results into the color buffer provided.  This is to be used with HDR buffers
+	 * the size of the resulting image is also output.
+	 *
+	 * @return true if taking the screenshot was successful.
+	 */
+	bool TakeHDRScreenshot(const TSharedRef<SWidget>& Widget, TArray<FLinearColor>& OutColorData, FIntVector& OutSize);
 
 	/**
 	 * Takes a screenshot of the widget writing the results into the color buffer provided, this version allows you to provide 
@@ -938,6 +946,14 @@ public:
 	 */
 	bool TakeScreenshot(const TSharedRef<SWidget>& Widget, const FIntRect& InnerWidgetArea, TArray<FColor>& OutColorData, FIntVector& OutSize);
 
+	/**
+	  * Takes a screenshot of the widget writing the results into the color buffer provided, this version allows you to provide
+	  * an inner area to screenshot.  Note that the format is BGRA.  The size of the resulting image is also output.
+	  *
+	  * @return true if taking the screenshot was successful.
+	*/
+	bool TakeHDRScreenshot(const TSharedRef<SWidget>& Widget, const FIntRect& InnerWidgetArea, TArray<FLinearColor>& OutColorData, FIntVector& OutSize);
+
 	/** Gets the user at the given index, null if the user does not exist. */
 	FORCEINLINE TSharedPtr<const FSlateUser> GetUser(int32 UserIndex) const
 	{
diff --git a/Engine/Source/Runtime/SlateCore/Public/Rendering/SlateRenderer.h b/Engine/Source/Runtime/SlateCore/Public/Rendering/SlateRenderer.h
--- a/Engine/Source/Runtime/SlateCore/Public/Rendering/SlateRenderer.h
+++ b/Engine/Source/Runtime/SlateCore/Public/Rendering/SlateRenderer.h
@@ -484,6 +484,12 @@ public:
 	 */
 	virtual void PrepareToTakeScreenshot(const FIntRect& Rect, TArray<FColor>* OutColorData, SWindow* InScreenshotWindow) {}
 
+	/** 
+	 * Prepares the renderer to take a screenshot of the UI.  The Rect is portion of the rendered output
+	 * that will be stored into the TArray of FColors.
+	 */
+	virtual void PrepareToTakeHDRScreenshot(const FIntRect& Rect, TArray<FLinearColor>* OutColorData, SWindow* InScreenshotWindow) {}
+
 	/**
 	 * Pushes the rendering of the specified window to the specified render target
 	 */
diff --git a/Engine/Source/Runtime/SlateRHIRenderer/Private/SlateRHIRenderer.cpp b/Engine/Source/Runtime/SlateRHIRenderer/Private/SlateRHIRenderer.cpp
--- a/Engine/Source/Runtime/SlateRHIRenderer/Private/SlateRHIRenderer.cpp
+++ b/Engine/Source/Runtime/SlateRHIRenderer/Private/SlateRHIRenderer.cpp
@@ -213,6 +213,7 @@ FSlateRHIRenderer::FSlateRHIRenderer(TSharedRef<FSlateFontServices> InSlateFontS
 
 	bTakingAScreenShot = false;
 	OutScreenshotData = NULL;
+	OutHDRScreenshotData = NULL;
 	ScreenshotViewportInfo = nullptr;
 	bIsStandaloneStereoOnlyDevice = IHeadMountedDisplayModule::IsAvailable() && IHeadMountedDisplayModule::Get().IsStandaloneStereoOnlyDevice();
 }
@@ -1362,7 +1363,14 @@ void FSlateRHIRenderer::DrawWindow_RenderThread(FRHICommandListImmediate& RHICmd
 
 		if (!ClampedScreenshotRect.IsEmpty())
 		{
-			RHICmdList.ReadSurfaceData(BackBuffer, ClampedScreenshotRect, *OutScreenshotData, FReadSurfaceDataFlags());
+			if (OutHDRScreenshotData != nullptr)
+			{
+				RHICmdList.ReadSurfaceData(BackBuffer, ClampedScreenshotRect, *OutHDRScreenshotData, FReadSurfaceDataFlags());
+			}
+			else
+			{
+				RHICmdList.ReadSurfaceData(BackBuffer, ClampedScreenshotRect, *OutScreenshotData, FReadSurfaceDataFlags());
+			}
 		}
 		else
 		{
@@ -1370,6 +1378,7 @@ void FSlateRHIRenderer::DrawWindow_RenderThread(FRHICommandListImmediate& RHICmd
 		}
 		bTakingAScreenShot = false;
 		OutScreenshotData = nullptr;
+		OutHDRScreenshotData = nullptr;
 		ScreenshotViewportInfo = nullptr;
 	}
 
@@ -1469,6 +1478,18 @@ void FSlateRHIRenderer::PrepareToTakeScreenshot(const FIntRect& Rect, TArray<FCo
 	bTakingAScreenShot = true;
 	ScreenshotRect = Rect;
 	OutScreenshotData = OutColorData;
+	OutHDRScreenshotData = nullptr;
+	ScreenshotViewportInfo = *WindowToViewportInfo.Find(InScreenshotWindow);
+}
+
+void FSlateRHIRenderer::PrepareToTakeHDRScreenshot(const FIntRect& Rect, TArray<FLinearColor>* OutColorData, SWindow* InScreenshotWindow)
+{
+	check(OutColorData);
+
+	bTakingAScreenShot = true;
+	ScreenshotRect = Rect;
+	OutScreenshotData = nullptr;
+	OutHDRScreenshotData = OutColorData;
 	ScreenshotViewportInfo = *WindowToViewportInfo.Find(InScreenshotWindow);
 }
 
diff --git a/Engine/Source/Runtime/SlateRHIRenderer/Private/SlateRHIRenderer.h b/Engine/Source/Runtime/SlateRHIRenderer/Private/SlateRHIRenderer.h
--- a/Engine/Source/Runtime/SlateRHIRenderer/Private/SlateRHIRenderer.h
+++ b/Engine/Source/Runtime/SlateRHIRenderer/Private/SlateRHIRenderer.h
@@ -237,6 +237,7 @@ public:
 	virtual void InvalidateAllViewports() override;
 
 	virtual void PrepareToTakeScreenshot(const FIntRect& Rect, TArray<FColor>* OutColorData, SWindow* ScreenshotWindow) override;
+	virtual void PrepareToTakeHDRScreenshot(const FIntRect& Rect, TArray<FLinearColor>* OutColorData, SWindow* ScreenshotWindow) override;
 
 	virtual void SetWindowRenderTarget(const SWindow& Window, class IViewportRenderTargetProvider* Provider) override;
 
@@ -304,6 +305,7 @@ private:
 	FIntRect ScreenshotRect;
 	FViewportInfo* ScreenshotViewportInfo;
 	TArray<FColor>* OutScreenshotData;
+	TArray<FLinearColor>* OutHDRScreenshotData;
 
 	/** These are state management variables for Scenes on the game thread. A similar copy exists on the RHI Rendering Policy for the rendering thread.*/
 	TArray<FSceneInterface*, TInlineAllocator<4>> ActiveScenes;
