Fix for tracepoint break triggering when frame-stepping with an active animation blueprint preview
Provided by [REDACTED_PERSON]

diff --git a/Engine/Source/Editor/UnrealEd/Private/Kismet2/KismetDebugUtilities.cpp b/Engine/Source/Editor/UnrealEd/Private/Kismet2/KismetDebugUtilities.cpp
--- a/Engine/Source/Editor/UnrealEd/Private/Kismet2/KismetDebugUtilities.cpp
+++ b/Engine/Source/Editor/UnrealEd/Private/Kismet2/KismetDebugUtilities.cpp
@@ -261,7 +261,7 @@ void FKismetDebugUtilities::OnScriptException(const UObject* ActiveObject, const
 			bShouldBreakExecution = true;
 			break;
 		case EBlueprintExceptionType::Tracepoint:
-			bShouldBreakExecution = bIsStepping;
+			bShouldBreakExecution = bIsStepping && TracepointBreakAllowedOnOwningWorld(ActiveObject);
 			break;
 		case EBlueprintExceptionType::WireTracepoint:
 			break;
@@ -504,6 +504,21 @@ void FKismetDebugUtilities::OnScriptException(const UObject* ActiveObject, const
 	}
 }
 
+bool FKismetDebugUtilities::TracepointBreakAllowedOnOwningWorld(const UObject* ObjOuter)
+{
+	bool bAllowTracepointBreak = true;
+
+	const UWorld* ObjWorld = ObjOuter->GetWorld();
+
+	// Disable tracepoints on EditorPreviews or Inactive worlds
+	if (ObjWorld && (ObjWorld->WorldType == EWorldType::EditorPreview || ObjWorld->WorldType == EWorldType::Inactive))
+	{
+		bAllowTracepointBreak = false;
+	}
+
+	return bAllowTracepointBreak;
+}
+
 UClass* FKismetDebugUtilities::FindClassForNode(const UObject* Object, UFunction* Function)
 {
 	if (NULL != Function)
diff --git a/Engine/Source/Editor/UnrealEd/Public/Kismet2/KismetDebugUtilities.h b/Engine/Source/Editor/UnrealEd/Public/Kismet2/KismetDebugUtilities.h
--- a/Engine/Source/Editor/UnrealEd/Public/Kismet2/KismetDebugUtilities.h
+++ b/Engine/Source/Editor/UnrealEd/Public/Kismet2/KismetDebugUtilities.h
@@ -486,6 +486,10 @@ protected:
 	static void CleanupBreakpoints(const UBlueprint* Blueprint);
 	static void CleanupWatches(const UBlueprint* Blueprint);
 	static void RemoveEmptySettings(const FString& BlueprintPath);
+
+	/** Looks along the outer chain to find the owning world and disallowing breaking on tracepoints for editor preview and inactive worlds. */
+	static bool TracepointBreakAllowedOnOwningWorld(const UObject* ObjOuter);
+
 private:
 	FKismetDebugUtilities() {}
 };
