CookMultiprocess - replicate bFullBuild flag to CookWorkers so they can disable EDLCookInfo which is expected non-active in incremental cooks.
#rnx

diff --git a/Engine/Source/Editor/UnrealEd/Classes/CookOnTheSide/CookOnTheFlyServer.h b/Engine/Source/Editor/UnrealEd/Classes/CookOnTheSide/CookOnTheFlyServer.h
--- a/Engine/Source/Editor/UnrealEd/Classes/CookOnTheSide/CookOnTheFlyServer.h
+++ b/Engine/Source/Editor/UnrealEd/Classes/CookOnTheSide/CookOnTheFlyServer.h
@@ -874,6 +874,7 @@ private:
 	/** Update accumulators of editor data and consume editor changes after last cook when starting cooks in the editor. */
 	void BeginCookEditorSystems();
 	void BeginCookPackageWriters(FBeginCookContext& BeginContext);
+	void BeginCookDirector(FBeginCookContext& BeginContext);
 	void InitializeSession();
 	void PreGarbageCollectImpl(TArray<UPackage*>& GCKeepPackages, TArray<UE::Cook::FPackageData*>& GCKeepPackageDatas,
 		TArray<UObject*>& LocalGCKeepObjects);
diff --git a/Engine/Source/Editor/UnrealEd/Private/CookOnTheFlyServer.cpp b/Engine/Source/Editor/UnrealEd/Private/CookOnTheFlyServer.cpp
--- a/Engine/Source/Editor/UnrealEd/Private/CookOnTheFlyServer.cpp
+++ b/Engine/Source/Editor/UnrealEd/Private/CookOnTheFlyServer.cpp
@@ -9221,6 +9221,7 @@ void UCookOnTheFlyServer::StartCookByTheBook( const FCookByTheBookStartupOptions
 
 	// Initialize systems that nothing in StartCookByTheBook references
 	// Functions in this section are not dependent upon each other and can be ordered arbitrarily or for async performance
+	BeginCookDirector(BeginContext);
 	BeginCookEditorSystems();
 	BeginCookEDLCookInfo(BeginContext);
 	BeginCookPackageWriters(BeginContext);
@@ -9714,6 +9715,14 @@ void UCookOnTheFlyServer::BeginCookEditorSystems()
 	ModifiedAssetFilenames.Empty();
 }
 
+void UCookOnTheFlyServer::BeginCookDirector(FBeginCookContext& BeginContext)
+{
+	if (CookDirector)
+	{
+		CookDirector->StartCook(BeginContext);
+	}
+}
+
 namespace UE::Cook
 {
 
diff --git a/Engine/Source/Editor/UnrealEd/Private/Cooker/CookDirector.cpp b/Engine/Source/Editor/UnrealEd/Private/Cooker/CookDirector.cpp
--- a/Engine/Source/Editor/UnrealEd/Private/Cooker/CookDirector.cpp
+++ b/Engine/Source/Editor/UnrealEd/Private/Cooker/CookDirector.cpp
@@ -122,6 +122,11 @@ FCookDirector::~FCookDirector()
 	Sockets::CloseSocket(WorkerConnectSocket);
 }
 
+void FCookDirector::StartCook(const FBeginCookContext& InBeginContext)
+{
+	BeginCookContext.Set(InBeginContext);
+}
+
 void FCookDirector::AssignRequests(TArrayView<UE::Cook::FPackageData*> Requests, TArray<FWorkerId>& OutAssignments,
 	TMap<FPackageData*, TArray<FPackageData*>>&& RequestGraph)
 {
diff --git a/Engine/Source/Editor/UnrealEd/Private/Cooker/CookDirector.h b/Engine/Source/Editor/UnrealEd/Private/Cooker/CookDirector.h
--- a/Engine/Source/Editor/UnrealEd/Private/Cooker/CookDirector.h
+++ b/Engine/Source/Editor/UnrealEd/Private/Cooker/CookDirector.h
@@ -34,6 +34,8 @@ public:
 	FCookDirector(UCookOnTheFlyServer& InCOTFS);
 	~FCookDirector();
 
+	void StartCook(const FBeginCookContext& Context);
+
 	/**
 	 * Assign the given requests out to CookWorkers (or keep on local COTFS), return the list of assignments.
 	 * Input requests have been sorted by leaf to root load order.
@@ -114,6 +116,7 @@ private:
 		Striped,
 		CookBurden,
 	};
+	FBeginCookContextForWorker BeginCookContext;
 	TMap<int32, TUniquePtr<FCookWorkerServer>> RemoteWorkers;
 	TMap<FCookWorkerServer*, TUniquePtr<FCookWorkerServer>> ShuttingDownWorkers;
 	TMap<FGuid, TRefCountPtr<IMPCollector>> MessageHandlers;
diff --git a/Engine/Source/Editor/UnrealEd/Private/Cooker/CookTypes.cpp b/Engine/Source/Editor/UnrealEd/Private/Cooker/CookTypes.cpp
--- a/Engine/Source/Editor/UnrealEd/Private/Cooker/CookTypes.cpp
+++ b/Engine/Source/Editor/UnrealEd/Private/Cooker/CookTypes.cpp
@@ -398,3 +398,54 @@ bool LoadFromCompactBinary(FCbFieldView Field, UE::Cook::FCookOnTheFlyOptions& O
 	bOk = LoadFromCompactBinary(Field["PlatformProtocol"], OutValue.bPlatformProtocol) & bOk;
 	return bOk;
 }
+
+void FBeginCookContextForWorkerPlatform::Set(const FBeginCookContextPlatform& InContext)
+{
+	bFullBuild = InContext.bFullBuild;
+	TargetPlatform = InContext.TargetPlatform;
+}
+
+FCbWriter& operator<<(FCbWriter& Writer, const FBeginCookContextForWorkerPlatform& Value)
+{
+	Writer.BeginObject();
+	Writer << "Platform" << (Value.TargetPlatform ? Value.TargetPlatform->PlatformName() : FString());
+	Writer << "FullBuild" << Value.bFullBuild;
+	Writer.EndObject();
+	return Writer;
+}
+
+bool LoadFromCompactBinary(FCbFieldView Field, FBeginCookContextForWorkerPlatform& Value)
+{
+	bool bOk = true;
+	FString PlatformName;
+	LoadFromCompactBinary(Field["Platform"], PlatformName);
+	Value.TargetPlatform = nullptr;
+	if (!PlatformName.IsEmpty())
+	{
+		ITargetPlatformManagerModule& TPM = GetTargetPlatformManagerRef();
+		Value.TargetPlatform = TPM.FindTargetPlatform(*PlatformName);
+		bOk = (Value.TargetPlatform != nullptr) & bOk;
+	}
+	bOk = LoadFromCompactBinary(Field["FullBuild"], Value.bFullBuild) & bOk;
+	return bOk;
+}
+
+void FBeginCookContextForWorker::Set(const FBeginCookContext& InContext)
+{
+	PlatformContexts.SetNum(InContext.PlatformContexts.Num());
+	for (int32 Index = 0; Index < InContext.PlatformContexts.Num(); ++Index)
+	{
+		PlatformContexts[Index].Set(InContext.PlatformContexts[Index]);
+	}
+}
+
+FCbWriter& operator<<(FCbWriter& Writer, const FBeginCookContextForWorker& Value)
+{
+	Writer << Value.PlatformContexts;
+	return Writer;
+}
+
+bool LoadFromCompactBinary(FCbFieldView Field, FBeginCookContextForWorker& Value)
+{
+	return LoadFromCompactBinary(Field, Value.PlatformContexts);
+}
diff --git a/Engine/Source/Editor/UnrealEd/Private/Cooker/CookTypes.h b/Engine/Source/Editor/UnrealEd/Private/Cooker/CookTypes.h
--- a/Engine/Source/Editor/UnrealEd/Private/Cooker/CookTypes.h
+++ b/Engine/Source/Editor/UnrealEd/Private/Cooker/CookTypes.h
@@ -462,6 +462,8 @@ struct FBeginCookContextPlatform
 	/** If true we are a CookWorker, and we are working on a Sandbox directory that has already been populated by a remote Director process. */
 	bool bWorkerOnSharedSandbox = false;
 };
+FCbWriter& operator<<(FCbWriter& Writer, const FBeginCookContextPlatform& Value);
+bool LoadFromCompactBinary(FCbFieldView Field, FBeginCookContextPlatform& Value);
 
 /** Data held on the stack and shared with multiple subfunctions when running StartCookByTheBook or StartCookOnTheFly */
 struct FBeginCookContext
@@ -479,6 +481,29 @@ struct FBeginCookContext
 	const UCookOnTheFlyServer& COTFS;
 };
 
+/** Helper struct for FBeginCookContextForWorker; holds the context data for each platform being cooked */
+struct FBeginCookContextForWorkerPlatform
+{
+	void Set(const FBeginCookContextPlatform& InContext);
+
+	const ITargetPlatform* TargetPlatform = nullptr;
+	/** If true, we are deleting all old results from disk and rebuilding every package. If false, we are building iteratively. */
+	bool bFullBuild = false;
+};
+FCbWriter& operator<<(FCbWriter& Writer, const FBeginCookContextForWorkerPlatform& Value);
+bool LoadFromCompactBinary(FCbFieldView Field, FBeginCookContextForWorkerPlatform& Value);
+
+/** Data from the director's FBeginCookContext that needs to be copied to workers. */
+struct FBeginCookContextForWorker
+{
+	void Set(const FBeginCookContext& InContext);
+
+	/** List of the platforms we are building, with startup context data about each one */
+	TArray<FBeginCookContextForWorkerPlatform> PlatformContexts;
+};
+FCbWriter& operator<<(FCbWriter& Writer, const FBeginCookContextForWorker& Value);
+bool LoadFromCompactBinary(FCbFieldView Field, FBeginCookContextForWorker& Value);
+
 void LogCookerMessage(const FString& MessageText, EMessageSeverity::Type Severity);
 LLM_DECLARE_TAG(Cooker);
 
diff --git a/Engine/Source/Editor/UnrealEd/Private/Cooker/CookWorkerClient.cpp b/Engine/Source/Editor/UnrealEd/Private/Cooker/CookWorkerClient.cpp
--- a/Engine/Source/Editor/UnrealEd/Private/Cooker/CookWorkerClient.cpp
+++ b/Engine/Source/Editor/UnrealEd/Private/Cooker/CookWorkerClient.cpp
@@ -104,6 +104,11 @@ FCookByTheBookOptions&& FCookWorkerClient::ConsumeCookByTheBookOptions()
 	check(InitialConfigMessage); // Should only be called after TryConnect and before DoneWithInitialSettings
 	return InitialConfigMessage->ConsumeCookByTheBookOptions();
 }
+const FBeginCookContextForWorker& FCookWorkerClient::GetBeginCookContext()
+{
+	check(InitialConfigMessage); // Should only be called after TryConnect and before DoneWithInitialSettings
+	return InitialConfigMessage->GetBeginCookContext();
+}
 FCookOnTheFlyOptions&& FCookWorkerClient::ConsumeCookOnTheFlyOptions()
 {
 	check(InitialConfigMessage); // Should only be called after TryConnect and before DoneWithInitialSettings
diff --git a/Engine/Source/Editor/UnrealEd/Private/Cooker/CookWorkerClient.h b/Engine/Source/Editor/UnrealEd/Private/Cooker/CookWorkerClient.h
--- a/Engine/Source/Editor/UnrealEd/Private/Cooker/CookWorkerClient.h
+++ b/Engine/Source/Editor/UnrealEd/Private/Cooker/CookWorkerClient.h
@@ -43,6 +43,7 @@ public:
 	FInitializeConfigSettings&& ConsumeInitializeConfigSettings();
 	FBeginCookConfigSettings&& ConsumeBeginCookConfigSettings();
 	FCookByTheBookOptions&& ConsumeCookByTheBookOptions();
+	const FBeginCookContextForWorker& GetBeginCookContext();
 	FCookOnTheFlyOptions&& ConsumeCookOnTheFlyOptions();
 	/** Mark that initialization is complete and we can free the memory for initialization settings. */
 	void DoneWithInitialSettings();
diff --git a/Engine/Source/Editor/UnrealEd/Private/Cooker/CookWorkerServer.cpp b/Engine/Source/Editor/UnrealEd/Private/Cooker/CookWorkerServer.cpp
--- a/Engine/Source/Editor/UnrealEd/Private/Cooker/CookWorkerServer.cpp
+++ b/Engine/Source/Editor/UnrealEd/Private/Cooker/CookWorkerServer.cpp
@@ -187,7 +187,7 @@ bool FCookWorkerServer::TryHandleConnectMessage(FWorkerConnectMessage& Message,
 		OrderedSessionPlatforms.Add(const_cast<ITargetPlatform*>(TargetPlatform));
 	}
 	ConfigMessage.ReadFromLocal(COTFS, OrderedSessionPlatforms,
-		*COTFS.CookByTheBookOptions, *COTFS.CookOnTheFlyOptions);
+		*COTFS.CookByTheBookOptions, *COTFS.CookOnTheFlyOptions, Director.BeginCookContext);
 	SendMessage(ConfigMessage);
 	return true;
 }
@@ -613,11 +613,13 @@ bool FAbortWorkerMessage::TryRead(FCbObject&& Object)
 
 FGuid FAbortWorkerMessage::MessageType(TEXT("83FD99DFE8DB4A9A8E71684C121BE6F3"));
 
-void FInitialConfigMessage::ReadFromLocal(const UCookOnTheFlyServer& COTFS, const TArray<ITargetPlatform*>& InOrderedSessionPlatforms,
-	const FCookByTheBookOptions& InCookByTheBookOptions, const FCookOnTheFlyOptions& InCookOnTheFlyOptions)
+void FInitialConfigMessage::ReadFromLocal(const UCookOnTheFlyServer& COTFS,
+	const TArray<ITargetPlatform*>& InOrderedSessionPlatforms, const FCookByTheBookOptions& InCookByTheBookOptions,
+	const FCookOnTheFlyOptions& InCookOnTheFlyOptions, const FBeginCookContextForWorker& InBeginContext)
 {
 	InitialSettings.CopyFromLocal(COTFS);
 	BeginCookSettings.CopyFromLocal(COTFS);
+	BeginCookContext = InBeginContext;
 	OrderedSessionPlatforms = InOrderedSessionPlatforms;
 	DirectorCookMode = COTFS.GetCookMode();
 	CookInitializationFlags = COTFS.GetCookFlags();
@@ -642,6 +644,7 @@ void FInitialConfigMessage::Write(FCbWriter& Writer) const
 	Writer.EndArray();
 	Writer << "InitialSettings" << InitialSettings;
 	Writer << "BeginCookSettings" << BeginCookSettings;
+	Writer << "BeginCookContext" << BeginCookContext;
 	Writer << "CookByTheBookOptions" << CookByTheBookOptions;
 	Writer << "CookOnTheFlyOptions" << CookOnTheFlyOptions;
 }
@@ -689,6 +692,7 @@ bool FInitialConfigMessage::TryRead(FCbObject&& Object)
 
 	bOk = LoadFromCompactBinary(Object["InitialSettings"], InitialSettings) & bOk;
 	bOk = LoadFromCompactBinary(Object["BeginCookSettings"], BeginCookSettings) & bOk;
+	bOk = LoadFromCompactBinary(Object["BeginCookContext"], BeginCookContext) & bOk;
 	bOk = LoadFromCompactBinary(Object["CookByTheBookOptions"], CookByTheBookOptions) & bOk;
 	bOk = LoadFromCompactBinary(Object["CookOnTheFlyOptions"], CookOnTheFlyOptions) & bOk;
 	return bOk;
diff --git a/Engine/Source/Editor/UnrealEd/Private/Cooker/CookWorkerServer.h b/Engine/Source/Editor/UnrealEd/Private/Cooker/CookWorkerServer.h
--- a/Engine/Source/Editor/UnrealEd/Private/Cooker/CookWorkerServer.h
+++ b/Engine/Source/Editor/UnrealEd/Private/Cooker/CookWorkerServer.h
@@ -181,7 +181,8 @@ public:
 	virtual FGuid GetMessageType() const override { return MessageType; }
 
 	void ReadFromLocal(const UCookOnTheFlyServer& COTFS, const TArray<ITargetPlatform*>& InOrderedSessionPlatforms,
-		const FCookByTheBookOptions& InCookByTheBookOptions, const FCookOnTheFlyOptions& InCookOnTheFlyOptions);
+		const FCookByTheBookOptions& InCookByTheBookOptions, const FCookOnTheFlyOptions& InCookOnTheFlyOptions,
+		const FBeginCookContextForWorker& InBeginContext);
 
 	ECookMode::Type GetDirectorCookMode() const { return DirectorCookMode; }
 	ECookInitializationFlags GetCookInitializationFlags() const { return CookInitializationFlags; }
@@ -189,6 +190,7 @@ public:
 	FBeginCookConfigSettings&& ConsumeBeginCookConfigSettings() { return MoveTemp(BeginCookSettings); }
 	FCookByTheBookOptions&& ConsumeCookByTheBookOptions() { return MoveTemp(CookByTheBookOptions); }
 	FCookOnTheFlyOptions&& ConsumeCookOnTheFlyOptions() { return MoveTemp(CookOnTheFlyOptions); }
+	const FBeginCookContextForWorker& GetBeginCookContext() { return BeginCookContext; }
 	const TArray<ITargetPlatform*>& GetOrderedSessionPlatforms() { return OrderedSessionPlatforms; }
 	bool IsZenStore() const { return bZenStore; }
 
@@ -197,6 +199,7 @@ public:
 private:
 	FInitializeConfigSettings InitialSettings;
 	FBeginCookConfigSettings BeginCookSettings;
+	FBeginCookContextForWorker BeginCookContext;
 	FCookByTheBookOptions CookByTheBookOptions;
 	FCookOnTheFlyOptions CookOnTheFlyOptions;
 	TArray<ITargetPlatform*> OrderedSessionPlatforms;
diff --git a/Engine/Source/Editor/UnrealEd/Private/Cooker/WorkerRequestsRemote.cpp b/Engine/Source/Editor/UnrealEd/Private/Cooker/WorkerRequestsRemote.cpp
--- a/Engine/Source/Editor/UnrealEd/Private/Cooker/WorkerRequestsRemote.cpp
+++ b/Engine/Source/Editor/UnrealEd/Private/Cooker/WorkerRequestsRemote.cpp
@@ -2,11 +2,13 @@
 
 #include "WorkerRequestsRemote.h"
 
+#include "Algo/Find.h"
 #include "Containers/SparseArray.h"
 #include "CookOnTheSide/CookLog.h"
 #include "CookPlatformManager.h"
 #include "CookTypes.h"
 #include "Cooker/CookRequests.h"
+#include "Interfaces/ITargetPlatform.h"
 #include "Logging/LogCategory.h"
 #include "Logging/LogMacros.h"
 #include "Misc/AssertionMacros.h"
@@ -133,16 +135,26 @@ void FWorkerRequestsRemote::GetBeginCookConfigSettings(UCookOnTheFlyServer& COTF
 
 void FWorkerRequestsRemote::GetBeginCookIterativeFlags(UCookOnTheFlyServer& COTFS, FBeginCookContext& BeginContext)
 {
+	const FBeginCookContextForWorker& DirectorBeginContext = CookWorkerClient.GetBeginCookContext();
+
 	for (FBeginCookContextPlatform& PlatformContext : BeginContext.PlatformContexts)
 	{
 		const ITargetPlatform* TargetPlatform = PlatformContext.TargetPlatform;
+		const FBeginCookContextForWorkerPlatform* DirectorPlatformContext = Algo::FindByPredicate(
+			DirectorBeginContext.PlatformContexts, [TargetPlatform](const FBeginCookContextForWorkerPlatform& Platform)
+			{
+				return Platform.TargetPlatform == TargetPlatform;
+			});
+		checkf(DirectorPlatformContext, TEXT("Director sent TargetPlatform %s, but this platform is not found in the DirectorBeginContext."),
+			*TargetPlatform->PlatformName());
+
 		UE::Cook::FPlatformData* PlatformData = PlatformContext.PlatformData;
-		PlatformContext.CurrentCookSettings = COTFS.CalculateCookSettingStrings(); // MPCOOKTODO: Copy from Director
-		PlatformContext.bHasMemoryResults = PlatformData->bIsSandboxInitialized;
-		PlatformContext.bFullBuild = true; // MPCOOKTODO: Copy from Director
-		PlatformContext.bClearMemoryResults = true;
-		PlatformContext.bPopulateMemoryResultsFromDiskResults = false;
-		PlatformContext.bIterateSharedBuild = false; // MPCOOKTODO: Copy from Director
+		PlatformContext.CurrentCookSettings.Empty(); // Not needed on CookWorkerClients
+		PlatformContext.bHasMemoryResults = false; // Not needed on CookWorkerClients
+		PlatformContext.bFullBuild = DirectorPlatformContext->bFullBuild;
+		PlatformContext.bClearMemoryResults = false; // Not needed on CookWorkerClients
+		PlatformContext.bPopulateMemoryResultsFromDiskResults = false; // Not needed on CookWorkerClients
+		PlatformContext.bIterateSharedBuild = false; // Not needed on CookWorkerClients
 		PlatformContext.bWorkerOnSharedSandbox = true;
 		PlatformData->bFullBuild = PlatformContext.bFullBuild;
 	}
