WorldPartitionBuilder: Add World Init Values getter ( UWorldPartitionBuilder::GetWorldInitializationValues() )

diff --git a/Engine/Source/Editor/UnrealEd/Private/WorldPartition/WorldPartitionBuilder.cpp b/Engine/Source/Editor/UnrealEd/Private/WorldPartition/WorldPartitionBuilder.cpp
--- a/Engine/Source/Editor/UnrealEd/Private/WorldPartition/WorldPartitionBuilder.cpp
+++ b/Engine/Source/Editor/UnrealEd/Private/WorldPartition/WorldPartitionBuilder.cpp
@@ -59,6 +59,19 @@ UWorldPartitionBuilder::UWorldPartitionBuilder(const FObjectInitializer& ObjectI
 	}
 }
 
+UWorld::InitializationValues UWorldPartitionBuilder::GetWorldInitializationValues() const
+{
+	UWorld::InitializationValues IVS;
+	IVS.RequiresHitProxies(false);
+	IVS.ShouldSimulatePhysics(false);
+	IVS.EnableTraceCollision(false);
+	IVS.CreateNavigation(false);
+	IVS.CreateAISystem(false);
+	IVS.AllowAudioPlayback(false);
+	IVS.CreatePhysicsScene(true);
+	return IVS;
+}
+
 bool UWorldPartitionBuilder::RunBuilder(UWorld* World)
 {
 	// Load configuration file & builder configuration
@@ -95,15 +108,7 @@ bool UWorldPartitionBuilder::RunBuilder(UWorld* World)
 		}
 		else
 		{
-			UWorld::InitializationValues IVS;
-			IVS.RequiresHitProxies(false);
-			IVS.ShouldSimulatePhysics(false);
-			IVS.EnableTraceCollision(false);
-			IVS.CreateNavigation(false);
-			IVS.CreateAISystem(false);
-			IVS.AllowAudioPlayback(false);
-			IVS.CreatePhysicsScene(true);
-			EditorWorld = MakeUnique<FScopedEditorWorld>(World, IVS);
+			EditorWorld = MakeUnique<FScopedEditorWorld>(World, GetWorldInitializationValues());
 		}
 
 		// Make sure the world is partitioned
diff --git a/Engine/Source/Editor/UnrealEd/Public/WorldPartition/WorldPartitionBuilder.h b/Engine/Source/Editor/UnrealEd/Public/WorldPartition/WorldPartitionBuilder.h
--- a/Engine/Source/Editor/UnrealEd/Public/WorldPartition/WorldPartitionBuilder.h
+++ b/Engine/Source/Editor/UnrealEd/Public/WorldPartition/WorldPartitionBuilder.h
@@ -87,6 +87,8 @@ protected:
 	bool AutoSubmitFiles(const TArray<FString>& InModifiedFiles, const FString& InChangelistDescription) const;
 	bool AutoSubmitPackages(const TArray<UPackage*>& InModifiedPackages, const FString& InChangelistDescription) const;
 
+	virtual UWorld::InitializationValues GetWorldInitializationValues() const;
+
 	int32 IterativeCellSize = 102400;
 	int32 IterativeCellOverlapSize = 0;
 	FBox  IterativeWorldBounds;
