CookMultiprocess - replicate package instigator when assigning packages to a cookworker, because FCookedSavePackageValidator::ValidateImports uses it to suppress missing module messages if the package is a startuppackage.
#rnx

diff --git a/Engine/Source/Editor/UnrealEd/Private/Cooker/CookWorkerClient.cpp b/Engine/Source/Editor/UnrealEd/Private/Cooker/CookWorkerClient.cpp
--- a/Engine/Source/Editor/UnrealEd/Private/Cooker/CookWorkerClient.cpp
+++ b/Engine/Source/Editor/UnrealEd/Private/Cooker/CookWorkerClient.cpp
@@ -543,10 +543,10 @@ void FCookWorkerClient::LogInvalidMessage(const TCHAR* MessageTypeName)
 
 void FCookWorkerClient::AssignPackages(FAssignPackagesMessage& Message)
 {
-	for (FConstructPackageData& ConstructPackageData : Message.PackageDatas)
+	for (FAssignPackageData& AssignData: Message.PackageDatas)
 	{
-		FPackageData& PackageData = COTFS.PackageDatas->FindOrAddPackageData(ConstructPackageData.PackageName,
-			ConstructPackageData.NormalizedFileName);
+		FPackageData& PackageData = COTFS.PackageDatas->FindOrAddPackageData(AssignData.ConstructData.PackageName,
+			AssignData.ConstructData.NormalizedFileName);
 		// If already InProgress, ignore the duplicate package silently
 		if (PackageData.IsInProgress())
 		{
@@ -562,7 +562,7 @@ void FCookWorkerClient::AssignPackages(FAssignPackagesMessage& Message)
 		}
 
 		PackageData.SetRequestData(OrderedSessionPlatforms , false /* bInIsUrgent */, FCompletionCallback(),
-			FInstigator(EInstigator::CookDirector));
+			FInstigator(AssignData.Instigator));
 		PackageData.SendToState(EPackageState::Request, ESendFlags::QueueAddAndRemove);
 	}
 }
diff --git a/Engine/Source/Editor/UnrealEd/Private/Cooker/CookWorkerServer.cpp b/Engine/Source/Editor/UnrealEd/Private/Cooker/CookWorkerServer.cpp
--- a/Engine/Source/Editor/UnrealEd/Private/Cooker/CookWorkerServer.cpp
+++ b/Engine/Source/Editor/UnrealEd/Private/Cooker/CookWorkerServer.cpp
@@ -373,15 +373,17 @@ void FCookWorkerServer::SendPendingPackages()
 		return;
 	}
 
-	TArray<FConstructPackageData> ConstructDatas;
-	ConstructDatas.Reserve(PackagesToAssign.Num());
+	TArray<FAssignPackageData> AssignDatas;
+	AssignDatas.Reserve(PackagesToAssign.Num());
 	for (FPackageData* PackageData : PackagesToAssign)
 	{
-		ConstructDatas.Add(PackageData->CreateConstructData());
+		FAssignPackageData& AssignData = AssignDatas.Emplace_GetRef();
+		AssignData.ConstructData = PackageData->CreateConstructData();
+		AssignData.Instigator = PackageData->GetInstigator();
 	}
 	PendingPackages.Append(PackagesToAssign);
 	PackagesToAssign.Empty();
-	SendMessage(FAssignPackagesMessage(MoveTemp(ConstructDatas)));
+	SendMessage(FAssignPackagesMessage(MoveTemp(AssignDatas)));
 }
 
 void FCookWorkerServer::PumpReceiveMessages()
@@ -561,7 +563,7 @@ void FCookWorkerServer::AddDiscoveredPackage(FDiscoveredPackage&& DiscoveredPack
 	COTFS.QueueDiscoveredPackageData(PackageData, MoveTemp(DiscoveredPackage.Instigator));
 }
 
-FAssignPackagesMessage::FAssignPackagesMessage(TArray<FConstructPackageData>&& InPackageDatas)
+FAssignPackagesMessage::FAssignPackagesMessage(TArray<FAssignPackageData>&& InPackageDatas)
 	: PackageDatas(MoveTemp(InPackageDatas))
 {
 }
@@ -578,6 +580,49 @@ bool FAssignPackagesMessage::TryRead(FCbObject&& Object)
 
 FGuid FAssignPackagesMessage::MessageType(TEXT("B7B1542B73254B679319D73F753DB6F8"));
 
+FCbWriter& operator<<(FCbWriter& Writer, const FAssignPackageData& AssignData)
+{
+	Writer.BeginObject();
+	Writer << "C" << AssignData.ConstructData;
+	Writer << "I" << AssignData.Instigator;
+	Writer.EndObject();
+	return Writer;
+}
+
+bool LoadFromCompactBinary(FCbFieldView Field, FAssignPackageData& AssignData)
+{
+	bool bOk = LoadFromCompactBinary(Field["C"], AssignData.ConstructData);
+	bOk = LoadFromCompactBinary(Field["I"], AssignData.Instigator) & bOk;
+	return bOk;
+}
+
+FCbWriter& operator<<(FCbWriter& Writer, const FInstigator& Instigator)
+{
+	Writer.BeginObject();
+	Writer << "C" << static_cast<uint8>(Instigator.Category);
+	Writer << "R" << Instigator.Referencer;
+	Writer.EndObject();
+	return Writer;
+}
+
+bool LoadFromCompactBinary(FCbFieldView Field, FInstigator& Instigator)
+{
+	uint8 CategoryInt;
+	bool bOk = true;
+	if (LoadFromCompactBinary(Field["C"], CategoryInt) &&
+		CategoryInt < static_cast<uint8>(EInstigator::Count))
+	{
+		Instigator.Category = static_cast<EInstigator>(CategoryInt);
+	}
+	else
+	{
+		Instigator.Category = EInstigator::InvalidCategory;
+		bOk = false;
+	}
+	bOk = LoadFromCompactBinary(Field["R"], Instigator.Referencer) & bOk;
+	return bOk;
+}
+
 FAbortPackagesMessage::FAbortPackagesMessage(TArray<FName>&& InPackageNames)
 	: PackageNames(MoveTemp(InPackageNames))
 {
diff --git a/Engine/Source/Editor/UnrealEd/Private/Cooker/CookWorkerServer.h b/Engine/Source/Editor/UnrealEd/Private/Cooker/CookWorkerServer.h
--- a/Engine/Source/Editor/UnrealEd/Private/Cooker/CookWorkerServer.h
+++ b/Engine/Source/Editor/UnrealEd/Private/Cooker/CookWorkerServer.h
@@ -7,6 +7,7 @@
 #include "Containers/ArrayView.h"
 #include "Containers/Set.h"
 #include "CookOnTheSide/CookOnTheFlyServer.h"
+#include "CookPackageData.h"
 #include "CookSockets.h"
 #include "CookTypes.h"
 #include "Misc/Guid.h"
@@ -16,7 +17,6 @@ class ITargetPlatform;
 struct FProcHandle;
 namespace UE::CompactBinaryTCP { struct FMarshalledMessage; }
 namespace UE::Cook { class FCookDirector; }
-namespace UE::Cook { struct FConstructPackageData; }
 namespace UE::Cook { struct FDiscoveredPackage; }
 namespace UE::Cook { struct FPackageData; }
 namespace UE::Cook { struct FPackageResultsMessage; }
@@ -116,19 +116,30 @@ private:
 	bool bTerminateImmediately = false;
 };
 
+/** Information about a PackageData the director sends to cookworkers. */
+struct FAssignPackageData
+{
+	FConstructPackageData ConstructData;
+	FInstigator Instigator;
+};
+FCbWriter& operator<<(FCbWriter& Writer, const FAssignPackageData& AssignData);
+bool LoadFromCompactBinary(FCbFieldView Field, FAssignPackageData& AssignData);
+FCbWriter& operator<<(FCbWriter& Writer, const FInstigator& Instigator);
+bool LoadFromCompactBinary(FCbFieldView Field, FInstigator& Instigator);
+
 /** Message from Server to Client to cook the given packages. */
 struct FAssignPackagesMessage : public UE::CompactBinaryTCP::IMessage
 {
 public:
 	FAssignPackagesMessage() = default;
-	FAssignPackagesMessage(TArray<FConstructPackageData>&& InPackageDatas);
+	FAssignPackagesMessage(TArray<FAssignPackageData>&& InPackageDatas);
 
 	virtual void Write(FCbWriter& Writer) const override;
 	virtual bool TryRead(FCbObject&& Object) override;
 	virtual FGuid GetMessageType() const override { return MessageType; }
 
 public:
-	TArray<FConstructPackageData> PackageDatas;
+	TArray<FAssignPackageData> PackageDatas;
 	static FGuid MessageType;
 };
 
diff --git a/Engine/Source/Runtime/Engine/Public/Engine/ICookInfo.h b/Engine/Source/Runtime/Engine/Public/Engine/ICookInfo.h
--- a/Engine/Source/Runtime/Engine/Public/Engine/ICookInfo.h
+++ b/Engine/Source/Runtime/Engine/Public/Engine/ICookInfo.h
@@ -41,7 +41,6 @@ namespace UE::Cook
 	CallbackMacro(AssetManagerModifyDLCCook, true) \
 	CallbackMacro(TargetPlatformExtraPackagesToCook, true) \
 	CallbackMacro(ConsoleCommand, true) \
-	CallbackMacro(CookDirector, true) \
 	CallbackMacro(CookOnTheFly, true) \
 	CallbackMacro(IterativeCook, true) \
 	CallbackMacro(PreviousAssetRegistry, true) \
