Fix non-unity build

diff --git a/Engine/Source/Editor/UnrealEd/Public/WorldPartition/WorldPartitionBuilder.h b/Engine/Source/Editor/UnrealEd/Public/WorldPartition/WorldPartitionBuilder.h
--- a/Engine/Source/Editor/UnrealEd/Public/WorldPartition/WorldPartitionBuilder.h
+++ b/Engine/Source/Editor/UnrealEd/Public/WorldPartition/WorldPartitionBuilder.h
@@ -3,6 +3,7 @@
 #pragma once
 
 #include "PackageSourceControlHelper.h"
+#include "Engine/World.h"
 #include "WorldPartitionBuilder.generated.h"
 
 typedef UE::Math::TIntVector3<int64> FWorldBuilderCellCoord;
