UE-160860 - Fix UnrealEd Module truncation warnings

Part 1 of N

diff --git a/Engine/Source/Editor/UnrealEd/Classes/Animation/DebugSkelMeshComponent.h b/Engine/Source/Editor/UnrealEd/Classes/Animation/DebugSkelMeshComponent.h
--- a/Engine/Source/Editor/UnrealEd/Classes/Animation/DebugSkelMeshComponent.h
+++ b/Engine/Source/Editor/UnrealEd/Classes/Animation/DebugSkelMeshComponent.h
@@ -130,7 +130,7 @@ public:
 
 	FDebugSkelMeshDynamicData* DynamicData;
 
-	uint32 GetAllocatedSize() const
+	SIZE_T GetAllocatedSize() const
 	{
 		return FSkeletalMeshSceneProxy::GetAllocatedSize();
 	}
diff --git a/Engine/Source/Editor/UnrealEd/Private/Animation/DebugSkelMeshComponent.cpp b/Engine/Source/Editor/UnrealEd/Private/Animation/DebugSkelMeshComponent.cpp
--- a/Engine/Source/Editor/UnrealEd/Private/Animation/DebugSkelMeshComponent.cpp
+++ b/Engine/Source/Editor/UnrealEd/Private/Animation/DebugSkelMeshComponent.cpp
@@ -159,10 +159,10 @@ bool UDebugSkelMeshComponent::CheckIfBoundsAreCorrrect()
 	return false;
 }
 
-float WrapInRange(float StartVal, float MinVal, float MaxVal)
+double WrapInRange(double StartVal, double MinVal, double MaxVal)
 {
-	float Size = MaxVal - MinVal;
-	float EndVal = StartVal;
+	double Size = MaxVal - MinVal;
+	double EndVal = StartVal;
 	while (EndVal < MinVal)
 	{
 		EndVal += Size;
diff --git a/Engine/Source/Editor/UnrealEd/Private/Commandlets/CookGlobalShaderCommandlet.cpp b/Engine/Source/Editor/UnrealEd/Private/Commandlets/CookGlobalShaderCommandlet.cpp
--- a/Engine/Source/Editor/UnrealEd/Private/Commandlets/CookGlobalShaderCommandlet.cpp
+++ b/Engine/Source/Editor/UnrealEd/Private/Commandlets/CookGlobalShaderCommandlet.cpp
@@ -178,10 +178,10 @@ void UCookGlobalShadersCommandlet::CookGlobalShadersOnDirectoriesChanges()
 	while (GIsRunning && !IsEngineExitRequested())
 	{
 		GEngine->UpdateTimeAndHandleMaxTickRate();
-		GEngine->Tick(FApp::GetDeltaTime(), false);
+		GEngine->Tick(static_cast<float>(FApp::GetDeltaTime()), false);
 		
 		// tick the directory watcher
-		DirectoryWatcherModule.Get()->Tick(FApp::GetDeltaTime());
+		DirectoryWatcherModule.Get()->Tick(static_cast<float>(FApp::GetDeltaTime()));
 		
 		// flush log
 		GLog->FlushThreadedLogs(EOutputDeviceRedirectorFlushOptions::Async);
diff --git a/Engine/Source/Editor/UnrealEd/Private/Commandlets/ExportPakDependencies.cpp b/Engine/Source/Editor/UnrealEd/Private/Commandlets/ExportPakDependencies.cpp
--- a/Engine/Source/Editor/UnrealEd/Private/Commandlets/ExportPakDependencies.cpp
+++ b/Engine/Source/Editor/UnrealEd/Private/Commandlets/ExportPakDependencies.cpp
@@ -121,8 +121,8 @@ public:
 		TSharedPtr<FJsonObject> JsonPackageObject = MakeShareable(new FJsonObject);
 		
 		JsonPackageObject->SetStringField(TEXT("Name"),*Name.ToString());
-		JsonPackageObject->SetNumberField(TEXT("InclusiveSize"),InclusiveSize);
-		JsonPackageObject->SetNumberField(TEXT("ExclusiveSize"),ExclusiveSize);
+		JsonPackageObject->SetNumberField(TEXT("InclusiveSize"),static_cast<double>(InclusiveSize));
+		JsonPackageObject->SetNumberField(TEXT("ExclusiveSize"),static_cast<double>(ExclusiveSize));
 
 		JsonPackageObject->SetArrayField(TEXT("DirectlyReferencing"),ToJsonHelper(DirectlyReferencing));
 		JsonPackageObject->SetArrayField(TEXT("DirectlyReferencedBy"),ToJsonHelper(DirectlyReferencedBy));
diff --git a/Engine/Source/Editor/UnrealEd/Private/Commandlets/FileServerCommandlet.cpp b/Engine/Source/Editor/UnrealEd/Private/Commandlets/FileServerCommandlet.cpp
--- a/Engine/Source/Editor/UnrealEd/Private/Commandlets/FileServerCommandlet.cpp
+++ b/Engine/Source/Editor/UnrealEd/Private/Commandlets/FileServerCommandlet.cpp
@@ -99,11 +99,11 @@ int32 UFileServerCommandlet::Main( const FString& Params )
 	while (GIsRunning && !IsEngineExitRequested())
 	{
 		GEngine->UpdateTimeAndHandleMaxTickRate();
-		GEngine->Tick(FApp::GetDeltaTime(), false);
+		GEngine->Tick(static_cast<float>(FApp::GetDeltaTime()), false);
 
 		// tick the directory watcher
 		FDirectoryWatcherModule& DirectoryWatcherModule = FModuleManager::Get().LoadModuleChecked<FDirectoryWatcherModule>(TEXT("DirectoryWatcher"));
-		DirectoryWatcherModule.Get()->Tick(FApp::GetDeltaTime());
+		DirectoryWatcherModule.Get()->Tick(static_cast<float>(FApp::GetDeltaTime()));
 
 		// update task graph
 		FTaskGraphInterface::Get().ProcessThreadUntilIdle(ENamedThreads::GameThread);
diff --git a/Engine/Source/Editor/UnrealEd/Private/Commandlets/GatherTextFromSourceCommandlet.cpp b/Engine/Source/Editor/UnrealEd/Private/Commandlets/GatherTextFromSourceCommandlet.cpp
--- a/Engine/Source/Editor/UnrealEd/Private/Commandlets/GatherTextFromSourceCommandlet.cpp
+++ b/Engine/Source/Editor/UnrealEd/Private/Commandlets/GatherTextFromSourceCommandlet.cpp
@@ -1523,7 +1523,7 @@ bool UGatherTextFromSourceCommandlet::FMacroDescriptor::ParseArgumentString(cons
 			if (Cursor - ArgStart > 0)
 			{
 				// Here, we found the end of the line, but we don't know if it's the end of the argument.
-				if (!ArgsGatherer.Gather(ArgStart, Cursor - ArgStart))
+				if (!ArgsGatherer.Gather(ArgStart, UE_PTRDIFF_TO_INT32(Cursor - ArgStart)))
 				{
 					UE_LOG(LogGatherTextFromSourceCommandlet, Warning, TEXT("%s(%d): Parsing Arguments failed in %s macro. %s"), *Context.Filename, Context.LineNumber, *GetToken(), *FLocTextHelper::SanitizeLogOutput(Context.LineText.TrimStartAndEnd()));
 					return false;
@@ -1545,7 +1545,7 @@ bool UGatherTextFromSourceCommandlet::FMacroDescriptor::ParseArgumentString(cons
 			{
 				if (Cursor - ArgStart - 1 > 0)
 				{
-					if (!ArgsGatherer.Gather(ArgStart, Cursor - ArgStart - 1))
+					if (!ArgsGatherer.Gather(ArgStart, UE_PTRDIFF_TO_INT32(Cursor - ArgStart) - 1))
 					{
 						UE_LOG(LogGatherTextFromSourceCommandlet, Warning, TEXT("%s(%d): Parsing Arguments failed in %s macro. %s"), *Context.Filename, Context.LineNumber, *GetToken(), *FLocTextHelper::SanitizeLogOutput(Context.LineText.TrimStartAndEnd()));
 						return false;
@@ -1603,7 +1603,7 @@ bool UGatherTextFromSourceCommandlet::FMacroDescriptor::ParseArgumentString(cons
 				if (Cursor - ArgStart > 0)
 				{
 					// We have a single bracket open and we found a ',' this is the end of the argument. If Bracket stack is > 1 it means that we are in a function call in one of the parameters.
-					if (!ArgsGatherer.Gather(ArgStart, Cursor - ArgStart))
+					if (!ArgsGatherer.Gather(ArgStart, UE_PTRDIFF_TO_INT32(Cursor - ArgStart)))
 					{
 						UE_LOG(LogGatherTextFromSourceCommandlet, Warning, TEXT("%s(%d): Parsing Arguments failed in %s macro. %s"), *Context.Filename, Context.LineNumber, *GetToken(), *FLocTextHelper::SanitizeLogOutput(Context.LineText.TrimStartAndEnd()));
 						return false;
@@ -1638,7 +1638,7 @@ bool UGatherTextFromSourceCommandlet::FMacroDescriptor::ParseArgumentString(cons
 			{
 				if (Cursor - ArgStart - 1 > 0)
 				{
-					if (!ArgsGatherer.Gather(ArgStart, Cursor - ArgStart - 1))
+					if (!ArgsGatherer.Gather(ArgStart, UE_PTRDIFF_TO_INT32(Cursor - ArgStart) - 1))
 					{
 						UE_LOG(LogGatherTextFromSourceCommandlet, Warning, TEXT("%s(%d): Parsing Arguments failed in %s macro. %s"), *Context.Filename, Context.LineNumber, *GetToken(), *FLocTextHelper::SanitizeLogOutput(Context.LineText.TrimStartAndEnd()));
 						return false;
@@ -1654,7 +1654,7 @@ bool UGatherTextFromSourceCommandlet::FMacroDescriptor::ParseArgumentString(cons
 			// If the arg is empty it means we found a closing braket after a ',' or at the begining of a line			
 			if (Cursor - ArgStart > 0)
 			{
-				if (!ArgsGatherer.Gather(ArgStart, Cursor - ArgStart))
+				if (!ArgsGatherer.Gather(ArgStart, UE_PTRDIFF_TO_INT32(Cursor - ArgStart)))
 				{
 					UE_LOG(LogGatherTextFromSourceCommandlet, Warning, TEXT("%s(%d): Parsing Arguments failed in %s macro. %s"), *Context.Filename, Context.LineNumber, *GetToken(), *FLocTextHelper::SanitizeLogOutput(Context.LineText.TrimStartAndEnd()));
 					return false;
diff --git a/Engine/Source/Editor/UnrealEd/Private/Commandlets/PackageUtilities.cpp b/Engine/Source/Editor/UnrealEd/Private/Commandlets/PackageUtilities.cpp
--- a/Engine/Source/Editor/UnrealEd/Private/Commandlets/PackageUtilities.cpp
+++ b/Engine/Source/Editor/UnrealEd/Private/Commandlets/PackageUtilities.cpp
@@ -879,7 +879,7 @@ namespace
 		// Comparison method
 		bool operator()( const FExportInfo& A, const FExportInfo& B ) const
 		{
-			int32 Result = 0;
+			int64 Result = 0;
 
 			for ( int32 PriorityType = 0; PriorityType < EXPORTSORT_MAX; PriorityType++ )
 			{
@@ -1782,7 +1782,7 @@ struct CompressAnimationsFunctor
 		// @todoanim: we expect this won't work properly since it won't have any skeletalmesh,
 		// but soon, the compression will changed based on skeleton. 
 		// when that happens, this doesn't have to worry about skeletalmesh not loaded
-	 	float LastSaveTime = FPlatformTime::Seconds();
+	 	double LastSaveTime = FPlatformTime::Seconds();
 		bool bDirtyPackage = false;
 		const FName& PackageName = Package->GetFName(); 
 		FString PackageFileName;
@@ -2001,7 +2001,7 @@ struct CompressAnimationsFunctor
 			if( bDirtyPackage )
 			{
 				// Save dirty package every 10 minutes at least, to avoid losing work in case of a crash on very large packages.
-				float const CurrentTime = FPlatformTime::Seconds();
+				const double CurrentTime = FPlatformTime::Seconds();
 				UE_LOG(LogPackageUtilities, Warning, TEXT("Time since last save: %f seconds"), (CurrentTime - LastSaveTime) );
 				if( (CurrentTime - LastSaveTime) > 10.f * 60.f )
 				{
diff --git a/Engine/Source/Editor/UnrealEd/Private/Commandlets/TextAssetCommandlet.cpp b/Engine/Source/Editor/UnrealEd/Private/Commandlets/TextAssetCommandlet.cpp
--- a/Engine/Source/Editor/UnrealEd/Private/Commandlets/TextAssetCommandlet.cpp
+++ b/Engine/Source/Editor/UnrealEd/Private/Commandlets/TextAssetCommandlet.cpp
@@ -444,8 +444,8 @@ bool UTextAssetCommandlet::DoTextAssetProcessing(const FProcessingArgs& InArgs)
 	const FString FailedDiffsPath = FPaths::ProjectSavedDir() / TEXT("FailedDiffs");
 	IFileManager::Get().DeleteDirectory(*FailedDiffsPath, false, true);
 
-	float TotalPackageLoadTime = 0.0;
-	float TotalPackageSaveTime = 0.0;
+	double TotalPackageLoadTime = 0.0;
+	double TotalPackageSaveTime = 0.0;
 
 	FArchive* CSVWriter = nullptr;
 	if (InArgs.CSVFilename.Len() > 0)
@@ -472,8 +472,8 @@ bool UTextAssetCommandlet::DoTextAssetProcessing(const FProcessingArgs& InArgs)
 		int64 NumFiles = 0;
 		FString MaxTimePackage;
 		FString MinTimePackage;
-		float IterationPackageLoadTime = 0.0;
-		float IterationPackageSaveTime = 0.0;
+		double IterationPackageLoadTime = 0.0;
+		double IterationPackageSaveTime = 0.0;
 		double ThisPackageLoadTime = 0.0;
 
 		TArray<FString> PhaseSuccess;
diff --git a/Engine/Source/Editor/UnrealEd/Private/Cooker/CookDirector.cpp b/Engine/Source/Editor/UnrealEd/Private/Cooker/CookDirector.cpp
--- a/Engine/Source/Editor/UnrealEd/Private/Cooker/CookDirector.cpp
+++ b/Engine/Source/Editor/UnrealEd/Private/Cooker/CookDirector.cpp
@@ -424,9 +424,9 @@ void FCookDirector::InitializeWorkers()
 	}
 
 	// Find any unused RemoteIndex less than the maximum used RemoteIndex
-	TArray<int32> UnusedRemoteIndexes;
+	TArray<uint8> UnusedRemoteIndexes;
 	RemoteWorkers.KeySort(TLess<>());
-	int32 NextPossiblyOpenIndex = 0;
+	uint8 NextPossiblyOpenIndex = 0;
 	for (TPair<int32, TUniquePtr<FCookWorkerServer>>& Pair : RemoteWorkers)
 	{
 		check(NextPossiblyOpenIndex <= Pair.Key);
@@ -439,7 +439,7 @@ void FCookDirector::InitializeWorkers()
 	// otherwise use the next integer because all indexes up to RemoteWorkers.Num() are in use.
 	while (RemoteWorkers.Num() < CookWorkerCount)
 	{
-		int32 RemoteIndex;
+		uint8 RemoteIndex;
 		if (UnusedRemoteIndexes.Num())
 		{
 			RemoteIndex = UnusedRemoteIndexes[0];
diff --git a/Engine/Source/Editor/UnrealEd/Private/Cooker/CookPackageData.cpp b/Engine/Source/Editor/UnrealEd/Private/Cooker/CookPackageData.cpp
--- a/Engine/Source/Editor/UnrealEd/Private/Cooker/CookPackageData.cpp
+++ b/Engine/Source/Editor/UnrealEd/Private/Cooker/CookPackageData.cpp
@@ -2644,7 +2644,7 @@ void FPackageDatas::PollPendingCookedPlatformDatas(bool bForce, double& LastCook
 	if (LastCookableObjectTickTime + TickCookableObjectsFrameTime <= CurrentTime)
 	{
 		UE_SCOPED_COOKTIMER(TickCookableObjects);
-		FTickableCookObject::TickObjects(CurrentTime - LastCookableObjectTickTime, false /* bTickComplete */);
+		FTickableCookObject::TickObjects(static_cast<float>(CurrentTime - LastCookableObjectTickTime), false /* bTickComplete */);
 		LastCookableObjectTickTime = CurrentTime;
 	}
 
diff --git a/Engine/Source/Editor/UnrealEd/Private/Cooker/CookProfiling.cpp b/Engine/Source/Editor/UnrealEd/Private/Cooker/CookProfiling.cpp
--- a/Engine/Source/Editor/UnrealEd/Private/Cooker/CookProfiling.cpp
+++ b/Engine/Source/Editor/UnrealEd/Private/Cooker/CookProfiling.cpp
@@ -39,7 +39,7 @@ public:
 		}
 		FirstChild = nullptr;
 	}
-	FHierarchicalTimerInfo* GetChild(int InId, const char* InName)
+	FHierarchicalTimerInfo* GetChild(uint16 InId, const char* InName)
 	{
 		for (FHierarchicalTimerInfo* Child = FirstChild; Child;)
 		{
@@ -91,7 +91,7 @@ FScopeTimer::FScopeTimer(int InId, const char* InName, bool IncrementScope /*= f
 {
 	checkSlow(IsInGameThread());
 
-	HierarchyTimerInfo = CurrentTimerInfo->GetChild(InId, InName);
+	HierarchyTimerInfo = CurrentTimerInfo->GetChild(static_cast<uint16>(InId), InName);
 
 	HierarchyTimerInfo->IncrementDepth = IncrementScope;
 
diff --git a/Engine/Source/Editor/UnrealEd/Private/Cooker/CookRequestCluster.cpp b/Engine/Source/Editor/UnrealEd/Private/Cooker/CookRequestCluster.cpp
--- a/Engine/Source/Editor/UnrealEd/Private/Cooker/CookRequestCluster.cpp
+++ b/Engine/Source/Editor/UnrealEd/Private/Cooker/CookRequestCluster.cpp
@@ -619,7 +619,7 @@ void FRequestCluster::FGraphSearch::AddVertices(TArray<TUniquePtr<FVertexData>>&
 
 void FRequestCluster::FGraphSearch::WaitForPollAvailability(double WaitTimeSeconds)
 {
-	uint32 WaitTime = WaitTimeSeconds > 0 ? FMath::Floor(WaitTimeSeconds * 1000) : MAX_uint32;
+	uint32 WaitTime = (WaitTimeSeconds > 0.0) ? static_cast<uint32>(FMath::Floor(WaitTimeSeconds * 1000)) : MAX_uint32;
 	PollReadyEvent->Wait(WaitTime);
 }
 
diff --git a/Engine/Source/Editor/UnrealEd/Private/Dialogs/Dialogs.cpp b/Engine/Source/Editor/UnrealEd/Private/Dialogs/Dialogs.cpp
--- a/Engine/Source/Editor/UnrealEd/Private/Dialogs/Dialogs.cpp
+++ b/Engine/Source/Editor/UnrealEd/Private/Dialogs/Dialogs.cpp
@@ -430,7 +430,7 @@ public:
 			.HAlign(HAlign_Fill)
 			.VAlign(VAlign_Fill)
 			.FillHeight(1.0f)
-			.Padding(5)
+			.Padding(5.0f)
 			[
 				SNew( STextBlock )
 				.WrapTextAt(615.0f)	// 400.0f
@@ -447,7 +447,7 @@ public:
 					SNew( SButton )
 					.Text( NSLOCTEXT("UnrealEd", "Yes", "Yes") )
 					.OnClicked( this, &SModalDialog::OnYesClicked )
-					.ContentPadding(7)
+					.ContentPadding(7.0f)
 				]
 				+ SHorizontalBox::Slot()
 				.AutoWidth()
@@ -455,7 +455,7 @@ public:
 					SNew( SButton )
 					.Text( NSLOCTEXT("UnrealEd", "No", "No") )
 					.OnClicked( this, &SModalDialog::OnNoClicked )
-					.ContentPadding(7)
+					.ContentPadding(7.0f)
 				]
 			]
 		];
@@ -904,7 +904,7 @@ void SGenericDialogWidget::Construct( const FArguments& InArgs )
 	{
 		ContentWidget = 
 			SNew(SBox)
-			.MaxDesiredHeight(InArgs._ScrollBoxMaxHeight)
+			.MaxDesiredHeight(static_cast<float>(InArgs._ScrollBoxMaxHeight))
 			[
 				SNew(SScrollBox)
 				+SScrollBox::Slot()
diff --git a/Engine/Source/Editor/UnrealEd/Private/Dialogs/DlgMoveAssets.cpp b/Engine/Source/Editor/UnrealEd/Private/Dialogs/DlgMoveAssets.cpp
--- a/Engine/Source/Editor/UnrealEd/Private/Dialogs/DlgMoveAssets.cpp
+++ b/Engine/Source/Editor/UnrealEd/Private/Dialogs/DlgMoveAssets.cpp
@@ -85,7 +85,7 @@ public:
 				// Add user input block
 				+SVerticalBox::Slot()
 				.AutoHeight()
-				.Padding(2)
+				.Padding(2.0f)
 				[
 						SNew(SVerticalBox)
 						+SVerticalBox::Slot()
diff --git a/Engine/Source/Editor/UnrealEd/Private/Dialogs/SDeleteAssetsDialog.cpp b/Engine/Source/Editor/UnrealEd/Private/Dialogs/SDeleteAssetsDialog.cpp
--- a/Engine/Source/Editor/UnrealEd/Private/Dialogs/SDeleteAssetsDialog.cpp
+++ b/Engine/Source/Editor/UnrealEd/Private/Dialogs/SDeleteAssetsDialog.cpp
@@ -789,8 +789,8 @@ TSharedRef<SWidget> SDeleteAssetsDialog::MakeConsolidationAssetPicker()
 	FContentBrowserModule& ContentBrowserModule = FModuleManager::Get().LoadModuleChecked<FContentBrowserModule>( TEXT( "ContentBrowser" ) );
 
 	return SNew( SBox )
-		.HeightOverride( 250 )
-		.WidthOverride( 300 )
+		.HeightOverride( 250.0f )
+		.WidthOverride( 300.0f )
 		[
 			ContentBrowserModule.Get().CreateAssetPicker( AssetPickerConfig )
 		];
diff --git a/Engine/Source/Editor/UnrealEd/Private/Editor/ActorPositioning.cpp b/Engine/Source/Editor/UnrealEd/Private/Editor/ActorPositioning.cpp
--- a/Engine/Source/Editor/UnrealEd/Private/Editor/ActorPositioning.cpp
+++ b/Engine/Source/Editor/UnrealEd/Private/Editor/ActorPositioning.cpp
@@ -142,7 +142,7 @@ FActorPositionTraceResult FActorPositioning::TraceWorldForPosition(const UWorld&
 
 		for (const FHitResult& Hit : Hits)
 		{
-			const float DistanceToHitSqr = (Hit.ImpactPoint - RayStart).SizeSquared();
+			const double DistanceToHitSqr = (Hit.ImpactPoint - RayStart).SizeSquared();
 			if (DistanceToHitSqr < ClosestHitDistanceSqr)
 			{
 				ClosestHitDistanceSqr = DistanceToHitSqr;
@@ -234,7 +234,7 @@ FTransform FActorPositioning::GetSurfaceAlignedTransform(const FPositioningData&
 	const ULevelEditorViewportSettings* ViewportSettings = GetDefault<ULevelEditorViewportSettings>();
 	const float SnapOffsetExtent = (ViewportSettings->SnapToSurface.bEnabled) ? (ViewportSettings->SnapToSurface.SnapOffsetExtent) : (0.0f);
 	const FVector PlacementExtent = (Data.ActorFactory && !Data.ActorFactory->bUsePlacementExtent) ? FVector::ZeroVector : Data.PlacementExtent;
-	const float CollisionOffsetExtent = FVector::BoxPushOut(Data.SurfaceNormal, PlacementExtent);
+	const double CollisionOffsetExtent = FVector::BoxPushOut(Data.SurfaceNormal, PlacementExtent);
 
 	FVector LocationOffset = Data.SurfaceNormal * FMath::Max(SnapOffsetExtent, CollisionOffsetExtent);
 	if (Data.ActorFactory && LocationOffset.SizeSquared() < Data.ActorFactory->SpawnPositionOffset.SizeSquared())
diff --git a/Engine/Source/Editor/UnrealEd/Private/EditorBuildUtils.cpp b/Engine/Source/Editor/UnrealEd/Private/EditorBuildUtils.cpp
--- a/Engine/Source/Editor/UnrealEd/Private/EditorBuildUtils.cpp
+++ b/Engine/Source/Editor/UnrealEd/Private/EditorBuildUtils.cpp
@@ -1670,7 +1670,7 @@ bool FEditorBuildUtils::EditorBuildVirtualTexture(UWorld* InWorld)
 			return true;
 		}
 
-		FScopedSlowTask BuildTask(Components.Num(), LOCTEXT("VirtualTextureBuild", "Building Virtual Textures"));
+		FScopedSlowTask BuildTask(static_cast<float>(Components.Num()), LOCTEXT("VirtualTextureBuild", "Building Virtual Textures"));
 		BuildTask.MakeDialog(true);
 
 		for (URuntimeVirtualTextureComponent* Component : Components)
diff --git a/Engine/Source/Editor/UnrealEd/Private/EditorComponents.cpp b/Engine/Source/Editor/UnrealEd/Private/EditorComponents.cpp
--- a/Engine/Source/Editor/UnrealEd/Private/EditorComponents.cpp
+++ b/Engine/Source/Editor/UnrealEd/Private/EditorComponents.cpp
@@ -236,8 +236,8 @@ void FGridWidget::DrawNewGrid(const FSceneView* View, FPrimitiveDrawInterface* P
 	{
 		// screenspace size looks better in 2d
 
-		float ScaleX = View->ViewMatrices.GetProjectionMatrix().M[0][0] * View->UnscaledViewRect.Width();
-		float ScaleY = View->ViewMatrices.GetProjectionMatrix().M[1][1] * View->UnscaledViewRect.Height();
+		float ScaleX = static_cast<float>(View->ViewMatrices.GetProjectionMatrix().M[0][0] * View->UnscaledViewRect.Width());
+		float ScaleY = static_cast<float>(View->ViewMatrices.GetProjectionMatrix().M[1][1] * View->UnscaledViewRect.Height());
 
 		float Scale = FMath::Min(ScaleX, ScaleY);
 
@@ -342,7 +342,7 @@ void FGridWidget::DrawNewGrid(const FSceneView* View, FPrimitiveDrawInterface* P
 	FVector UVMax = UVMid + FVector(UVRadi, UVRadi, 0);
 
 	// vertex pos is in -1..1 range
-	DrawPlane10x10(PDI, ObjectToWorld, Radii, FVector2D(UVMin), FVector2D(UVMax), MaterialInst->GetRenderProxy(), SDPG_World );
+	DrawPlane10x10(PDI, ObjectToWorld, static_cast<float>(Radii), FVector2D(UVMin), FVector2D(UVMax), MaterialInst->GetRenderProxy(), SDPG_World );
 }
 
 
@@ -468,11 +468,11 @@ void FEditorCommonDrawHelper::DrawOldGrid(const FSceneView* View,FPrimitiveDrawI
 				LineColor = GridColorMinor;
 			}
 
-			PDI->DrawLine(A,B,LineColor,eDPG, LineThickness, GridDepthBias);
+			PDI->DrawLine(A, B, LineColor, static_cast<uint8>(eDPG), LineThickness, GridDepthBias);
 
 			A.Y=A.X;							B.Y=B.X;
 			A.X=(PerspectiveGridSize/4.f);		B.X=-(PerspectiveGridSize/4.f);
-			PDI->DrawLine(A,B,LineColor,eDPG, LineThickness, GridDepthBias);
+			PDI->DrawLine(A, B, LineColor, static_cast<uint8>(eDPG), LineThickness, GridDepthBias);
 		}
 	}
 	// Draw ortho grid.
@@ -540,7 +540,7 @@ void FEditorCommonDrawHelper::DrawOldGrid(const FSceneView* View,FPrimitiveDrawI
 	// Draw orthogonal worldframe.
 	if(bDrawWorldBox)
 	{
-		DrawWireBox(PDI, FBox(FVector(-MaxGridExtent, -MaxGridExtent, -MaxGridExtent), FVector(MaxGridExtent, MaxGridExtent, MaxGridExtent)), GEngine->C_WorldBox, eDPG);
+		DrawWireBox(PDI, FBox(FVector(-MaxGridExtent, -MaxGridExtent, -MaxGridExtent), FVector(MaxGridExtent, MaxGridExtent, MaxGridExtent)), GEngine->C_WorldBox, static_cast<uint8>(eDPG));
 	}
 }
 
diff --git a/Engine/Source/Editor/UnrealEd/Private/EditorCsg.cpp b/Engine/Source/Editor/UnrealEd/Private/EditorCsg.cpp
--- a/Engine/Source/Editor/UnrealEd/Private/EditorCsg.cpp
+++ b/Engine/Source/Editor/UnrealEd/Private/EditorCsg.cpp
@@ -289,7 +289,7 @@ void UEditorEngine::csgRebuild( UWorld* InWorld )
 	}
 
 	{
-		FScopedSlowTask SlowTask(DynamicBrushes.Num(), NSLOCTEXT("UnrealEd", "RebuildCSGRebuildingDynamicBrushBSPs", "Rebuild CSG: Rebuilding Dynamic Brush BSPs") );
+		FScopedSlowTask SlowTask(static_cast<float>(DynamicBrushes.Num()), NSLOCTEXT("UnrealEd", "RebuildCSGRebuildingDynamicBrushBSPs", "Rebuild CSG: Rebuilding Dynamic Brush BSPs") );
 		for ( int32 BrushIndex = 0; BrushIndex < DynamicBrushes.Num(); BrushIndex++ )
 		{
 			SlowTask.EnterProgressFrame();
diff --git a/Engine/Source/Editor/UnrealEd/Private/EditorModeManager.cpp b/Engine/Source/Editor/UnrealEd/Private/EditorModeManager.cpp
--- a/Engine/Source/Editor/UnrealEd/Private/EditorModeManager.cpp
+++ b/Engine/Source/Editor/UnrealEd/Private/EditorModeManager.cpp
@@ -1536,8 +1536,8 @@ void FEditorModeTools::DrawHUD( FEditorViewportClient* InViewportClient,FViewpor
 							PixelLocation.Y < 0.0f || PixelLocation.Y > View->UnscaledViewRect.Height() * InvDpiScale;
 						if (!bOutside)
 						{
-							const float X = PixelLocation.X - (TextureSizeX / 2);
-							const float Y = PixelLocation.Y - (TextureSizeY / 2);
+							const double X = PixelLocation.X - (TextureSizeX / 2);
+							const double Y = PixelLocation.Y - (TextureSizeY / 2);
 							if (bIsHitTesting)
 							{
 								Canvas->SetHitProxy(new HStaticMeshVert(Actor, Vertex));
diff --git a/Engine/Source/Editor/UnrealEd/Private/EditorViewportClient.cpp b/Engine/Source/Editor/UnrealEd/Private/EditorViewportClient.cpp
--- a/Engine/Source/Editor/UnrealEd/Private/EditorViewportClient.cpp
+++ b/Engine/Source/Editor/UnrealEd/Private/EditorViewportClient.cpp
@@ -125,7 +125,7 @@ void PixelInspectorRealtimeManagement(FEditorViewportClient *CurrentViewport, bo
 namespace EditorViewportClient
 {
 	static const float GridSize = 2048.0f;
-	static const int32 CellSize = 16;
+	static const int8 CellSize = 16;
 	static const float LightRotSpeed = 0.22f;
 }
 
@@ -181,13 +181,13 @@ FViewportCursorLocation::FViewportCursorLocation(const FSceneView* View, FEditor
 	: Origin(ForceInit), Direction(ForceInit), CursorPos(X, Y)
 {
 
-	FVector4 ScreenPos = View->CursorToScreen(X, Y, 0);
+	FVector4 ScreenPos = View->CursorToScreen(static_cast<float>(X), static_cast<float>(Y), 0);
 
 	const FMatrix InvViewMatrix = View->ViewMatrices.GetInvViewMatrix();
 	const FMatrix InvProjMatrix = View->ViewMatrices.GetInvProjectionMatrix();
 
-	const float ScreenX = ScreenPos.X;
-	const float ScreenY = ScreenPos.Y;
+	const double ScreenX = ScreenPos.X;
+	const double ScreenY = ScreenPos.Y;
 
 	ViewportClient = InViewportClient;
 
@@ -684,7 +684,7 @@ void FEditorViewportClient::OnEditorModeIDChanged(const FEditorModeID& EditorMod
 
 float FEditorViewportClient::GetOrthoUnitsPerPixel(const FViewport* InViewport) const
 {
-	const float SizeX = InViewport->GetSizeXY().X;
+	const float SizeX = static_cast<float>(InViewport->GetSizeXY().X);
 
 	// 15.0f was coming from the CAMERA_ZOOM_DIV marco, seems it was chosen arbitrarily
 	return (GetOrthoZoom() / (SizeX * 15.f)) * ComputeOrthoZoomFactor(SizeX);
@@ -2008,8 +2008,8 @@ void FEditorViewportClient::UpdateLightingShowFlags( FEngineShowFlags& InOutShow
 
 bool FEditorViewportClient::CalculateEditorConstrainedViewRect(FSlateRect& OutSafeFrameRect, FViewport* InViewport, float DPIScale)
 {
-	const int32 SizeX = InViewport->GetSizeXY().X / DPIScale;
-	const int32 SizeY = InViewport->GetSizeXY().Y / DPIScale;
+	const float SizeX = InViewport->GetSizeXY().X / DPIScale;
+	const float SizeY = InViewport->GetSizeXY().Y / DPIScale;
 
 	OutSafeFrameRect = FSlateRect(0, 0, SizeX, SizeY);
 	float FixedAspectRatio;
@@ -3317,7 +3317,6 @@ void FEditorViewportClient::DrawAxes(FViewport* InViewport, FCanvas* Canvas, con
 		ViewTM = FRotationMatrix( *InRotation );
 	}
 
-	const int32 SizeX = InViewport->GetSizeXY().X / Canvas->GetDPIScale();
 	const int32 SizeY = InViewport->GetSizeXY().Y / Canvas->GetDPIScale();
 
 	const FIntPoint AxisOrigin( 30, SizeY - 30 );
diff --git a/Engine/Source/Editor/UnrealEd/Private/Factories/ActorFactory.cpp b/Engine/Source/Editor/UnrealEd/Private/Factories/ActorFactory.cpp
--- a/Engine/Source/Editor/UnrealEd/Private/Factories/ActorFactory.cpp
+++ b/Engine/Source/Editor/UnrealEd/Private/Factories/ActorFactory.cpp
@@ -161,7 +161,7 @@ FQuat FindActorAlignmentRotation(const FQuat& InActorRotation, const FVector& In
 		*OutDeltaRotation = FQuat::Identity;
 	}
 
-	const float VectorDot = InWorldNormal | TransformedModelAxis;
+	const double VectorDot = InWorldNormal | TransformedModelAxis;
 	if (1.f - FMath::Abs(VectorDot) <= KINDA_SMALL_NUMBER)
 	{
 		if (VectorDot < 0.f)
@@ -183,7 +183,7 @@ FQuat FindActorAlignmentRotation(const FQuat& InActorRotation, const FVector& In
 		
 		// Find a local 'pitch' axis to rotate around
 		const FVector OrthoPitchAxis = FVector::CrossProduct(PitchReferenceAxis, InModelAxis);
-		const float Pitch = FMath::Acos(PitchReferenceAxis | DestNormalModelSpace) - FMath::Acos(PitchReferenceAxis | InModelAxis);//FMath::Asin(OrthoPitchAxis.Size());
+		const double Pitch = FMath::Acos(PitchReferenceAxis | DestNormalModelSpace) - FMath::Acos(PitchReferenceAxis | InModelAxis);//FMath::Asin(OrthoPitchAxis.Size());
 
 		DeltaRotation = FQuat(OrthoPitchAxis.GetSafeNormal(), Pitch);
 		DeltaRotation.Normalize();
@@ -194,7 +194,7 @@ FQuat FindActorAlignmentRotation(const FQuat& InActorRotation, const FVector& In
 		const float ParallelDotThreshold = 0.98f; // roughly 11.4 degrees (!)
 		if (!FVector::Coincident(InWorldNormal, TransformedModelAxis, ParallelDotThreshold))
 		{
-			const float Yaw = FMath::Atan2(InWorldNormal.X, InWorldNormal.Y) - FMath::Atan2(TransformedModelAxis.X, TransformedModelAxis.Y);
+			const double Yaw = FMath::Atan2(InWorldNormal.X, InWorldNormal.Y) - FMath::Atan2(TransformedModelAxis.X, TransformedModelAxis.Y);
 
 			// Rotation axis for yaw is the Z axis in world space
 			const FVector WorldYawAxis = (InActorRotation * DeltaRotation).Inverse().RotateVector(Z);
diff --git a/Engine/Source/Editor/UnrealEd/Private/Factories/MirrorDataTableFactory.cpp b/Engine/Source/Editor/UnrealEd/Private/Factories/MirrorDataTableFactory.cpp
--- a/Engine/Source/Editor/UnrealEd/Private/Factories/MirrorDataTableFactory.cpp
+++ b/Engine/Source/Editor/UnrealEd/Private/Factories/MirrorDataTableFactory.cpp
@@ -111,7 +111,7 @@ public:
 				.MaxHeight(500)
 				[
 					SNew(SBorder)
-					.Padding(4)
+					.Padding(4.0f)
 					.BorderImage(FAppStyle::GetBrush("ToolPanel.GroupBorder"))
 					[
 						StructViewerModule.CreateStructViewer(Options, FOnStructPicked::CreateSP(this, &SMirrorDataTableFactoryWindow::OnPickedStruct))
diff --git a/Engine/Source/Editor/UnrealEd/Private/Fbx/FbxMaterialConflictWindow.cpp b/Engine/Source/Editor/UnrealEd/Private/Fbx/FbxMaterialConflictWindow.cpp
--- a/Engine/Source/Editor/UnrealEd/Private/Fbx/FbxMaterialConflictWindow.cpp
+++ b/Engine/Source/Editor/UnrealEd/Private/Fbx/FbxMaterialConflictWindow.cpp
@@ -81,7 +81,7 @@ void SFbxMaterialConflictWindow::Construct(const FArguments& InArgs)
 						SNew(SVerticalBox)
 						+ SVerticalBox::Slot()
 						.FillHeight(1.0f)
-						.Padding(2)
+						.Padding(2.0f)
 						[
 							SNew(SVerticalBox)
 							+ SVerticalBox::Slot()
@@ -131,7 +131,7 @@ TSharedPtr<SWidget> SFbxMaterialConflictWindow::ConstructMaterialComparison()
 	FText MaterialCompareInstruction = bIsPreviewConflict ? LOCTEXT("SFbxMaterialConflictWindow_MaterialCompareDocumentationPreview", "Material conflict preview mode") : LOCTEXT("SFbxMaterialConflictWindow_MaterialCompareDocumentation", "To fix a material match, right click on the reimport asset material.");
 	FText MaterialCompareInstructionTooltip = bIsPreviewConflict ? LOCTEXT("SFbxMaterialConflictWindow_MaterialCompareDocumentationPreviewTooltip", "This is only a conflict preview, the material conflict dialog will show up during import to allow you to fix those conflicts.") : LOCTEXT("SFbxMaterialConflictWindow_MaterialCompareDocumentationTooltip", "To fix a material match, right click on the reimport asset material.\nUse the [Clear] option in the context menu to clear a material match.");
 	return SNew(SBox)
-	.MaxDesiredHeight(500)
+	.MaxDesiredHeight(500.0f)
 	[
 		SNew(SBorder)
 		.Padding(FMargin(3))
diff --git a/Engine/Source/Editor/UnrealEd/Private/Fbx/FbxSceneImportFactory.cpp b/Engine/Source/Editor/UnrealEd/Private/Fbx/FbxSceneImportFactory.cpp
--- a/Engine/Source/Editor/UnrealEd/Private/Fbx/FbxSceneImportFactory.cpp
+++ b/Engine/Source/Editor/UnrealEd/Private/Fbx/FbxSceneImportFactory.cpp
@@ -388,8 +388,8 @@ void FetchFbxCameraInScene(UnFbx::FFbxImporter *FbxImporter, FbxNode* ParentNode
 			}
 			CameraInfo->UniqueId = CameraAttribute->GetUniqueID();
 
-			float FieldOfView;
-			float FocalLength;
+			FbxDouble FieldOfView;
+			FbxDouble FocalLength;
 
 			if (CameraAttribute->GetApertureMode() == FbxCamera::eFocalLength)
 			{
@@ -402,16 +402,16 @@ void FetchFbxCameraInScene(UnFbx::FFbxImporter *FbxImporter, FbxNode* ParentNode
 				FocalLength = CameraAttribute->ComputeFocalLength(FieldOfView);
 			}
 
-			CameraInfo->AspectWidth = CameraAttribute->AspectWidth.Get();
-			CameraInfo->AspectHeight = CameraAttribute->AspectHeight.Get();
-			CameraInfo->NearPlane = CameraAttribute->NearPlane.Get();
-			CameraInfo->FarPlane = CameraAttribute->FarPlane.Get();
-			CameraInfo->ProjectionPerspective = CameraAttribute->ProjectionType.Get() == FbxCamera::ePerspective;
-			CameraInfo->OrthoZoom = CameraAttribute->OrthoZoom.Get();
-			CameraInfo->FieldOfView = FieldOfView;
-			CameraInfo->FocalLength = FocalLength;
-			CameraInfo->ApertureWidth = CameraAttribute->GetApertureWidth();
-			CameraInfo->ApertureHeight = CameraAttribute->GetApertureHeight();
+			CameraInfo->AspectWidth				= static_cast<float>(CameraAttribute->AspectWidth.Get());
+			CameraInfo->AspectHeight			= static_cast<float>(CameraAttribute->AspectHeight.Get());
+			CameraInfo->NearPlane				= static_cast<float>(CameraAttribute->NearPlane.Get());
+			CameraInfo->FarPlane				= static_cast<float>(CameraAttribute->FarPlane.Get());
+			CameraInfo->ProjectionPerspective	= CameraAttribute->ProjectionType.Get() == FbxCamera::ePerspective;
+			CameraInfo->OrthoZoom				= static_cast<float>(CameraAttribute->OrthoZoom.Get());
+			CameraInfo->FieldOfView				= static_cast<float>(FieldOfView);
+			CameraInfo->FocalLength				= static_cast<float>(FocalLength);
+			CameraInfo->ApertureWidth			= static_cast<float>(CameraAttribute->GetApertureWidth());
+			CameraInfo->ApertureHeight			= static_cast<float>(CameraAttribute->GetApertureHeight());
 			SceneInfoPtr->CameraInfo.Add(CameraInfo->UniqueId, CameraInfo);
 		}
 	}
diff --git a/Engine/Source/Editor/UnrealEd/Private/Fbx/FbxSkeletalMeshImport.cpp b/Engine/Source/Editor/UnrealEd/Private/Fbx/FbxSkeletalMeshImport.cpp
--- a/Engine/Source/Editor/UnrealEd/Private/Fbx/FbxSkeletalMeshImport.cpp
+++ b/Engine/Source/Editor/UnrealEd/Private/Fbx/FbxSkeletalMeshImport.cpp
@@ -2595,7 +2595,7 @@ USkeletalMesh* UnFbx::FFbxImporter::ReimportSkeletalMesh(USkeletalMesh* Mesh, UF
 
 		int32 LODIndex;
 		int32 SuccessfulLodIndex = 0;
-		float ProgressStep = 90.0 / MaxLODLevel;
+		float ProgressStep = static_cast<float>(90.0f / MaxLODLevel);
 		for (LODIndex = 0; LODIndex < MaxLODLevel; LODIndex++)
 		{
 			SlowTask.EnterProgressFrame(ProgressStep);
diff --git a/Engine/Source/Editor/UnrealEd/Private/Fbx/FbxUtilsImport.cpp b/Engine/Source/Editor/UnrealEd/Private/Fbx/FbxUtilsImport.cpp
--- a/Engine/Source/Editor/UnrealEd/Private/Fbx/FbxUtilsImport.cpp
+++ b/Engine/Source/Editor/UnrealEd/Private/Fbx/FbxUtilsImport.cpp
@@ -173,7 +173,12 @@ FbxAMatrix FFbxDataConverter::ConvertMatrix(const FMatrix& UEMatrix)
 FColor FFbxDataConverter::ConvertColor(FbxDouble3 Color)
 {
 	//Fbx is in linear color space
-	FColor SRGBColor = FLinearColor(Color[0], Color[1], Color[2]).ToFColor(true);
+	FColor SRGBColor = 
+		FLinearColor(
+			static_cast<float>(Color[0]),
+			static_cast<float>(Color[1]),
+			static_cast<float>(Color[2])
+		).ToFColor(true);
 	return SRGBColor;
 }
 
diff --git a/Engine/Source/Editor/UnrealEd/Private/Fbx/ReimportFbxSceneFactory.cpp b/Engine/Source/Editor/UnrealEd/Private/Fbx/ReimportFbxSceneFactory.cpp
--- a/Engine/Source/Editor/UnrealEd/Private/Fbx/ReimportFbxSceneFactory.cpp
+++ b/Engine/Source/Editor/UnrealEd/Private/Fbx/ReimportFbxSceneFactory.cpp
@@ -1342,7 +1342,7 @@ EReimportResult::Type UReimportFbxSceneFactory::ReimportSkeletalMesh(void* VoidF
 			if (SortedLinks.Num() != 0)
 			{
 				//Find the number of take
-				int32 ResampleRate = DEFAULT_SAMPLERATE;
+				int32 ResampleRate = static_cast<int32>(DEFAULT_SAMPLERATE);
 				if (GlobalImportSettings->bResample)
 				{
 					if(FbxImporter->ImportOptions->ResampleRate > 0)
@@ -1412,7 +1412,7 @@ EReimportResult::Type UReimportFbxSceneFactory::ReimportSkeletalMesh(void* VoidF
 						}
 						
 						//Get the sequence timespan
-						ResampleRate = DEFAULT_SAMPLERATE;
+						ResampleRate = static_cast<int32>(DEFAULT_SAMPLERATE);
 						if (FbxImporter->ImportOptions->bResample)
 						{
 							if(FbxImporter->ImportOptions->ResampleRate > 0)
diff --git a/Engine/Source/Editor/UnrealEd/Private/Fbx/SFbxSceneOptionWindow.cpp b/Engine/Source/Editor/UnrealEd/Private/Fbx/SFbxSceneOptionWindow.cpp
--- a/Engine/Source/Editor/UnrealEd/Private/Fbx/SFbxSceneOptionWindow.cpp
+++ b/Engine/Source/Editor/UnrealEd/Private/Fbx/SFbxSceneOptionWindow.cpp
@@ -112,7 +112,7 @@ TSharedRef<SDockTab> SFbxSceneOptionWindow::SpawnSceneTab(const FSpawnTabArgs& A
 				.AutoHeight()
 				[
 					SNew(SUniformGridPanel)
-					.SlotPadding(2)
+					.SlotPadding(2.0f)
 					+ SUniformGridPanel::Slot(0, 0)
 					[
 						SNew(SHorizontalBox)
@@ -318,7 +318,7 @@ TSharedRef<SDockTab> SFbxSceneOptionWindow::SpawnSkeletalMeshReimportTab(const F
 					.AutoHeight()
 					[
 						SNew(SUniformGridPanel)
-						.SlotPadding(2)
+						.SlotPadding(2.0f)
 						+ SUniformGridPanel::Slot(0, 0)
 						.VAlign(VAlign_Center)
 						.HAlign(HAlign_Left)
diff --git a/Engine/Source/Editor/UnrealEd/Private/FeedbackContextEditor.cpp b/Engine/Source/Editor/UnrealEd/Private/FeedbackContextEditor.cpp
--- a/Engine/Source/Editor/UnrealEd/Private/FeedbackContextEditor.cpp
+++ b/Engine/Source/Editor/UnrealEd/Private/FeedbackContextEditor.cpp
@@ -108,7 +108,7 @@ public:
 				.AutoHeight()
 				[
 					SNew(SBox)
-					.HeightOverride(MainBarHeight)
+					.HeightOverride(static_cast<float>(MainBarHeight))
 					[
 						SNew(SProgressBar)
 						.BorderPadding(FVector2D::ZeroVector)
@@ -551,7 +551,7 @@ void FFeedbackContextEditor::ProgressReported( const float TotalProgressInterp,
 				BuildProgressWidget->SetBuildStatusText(DisplayMessage);
 			}
 
-			BuildProgressWidget->SetBuildProgressPercent(TotalProgressInterp * 100, 100);
+			BuildProgressWidget->SetBuildProgressPercent(static_cast<int32>(TotalProgressInterp * 100), 100);
 			TickSlate(BuildProgressWindow.Pin());
 		}
 		else if (SlowTaskWindow.IsValid())
diff --git a/Engine/Source/Editor/UnrealEd/Private/FileHelpers.cpp b/Engine/Source/Editor/UnrealEd/Private/FileHelpers.cpp
--- a/Engine/Source/Editor/UnrealEd/Private/FileHelpers.cpp
+++ b/Engine/Source/Editor/UnrealEd/Private/FileHelpers.cpp
@@ -1618,10 +1618,10 @@ bool FEditorFileUtils::AddCheckoutPackageItems(bool bCheckDirty, TArray<UPackage
 		if (SourceControlCheckPackages.Num())
 		{
 			// Update the source control status of all potentially relevant packages
-			FScopedSlowTask SlowTask(SourceControlCheckPackages.Num(), LOCTEXT("UpdatingSourceControlStatus", "Updating source control status..."));
+			FScopedSlowTask SlowTask(static_cast<float>(SourceControlCheckPackages.Num()), LOCTEXT("UpdatingSourceControlStatus", "Updating source control status..."));
 			SlowTask.MakeDialogDelayed(0.5f);
 			SourceControlProvider.Execute(ISourceControlOperation::Create<FUpdateStatus>(), SourceControlCheckPackages);
-			SlowTask.EnterProgressFrame(SourceControlCheckPackages.Num());
+			SlowTask.EnterProgressFrame(static_cast<float>(SourceControlCheckPackages.Num()));
 		}
 	}
 
@@ -2116,10 +2116,10 @@ ECommandResult::Type FEditorFileUtils::CheckoutPackages(const TArray<UPackage*>&
 	// Attempt to check out each package the user specified to be checked out that is not read only
 	if(FinalPackageCheckoutList.Num() > 0)
 	{
-		FScopedSlowTask SlowTask(FinalPackageCheckoutList.Num(), LOCTEXT("CheckingOutPackages", "Checking out packages..."));
+		FScopedSlowTask SlowTask(static_cast<float>(FinalPackageCheckoutList.Num()), LOCTEXT("CheckingOutPackages", "Checking out packages..."));
 		SlowTask.MakeDialog();
 		CheckOutResult = SourceControlProvider.Execute(ISourceControlOperation::Create<FCheckOut>(), FinalPackageCheckoutList);
-		SlowTask.EnterProgressFrame(FinalPackageCheckoutList.Num());
+		SlowTask.EnterProgressFrame(static_cast<float>(FinalPackageCheckoutList.Num()));
 	}
 
 	// Attempt to mark for add each package the user specified that is not already tracked by source control
@@ -3167,7 +3167,7 @@ EAutosaveContentPackagesResult::Type FEditorFileUtils::AutosaveContentPackagesEx
 		}
 	}
 
-	FScopedSlowTask SlowTask(PackagesToSave.Num()*2, LOCTEXT("PerformingAutoSave_Caption", "Auto-saving out of date packages..."));
+	FScopedSlowTask SlowTask(static_cast<float>(PackagesToSave.Num() * 2), LOCTEXT("PerformingAutoSave_Caption", "Auto-saving out of date packages..."));
 
 	for (UPackage* CurPackage : PackagesToSave)
 	{
@@ -4004,7 +4004,7 @@ FEditorFileUtils::EPromptReturnCode InternalPromptForCheckoutAndSave(const TArra
 	PackagesToSave.Reserve(FinalSaveList.Num());
 
 	{
-		FScopedSlowTask SlowTask(FinalSaveList.Num() * 2, NSLOCTEXT("UnrealEd", "SavingPackagesE", "Saving packages..."));
+		FScopedSlowTask SlowTask(static_cast<float>(FinalSaveList.Num() * 2), NSLOCTEXT("UnrealEd", "SavingPackagesE", "Saving packages..."));
 		SlowTask.MakeDialog();
 
 		UWorld* ActorsWorld = nullptr;
@@ -4047,7 +4047,7 @@ FEditorFileUtils::EPromptReturnCode InternalPromptForCheckoutAndSave(const TArra
 		if (PackagesToClean.Num() > 0)
 		{
 			ObjectTools::CleanupAfterSuccessfulDelete(PackagesToClean, true);
-			SlowTask.EnterProgressFrame(PackagesToClean.Num());
+			SlowTask.EnterProgressFrame(static_cast<float>(PackagesToClean.Num()));
 		}
 
 		PrepareSavePackages(PackagesToSave);
diff --git a/Engine/Source/Editor/UnrealEd/Private/LevelEditorViewport.cpp b/Engine/Source/Editor/UnrealEd/Private/LevelEditorViewport.cpp
--- a/Engine/Source/Editor/UnrealEd/Private/LevelEditorViewport.cpp
+++ b/Engine/Source/Editor/UnrealEd/Private/LevelEditorViewport.cpp
@@ -1269,8 +1269,8 @@ static FActorPositionTraceResult TraceForPositionOn2DLayer(const FViewportCursor
 	}
 
 	FActorPositionTraceResult Result;
-	const float Numerator = FVector::DotProduct(PlaneCenter - Cursor.GetOrigin(), PlaneNormal);
-	const float Denominator = FVector::DotProduct(PlaneNormal, Cursor.GetDirection());
+	const double Numerator = FVector::DotProduct(PlaneCenter - Cursor.GetOrigin(), PlaneNormal);
+	const double Denominator = FVector::DotProduct(PlaneNormal, Cursor.GetDirection());
 	if (FMath::Abs(Denominator) < SMALL_NUMBER)
 	{
 		Result.State = FActorPositionTraceResult::Failed;
@@ -1279,7 +1279,7 @@ static FActorPositionTraceResult TraceForPositionOn2DLayer(const FViewportCursor
 	{
 		Result.State = FActorPositionTraceResult::HitSuccess;
 		Result.SurfaceNormal = PlaneNormal;
-		float D = Numerator / Denominator;
+		double D = Numerator / Denominator;
 		Result.Location = Cursor.GetOrigin() + D * Cursor.GetDirection();
 	}
 
@@ -2558,7 +2558,7 @@ void FLevelEditorViewportClient::UpdateViewForLockedActor(float DeltaTime)
 				}
 			}
 
-			const float DistanceToCurrentLookAt = FVector::Dist( GetViewLocation() , GetLookAtLocation() );
+			const double DistanceToCurrentLookAt = FVector::Dist( GetViewLocation() , GetLookAtLocation() );
 
 			const FQuat CameraOrientation = FQuat::MakeFromEuler( GetViewRotation().Euler() );
 			FVector Direction = CameraOrientation.RotateVector( FVector(1,0,0) );
@@ -2658,7 +2658,7 @@ void FLevelEditorViewportClient::ProjectActorsIntoWorld(const TArray<AActor*>& A
 			if (SceneView->WorldToPixel(NewActorPosition, ScreenPos) && FMath::IsWithin<float>(ScreenPos.X, 0, ViewportSize.X) && FMath::IsWithin<float>(ScreenPos.Y, 0, ViewportSize.Y))
 			{
 				bIsOnScreen = true;
-				Cursor = FViewportCursorLocation(SceneView, this, ScreenPos.X, ScreenPos.Y);
+				Cursor = FViewportCursorLocation(SceneView, this, static_cast<int32>(ScreenPos.X), static_cast<int32>(ScreenPos.Y));
 			}
 		}
 
@@ -3383,7 +3383,7 @@ void FLevelEditorViewportClient::NudgeSelectedObjects( const struct FInputEventS
 
 		bWidgetAxisControlledByDrag = false;
 		Widget->SetCurrentAxis( VirtualAxis );
-		MouseDeltaTracker->AddDelta( this, VirtualKey , VirtualDelta, 1 );
+		MouseDeltaTracker->AddDelta( this, VirtualKey, static_cast<int32>(VirtualDelta), 1 );
 		Widget->SetCurrentAxis( VirtualAxis );
 		UpdateMouseDelta();
 		InViewport->SetMouse( StartMousePos.X , StartMousePos.Y );
@@ -4549,7 +4549,7 @@ void FLevelEditorViewportClient::DrawBrushDetails(const FSceneView* View, FPrimi
 			{
 				// Allocate the material proxy and register it so it can be deleted properly once the rendering is done with it.
 				FColor VolumeColor = Brush->GetWireColor();
-				VolumeColor.A = Brush->ShadedVolumeOpacityValue * 255.0f;
+				VolumeColor.A = static_cast<uint8>(Brush->ShadedVolumeOpacityValue * 255.0f);
 				FDynamicColoredMaterialRenderProxy* MaterialProxy = new FDynamicColoredMaterialRenderProxy(GEngine->GeomMaterial->GetRenderProxy(), VolumeColor);
 				PDI->RegisterDynamicResource(MaterialProxy);
 
@@ -4593,7 +4593,8 @@ void FLevelEditorViewportClient::DrawBrushDetails(const FSceneView* View, FPrimi
 							const FVector3f& PolyVertex = poly->Vertices[VertexIndex];
 							const FVector WorldLocation = FVector(BrushTransform.TransformPosition((FVector)PolyVertex));
 
-							const float Scale = View->WorldToScreen(WorldLocation).W * (4.0f / View->UnscaledViewRect.Width() / View->ViewMatrices.GetProjectionMatrix().M[0][0]);
+							const float Scale =
+								static_cast<float>(View->WorldToScreen(WorldLocation).W * (4.0f / View->UnscaledViewRect.Width() / View->ViewMatrices.GetProjectionMatrix().M[0][0]));
 
 							const FColor Color(Brush->GetWireColor());
 							PDI->SetHitProxy(new HBSPBrushVert(Brush, &poly->Vertices[VertexIndex]));
@@ -4703,7 +4704,7 @@ void FLevelEditorViewportClient::SetupViewForRendering( FSceneViewFamily& ViewFa
 
 		if(bEnableColorScaling)
 		{
-				View.ColorScale = FLinearColor(ColorScale.X,ColorScale.Y,ColorScale.Z);
+			View.ColorScale = FLinearColor(FVector3f{ ColorScale });
 		}
 	}
 
diff --git a/Engine/Source/Editor/UnrealEd/Private/Lightmass/ImportVolumetricLightmap.cpp b/Engine/Source/Editor/UnrealEd/Private/Lightmass/ImportVolumetricLightmap.cpp
--- a/Engine/Source/Editor/UnrealEd/Private/Lightmass/ImportVolumetricLightmap.cpp
+++ b/Engine/Source/Editor/UnrealEd/Private/Lightmass/ImportVolumetricLightmap.cpp
@@ -339,7 +339,8 @@ void FilterWithNeighbors(
 
 											if (!NeighborVoxelImportData.bInsideGeometry && !NeighborVoxelImportData.bBorderVoxel)
 											{
-												const float Weight = 1.0f / FMath::Max<float>(FMath::Abs(NeighborX) + FMath::Abs(NeighborY) + FMath::Abs(NeighborZ), .5f);
+												const float NeighborTotal = static_cast<float>(FMath::Abs(NeighborX) + FMath::Abs(NeighborY) + FMath::Abs(NeighborZ));
+												const float Weight = 1.0f / FMath::Max<float>(NeighborTotal, .5f);
 												FLinearColor NeighborAmbientVector = FilteredVolumeLookup<FFloat3Packed>(BrickTextureCoordinate, CurrentLevelData.BrickDataDimensions, (const FFloat3Packed*)CurrentLevelData.BrickData.AmbientVector.Data.GetData());
 												AmbientVector += NeighborAmbientVector * Weight;
 
@@ -852,7 +853,7 @@ int32 TrimBricksByInterpolationError(
 				}
 			}
 
-			const float RMSE = FMath::Sqrt((ErrorSquared * InvTotalBrickSize).GetMax());
+			const double RMSE = FMath::Sqrt((ErrorSquared * InvTotalBrickSize).GetMax());
 			const bool bCullBrick = RMSE < VolumetricLightmapSettings.MinBrickError;
 
 			if (bCullBrick)
@@ -885,7 +886,7 @@ int32 TrimBricksForMemoryLimit(
 	TArray<const FImportedVolumetricLightmapBrick*>& HighestDensityBricks = BricksByDepth[VolumetricLightmapSettings.MaxRefinementLevels - 1];
 
 	const uint64 BrickSizeBytes = VoxelSizeBytes * PaddedBrickSize * PaddedBrickSize * PaddedBrickSize;
-	const uint64 MaxBrickBytes = MaximumBrickMemoryMb * 1024 * 1024;
+	const uint64 MaxBrickBytes = static_cast<uint64>(MaximumBrickMemoryMb * 1024 * 1024);
 	check(FMath::DivideAndRoundUp(MaxBrickBytes, BrickSizeBytes) <= 0x7FFFFFFFull);
 	const int32 NumBricksBudgeted = (int32)FMath::DivideAndRoundUp(MaxBrickBytes, BrickSizeBytes);
 	const int32 NumBricksToRemove = FMath::Clamp<int32>(NumBricksBeforeTrimming - NumBricksBudgeted, 0, HighestDensityBricks.Num());
@@ -957,10 +958,10 @@ void BuildIndirectionTexture(
 						const FIntVector IndirectionDestDataCoordinate = Brick.IndirectionTexturePosition + FIntVector(X, Y, Z);
 						const int32 IndirectionDestDataIndex = ((IndirectionDestDataCoordinate.Z * CurrentLevelData.IndirectionTextureDimensions.Y) + IndirectionDestDataCoordinate.Y) * CurrentLevelData.IndirectionTextureDimensions.X + IndirectionDestDataCoordinate.X;
 						uint8* IndirectionVoxelPtr = (uint8*)&CurrentLevelData.IndirectionTexture.Data[IndirectionDestDataIndex * IndirectionTextureDataStride];
-						*(IndirectionVoxelPtr + 0) = BrickLayoutPosition.X;
-						*(IndirectionVoxelPtr + 1) = BrickLayoutPosition.Y;
-						*(IndirectionVoxelPtr + 2) = BrickLayoutPosition.Z;
-						*(IndirectionVoxelPtr + 3) = NumBottomLevelBricks;
+						*(IndirectionVoxelPtr + 0) = static_cast<uint8>(BrickLayoutPosition.X);
+						*(IndirectionVoxelPtr + 1) = static_cast<uint8>(BrickLayoutPosition.Y);
+						*(IndirectionVoxelPtr + 2) = static_cast<uint8>(BrickLayoutPosition.Z);
+						*(IndirectionVoxelPtr + 3) = static_cast<uint8>(NumBottomLevelBricks);
 					}
 				}
 			}
@@ -1534,7 +1535,7 @@ void FLightmassProcessor::ImportVolumetricLightmap()
 		 * Statistics
 		 */
 
-		float ImportTime = FPlatformTime::Seconds() - StartTime;
+		double ImportTime = FPlatformTime::Seconds() - StartTime;
 		UE_LOG(LogVolumetricLightmapImport, Log, TEXT("Imported Volumetric Lightmap in %.3fs"), ImportTime);
 		UE_LOG(LogVolumetricLightmapImport, Log, TEXT("     Indirection Texture %ux%ux%u = %.1fMb"),
 			CurrentLevelData.IndirectionTextureDimensions.X,
@@ -1576,7 +1577,7 @@ void FLightmassProcessor::ImportVolumetricLightmap()
 			TotalNumBricks += BricksAtCurrentDepth.Num();
 		}
 
-		const int32 ActualBrickSizeBytes = BrickDataSize / TotalNumBricks;
+		const uint64 ActualBrickSizeBytes = BrickDataSize / TotalNumBricks;
 
 		FString TrimmedString;
 
@@ -1607,7 +1608,7 @@ void FLightmassProcessor::ImportVolumetricLightmap()
 			const int32 DetailCellsPerCurrentLevelBrick = 1 << ((VolumetricLightmapSettings.MaxRefinementLevels - CurrentDepth) * BrickSizeLog2);
 			const FVector CurrentDepthBrickSize = DetailCellSize * DetailCellsPerCurrentLevelBrick;
 			const TArray<const FImportedVolumetricLightmapBrick*>& BricksAtCurrentDepth = BricksByDepth[CurrentDepth];
-			const float CurrentDepthBrickVolume = CurrentDepthBrickSize.X * CurrentDepthBrickSize.Y * CurrentDepthBrickSize.Z;
+			const double CurrentDepthBrickVolume = CurrentDepthBrickSize.X * CurrentDepthBrickSize.Y * CurrentDepthBrickSize.Z;
 
 			UE_LOG(LogVolumetricLightmapImport, Log, TEXT("         %u: %.1f%% covering %.1f%% of volume"),
 				CurrentDepth,
diff --git a/Engine/Source/Editor/UnrealEd/Private/MouseDeltaTracker.cpp b/Engine/Source/Editor/UnrealEd/Private/MouseDeltaTracker.cpp
--- a/Engine/Source/Editor/UnrealEd/Private/MouseDeltaTracker.cpp
+++ b/Engine/Source/Editor/UnrealEd/Private/MouseDeltaTracker.cpp
@@ -437,7 +437,7 @@ void FMouseDeltaTracker::AddDelta(FEditorViewportClient* InViewportClient, FKey
 		InViewportClient->SetCurrentWidgetAxis( EAxisList::None );
 	}
 
-	FVector Wk = InViewportClient->TranslateDelta( InKey, InDelta, InNudge );
+	FVector Wk = InViewportClient->TranslateDelta( InKey, static_cast<float>(InDelta), InNudge );
 
 	EndScreen += Wk;
 
@@ -465,11 +465,11 @@ void FMouseDeltaTracker::AddDelta(FEditorViewportClient* InViewportClient, FKey
 		{
 			if (bIsRotation)
 			{
-				float Scale = 1.0f;
+				double Scale = 1.0;
 
 				if( InViewportClient->IsOrtho() )
 				{
-					Scale = DEFAULT_ORTHOZOOM / (float)InViewportClient->GetOrthoZoom();
+					Scale = DEFAULT_ORTHOZOOM / InViewportClient->GetOrthoZoom();
 				}
 
 				Wk *= Scale;
diff --git a/Engine/Source/Editor/UnrealEd/Private/PackageAutoSaver.cpp b/Engine/Source/Editor/UnrealEd/Private/PackageAutoSaver.cpp
--- a/Engine/Source/Editor/UnrealEd/Private/PackageAutoSaver.cpp
+++ b/Engine/Source/Editor/UnrealEd/Private/PackageAutoSaver.cpp
@@ -558,7 +558,7 @@ void FPackageAutoSaver::UpdateAutoSaveNotification()
 			}
 			else // defer until the user finishes using pop-up menus or the notification will dismiss them...
 			{
-				ForceMinimumTimeTillAutoSave(LoadingSavingSettings->AutoSaveWarningInSeconds);
+				ForceMinimumTimeTillAutoSave(static_cast<float>(LoadingSavingSettings->AutoSaveWarningInSeconds));
 			}
 		}
 		else
diff --git a/Engine/Source/Editor/UnrealEd/Private/PackageBackup.cpp b/Engine/Source/Editor/UnrealEd/Private/PackageBackup.cpp
--- a/Engine/Source/Editor/UnrealEd/Private/PackageBackup.cpp
+++ b/Engine/Source/Editor/UnrealEd/Private/PackageBackup.cpp
@@ -25,7 +25,7 @@
  */
 struct FBackupFileInfo
 {
-	int32 FileSize;						/** Size of the file */
+	int64 FileSize;						/** Size of the file */
 	FString FileName;					/** Name of the file */
 	FDateTime FileTimeStamp;			/** Timestamp of the file */
 };
@@ -131,7 +131,7 @@ bool FAutoPackageBackup::ShouldBackupPackage( const UPackage& InPackage, FString
 	
 	// If the package passed the initial backup checks, proceed to check more specific conditions
 	// that might disqualify the package from being backed up
-	const int32 FileSizeOfBackup = IFileManager::Get().FileSize( *OutFilename );
+	const int64 FileSizeOfBackup = IFileManager::Get().FileSize( *OutFilename );
 	if ( bShouldBackup )
 	{
 		// Ensure that the size the backup would require is less than that of the maximum allowed
@@ -156,7 +156,7 @@ bool FAutoPackageBackup::ShouldBackupPackage( const UPackage& InPackage, FString
 		FString ExistingFileNameExtension = FPaths::GetExtension(OutFilename);
 
 		bool bFoundExistingBackup = false;
-		int32 DirectorySize = 0;
+		int64 DirectorySize = 0;
 		FDateTime LastBackupTimeStamp = FDateTime::MinValue();
 
 		TArray<FBackupFileInfo> BackupFileInfoArray;
@@ -303,7 +303,7 @@ FString FAutoPackageBackup::GetBackupDirectory()
  * @return	true if the space was successfully provided, false if not or if the requested space was
  *			greater than the max allowed space by the user
  */
-bool FAutoPackageBackup::PerformBackupSpaceMaintenance( TArray<FBackupFileInfo>& InBackupFiles, int32 InSpaceUsed, int32 InSpaceRequired )
+bool FAutoPackageBackup::PerformBackupSpaceMaintenance( TArray<FBackupFileInfo>& InBackupFiles, int64 InSpaceUsed, int64 InSpaceRequired )
 {
 	bool bSpaceFreed = false;
 
diff --git a/Engine/Source/Editor/UnrealEd/Private/PerformanceMonitor.h b/Engine/Source/Editor/UnrealEd/Private/PerformanceMonitor.h
--- a/Engine/Source/Editor/UnrealEd/Private/PerformanceMonitor.h
+++ b/Engine/Source/Editor/UnrealEd/Private/PerformanceMonitor.h
@@ -95,7 +95,8 @@ struct FMovingAverage
 			SampleAverage = Sum / Samples.Num();
 
 			// reset the accumulator and counter
-			CurrentSampleAccumulator = CurrentSampleCount = 0;
+			CurrentSampleAccumulator = 0.0f;
+			CurrentSampleCount = 0;
 			CurrentSampleStartTime = CurrentTime;
 		}
 	}
diff --git a/Engine/Source/Editor/UnrealEd/Private/SClassPickerDialog.cpp b/Engine/Source/Editor/UnrealEd/Private/SClassPickerDialog.cpp
--- a/Engine/Source/Editor/UnrealEd/Private/SClassPickerDialog.cpp
+++ b/Engine/Source/Editor/UnrealEd/Private/SClassPickerDialog.cpp
@@ -295,7 +295,7 @@ TSharedRef<ITableRow> SClassPickerDialog::GenerateListRow(TSharedPtr<FClassPicke
 			.AutoWidth()
 			[
 				SNew(SBox)
-				.WidthOverride(60)
+				.WidthOverride(60.0f)
 				[
 					FEditorClassUtils::GetDocumentationLinkWidget(ItemClass)
 				]
diff --git a/Engine/Source/Editor/UnrealEd/Private/Settings/SettingsClasses.cpp b/Engine/Source/Editor/UnrealEd/Private/Settings/SettingsClasses.cpp
--- a/Engine/Source/Editor/UnrealEd/Private/Settings/SettingsClasses.cpp
+++ b/Engine/Source/Editor/UnrealEd/Private/Settings/SettingsClasses.cpp
@@ -642,8 +642,8 @@ void ULevelEditorPlaySettings::UpdateCustomSafeZones()
 
 FMargin ULevelEditorPlaySettings::CalculateCustomUnsafeZones(TArray<FVector2D>& CustomSafeZoneStarts, TArray<FVector2D>& CustomSafeZoneDimensions, FString& DeviceType, FVector2D PreviewSize)
 {
-	int32 PreviewHeight = PreviewSize.Y;
-	int32 PreviewWidth = PreviewSize.X;
+	double PreviewHeight = PreviewSize.Y;
+	double PreviewWidth = PreviewSize.X;
 	bool bPreviewIsPortrait = PreviewHeight > PreviewWidth;
 	FMargin CustomSafeZoneOverride = FMargin();
 	CustomSafeZoneStarts.Empty();
@@ -703,7 +703,7 @@ FMargin ULevelEditorPlaySettings::CalculateCustomUnsafeZones(TArray<FVector2D>&
 				if (!bAdjustsToDeviceRotation && ((Orientation.Contains(TEXT("L")) && bPreviewIsPortrait) ||
 					(Orientation.Contains(TEXT("P")) && !bPreviewIsPortrait)))
 				{
-					float Placeholder = Start.X;
+					double Placeholder = Start.X;
 					Start.X = Start.Y;
 					Start.Y = Placeholder;
 
diff --git a/Engine/Source/Editor/UnrealEd/Private/UnrealWidget.cpp b/Engine/Source/Editor/UnrealEd/Private/UnrealWidget.cpp
--- a/Engine/Source/Editor/UnrealEd/Private/UnrealWidget.cpp
+++ b/Engine/Source/Editor/UnrealEd/Private/UnrealWidget.cpp
@@ -168,7 +168,7 @@ void FWidget::ConvertMouseToAxis_Rotate(FVector2D TangentDir, FVector2D DragDir,
 		FSnappingUtils::SnapRotatorToGrid(Rotation);
 
 		// Record delta rotation (used by the widget to render the accumulated delta)
-		CurrentDeltaRotation = Rotation.Roll;
+		CurrentDeltaRotation = static_cast<float>(Rotation.Roll);
 
 		// Use to calculate the new input delta
 		EffectiveDelta = AxisDir * Rotation.Roll;
@@ -187,7 +187,7 @@ void FWidget::ConvertMouseToAxis_Rotate(FVector2D TangentDir, FVector2D DragDir,
 		Rotation = FRotator(FVector2D::DotProduct(AxisDir, DragDir), 0, 0);
 		FSnappingUtils::SnapRotatorToGrid(Rotation);
 
-		CurrentDeltaRotation = Rotation.Pitch;
+		CurrentDeltaRotation = static_cast<float>(Rotation.Pitch);
 		EffectiveDelta = AxisDir * Rotation.Pitch;
 		// Adjust the input delta according to how much rotation was actually applied
 		InOutDelta = FVector(EffectiveDelta.X, -EffectiveDelta.Y, 0.0f);
@@ -203,7 +203,7 @@ void FWidget::ConvertMouseToAxis_Rotate(FVector2D TangentDir, FVector2D DragDir,
 		Rotation = FRotator(0, FVector2D::DotProduct(AxisDir, DragDir), 0);
 		FSnappingUtils::SnapRotatorToGrid(Rotation);
 
-		CurrentDeltaRotation = Rotation.Yaw;
+		CurrentDeltaRotation = static_cast<float>(Rotation.Yaw);
 		EffectiveDelta = AxisDir * Rotation.Yaw;
 		// Adjust the input delta according to how much rotation was actually applied
 		InOutDelta = FVector(EffectiveDelta.X, -EffectiveDelta.Y, 0.0f);
diff --git a/Engine/Source/Editor/UnrealEd/Private/Utils.cpp b/Engine/Source/Editor/UnrealEd/Private/Utils.cpp
--- a/Engine/Source/Editor/UnrealEd/Private/Utils.cpp
+++ b/Engine/Source/Editor/UnrealEd/Private/Utils.cpp
@@ -70,8 +70,8 @@ void HWidgetUtilProxy::CalcVectors(FSceneView* SceneView, const FViewportClick&
 	else
 	{
 		ScreenDir.Normalize();
-		DragDirX = ScreenDir.X;
-		DragDirY = ScreenDir.Y;
+		DragDirX = static_cast<float>(ScreenDir.X);
+		DragDirY = static_cast<float>(ScreenDir.Y);
 	}
 }
 
@@ -89,8 +89,8 @@ void FUnrealEdUtils::DrawWidget(const FSceneView* View,FPrimitiveDrawInterface*
 	const FVector WidgetOrigin = WidgetMatrix.GetOrigin();
 
 	// Calculate size to draw widget so it takes up the same screen space.
-	const float ZoomFactor = FMath::Min<float>(View->ViewMatrices.GetProjectionMatrix().M[0][0], View->ViewMatrices.GetProjectionMatrix().M[1][1]);
-	const float WidgetRadius = View->Project(WidgetOrigin).W * (UnrealEd_WidgetSize / ZoomFactor);
+	const double ZoomFactor = FMath::Min<float>(View->ViewMatrices.GetProjectionMatrix().M[0][0], View->ViewMatrices.GetProjectionMatrix().M[1][1]);
+	const double WidgetRadius = View->Project(WidgetOrigin).W * (UnrealEd_WidgetSize / ZoomFactor);
 
 	// Choose its color. Highlight manipulated axis in yellow.
 	FColor XColor(FColor::Red);
diff --git a/Engine/Source/Editor/UnrealEd/Private/WorkflowOrientedApp/SContentReference.cpp b/Engine/Source/Editor/UnrealEd/Private/WorkflowOrientedApp/SContentReference.cpp
--- a/Engine/Source/Editor/UnrealEd/Private/WorkflowOrientedApp/SContentReference.cpp
+++ b/Engine/Source/Editor/UnrealEd/Private/WorkflowOrientedApp/SContentReference.cpp
@@ -99,7 +99,7 @@ void SContentReference::Construct(const FArguments& InArgs)
 			SNew(SButton)
 			.ButtonStyle( FAppStyle::Get(), "NoBorder" )
 			.OnClicked(this, &SContentReference::OnClickFindButton)
-			.ContentPadding(0)
+			.ContentPadding(0.0f)
 			.Visibility(this, &SContentReference::GetFindButtonVisibility)
 			.ToolTipText(LOCTEXT("Find", "Find in content browser"))
 			[
@@ -224,8 +224,8 @@ TSharedRef<SWidget> SContentReference::MakeAssetPickerMenu()
 	AssetPickerConfig.InitialAssetViewType = InitialAssetViewType;
 
 	return SNew(SBox)
-		.WidthOverride(AssetPickerSizeOverride.Get().X)
-		.HeightOverride(AssetPickerSizeOverride.Get().Y)
+		.WidthOverride(static_cast<float>(AssetPickerSizeOverride.Get().X))
+		.HeightOverride(static_cast<float>(AssetPickerSizeOverride.Get().Y))
 		[
 			ContentBrowserModule.Get().CreateAssetPicker(AssetPickerConfig)
 		];
diff --git a/Engine/Source/Editor/UnrealEd/Private/WorldPartition/WorldPartitionBuilder.cpp b/Engine/Source/Editor/UnrealEd/Private/WorldPartition/WorldPartitionBuilder.cpp
--- a/Engine/Source/Editor/UnrealEd/Private/WorldPartition/WorldPartitionBuilder.cpp
+++ b/Engine/Source/Editor/UnrealEd/Private/WorldPartition/WorldPartitionBuilder.cpp
@@ -228,7 +228,7 @@ bool UWorldPartitionBuilder::Run(UWorld* World, FPackageSourceControlHelper& Pac
 		const FWorldBuilderCellCoord BeginCellCoords = MinCellCoords;
 		const FWorldBuilderCellCoord EndCellCoords = BeginCellCoords + NumCellsIterations;
 
-		auto CanIterateZ = [&BeginCellCoords, &EndCellCoords, LoadingMode](const bool bInResult, const int32 InZ) -> bool
+		auto CanIterateZ = [&BeginCellCoords, &EndCellCoords, LoadingMode](const bool bInResult, const int64 InZ) -> bool
 		{
 			if (LoadingMode == ELoadingMode::IterativeCells2D)
 			{
@@ -238,7 +238,7 @@ bool UWorldPartitionBuilder::Run(UWorld* World, FPackageSourceControlHelper& Pac
 			return bInResult && (InZ < EndCellCoords.Z);
 		};
 
-		const int32 IterationCount = ((LoadingMode == ELoadingMode::IterativeCells2D) ? 1 : NumCellsIterations.Z) * NumCellsIterations.Y * NumCellsIterations.X;
+		const int64 IterationCount = ((LoadingMode == ELoadingMode::IterativeCells2D) ? 1 : NumCellsIterations.Z) * NumCellsIterations.Y * NumCellsIterations.X;
 		int32 IterationIndex = 0;
 
 		UE_LOG(LogWorldPartitionBuilder, Display, TEXT("Iterative Cell Mode"));
diff --git a/Engine/Source/Editor/UnrealEd/Private/WorldPartition/WorldPartitionMiniMapBuilder.cpp b/Engine/Source/Editor/UnrealEd/Private/WorldPartition/WorldPartitionMiniMapBuilder.cpp
--- a/Engine/Source/Editor/UnrealEd/Private/WorldPartition/WorldPartitionMiniMapBuilder.cpp
+++ b/Engine/Source/Editor/UnrealEd/Private/WorldPartition/WorldPartitionMiniMapBuilder.cpp
@@ -107,8 +107,14 @@ bool UWorldPartitionMiniMapBuilder::RunInternal(UWorld* World, const FCellInfo&
 	// World X,Y to minimap X,Y
 	const FVector3d MinimapMin = WorldToMinimap.TransformPosition(ClampedBounds.Min);
 	const FVector3d MinimapMax = WorldToMinimap.TransformPosition(ClampedBounds.Max);
-	const FIntVector2 DstMin(FMath::Max(FMath::Floor(MinimapMin.X), 0), FMath::Max(FMath::Floor(MinimapMin.Y), 0));
-	const FIntVector2 DstMax(FMath::Max(FMath::Floor(MinimapMax.X), 0), FMath::Max(FMath::Floor(MinimapMax.Y), 0));
+
+	int32 XMin = static_cast<int32>(FMath::Max(FMath::Floor(MinimapMin.X), 0));
+	int32 YMin = static_cast<int32>(FMath::Max(FMath::Floor(MinimapMin.Y), 0));
+	int32 XMax = static_cast<int32>(FMath::Max(FMath::Floor(MinimapMax.X), 0));
+	int32 YMax = static_cast<int32>(FMath::Max(FMath::Floor(MinimapMax.Y), 0));
+
+	const FIntVector2 DstMin(XMin, YMin);
+	const FIntVector2 DstMax(XMax, YMax);
 
 	const uint32 CaptureWidthPixels = DstMax.X - DstMin.X;
 	const uint32 CaptureHeightPixels = DstMax.Y - DstMin.Y;
diff --git a/Engine/Source/Editor/UnrealEd/Private/WorldPartition/WorldPartitionNavigationDataBuilder.cpp b/Engine/Source/Editor/UnrealEd/Private/WorldPartition/WorldPartitionNavigationDataBuilder.cpp
--- a/Engine/Source/Editor/UnrealEd/Private/WorldPartition/WorldPartitionNavigationDataBuilder.cpp
+++ b/Engine/Source/Editor/UnrealEd/Private/WorldPartition/WorldPartitionNavigationDataBuilder.cpp
@@ -328,7 +328,13 @@ bool UWorldPartitionNavigationDataBuilder::GenerateNavigationData(UWorldPartitio
 
 	// A DataChunkActor will be generated for each tile touching the generating bounds.
 	const TSubclassOf<APartitionActor>& NavigationDataActorClass = ANavigationDataChunkActor::StaticClass();
-	const FIntRect GeneratingBounds2D(GeneratingBounds.Min.X, GeneratingBounds.Min.Y, GeneratingBounds.Max.X, GeneratingBounds.Max.Y);
+
+	int32 XMin = static_cast<int32>(GeneratingBounds.Min.X);
+	int32 YMin = static_cast<int32>(GeneratingBounds.Min.Y);
+	int32 XMax = static_cast<int32>(GeneratingBounds.Max.X);
+	int32 YMax = static_cast<int32>(GeneratingBounds.Max.Y);
+
+	const FIntRect GeneratingBounds2D(XMin, YMin, XMax, YMax);
 	FActorPartitionGridHelper::ForEachIntersectingCell(NavigationDataActorClass, GeneratingBounds2D, World->PersistentLevel,
 		[&WorldPartition, &ActorCount, World, &ValidNavigationDataChunkActors, &NavDataBounds, this](const UActorPartitionSubsystem::FCellCoord& InCellCoord, const FIntRect& InCellBounds)->bool
 		{
diff --git a/Engine/Source/Editor/UnrealEd/Public/Dialogs/Dialogs.h b/Engine/Source/Editor/UnrealEd/Public/Dialogs/Dialogs.h
--- a/Engine/Source/Editor/UnrealEd/Public/Dialogs/Dialogs.h
+++ b/Engine/Source/Editor/UnrealEd/Public/Dialogs/Dialogs.h
@@ -176,7 +176,7 @@ class SGenericDialogWidget : public SCompoundWidget
 public:
 	SLATE_BEGIN_ARGS( SGenericDialogWidget )
 		: _UseScrollBox(true)
-		, _ScrollBoxMaxHeight(300.0f)
+		, _ScrollBoxMaxHeight(300)
 	{
 	}
 		
diff --git a/Engine/Source/Editor/UnrealEd/Public/PackageBackup.h b/Engine/Source/Editor/UnrealEd/Public/PackageBackup.h
--- a/Engine/Source/Editor/UnrealEd/Public/PackageBackup.h
+++ b/Engine/Source/Editor/UnrealEd/Public/PackageBackup.h
@@ -107,7 +107,7 @@ private:
 	 * @return	true if the space was successfully provided, false if not or if the requested space was
 	 *			greater than the max allowed space by the user
 	 */
-	static bool PerformBackupSpaceMaintenance( TArray<struct FBackupFileInfo>& InBackupFiles, int32 InSpaceUsed, int32 InSpaceRequired );
+	static bool PerformBackupSpaceMaintenance( TArray<struct FBackupFileInfo>& InBackupFiles, int64 InSpaceUsed, int64 InSpaceRequired );
 
 	// Copy constructor, Destructor, and Self-assignment operator all left private and un-implemented in
 	// order to prevent the class from being instantiated
