Fix import into level material assignment picker by removing the filter /Game which should be /All/Game and replace the path filter by boolean to control engine and developper content.
Clear the selection to avoid crash when re-importing a scene blueprint.
#rnx

diff --git a/Engine/Source/Editor/UnrealEd/Private/Fbx/ReimportFbxSceneFactory.cpp b/Engine/Source/Editor/UnrealEd/Private/Fbx/ReimportFbxSceneFactory.cpp
--- a/Engine/Source/Editor/UnrealEd/Private/Fbx/ReimportFbxSceneFactory.cpp
+++ b/Engine/Source/Editor/UnrealEd/Private/Fbx/ReimportFbxSceneFactory.cpp
@@ -958,8 +958,13 @@ UBlueprint *UReimportFbxSceneFactory::UpdateOriginalBluePrint(FString &BluePrint
 	{
 		return nullptr;
 	}
-	//Close all editor that edit this blueprint
-	GEditor->GetEditorSubsystem<UAssetEditorSubsystem>()->CloseAllEditorsForAsset(BluePrint);
+	if (GEditor)
+	{
+		//Close all editor that edit this blueprint
+		GEditor->GetEditorSubsystem<UAssetEditorSubsystem>()->CloseAllEditorsForAsset(BluePrint);
+		//Make sure all editor selection are cleared
+		GEditor->ResetAllSelectionSets();
+	}
 	//Set the import status for the next reimport
 	for (TSharedPtr<FFbxNodeInfo> NodeInfo : SceneInfoPtr->HierarchyInfo)
 	{
diff --git a/Engine/Source/Editor/UnrealEd/Private/Fbx/SSceneMaterialsListView.cpp b/Engine/Source/Editor/UnrealEd/Private/Fbx/SSceneMaterialsListView.cpp
--- a/Engine/Source/Editor/UnrealEd/Private/Fbx/SSceneMaterialsListView.cpp
+++ b/Engine/Source/Editor/UnrealEd/Private/Fbx/SSceneMaterialsListView.cpp
@@ -341,13 +341,15 @@ TSharedPtr<SWidget> SFbxSceneMaterialsListView::OnOpenContextMenu()
 			Config.bAllowNullSelection = false;
 			Config.bFocusSearchBoxWhenOpened = true;
 			Config.bAllowDragging = false;
-			//Thumbnail are not working since we are in a modal dialog
-			Config.bCanShowRealTimeThumbnails = true;
-			// Don't show stuff in Engine
-			Config.Filter.PackagePaths.Add("/Game");
 			Config.Filter.bRecursivePaths = true;
+			Config.bForceShowEngineContent = false;
+			Config.bForceShowPluginContent = false;
+
+			//We have to limit the height in case there is a lot of item, the computed height can overpass the 16 bit integer maximum.
 			TSharedRef<SWidget> Widget =
 				SNew(SBox)
+				.WidthOverride(300)
+				.HeightOverride(300)
 				[
 					ContentBrowserModule.Get().CreateAssetPicker(Config)
 				];
