Fix potential crash by making sure the fbx exporter get cleanup properly after any export.

diff --git a/Engine/Source/Editor/UnrealEd/Private/EditorExporters.cpp b/Engine/Source/Editor/UnrealEd/Private/EditorExporters.cpp
--- a/Engine/Source/Editor/UnrealEd/Private/EditorExporters.cpp
+++ b/Engine/Source/Editor/UnrealEd/Private/EditorExporters.cpp
@@ -89,6 +89,23 @@
 
 DEFINE_LOG_CATEGORY_STATIC(LogEditorExporters, Log, All);
 
+struct FScopedFbxExporterInstance
+{
+	FScopedFbxExporterInstance()
+	{
+		Exporter = UnFbx::FFbxExporter::GetInstance();
+	}
+	~FScopedFbxExporterInstance()
+	{
+		UnFbx::FFbxExporter::DeleteInstance();
+		Exporter = nullptr;
+	}
+
+	UnFbx::FFbxExporter* GetExporter() { return Exporter; }
+private:
+	UnFbx::FFbxExporter* Exporter = nullptr;
+};
+
 /*------------------------------------------------------------------------------
 	UTextBufferExporterTXT implementation.
 ------------------------------------------------------------------------------*/
@@ -1660,22 +1677,22 @@ bool ULevelExporterFBX::ExportBinary( UObject* Object, const TCHAR* Type, FArchi
 
 	GWarn->BeginSlowTask(NSLOCTEXT("UnrealEd", "ExportingLevelToFBX", "Exporting Level To FBX"), true);
 
-	UnFbx::FFbxExporter* Exporter = UnFbx::FFbxExporter::GetInstance();
+	FScopedFbxExporterInstance ScopedExporterInstance;
 	bool ExportCancel = false;
 	if (UFbxExportOption* AutomatedExportOptions = GetAutomatedExportOptionsFbx())
 	{
-		Exporter->SetExportOptionsOverride(AutomatedExportOptions);
+		ScopedExporterInstance.GetExporter()->SetExportOptionsOverride(AutomatedExportOptions);
 	}
 	else
 	{
 		//Show the fbx export dialog options
 		bool ExportAll = false;
-		Exporter->FillExportOptions(false, true, UExporter::CurrentFilename, ExportCancel, ExportAll);
+		ScopedExporterInstance.GetExporter()->FillExportOptions(false, true, UExporter::CurrentFilename, ExportCancel, ExportAll);
 	}
 
 	if (!ExportCancel)
 	{
-		Exporter->CreateDocument();
+		ScopedExporterInstance.GetExporter()->CreateDocument();
 
 		GWarn->StatusUpdate(0, 1, NSLOCTEXT("UnrealEd", "ExportingLevelToFBX", "Exporting Level To FBX"));
 		{
@@ -1683,7 +1700,7 @@ bool ULevelExporterFBX::ExportBinary( UObject* Object, const TCHAR* Type, FArchi
 
 			if (bSelectedOnly)
 			{
-				Exporter->ExportBSP(World->GetModel(), true);
+				ScopedExporterInstance.GetExporter()->ExportBSP(World->GetModel(), true);
 			}
 
 			class FLevelExporterFBXNodeNameAdapter : public INodeNameAdapter
@@ -1724,7 +1741,7 @@ bool ULevelExporterFBX::ExportBinary( UObject* Object, const TCHAR* Type, FArchi
 
 			FLevelExporterFBXNodeNameAdapter NodeNameAdapter;
 
-			Exporter->ExportLevelMesh(Level, bSelectedOnly, NodeNameAdapter);
+			ScopedExporterInstance.GetExporter()->ExportLevelMesh(Level, bSelectedOnly, NodeNameAdapter);
 
 			// Export streaming levels and actors
 			for (int32 CurLevelIndex = 0; CurLevelIndex < World->GetNumLevels(); ++CurLevelIndex)
@@ -1732,14 +1749,14 @@ bool ULevelExporterFBX::ExportBinary( UObject* Object, const TCHAR* Type, FArchi
 				ULevel* CurLevel = World->GetLevel(CurLevelIndex);
 				if (CurLevel != NULL && CurLevel != Level)
 				{
-					Exporter->ExportLevelMesh(CurLevel, bSelectedOnly, NodeNameAdapter);
+					ScopedExporterInstance.GetExporter()->ExportLevelMesh(CurLevel, bSelectedOnly, NodeNameAdapter);
 				}
 			}
 		}
-		Exporter->WriteToFile(*UExporter::CurrentFilename);
+		ScopedExporterInstance.GetExporter()->WriteToFile(*UExporter::CurrentFilename);
 	}
 
-	Exporter->SetExportOptionsOverride(nullptr);
+	ScopedExporterInstance.GetExporter()->SetExportOptionsOverride(nullptr);
 
 	GWarn->EndSlowTask();
 
@@ -2111,10 +2128,10 @@ UStaticMeshExporterFBX::UStaticMeshExporterFBX(const FObjectInitializer& ObjectI
 bool UStaticMeshExporterFBX::ExportBinary( UObject* Object, const TCHAR* Type, FArchive& Ar, FFeedbackContext* Warn, int32 FileIndex, uint32 PortFlags )
 {
 	UStaticMesh* StaticMesh = CastChecked<UStaticMesh>( Object );
-	UnFbx::FFbxExporter* Exporter = UnFbx::FFbxExporter::GetInstance();
+	FScopedFbxExporterInstance ScopedExporterInstance;
 	if (UFbxExportOption* AutomatedExportOptions = GetAutomatedExportOptionsFbx())
 	{
-		Exporter->SetExportOptionsOverride(AutomatedExportOptions);
+		ScopedExporterInstance.GetExporter()->SetExportOptionsOverride(AutomatedExportOptions);
 		SetShowExportOption(false);
 	}
 	else
@@ -2122,7 +2139,7 @@ bool UStaticMeshExporterFBX::ExportBinary( UObject* Object, const TCHAR* Type, F
 		//Show the fbx export dialog options
 		bool ExportAll = GetBatchMode() && !GetShowExportOption();
 		bool ExportCancel = false;
-		Exporter->FillExportOptions(GetBatchMode(), GetShowExportOption(), UExporter::CurrentFilename, ExportCancel, ExportAll);
+		ScopedExporterInstance.GetExporter()->FillExportOptions(GetBatchMode(), GetShowExportOption(), UExporter::CurrentFilename, ExportCancel, ExportAll);
 		if (ExportCancel)
 		{
 			SetCancelBatch(GetBatchMode());
@@ -2132,10 +2149,10 @@ bool UStaticMeshExporterFBX::ExportBinary( UObject* Object, const TCHAR* Type, F
 		SetShowExportOption(!ExportAll);
 	}
 
-	Exporter->CreateDocument();
-	Exporter->ExportStaticMesh(StaticMesh);
-	Exporter->WriteToFile(*UExporter::CurrentFilename);
-	Exporter->SetExportOptionsOverride(nullptr);
+	ScopedExporterInstance.GetExporter()->CreateDocument();
+	ScopedExporterInstance.GetExporter()->ExportStaticMesh(StaticMesh);
+	ScopedExporterInstance.GetExporter()->WriteToFile(*UExporter::CurrentFilename);
+	ScopedExporterInstance.GetExporter()->SetExportOptionsOverride(nullptr);
 
 	return true;
 }
@@ -2158,10 +2175,10 @@ USkeletalMeshExporterFBX::USkeletalMeshExporterFBX(const FObjectInitializer& Obj
 bool USkeletalMeshExporterFBX::ExportBinary( UObject* Object, const TCHAR* Type, FArchive& Ar, FFeedbackContext* Warn, int32 FileIndex, uint32 PortFlags )
 {
 	USkeletalMesh* SkeletalMesh = CastChecked<USkeletalMesh>( Object );
-	UnFbx::FFbxExporter* Exporter = UnFbx::FFbxExporter::GetInstance();
+	FScopedFbxExporterInstance ScopedExporterInstance;
 	if (UFbxExportOption* AutomatedExportOptions = GetAutomatedExportOptionsFbx())
 	{
-		Exporter->SetExportOptionsOverride(AutomatedExportOptions);
+		ScopedExporterInstance.GetExporter()->SetExportOptionsOverride(AutomatedExportOptions);
 		SetShowExportOption(false);
 	}
 	else
@@ -2169,7 +2186,7 @@ bool USkeletalMeshExporterFBX::ExportBinary( UObject* Object, const TCHAR* Type,
 		//Show the fbx export dialog options
 		bool ExportAll = GetBatchMode() && !GetShowExportOption();
 		bool ExportCancel = false;
-		Exporter->FillExportOptions(GetBatchMode(), GetShowExportOption(), UExporter::CurrentFilename, ExportCancel, ExportAll);
+		ScopedExporterInstance.GetExporter()->FillExportOptions(GetBatchMode(), GetShowExportOption(), UExporter::CurrentFilename, ExportCancel, ExportAll);
 		if (ExportCancel)
 		{
 			SetCancelBatch(GetBatchMode());
@@ -2179,10 +2196,10 @@ bool USkeletalMeshExporterFBX::ExportBinary( UObject* Object, const TCHAR* Type,
 		SetShowExportOption(!ExportAll);
 	}
 
-	Exporter->CreateDocument();
-	Exporter->ExportSkeletalMesh(SkeletalMesh);
-	Exporter->WriteToFile(*UExporter::CurrentFilename);
-	Exporter->SetExportOptionsOverride(nullptr);
+	ScopedExporterInstance.GetExporter()->CreateDocument();
+	ScopedExporterInstance.GetExporter()->ExportSkeletalMesh(SkeletalMesh);
+	ScopedExporterInstance.GetExporter()->WriteToFile(*UExporter::CurrentFilename);
+	ScopedExporterInstance.GetExporter()->SetExportOptionsOverride(nullptr);
 
 	return true;
 }
@@ -2209,10 +2226,10 @@ bool UAnimSequenceExporterFBX::ExportBinary( UObject* Object, const TCHAR* Type,
 
 	if (AnimSkeleton && PreviewMesh)
 	{
-		UnFbx::FFbxExporter* Exporter = UnFbx::FFbxExporter::GetInstance();
+		FScopedFbxExporterInstance ScopedExporterInstance;
 		if (UFbxExportOption* AutomatedExportOptions = GetAutomatedExportOptionsFbx())
 		{
-			Exporter->SetExportOptionsOverride(AutomatedExportOptions);
+			ScopedExporterInstance.GetExporter()->SetExportOptionsOverride(AutomatedExportOptions);
 			SetShowExportOption(false);
 		}
 		else
@@ -2220,7 +2237,7 @@ bool UAnimSequenceExporterFBX::ExportBinary( UObject* Object, const TCHAR* Type,
 			//Show the fbx export dialog options
 			bool ExportAll = GetBatchMode() && !GetShowExportOption();
 			bool ExportCancel = false;
-			Exporter->FillExportOptions(GetBatchMode(), GetShowExportOption(), UExporter::CurrentFilename, ExportCancel, ExportAll);
+			ScopedExporterInstance.GetExporter()->FillExportOptions(GetBatchMode(), GetShowExportOption(), UExporter::CurrentFilename, ExportCancel, ExportAll);
 			if (ExportCancel)
 			{
 				SetCancelBatch(GetBatchMode());
@@ -2230,10 +2247,10 @@ bool UAnimSequenceExporterFBX::ExportBinary( UObject* Object, const TCHAR* Type,
 			SetShowExportOption(!ExportAll);
 		}
 
-		Exporter->CreateDocument();
-		Exporter->ExportAnimSequence(AnimSequence, PreviewMesh, Exporter->GetExportOptions()->bExportPreviewMesh);
-		Exporter->WriteToFile( *UExporter::CurrentFilename );
-		Exporter->SetExportOptionsOverride(nullptr);
+		ScopedExporterInstance.GetExporter()->CreateDocument();
+		ScopedExporterInstance.GetExporter()->ExportAnimSequence(AnimSequence, PreviewMesh, ScopedExporterInstance.GetExporter()->GetExportOptions()->bExportPreviewMesh);
+		ScopedExporterInstance.GetExporter()->WriteToFile( *UExporter::CurrentFilename );
+		ScopedExporterInstance.GetExporter()->SetExportOptionsOverride(nullptr);
 
 		return true;
 	}
