Fbx importer: Fix the staticmesh reimport overriding the nanite enable settings from the import option

diff --git a/Engine/Source/Editor/UnrealEd/Private/Fbx/FbxStaticMeshImport.cpp b/Engine/Source/Editor/UnrealEd/Private/Fbx/FbxStaticMeshImport.cpp
--- a/Engine/Source/Editor/UnrealEd/Private/Fbx/FbxStaticMeshImport.cpp
+++ b/Engine/Source/Editor/UnrealEd/Private/Fbx/FbxStaticMeshImport.cpp
@@ -2022,6 +2022,7 @@ UStaticMesh* UnFbx::FFbxImporter::ImportStaticMeshAsSingle(UObject* InParent, TA
 			ExistingBuildSettings.bGenerateLightmapUVs = SrcModel.BuildSettings.bGenerateLightmapUVs;
 			ExistingBuildSettings.DstLightmapIndex = SrcModel.BuildSettings.DstLightmapIndex;
 			MutableExistMeshDataPtr->ExistingLightMapCoordinateIndex = SrcModel.BuildSettings.DstLightmapIndex;
+			MutableExistMeshDataPtr->ExistingNaniteSettings.bEnabled = StaticMesh->NaniteSettings.bEnabled;
 
 			StaticMeshImportUtils::RestoreExistingMeshSettings(ExistMeshDataPtr, InStaticMesh, StaticMesh->LODGroup != NAME_None ? INDEX_NONE : LODIndex);
 		}
