Fix to UPackageTools::UnloadPackages to deal with stale external actor packages when unloading the current world
Remove code in FPluginUtils::UnloadPluginsAssets that closes the current world since UnloadPackages already takes care of it

diff --git a/Engine/Plugins/Developer/PluginUtils/Source/PluginUtils/Private/PluginUtils.cpp b/Engine/Plugins/Developer/PluginUtils/Source/PluginUtils/Private/PluginUtils.cpp
--- a/Engine/Plugins/Developer/PluginUtils/Source/PluginUtils/Private/PluginUtils.cpp
+++ b/Engine/Plugins/Developer/PluginUtils/Source/PluginUtils/Private/PluginUtils.cpp
@@ -7,7 +7,6 @@
 #include "IContentBrowserSingleton.h"
 #include "ContentBrowserModule.h"
 #include "Editor.h"
-#include "FileHelpers.h"
 #include "GameProjectUtils.h"
 #include "Interfaces/IProjectManager.h"
 #include "Interfaces/IPluginManager.h"
@@ -900,21 +899,6 @@ bool FPluginUtils::UnloadPluginsAssets(const TConstArrayView<TSharedRef<IPlugin>
 		return true;
 	}
 
-	// Close the current map if it's in one of the plugins
-	if (UWorld* CurrentWorld = GEditor ? GEditor->GetEditorWorldContext().World() : nullptr)
-	{
-		const FNameBuilder CurrentWorldPath(CurrentWorld->GetPackage()->GetFName());
-		const FStringView CurrentWorldMountPoint = FPathViews::GetMountPointNameFromPath(CurrentWorldPath);
-		if (PluginNames.ContainsByHash(GetTypeHash(CurrentWorldMountPoint), CurrentWorldMountPoint))
-		{
-			UWorld* BlankWorld = UEditorLoadingAndSavingUtils::NewBlankMap(/*bSaveExistingMap*/ false);
-			if (ensure(BlankWorld))
-			{
-				BlankWorld->GetOutermost()->ClearDirtyFlag();
-			}
-		}
-	}
-
 	// Synchronous scan plugins to make sure we find all their assets.
 	IAssetRegistry& AssetRegistry = FModuleManager::LoadModuleChecked<FAssetRegistryModule>(AssetRegistryConstants::ModuleName).Get();
 	AssetRegistry.ScanPathsSynchronous(PluginContentMountPoints, /*bForceRescan=*/ true);
diff --git a/Engine/Source/Editor/UnrealEd/Private/PackageTools.cpp b/Engine/Source/Editor/UnrealEd/Private/PackageTools.cpp
--- a/Engine/Source/Editor/UnrealEd/Private/PackageTools.cpp
+++ b/Engine/Source/Editor/UnrealEd/Private/PackageTools.cpp
@@ -48,7 +48,6 @@
 #include "Engine/BlueprintGeneratedClass.h"
 #include "Engine/GameEngine.h"
 #include "Engine/LevelStreaming.h"
-#include "Engine/MapBuildDataRegistry.h"
 #include "Engine/Selection.h"
 #include "AssetRegistry/IAssetRegistry.h"
 #include "Logging/MessageLog.h"
@@ -383,35 +382,28 @@ UPackageTools::UPackageTools(const FObjectInitializer& ObjectInitializer)
 		{
 			if (UWorld* EditorWorld = GEditor->GetEditorWorldContext().World())
 			{
-				// Is the currently loaded world being unloaded? If so, we just reset the current world.
-				// We also need to skip the build data package as that will also be destroyed by the call to CreateNewMapForEditing.
-				if (PackagesToUnload.Contains(EditorWorld->GetOutermost()))
+				// Is the current world being unloaded?
+				if (PackagesToUnload.Contains(EditorWorld->GetPackage()))
 				{
-					// Remove the world package from the unload list
-					PackagesToUnload.Remove(EditorWorld->GetOutermost());
-
-					// Remove the level build data package from the unload list as creating a new map will unload build data for the current world
-					for (int32 LevelIndex = 0; LevelIndex < EditorWorld->GetNumLevels(); ++LevelIndex)
+					TArray<TWeakObjectPtr<UPackage>> WeakPackages;
+					WeakPackages.Reserve(PackagesToUnload.Num());
+					for (UPackage* Package : PackagesToUnload)
 					{
-						ULevel* Level = EditorWorld->GetLevel(LevelIndex);
-						if (Level->MapBuildData)
-						{
-							PackagesToUnload.Remove(Level->MapBuildData->GetOutermost());
-						}
+						WeakPackages.Add(Package);
 					}
 
-					// Remove any streaming levels from the unload list as creating a new map will unload streaming levels for the current world
-					for (ULevelStreaming* EditorStreamingLevel : EditorWorld->GetStreamingLevels())
+					// Unload the current world
+					GEditor->CreateNewMapForEditing();
+
+					// Remove stale entries in PackagesToUnload (unloaded world, level build data, streaming levels, external actors, etc)
+					PackagesToUnload.Reset();
+					for (const TWeakObjectPtr<UPackage>& WeakPackage : WeakPackages)
 					{
-						if (EditorStreamingLevel->IsLevelLoaded())
+						if (UPackage* Package = WeakPackage.Get())
 						{
-							UPackage* EditorStreamingLevelPackage = EditorStreamingLevel->GetLoadedLevel()->GetOutermost();
-							PackagesToUnload.Remove(EditorStreamingLevelPackage);
+							PackagesToUnload.Add(Package);
 						}
 					}
-
-					// Unload the current world
-					GEditor->CreateNewMapForEditing();
 				}
 			}
 		}
