Cooker - LooseCookedPackageWriter - fix null pointer dereference when scanning the previous cook directory if there is a filename with an empty extension.
#rnx

diff --git a/Engine/Source/Editor/UnrealEd/Private/Cooker/LooseCookedPackageWriter.cpp b/Engine/Source/Editor/UnrealEd/Private/Cooker/LooseCookedPackageWriter.cpp
--- a/Engine/Source/Editor/UnrealEd/Private/Cooker/LooseCookedPackageWriter.cpp
+++ b/Engine/Source/Editor/UnrealEd/Private/Cooker/LooseCookedPackageWriter.cpp
@@ -698,11 +698,14 @@ public:
 		{
 			FString Filename(FilenameOrDirectory);
 			FStringView Extension(FPathViews::GetExtension(Filename, true /* bIncludeDot */));
-			const TCHAR* ExtensionStr = Extension.GetData();
-			check(ExtensionStr[Extension.Len()] == '\0'); // IsPackageExtension takes a null-terminated TCHAR; we should have it since GetExtension is from the end of the filename
-			if (FPackageName::IsPackageExtension(ExtensionStr))
+			if (!Extension.IsEmpty())
 			{
-				FoundFiles.Add(Filename);
+				const TCHAR* ExtensionStr = Extension.GetData();
+				check(ExtensionStr[Extension.Len()] == '\0'); // IsPackageExtension takes a null-terminated TCHAR; we should have it since GetExtension is from the end of the filename
+				if (FPackageName::IsPackageExtension(ExtensionStr))
+				{
+					FoundFiles.Add(Filename);
+				}
 			}
 		}
 		return true;
