EditorDomain: When a native class used by a package does not exist, log a warning and create placeholder data for the class rather than refusing to calculate the package digest for any packages using it.
#rnx

diff --git a/Engine/Source/Editor/UnrealEd/Private/EditorDomain/EditorDomainUtils.cpp b/Engine/Source/Editor/UnrealEd/Private/EditorDomain/EditorDomainUtils.cpp
--- a/Engine/Source/Editor/UnrealEd/Private/EditorDomain/EditorDomainUtils.cpp
+++ b/Engine/Source/Editor/UnrealEd/Private/EditorDomain/EditorDomainUtils.cpp
@@ -492,18 +492,6 @@ FClassDigestData* FPrecacheClassDigest::GetRecursive(FName ClassName, bool bAllo
 		}
 	}
 
-	UStruct* Struct = nullptr;
-	if (FPackageName::IsScriptPackage(NameStringBuffer))
-	{
-		Struct = FindObject<UStruct>(nullptr, *NameStringBuffer);
-		if (!Struct)
-		{
-			// If ClassName is native but is not yet loaded then abort and the caller gives an error or retries later
-			ClassDigests.Remove(ClassName);
-			return nullptr;
-		}
-	}
-
 	// Fill in digest data config-driven flags
 	DigestData->EditorDomainUse = EDomainUse::LoadEnabled | EDomainUse::SaveEnabled;
 	DigestData->EditorDomainUse &= ~MapFindRef(GClassBlockedUses, ClassName, EDomainUse::None);
@@ -521,18 +509,39 @@ FClassDigestData* FPrecacheClassDigest::GetRecursive(FName ClassName, bool bAllo
 	}
 
 	// Fill in native-specific digest data, get the ParentName, and if non-native, get the native ancestor struct
+	UStruct* Struct = nullptr;
+	bool bIsNative = FPackageName::IsScriptPackage(NameStringBuffer);
+	if (bIsNative)
+	{
+		Struct = FindObject<UStruct>(nullptr, *NameStringBuffer);
+	}
 	FName ParentName;
-	if (Struct)
+	if (bIsNative)
 	{
-		DigestData->ResolvedClosestNative = RedirectsResolvedName;
-		DigestData->bNative = true;
-		DigestData->SchemaHash = Struct->GetSchemaHash(false /* bSkipEditorOnly */);
-		UStruct* ParentStruct = Struct->GetSuperStruct();
-		if (ParentStruct)
+		if (Struct)
 		{
-			NameStringBuffer.Reset();
-			ParentStruct->GetPathName(nullptr, NameStringBuffer);
-			ParentName = FName(*NameStringBuffer);
+			DigestData->ResolvedClosestNative = RedirectsResolvedName;
+			DigestData->bNative = true;
+			DigestData->SchemaHash = Struct->GetSchemaHash(false /* bSkipEditorOnly */);
+			UStruct* ParentStruct = Struct->GetSuperStruct();
+			if (ParentStruct)
+			{
+				NameStringBuffer.Reset();
+				ParentStruct->GetPathName(nullptr, NameStringBuffer);
+				ParentName = FName(*NameStringBuffer);
+			}
+		}
+		else
+		{
+			UE_LOG(LogEditorDomain, Display, TEXT("Class %s is imported by a package but does not exist in memory. EditorDomain keys for packages using it will be invalid if it still exists.")
+				TEXT("\n\tTo clear this message, resave packages that use the deleted class, or load its module earlier than the packages that use it are referenced."),
+				*NameStringBuffer);
+
+			// Create placeholder data that acts as if this class is a leaf off of UObject
+			DigestData->ResolvedClosestNative = FName(*UObject::StaticClass()->GetPathName());
+			DigestData->bNative = true;
+			DigestData->SchemaHash.Reset();
+			ParentName = DigestData->ResolvedClosestNative;
 		}
 	}
 	else
