Make sure auto save is disabled when interchange is working
#rnx

diff --git a/Engine/Source/Editor/UnrealEd/Private/PackageAutoSaver.cpp b/Engine/Source/Editor/UnrealEd/Private/PackageAutoSaver.cpp
--- a/Engine/Source/Editor/UnrealEd/Private/PackageAutoSaver.cpp
+++ b/Engine/Source/Editor/UnrealEd/Private/PackageAutoSaver.cpp
@@ -21,6 +21,7 @@
 #include "EditorModes.h"
 #include "UnrealEdMisc.h"
 #include "FileHelpers.h"
+#include "InterchangeManager.h"
 #include "UnrealEdGlobals.h"
 #include "PackageRestore.h"
 #include "Framework/Notifications/NotificationManager.h"
@@ -424,6 +425,7 @@ bool FPackageAutoSaver::CanAutoSave() const
 	const bool bAreAssetsCompiling = FAssetCompilingManager::Get().GetNumRemainingAssets() > 0;
 	const bool bIsVREditorActive = IVREditorModule::Get().IsVREditorEnabled();	// @todo vreditor: Eventually we should support this while in VR (modal VR progress, with sufficient early warning)
 	const bool bAreAnimationsCompressing = GAsyncCompressedAnimationsTracker ? GAsyncCompressedAnimationsTracker->GetNumRemainingJobs() > 0 : false;
+	const bool bIsInterchangeActive = UInterchangeManager::GetInterchangeManager().IsInterchangeActive();
 
 	bool bIsSequencerPlaying = false;
 	for (FLevelEditorViewportClient* LevelVC : GEditor->GetLevelViewportClients())
@@ -438,7 +440,21 @@ bool FPackageAutoSaver::CanAutoSave() const
 	// query any active editor modes and allow them to prevent autosave
 	const bool bActiveModesAllowAutoSave = GLevelEditorModeTools().CanAutoSave();
 
-	return (bAutosaveEnabled && !bSlowTask && !bPlayWorldValid && !bAnyMenusVisible && !bAutomationTesting && !bIsInteracting && !GIsDemoMode && bHasGameOrProjectLoaded && !bAreShadersCompiling && !bAreAssetsCompiling && !bAreAnimationsCompressing && !bIsVREditorActive && !bIsSequencerPlaying && bActiveModesAllowAutoSave);
+	return (bAutosaveEnabled
+		&& !bSlowTask
+		&& !bPlayWorldValid
+		&& !bAnyMenusVisible
+		&& !bAutomationTesting
+		&& !bIsInteracting
+		&& !GIsDemoMode
+		&& bHasGameOrProjectLoaded
+		&& !bAreShadersCompiling
+		&& !bAreAssetsCompiling
+		&& !bAreAnimationsCompressing
+		&& !bIsVREditorActive
+		&& !bIsSequencerPlaying
+		&& !bIsInterchangeActive
+		&& bActiveModesAllowAutoSave);
 }
 
 bool FPackageAutoSaver::DoPackagesNeedAutoSave() const
diff --git a/Engine/Source/Runtime/Interchange/Engine/Private/InterchangeManager.cpp b/Engine/Source/Runtime/Interchange/Engine/Private/InterchangeManager.cpp
--- a/Engine/Source/Runtime/Interchange/Engine/Private/InterchangeManager.cpp
+++ b/Engine/Source/Runtime/Interchange/Engine/Private/InterchangeManager.cpp
@@ -1463,6 +1463,11 @@ UInterchangeTranslatorBase* UInterchangeManager::GetTranslatorForSourceData(cons
 	return nullptr;
 }
 
+bool UInterchangeManager::IsInterchangeActive()
+{
+	return bIsActive;
+}
+
 bool UInterchangeManager::WarnIfInterchangeIsActive()
 {
 	if (!bIsActive)
diff --git a/Engine/Source/Runtime/Interchange/Engine/Public/InterchangeManager.h b/Engine/Source/Runtime/Interchange/Engine/Public/InterchangeManager.h
--- a/Engine/Source/Runtime/Interchange/Engine/Public/InterchangeManager.h
+++ b/Engine/Source/Runtime/Interchange/Engine/Public/InterchangeManager.h
@@ -473,6 +473,11 @@ public:
 	 */
 	UInterchangeTranslatorBase* GetTranslatorForSourceData(const UInterchangeSourceData* SourceData) const;
 
+	/**
+	 * Return true if the Interchange is active (importing or exporting), return false otherwise.
+	 */
+	bool IsInterchangeActive();
+
 	/**
 	 * Return false if the Interchange is not active (importing or exporting).
 	 * If the interchange is active, it will display a notification to let the user know they can cancel the asynchronous import/export
