Editor: Fix empty notification being shown when all remaining undo/redo transactions are expired.

diff --git a/Engine/Source/Editor/UnrealEd/Private/EditorServer.cpp b/Engine/Source/Editor/UnrealEd/Private/EditorServer.cpp
--- a/Engine/Source/Editor/UnrealEd/Private/EditorServer.cpp
+++ b/Engine/Source/Editor/UnrealEd/Private/EditorServer.cpp
@@ -1221,7 +1221,11 @@ void UEditorEngine::HandleTransactorRedoUndo(const FTransactionContext& Transact
 		}
 	}
 
-	if (!bSquelchTransactionNotification)
+	if (!bSquelchTransactionNotification 
+		// We skip this notification if we failed and there's nothing valid to display. Currently this can occur because
+		// UTransBuffer::Undo/Redo issues broadcasts even when there wasn't a transaction to Undo/Redo, but in these cases
+		// Succeeded is set to false and the context is invalid.
+		&& !(!Succeeded && TransactionContext.Title.IsEmpty() && !TransactionContext.IsValid()))
 	{
 		const FText UndoRedoMessage = WasUndo ? NSLOCTEXT("UnrealEd", "UndoMessageFormat", "Undo: {0}") : NSLOCTEXT("UnrealEd", "RedoMessageFormat", "Redo: {0}");
 		ShowUndoRedoNotification(FText::Format(UndoRedoMessage, TransactionContext.Title), Succeeded);
