Always create a redirector when renaming an asset referenced by external actors (FEditorFileUtils::IsMapPackageAsset will now report external actors as map packages).

diff --git a/Engine/Source/Editor/UnrealEd/Private/FileHelpers.cpp b/Engine/Source/Editor/UnrealEd/Private/FileHelpers.cpp
--- a/Engine/Source/Editor/UnrealEd/Private/FileHelpers.cpp
+++ b/Engine/Source/Editor/UnrealEd/Private/FileHelpers.cpp
@@ -4631,6 +4631,12 @@ bool FEditorFileUtils::IsMapPackageAsset(const FString& ObjectPath, FString& Map
 				MapFilePath = PackagePath;
 				return true;
 			}
+			const FString ExternalPackagePath = FString(TEXT("/Game")) / FPackagePath::GetExternalActorsFolderName();
+			if ( FPackageName::FilenameToLongPackageName(PackagePath).StartsWith(ExternalPackagePath) )
+			{
+				MapFilePath = PackagePath;
+				return true;
+			}
 		}
 	}
 
