Fix unsaved writable files appearing in Uncontrolled Changelist
- Files are now added to the reconcile list when marked writable
- They are later on added to the UC when they are saved (if they still are not checked out)

diff --git a/Engine/Source/Developer/UncontrolledChangelists/Private/UncontrolledChangelistsModule.cpp b/Engine/Source/Developer/UncontrolledChangelists/Private/UncontrolledChangelistsModule.cpp
--- a/Engine/Source/Developer/UncontrolledChangelists/Private/UncontrolledChangelistsModule.cpp
+++ b/Engine/Source/Developer/UncontrolledChangelists/Private/UncontrolledChangelistsModule.cpp
@@ -111,14 +111,45 @@ TArray<FUncontrolledChangelistStateRef> FUncontrolledChangelistsModule::GetChang
 	return UncontrolledChangelistStates;
 }
 
-bool FUncontrolledChangelistsModule::OnMakeWritable(const TArray<FString>& InFilenames)
+bool FUncontrolledChangelistsModule::OnMakeWritable(const FString& InFilename)
 {
 	if (!IsEnabled())
 	{
 		return false;
 	}
 
-	return AddFilesToDefaultUncontrolledChangelist(InFilenames, FUncontrolledChangelistState::ECheckFlags::NotCheckedOut);
+	AddedAssetsCache.Add(FPaths::ConvertRelativePathToFull(InFilename));
+	return true;
+}
+
+bool FUncontrolledChangelistsModule::OnNewFilesAdded(const TArray<FString>& InFilenames)
+{
+	return AddToUncontrolledChangelist(InFilenames);
+}
+
+bool FUncontrolledChangelistsModule::OnSaveWritable(const FString& InFilename)
+{
+	return AddToUncontrolledChangelist({ InFilename });
+}
+
+bool FUncontrolledChangelistsModule::AddToUncontrolledChangelist(const TArray<FString>& InFilenames)
+{
+	if (!IsEnabled())
+	{
+		return false;
+	}
+
+	TArray<FString> FullPaths;
+	FullPaths.Reserve(InFilenames.Num());
+	Algo::Transform(InFilenames, FullPaths, [](const FString& Filename) { return FPaths::ConvertRelativePathToFull(Filename); });
+
+	// Remove from reconcile cache
+	for (const FString& FullPath : FullPaths)
+	{
+		AddedAssetsCache.Remove(FullPath);
+	}
+
+	return AddFilesToDefaultUncontrolledChangelist(FullPaths, FUncontrolledChangelistState::ECheckFlags::NotCheckedOut);
 }
 
 void FUncontrolledChangelistsModule::UpdateStatus()
diff --git a/Engine/Source/Developer/UncontrolledChangelists/Public/UncontrolledChangelistsModule.h b/Engine/Source/Developer/UncontrolledChangelists/Public/UncontrolledChangelistsModule.h
--- a/Engine/Source/Developer/UncontrolledChangelists/Public/UncontrolledChangelistsModule.h
+++ b/Engine/Source/Developer/UncontrolledChangelists/Public/UncontrolledChangelistsModule.h
@@ -43,12 +43,26 @@ public:
 	 */
 	 TArray<FUncontrolledChangelistStateRef> GetChangelistStates() const;
 
-	 /**
-	  * Called when files have been made writable. Adds the files to the Default Uncontrolled Changelist
-	  * @param	InFilenames			The files to be added.
-	  * @return True if the files have been added to the Default Uncontrolled Changelist.
-	  */
-	 bool OnMakeWritable(const TArray<FString>& InFilenames);
+	/**
+	 * Called when file has been made writable. Adds the file to the reconcile list because we don't know yet if it will change.
+	 * @param	InFilename		The file to be reconciled.
+	 * @return True if the file have been handled by the Uncontrolled module.
+	 */
+	bool OnMakeWritable(const FString& InFilename);
+	 	 
+	/**
+	 * Called when file has been saved without an available Provider. Adds the file to the Default Uncontrolled Changelist
+	 * @param	InFilename		The file to be added.
+	 * @return	True if the file have been handled by the Uncontrolled module.
+	 */
+	bool OnSaveWritable(const FString& InFilename);
+
+	/**
+	 * Called when files should have been marked for add without an available Provider. Adds the files to the Default Uncontrolled Changelist
+	 * @param	InFilenames		The files to be added.
+	 * @return	True if the files have been handled by the Uncontrolled module.
+	 */
+	bool OnNewFilesAdded(const TArray<FString>& InFilenames);
 
 	/**
 	 * Updates the status of Uncontrolled Changelists and files.
@@ -156,6 +170,11 @@ private:
 	 */
 	void OnAssetAddedInternal(const FAssetData& InAssetData, TSet<FString>& InAddedAssetsCache, bool bInStartupTask);
 
+	/**
+	 * Add files to Uncontrolled Changelist. Also adds them to files to reconcile.
+	 */
+	bool AddToUncontrolledChangelist(const TArray<FString>& InFilenames);
+
 	/**
 	 * Saves the state of UncontrolledChangelists to Json for persistency.
 	 */
diff --git a/Engine/Source/Editor/UnrealEd/Private/EditorEngine.cpp b/Engine/Source/Editor/UnrealEd/Private/EditorEngine.cpp
--- a/Engine/Source/Editor/UnrealEd/Private/EditorEngine.cpp
+++ b/Engine/Source/Editor/UnrealEd/Private/EditorEngine.cpp
@@ -4319,7 +4319,7 @@ void UEditorEngine::RunDeferredMarkForAddFiles(bool)
 	else if (FUncontrolledChangelistsModule::Get().IsEnabled())
 	{
 		FUncontrolledChangelistsModule& UncontrolledChangelistsModule = FUncontrolledChangelistsModule::Get();
-		UncontrolledChangelistsModule.OnMakeWritable(DeferredFilesToAddToSourceControl);
+		UncontrolledChangelistsModule.OnNewFilesAdded(DeferredFilesToAddToSourceControl);
 	}
 
 	// Clear the list when this run whether source control is active or not, since we do not want to accumulate those if the user is running without source control
diff --git a/Engine/Source/Editor/UnrealEd/Private/FileHelpers.cpp b/Engine/Source/Editor/UnrealEd/Private/FileHelpers.cpp
--- a/Engine/Source/Editor/UnrealEd/Private/FileHelpers.cpp
+++ b/Engine/Source/Editor/UnrealEd/Private/FileHelpers.cpp
@@ -1918,7 +1918,7 @@ bool FEditorFileUtils::PromptToCheckoutPackagesInternal(bool bCheckDirty, const
 						{
 							// Add to PackagesNotToPromptAnyMore only if not added to Uncontrolled Changelist.
 							// If added to Uncontrolled Changelist, we want the checkout prompt to be displayed again if the file is reverted
-							if (!UncontrolledChangelistModule.OnMakeWritable({Filename}))
+							if (!UncontrolledChangelistModule.OnMakeWritable(Filename))
 							{
 								PackagesNotToPromptAnyMore.Add(PackageToMakeWritable->GetName());
 							}
@@ -3438,6 +3438,8 @@ static InternalSavePackageResult InternalSavePackage(UPackage* PackageToSave, bo
 		}
 
 		ISourceControlProvider& SourceControlProvider = ISourceControlModule::Get().GetProvider();
+		FUncontrolledChangelistsModule& UncontrolledChangelistsModule = FUncontrolledChangelistsModule::Get();
+
 		if (ISourceControlModule::Get().IsEnabled())
 		{
 			// Assume the package was correctly checked out from SCC
@@ -3456,19 +3458,13 @@ static InternalSavePackageResult InternalSavePackage(UPackage* PackageToSave, bo
 		}
 		else
 		{
-			FUncontrolledChangelistsModule& UncontrolledChangelistsModule = FUncontrolledChangelistsModule::Get();
-
 			// If we are in offline mode, automatically add the modified package to an Uncontrolled Changelist
-			if (UncontrolledChangelistsModule.IsEnabled() && (!IFileManager::Get().IsReadOnly(*FinalPackageSavePath)))
-			{
-				UncontrolledChangelistsModule.OnMakeWritable({ FinalPackageSavePath });
-				bOutPackageLocallyWritable = true;
-			}
-			else
-			{
-				// If source control is disabled then we don't care if the package is locally writable
-				bOutPackageLocallyWritable = false;
-			}
+			bOutPackageLocallyWritable = UncontrolledChangelistsModule.IsEnabled() && (!IFileManager::Get().IsReadOnly(*FinalPackageSavePath));
+		}
+
+		if (bWasSuccessful && bOutPackageLocallyWritable)
+		{
+			UncontrolledChangelistsModule.OnSaveWritable({ FinalPackageSavePath });
 		}
 
 		// Handle all failures the same way.
