Message boxes with "Yes All" and "No All" options now use an Apply to All checkbox instead of separate buttons.

diff --git a/Engine/Source/Editor/UnrealEd/Private/Dialogs/Dialogs.cpp b/Engine/Source/Editor/UnrealEd/Private/Dialogs/Dialogs.cpp
--- a/Engine/Source/Editor/UnrealEd/Private/Dialogs/Dialogs.cpp
+++ b/Engine/Source/Editor/UnrealEd/Private/Dialogs/Dialogs.cpp
@@ -50,6 +50,7 @@ public:
 		ParentWindow = InArgs._ParentWindow.Get();
 		ParentWindow->SetWidgetToFocusOnActivate(SharedThis(this));
 		Response = EAppReturnType::Cancel;
+		MessageType = InArgs._MessageType.Get();
 
 		FSlateFontInfo MessageFont( FAppStyle::GetFontStyle("StandardDialog.LargeFont"));
 		MyMessage = InArgs._Message;
@@ -103,7 +104,7 @@ public:
 					SNew(SHorizontalBox)
 
 					+ SHorizontalBox::Slot()
-					.FillWidth(1.0f)
+					.AutoWidth()
 					.HAlign(HAlign_Left)
 					.VAlign(VAlign_Center)
 					[
@@ -120,6 +121,24 @@ public:
 						]
 					]
 
+					+ SHorizontalBox::Slot()
+						.FillWidth(1.0f)
+						.HAlign(HAlign_Left)
+						.VAlign(VAlign_Center)
+						.Padding(FMargin(16.f, 0.f, 0.f, 0.f))
+						[
+							SNew(SCheckBox)
+							.IsChecked(ECheckBoxState::Unchecked)
+							.OnCheckStateChanged(this, &SChoiceDialog::OnCheckboxClicked)
+							.Visibility(this, &SChoiceDialog::GetCheckboxVisibility)
+							.ToolTipText(NSLOCTEXT("SChoiceDialog", "ApplyToAllTooltip", "Make your choice of Yes or No apply to all remaining items in the current operation"))
+							[
+								SNew(STextBlock)
+								.WrapTextAt(615.0f)
+								.Text(NSLOCTEXT("SChoiceDialog", "ApplyToAllLabel", "Apply to All"))
+							]
+						]
+
 					+ SHorizontalBox::Slot()
 					.AutoWidth()
 					.HAlign(HAlign_Right)
@@ -159,7 +178,7 @@ public:
 
 
 
-		switch ( InArgs._MessageType.Get() )
+		switch ( MessageType )
 		{	
 		case EAppMsgType::Ok:
 			ADD_SLOT_PRIMARY(Ok)
@@ -185,14 +204,10 @@ public:
 		case EAppMsgType::YesNoYesAllNoAll:
 			ADD_SLOT_PRIMARY(Yes)
 			ADD_SLOT(No)
-			ADD_SLOT(YesAll)
-			ADD_SLOT(NoAll)
 			break;
 		case EAppMsgType::YesNoYesAllNoAllCancel:
 			ADD_SLOT_PRIMARY(Yes)
 			ADD_SLOT(No)
-			ADD_SLOT(YesAll)
-			ADD_SLOT(NoAll)
 			ADD_SLOT(Cancel)
 			break;
 		case EAppMsgType::YesNoYesAll:
@@ -282,6 +297,18 @@ private:
 	FReply HandleButtonClicked( EAppReturnType::Type InResponse )
 	{
 		Response = InResponse;
+		if ((MessageType == EAppMsgType::YesNoYesAllNoAll || MessageType == EAppMsgType::YesNoYesAllNoAllCancel)
+			&& bApplyToAllChecked)
+		{
+			if (Response == EAppReturnType::Yes)
+			{
+				Response = EAppReturnType::YesAll;
+			}
+			else if (Response == EAppReturnType::No)
+			{
+				Response = EAppReturnType::NoAll;
+			}
+		}
 
 		ResultCallback.ExecuteIfBound(ParentWindow.ToSharedRef(), Response);
 
@@ -298,6 +325,22 @@ private:
 		return FReply::Handled();
 	}
 		
+	// Used as a delegate for the OnClicked property of the Apply to All checkbox
+	void OnCheckboxClicked(ECheckBoxState InNewState)
+	{
+		bApplyToAllChecked = InNewState == ECheckBoxState::Checked;
+	}
+
+	// Used as a delegate for the Visibility property of the Apply to All checkbox
+	EVisibility GetCheckboxVisibility() const
+	{
+		if (MessageType == EAppMsgType::YesNoYesAllNoAll || MessageType == EAppMsgType::YesNoYesAllNoAllCancel)
+		{
+			return EVisibility::Visible;
+		}
+		return EVisibility::Hidden;
+	}
+
 public:
 	/** Callback delegate that is triggered, when the dialog is run in non-modal mode */
 	FOnMsgDlgResult ResultCallback;
@@ -307,6 +350,8 @@ private:
 	EAppReturnType::Type Response;
 	TSharedPtr<SWindow> ParentWindow;
 	TAttribute<FText> MyMessage;
+	bool bApplyToAllChecked;
+	EAppMsgType::Type MessageType;
 };
 
 
diff --git a/Engine/Source/Editor/UnrealEd/Private/Factories/Factory.cpp b/Engine/Source/Editor/UnrealEd/Private/Factories/Factory.cpp
--- a/Engine/Source/Editor/UnrealEd/Private/Factories/Factory.cpp
+++ b/Engine/Source/Editor/UnrealEd/Private/Factories/Factory.cpp
@@ -358,7 +358,7 @@ void UFactory::DisplayOverwriteOptionsDialog(const FText& Message)
 	else if (OverwriteYesOrNoToAllState != EAppReturnType::YesAll && OverwriteYesOrNoToAllState != EAppReturnType::NoAll)
 	{
 		OverwriteYesOrNoToAllState = FMessageDialog::Open(EAppMsgType::YesNoYesAllNoAllCancel, FText::Format(
-			NSLOCTEXT("UnrealEd", "ImportedAssetAlreadyExists", "{0} Would you like to overwrite the existing settings?\n\nYes or Yes to All: Overwrite the existing settings.\nNo or No to All: Preserve the existing settings.\nCancel: Abort the operation."),
+			NSLOCTEXT("UnrealEd", "ImportedAssetAlreadyExists", "{0}\n\nWould you also like to overwrite its existing settings?"),
 			Message));
 	}
 }
diff --git a/Engine/Source/Editor/UnrealEd/Private/ObjectTools.cpp b/Engine/Source/Editor/UnrealEd/Private/ObjectTools.cpp
--- a/Engine/Source/Editor/UnrealEd/Private/ObjectTools.cpp
+++ b/Engine/Source/Editor/UnrealEd/Private/ObjectTools.cpp
@@ -2423,7 +2423,7 @@ namespace ObjectTools
 					{
 						FFormatNamedArguments Args;
 						Args.Add(TEXT("Filename"), FText::FromString(PackageFilename));
-						const FText Message = FText::Format(NSLOCTEXT("ObjectTools", "DeleteReadOnlyWarning", "File '{Filename}' is read-only on disk, are you sure you want to delete it?"), Args);
+						const FText Message = FText::Format(NSLOCTEXT("ObjectTools", "DeleteReadOnlyWarning", "This file is read-only on disk:\n\n{Filename}\n\nDelete it anyway?"), Args);
 
 						ReturnType = FMessageDialog::Open(EAppMsgType::YesNoYesAllNoAll, Message);
 						bMakeWritable = ReturnType == EAppReturnType::YesAll;
@@ -2873,7 +2873,7 @@ namespace ObjectTools
 					{
 						FFormatNamedArguments Args;
 						Args.Add(TEXT("Filename"), FText::FromString(PackageFilename));
-						const FText Message = FText::Format(NSLOCTEXT("ObjectTools", "DeleteReadOnlyWarning", "File '{Filename}' is read-only on disk, are you sure you want to delete it?"), Args);
+						const FText Message = FText::Format(NSLOCTEXT("ObjectTools", "DeleteReadOnlyWarning", "This file is read-only on disk:\n\n{Filename}\n\nDelete it anyway?"), Args);
 
 						ReturnType = FMessageDialog::Open(EAppMsgType::YesNoYesAllNoAll, EAppReturnType::No, Message);
 						bMakeWritable = ReturnType == EAppReturnType::YesAll;
