UE interface to manage external content in a project

diff --git a/Engine/Source/Editor/UnrealEd/Classes/Editor/EditorEngine.h b/Engine/Source/Editor/UnrealEd/Classes/Editor/EditorEngine.h
--- a/Engine/Source/Editor/UnrealEd/Classes/Editor/EditorEngine.h
+++ b/Engine/Source/Editor/UnrealEd/Classes/Editor/EditorEngine.h
@@ -76,6 +76,7 @@ class UEditorWorldExtensionManager;
 class ITargetDevice;
 class ULevelEditorDragDropHandler;
 class UTypedElementSelectionSet;
+class IProjectExternalContentInterface;
 
 //
 // Things to set in mapSetBrush.
@@ -2475,6 +2476,16 @@ public:
 	FOnCreateLevelEditorDragDropHandler& OnCreateLevelEditorDragDropHandler() { return OnCreateLevelEditorDragDropHandlerDelegate; }
 	ULevelEditorDragDropHandler* GetLevelEditorDragDropHandler() const;
 
+	/** 
+	 * Gets the interface to manage project references to external content
+	 * @note the returned pointer cannot be null
+	 */
+	IProjectExternalContentInterface* GetProjectExternalContentInterface();
+
+	/** Delegate to override the IProjectExternalContentInterface */
+	DECLARE_DELEGATE_RetVal(IProjectExternalContentInterface*, FProjectExternalContentInterfaceGetter);
+	FProjectExternalContentInterfaceGetter ProjectExternalContentInterfaceGetter;
+
 private:
 	UPROPERTY()
 	mutable TObjectPtr<ULevelEditorDragDropHandler> DragDropHandler;
diff --git a/Engine/Source/Editor/UnrealEd/Private/EditorEngine.cpp b/Engine/Source/Editor/UnrealEd/Private/EditorEngine.cpp
--- a/Engine/Source/Editor/UnrealEd/Private/EditorEngine.cpp
+++ b/Engine/Source/Editor/UnrealEd/Private/EditorEngine.cpp
@@ -253,7 +253,7 @@
 #include "DeviceProfiles/DeviceProfileManager.h"
 #include "Rendering/StaticLightingSystemInterface.h"
 #include "LevelEditorDragDropHandler.h"
-
+#include "IProjectExternalContentInterface.h"
 
 DEFINE_LOG_CATEGORY_STATIC(LogEditor, Log, All);
 
@@ -7863,4 +7863,27 @@ ULevelEditorDragDropHandler* UEditorEngine::GetLevelEditorDragDropHandler() cons
 	return DragDropHandler;
 }
 
-#undef LOCTEXT_NAMESPACE 
+namespace
+{
+	class FProjectExternalContentDefault : public IProjectExternalContentInterface
+	{
+	private:
+		virtual bool IsEnabled() const override { return false; }
+		virtual bool HasExternalContent(const FString& ExternalContentId) const override { return false; }
+		virtual bool IsExternalContentLoaded(const FString& ExternalContentId) const override { return false; }
+		virtual TArray<FString> GetExternalContentIds() const override { return {}; }
+		virtual void AddExternalContent(const FString& ExternalContentId, FAddExternalContentComplete CompleteCallback) override { CompleteCallback.ExecuteIfBound(false, /*Plugins=*/{}); }
+		virtual void RemoveExternalContent(const FString& ExternalContentId, FRemoveExternalContentComplete CompleteCallback) override { CompleteCallback.ExecuteIfBound(false); }
+	};
+
+	FProjectExternalContentDefault ProjectExternalContentDefault;
+}
+
+IProjectExternalContentInterface* UEditorEngine::GetProjectExternalContentInterface()
+{
+	IProjectExternalContentInterface* ProjectExternalContentInterface = ProjectExternalContentInterfaceGetter.IsBound() ? ProjectExternalContentInterfaceGetter.Execute() : &ProjectExternalContentDefault;
+	check(ProjectExternalContentInterface);
+	return ProjectExternalContentInterface;
+}
+
+#undef LOCTEXT_NAMESPACE
diff --git a/Engine/Source/Editor/UnrealEd/Public/IProjectExternalContentInterface.h b/Engine/Source/Editor/UnrealEd/Public/IProjectExternalContentInterface.h
new file mode 100644
--- /dev/null
+++ b/Engine/Source/Editor/UnrealEd/Public/IProjectExternalContentInterface.h
@@ -0,0 +1,54 @@
+// Copyright Epic Games, Inc. All Rights Reserved.
+#pragma once
+
+#include "CoreMinimal.h"
+
+/** Interface to manage project references to external content */
+class IProjectExternalContentInterface
+{
+public:
+	/** Returns whether the project can reference external content */
+	virtual bool IsEnabled() const = 0;
+
+	/**
+     * Returns whether the specified external content is referenced by the project
+	 * @param ExternalContentId External content identifier (verse path, link code, etc)
+	 */
+	virtual bool HasExternalContent(const FString& ExternalContentId) const = 0;
+
+	/** 
+	 * Returns whether the specified external content is loaded (and referenced by the project)
+	 * @param ExternalContentId External content identifier (verse path, link code, etc)
+	 */
+	virtual bool IsExternalContentLoaded(const FString& ExternalContentId) const = 0;
+
+	/** Returns the list of external content referenced by the project (verse paths, link codes, etc) */
+	virtual TArray<FString> GetExternalContentIds() const = 0;
+
+	/**
+	 * Called upon AddExternalContent completion
+	 * @param bSuccess Whether the external content was successfully added to the project
+	 * @param Plugins List of loaded plugins hosting the external content
+	 */
+	DECLARE_DELEGATE_TwoParams(FAddExternalContentComplete, bool /*bSuccess*/, const TArray<TSharedRef<IPlugin>>& /*Plugins*/);
+
+	/**
+	 * Adds a reference to external content to the project and asynchronously downloads/loads the external content
+	 * @param ExternalContentId External content identifier (verse path, link code, etc)
+	 * @param CompleteCallback See FAddExternalContentComplete
+	 */
+	virtual void AddExternalContent(const FString& ExternalContentId, FAddExternalContentComplete CompleteCallback = FAddExternalContentComplete()) = 0;
+
+	/**
+	 * Called upon RemoveExternalContent completion
+	 * @param bSuccess Whether the external content was successfully removed from the project (could be canceled by the user)
+	 */
+	DECLARE_DELEGATE_OneParam(FRemoveExternalContentComplete, bool /*bSuccess*/);
+
+	/**
+	 * Removes a reference to external content from the project and unloads the external content
+	 * @param ExternalContentId External content identifier (verse path, link code, etc)
+	 * @param Params See FRemoveExternalContentParams
+	 */
+	virtual void RemoveExternalContent(const FString& ExternalContentId, FRemoveExternalContentComplete CompleteCallback = FRemoveExternalContentComplete()) = 0;
+};
