Making a group gives it the rotation of the last selected actor

diff --git a/Engine/Source/Editor/UnrealEd/Private/ActorGroupingUtils.cpp b/Engine/Source/Editor/UnrealEd/Private/ActorGroupingUtils.cpp
--- a/Engine/Source/Editor/UnrealEd/Private/ActorGroupingUtils.cpp
+++ b/Engine/Source/Editor/UnrealEd/Private/ActorGroupingUtils.cpp
@@ -114,6 +114,7 @@ AGroupActor* UActorGroupingUtils::GroupActors(const TArray<AActor*>& ActorsToGro
 
 					SpawnedGroupActor->SetFolderPath(FolderPath);
 					SpawnedGroupActor->CenterGroupLocation();
+					SpawnedGroupActor->SetActorRotation(FinalActorList.Last()->GetActorRotation());
 					SpawnedGroupActor->Lock();
 
 					return SpawnedGroupActor;
diff --git a/Engine/Source/Editor/UnrealEd/Private/GroupActor.cpp b/Engine/Source/Editor/UnrealEd/Private/GroupActor.cpp
--- a/Engine/Source/Editor/UnrealEd/Private/GroupActor.cpp
+++ b/Engine/Source/Editor/UnrealEd/Private/GroupActor.cpp
@@ -589,6 +589,7 @@ void AGroupActor::AddSelectedActorsToSelectedGroup()
 					}
 
 					SelectedGroup->CenterGroupLocation();
+					SelectedGroup->SetActorRotation(ActorsToAdd.Last()->GetActorRotation());
 				}
 			}
 			else
