Only throw a Blueprint compile error for a failed graph merge due to mismatched schemas if we are actually doing a merge.

Notes:
- Fixes a recent regression introduced at 22204204 that had impacted an external project.

diff --git a/Engine/Source/Editor/UnrealEd/Private/EdGraphUtilities.cpp b/Engine/Source/Editor/UnrealEd/Private/EdGraphUtilities.cpp
--- a/Engine/Source/Editor/UnrealEd/Private/EdGraphUtilities.cpp
+++ b/Engine/Source/Editor/UnrealEd/Private/EdGraphUtilities.cpp
@@ -410,13 +410,14 @@ void FEdGraphUtilities::MergeChildrenGraphsIn(UEdGraph* MergeTarget, UEdGraph* P
 			{
 				ChildGraph->MoveNodesToAnotherGraph(MergeTarget, IsAsyncLoading() || bIsLoading, bInIsCompiling);
 			}
+			else if (bRequireSchemaMatch && NodeOwner && MessageLog)
+			{
+				// If merging requires a schema match and we have a valid owner for the child graph, throw an error message on the outer node context.
+				MessageLog->Error(*NSLOCTEXT("EdGraphUtilities", "CannotMergeChildGraph_Error", "Unable to merge child graph for node @@. The child graph's schema is incompatible with the outer graph.").ToString(), MessageLog->FindSourceObject(NodeOwner));
+			}
 
 			MergeChildrenGraphsIn(MergeTarget, ChildGraph, bRequireSchemaMatch, bInIsCompiling, MessageLog);
 		}
-		else if (bNonVirtualGraph && bRequireSchemaMatch && MessageLog)
-		{
-			MessageLog->Error(*NSLOCTEXT("EdGraphUtilities", "CannotMergeChildGraph_Error", "Unable to merge child graph for node @@. The child graph's schema is incompatible with the outer graph.").ToString(), MessageLog->FindSourceObject(NodeOwner));
-		}
 	}
 }
 
