Fix crash when we select interchange factory format in the objecttools import file picker.
Make sure to avoid legacy factory in the file picker extension when interchange replace those.
#rnx

diff --git a/Engine/Plugins/Editor/GLTFImporter/Source/GLTFImporter/Private/GLTFImportFactory.cpp b/Engine/Plugins/Editor/GLTFImporter/Source/GLTFImporter/Private/GLTFImportFactory.cpp
--- a/Engine/Plugins/Editor/GLTFImporter/Source/GLTFImporter/Private/GLTFImportFactory.cpp
+++ b/Engine/Plugins/Editor/GLTFImporter/Source/GLTFImporter/Private/GLTFImportFactory.cpp
@@ -179,6 +179,18 @@ bool UGLTFImportFactory::FactoryCanImport(const FString& Filename)
 	return false;
 }
 
+TArray<FString> UGLTFImportFactory::GetFormats() const
+{
+	static TArray<FString> EmptyArray;
+	static const auto CVarGltf = IConsoleManager::Get().FindConsoleVariable(TEXT("Interchange.FeatureFlags.Import.GLTF"));
+	const bool bUseLegacyGltf = (!CVarGltf || !CVarGltf->GetBool());
+	if (bUseLegacyGltf)
+	{
+		return Formats;
+	}
+	return EmptyArray;
+}
+
 void UGLTFImportFactory::CleanUp()
 {
 	// cleanup any resources/buffers
diff --git a/Engine/Plugins/Editor/GLTFImporter/Source/GLTFImporter/Private/GLTFImportFactory.h b/Engine/Plugins/Editor/GLTFImporter/Source/GLTFImporter/Private/GLTFImportFactory.h
--- a/Engine/Plugins/Editor/GLTFImporter/Source/GLTFImporter/Private/GLTFImportFactory.h
+++ b/Engine/Plugins/Editor/GLTFImporter/Source/GLTFImporter/Private/GLTFImportFactory.h
@@ -21,6 +21,7 @@ public:
 
 	virtual void CleanUp() override;
 	virtual bool FactoryCanImport(const FString& Filename) override;
+	virtual TArray<FString> GetFormats() const override;
 
 private:
 	void UpdateMeshes() const;
diff --git a/Engine/Source/Developer/AssetTools/Private/AssetTools.cpp b/Engine/Source/Developer/AssetTools/Private/AssetTools.cpp
--- a/Engine/Source/Developer/AssetTools/Private/AssetTools.cpp
+++ b/Engine/Source/Developer/AssetTools/Private/AssetTools.cpp
@@ -4278,7 +4278,7 @@ TArray<UObject*> UAssetToolsImpl::ImportAssetsWithDialogImplementation(const FSt
 	if (UInterchangeManager::IsInterchangeImportEnabled())
 	{
 		TArray<FString> InterchangeFileExtensions = UInterchangeManager::GetInterchangeManager().GetSupportedFormats(EInterchangeTranslatorType::Assets);
-		ObjectTools::AppendFormatsFileExtensions(InterchangeFileExtensions, FileTypes, AllExtensions);
+		ObjectTools::AppendFormatsFileExtensions(InterchangeFileExtensions, FileTypes, AllExtensions, FilterIndexToFactory);
 	}
 
 	FileTypes = FString::Printf(TEXT("All Files (%s)|%s|%s"), *AllExtensions, *AllExtensions, *FileTypes);
diff --git a/Engine/Source/Editor/UnrealEd/Classes/Factories/Factory.h b/Engine/Source/Editor/UnrealEd/Classes/Factories/Factory.h
--- a/Engine/Source/Editor/UnrealEd/Classes/Factories/Factory.h
+++ b/Engine/Source/Editor/UnrealEd/Classes/Factories/Factory.h
@@ -221,6 +221,11 @@ public:
 	 * @return true if this factory is being used for automated import.  Dialogs and user input should be disabled if this method returns true
 	 */
 	virtual bool IsAutomatedImport() const;
+
+	/**
+	 * @return the supported factory formats list.
+	 */
+	virtual TArray<FString> GetFormats() const { return Formats; }
 public:
 
 	/**
diff --git a/Engine/Source/Editor/UnrealEd/Classes/Factories/FbxFactory.h b/Engine/Source/Editor/UnrealEd/Classes/Factories/FbxFactory.h
--- a/Engine/Source/Editor/UnrealEd/Classes/Factories/FbxFactory.h
+++ b/Engine/Source/Editor/UnrealEd/Classes/Factories/FbxFactory.h
@@ -42,9 +42,11 @@ class UNREALED_API UFbxFactory : public UFactory
 	virtual bool FactoryCanImport(const FString& Filename) override;
 	virtual bool CanImportBeCanceled() const override;
 	virtual IImportSettingsParser* GetImportSettingsParser() override;
-
+	virtual TArray<FString> GetFormats() const override;
 	//~ End UFactory Interface
 	
+	static TArray<FString> GetFbxFormats(const UFactory* Factory);
+
 	/**
 	 * Detect mesh type to import: Static Mesh or Skeletal Mesh.
 	 * Only the first mesh will be detected.
diff --git a/Engine/Source/Editor/UnrealEd/Classes/Factories/FbxSceneImportFactory.h b/Engine/Source/Editor/UnrealEd/Classes/Factories/FbxSceneImportFactory.h
--- a/Engine/Source/Editor/UnrealEd/Classes/Factories/FbxSceneImportFactory.h
+++ b/Engine/Source/Editor/UnrealEd/Classes/Factories/FbxSceneImportFactory.h
@@ -351,6 +351,7 @@ class UNREALED_API UFbxSceneImportFactory : public USceneImportFactory
 	virtual UObject* FactoryCreateBinary(UClass* InClass, UObject* InParent, FName InName, EObjectFlags Flags, UObject* Context, const TCHAR* Type, const uint8*& Buffer, const uint8* BufferEnd, FFeedbackContext* Warn) override;
 	virtual UObject* FactoryCreateBinary(UClass* InClass, UObject* InParent, FName InName, EObjectFlags Flags, UObject* Context, const TCHAR* Type, const uint8*& Buffer, const uint8* BufferEnd, FFeedbackContext* Warn, bool& bOutOperationCanceled) override;
 	virtual bool FactoryCanImport(const FString& Filename) override;
+	virtual TArray<FString> GetFormats() const override;
 
 	/** USceneImportFactory Interface */
 	virtual bool ImportsAssets() const override { return true; }
diff --git a/Engine/Source/Editor/UnrealEd/Private/Fbx/FbxFactory.cpp b/Engine/Source/Editor/UnrealEd/Private/Fbx/FbxFactory.cpp
--- a/Engine/Source/Editor/UnrealEd/Private/Fbx/FbxFactory.cpp
+++ b/Engine/Source/Editor/UnrealEd/Private/Fbx/FbxFactory.cpp
@@ -1149,6 +1149,45 @@ IImportSettingsParser* UFbxFactory::GetImportSettingsParser()
 	return ImportUI;
 }
 
+TArray<FString> UFbxFactory::GetFormats() const
+{
+	return GetFbxFormats(this);
+}
+
+TArray<FString> UFbxFactory::GetFbxFormats(const UFactory* Factory)
+{
+	TArray<FString> FormatArray;
+	const IConsoleVariable* CVarFbx = IConsoleManager::Get().FindConsoleVariable(TEXT("Interchange.FeatureFlags.Import.FBX"));
+	const bool bUseLegacyFbx = (!CVarFbx || !CVarFbx->GetBool());
+	const IConsoleVariable* CVarObj = IConsoleManager::Get().FindConsoleVariable(TEXT("Interchange.FeatureFlags.Import.OBJ"));
+	const bool bUseLegacyObj = (!CVarObj || !CVarObj->GetBool());
+
+	for (const FString& Format : Factory->Formats)
+	{
+		if (Format.StartsWith(TEXT("fbx")))
+		{
+			//Skip if interchange fbx is enabled
+			if (bUseLegacyFbx)
+			{
+				FormatArray.Add(Format);
+			}
+		}
+		else if (Format.StartsWith(TEXT("obj")))
+		{
+			//Skip if interchange obj is enabled
+			if (bUseLegacyObj)
+			{
+				FormatArray.Add(Format);
+			}
+		}
+		else
+		{
+			FormatArray.Add(Format);
+		}
+	}
+	return FormatArray;
+}
+
 UFbxImportUI::UFbxImportUI(const FObjectInitializer& ObjectInitializer)
 	: Super(ObjectInitializer)
 {
diff --git a/Engine/Source/Editor/UnrealEd/Private/Fbx/FbxSceneImportFactory.cpp b/Engine/Source/Editor/UnrealEd/Private/Fbx/FbxSceneImportFactory.cpp
--- a/Engine/Source/Editor/UnrealEd/Private/Fbx/FbxSceneImportFactory.cpp
+++ b/Engine/Source/Editor/UnrealEd/Private/Fbx/FbxSceneImportFactory.cpp
@@ -20,6 +20,7 @@
 #include "Components/SkeletalMeshComponent.h"
 #include "Materials/Material.h"
 #include "Engine/Texture.h"
+#include "Factories/FbxFactory.h"
 #include "Factories/FbxAnimSequenceImportData.h"
 #include "Factories/FbxSkeletalMeshImportData.h"
 #include "Factories/FbxStaticMeshImportData.h"
@@ -916,6 +917,10 @@ bool UFbxSceneImportFactory::FactoryCanImport(const FString& Filename)
 	return false;
 }
 
+TArray<FString> UFbxSceneImportFactory::GetFormats() const
+{
+	return UFbxFactory::GetFbxFormats(this);
+}
 
 TSharedPtr<FFbxNodeInfo> GetNodeInfoPtrById(TArray<TSharedPtr<FFbxNodeInfo>> &HierarchyInfo, uint64 SearchId)
 {
diff --git a/Engine/Source/Editor/UnrealEd/Private/Fbx/ReimportFbxSceneFactory.cpp b/Engine/Source/Editor/UnrealEd/Private/Fbx/ReimportFbxSceneFactory.cpp
--- a/Engine/Source/Editor/UnrealEd/Private/Fbx/ReimportFbxSceneFactory.cpp
+++ b/Engine/Source/Editor/UnrealEd/Private/Fbx/ReimportFbxSceneFactory.cpp
@@ -116,7 +116,6 @@ UReimportFbxSceneFactory::UReimportFbxSceneFactory(const FObjectInitializer& Obj
 {
 
 	SupportedClass = UFbxSceneImportData::StaticClass();
-	Formats.Add(TEXT("fbx;FBX scene"));
 
 	bCreateNew = false;
 	bText = false;
diff --git a/Engine/Source/Editor/UnrealEd/Private/ObjectTools.cpp b/Engine/Source/Editor/UnrealEd/Private/ObjectTools.cpp
--- a/Engine/Source/Editor/UnrealEd/Private/ObjectTools.cpp
+++ b/Engine/Source/Editor/UnrealEd/Private/ObjectTools.cpp
@@ -4329,7 +4329,7 @@ namespace ObjectTools
 
 			TArray<FString> Descriptions;
 			TArray<FString> Extensions;
-			InternalGetFormatInfo( CurFactory->Formats, Descriptions, Extensions );
+			InternalGetFormatInfo( CurFactory->GetFormats(), Descriptions, Extensions);
 			check( Descriptions.Num() == Extensions.Num() );
 
 			// Make sure to only store each key, value pair once
@@ -4449,6 +4449,29 @@ namespace ObjectTools
 		InternalAppendFileExtensions(Descriptions, Extensions, out_FileTypes, out_Extensions);
 	}
 
+	void AppendFormatsFileExtensions(const TArray<FString>& InFormats, FString& out_FileTypes, FString& out_Extensions, TMultiMap<uint32, UFactory*>& out_FilterIndexToFactory)
+	{
+		TArray<FString> Descriptions;
+		TArray<FString> Extensions;
+		InternalGetFormatInfo(InFormats, Descriptions, Extensions);
+		InternalAppendFileExtensions(Descriptions, Extensions, out_FileTypes, out_Extensions);
+		uint32 MaxKeyNumber = 0;
+		TSet<uint32> Keys;
+		out_FilterIndexToFactory.GetKeys(Keys);
+		for (uint32 Key : Keys)
+		{
+			MaxKeyNumber = FMath::Max(MaxKeyNumber, Key);
+		}
+		if (Keys.Num() > 0)
+		{
+			MaxKeyNumber++;
+		}
+		for (const FString& Extension : Extensions)
+		{
+			out_FilterIndexToFactory.Add(MaxKeyNumber++, nullptr);
+		}
+	}
+
 	/**
 	 * Generates a list of file types for a given class.
 	 */
@@ -4457,7 +4480,7 @@ namespace ObjectTools
 		check(InFactory);
 		TArray<FString> Descriptions;
 		TArray<FString> Extensions;
-		InternalGetFormatInfo( InFactory->Formats, Descriptions, Extensions );
+		InternalGetFormatInfo( InFactory->GetFormats(), Descriptions, Extensions);
 		InternalAppendFileExtensions( Descriptions, Extensions, out_Filetypes, out_Extensions);
 	}
 
diff --git a/Engine/Source/Editor/UnrealEd/Public/ObjectTools.h b/Engine/Source/Editor/UnrealEd/Public/ObjectTools.h
--- a/Engine/Source/Editor/UnrealEd/Public/ObjectTools.h
+++ b/Engine/Source/Editor/UnrealEd/Public/ObjectTools.h
@@ -514,6 +514,16 @@ namespace ObjectTools
 	 */
 	UNREALED_API void AppendFormatsFileExtensions(const TArray<FString>& InFormats, FString& out_FileTypes, FString& out_Extensions);
 
+	/**
+	 * Populates two strings with all of the file types and extensions the format list provides.
+	 *
+	 * @param	InFormats		Array of supported file types. Each entry needs to be of the form "ext;Description" where ext is the file extension.
+	 * @param	out_FileTypes	File types supported by the provided array of formats, concatenated into a string
+	 * @param	out_Extensions	Extensions supported by the provided array of formats, concatenated into a string
+	 * @param	out_FilterIndexToFactory	Add INDEX_NONE entry for all provided Formats
+	 */
+	UNREALED_API void AppendFormatsFileExtensions(const TArray<FString>& InFormats, FString& out_FileTypes, FString& out_Extensions, TMultiMap<uint32, UFactory*>& out_FilterIndexToFactory);
+
 	/**
 	 * Iterates over all classes and assembles a list of non-abstract UExport-derived type instances.
 	 */
