Handle a very obtuse case of weak reference when validating object deletion

diff --git a/Engine/Source/Editor/UnrealEd/Private/ObjectTools.cpp b/Engine/Source/Editor/UnrealEd/Private/ObjectTools.cpp
--- a/Engine/Source/Editor/UnrealEd/Private/ObjectTools.cpp
+++ b/Engine/Source/Editor/UnrealEd/Private/ObjectTools.cpp
@@ -213,7 +213,7 @@ namespace ObjectTools
 		while (InObjects.Num() != LastObjectCount)
 		{
 			LastObjectCount = InObjects.Num();
-			for (UObject* NewReferencer : FReferencerFinder::GetAllReferencers(InObjects, &InObjects))
+			for (UObject* NewReferencer : FReferencerFinder::GetAllReferencers(InObjects, &InObjects, EReferencerFinderFlags::SkipWeakReferences))
 			{
 				// Stop walking the dependency chain when the transactor is the referencer
 				if (Transactor == NewReferencer)
@@ -249,7 +249,7 @@ namespace ObjectTools
 
 		bOutIsReferenced = false;
 		bOutIsReferencedInMemoryByUndo = false;
-
+		
 		if (!CVarUseLegacyGetReferencersForDeletion.GetValueOnAnyThread())
 		{
 			const UTransactor* Transactor = GEditor ? ToRawPtr(GEditor->Trans) : nullptr;
@@ -284,7 +284,7 @@ namespace ObjectTools
 			}
 
 			// Check and see whether we are referenced by any objects that won't be garbage collected (*including* the undo buffer)
-			for (UObject* Referencer : FReferencerFinder::GetAllReferencers(ObjectsToDelete, nullptr))
+			for (UObject* Referencer : FReferencerFinder::GetAllReferencers(ObjectsToDelete, nullptr, EReferencerFinderFlags::SkipWeakReferences))
 			{
 				if (Referencer->IsIn(InObject))
 				{
@@ -323,7 +323,7 @@ namespace ObjectTools
 						}
 					}
 
-					if (FReferencerFinder::GetAllReferencers(Objects, nullptr).Contains(Transactor))
+					if (FReferencerFinder::GetAllReferencers(Objects, nullptr, EReferencerFinderFlags::SkipWeakReferences).Contains(Transactor))
 					{
 						bOutIsReferencedInMemoryByUndo = true;
 					}
diff --git a/Engine/Source/Runtime/CoreUObject/Private/UObject/ReferencerFinder.cpp b/Engine/Source/Runtime/CoreUObject/Private/UObject/ReferencerFinder.cpp
--- a/Engine/Source/Runtime/CoreUObject/Private/UObject/ReferencerFinder.cpp
+++ b/Engine/Source/Runtime/CoreUObject/Private/UObject/ReferencerFinder.cpp
@@ -42,8 +42,28 @@ public:
 			}
 		}
 	}
+
+	bool MarkWeakObjectReferenceForClearing(UObject** WeakReference)
+	{
+		return (Flags & EReferencerFinderFlags::SkipWeakReferences) == EReferencerFinderFlags::SkipWeakReferences;
+	}
+};
+
+class FAllReferencesCollector : public TDefaultReferenceCollector<FAllReferencesProcessor>
+{
+	FAllReferencesProcessor& Processor;
+public:
+	FAllReferencesCollector(FAllReferencesProcessor& InProcessor, FGCArrayStruct& InObjectArrayStruct)
+		: TDefaultReferenceCollector<FAllReferencesProcessor>(InProcessor, InObjectArrayStruct)
+		, Processor(InProcessor)
+	{
+	}
+
+	virtual bool MarkWeakObjectReferenceForClearing(UObject** WeakReference) override
+	{
+		return Processor.MarkWeakObjectReferenceForClearing(WeakReference);
+	}
 };
-typedef TDefaultReferenceCollector<FAllReferencesProcessor> FAllReferencesCollector;
 
 // Allow parallel reference collection to be overridden to single threaded via console command.
 static int32 GAllowParallelReferenceCollection = 1;
@@ -137,4 +157,4 @@ TArray<UObject*> FReferencerFinder::GetAllReferencers(const TSet<UObject*>& Refe
 		GUObjectArray.UnlockInternalArray();
 	}
 	return Ret;
-}
+}
\ No newline at end of file
diff --git a/Engine/Source/Runtime/CoreUObject/Public/UObject/ReferencerFinder.h b/Engine/Source/Runtime/CoreUObject/Public/UObject/ReferencerFinder.h
--- a/Engine/Source/Runtime/CoreUObject/Public/UObject/ReferencerFinder.h
+++ b/Engine/Source/Runtime/CoreUObject/Public/UObject/ReferencerFinder.h
@@ -14,7 +14,8 @@ class UObject;
 enum class EReferencerFinderFlags : uint8
 {
 	None = 0,
-	SkipInnerReferences = 1, // Do not add inner objects to the referencers of an outer.
+	SkipInnerReferences = 1 << 0, // Do not add inner objects to the referencers of an outer.
+	SkipWeakReferences  = 1 << 1, // Do not add weak references to the referencers
 };
 
 ENUM_CLASS_FLAGS(EReferencerFinderFlags);
