UPackageTools::ReloadPackages now sets GIsEditorLoadingPackage to true to avoid marking reloaded packages dirty.

diff --git a/Engine/Source/Editor/UnrealEd/Private/PackageTools.cpp b/Engine/Source/Editor/UnrealEd/Private/PackageTools.cpp
--- a/Engine/Source/Editor/UnrealEd/Private/PackageTools.cpp
+++ b/Engine/Source/Editor/UnrealEd/Private/PackageTools.cpp
@@ -803,6 +803,8 @@ UPackageTools::UPackageTools(const FObjectInitializer& ObjectInitializer)
 	{
 		bool bResult = false;
 
+		TGuardValue<bool> IsEditorLoadingPackageGuard(GIsEditorLoadingPackage, true);
+
 		FTextBuilder ErrorMessageBuilder;
 
 		using namespace UE::PackageTools::Private;
