Fixing bug where Adding and Renaming collapsed graphs in functions breaks refs to local vars

diff --git a/Engine/Source/Editor/UnrealEd/Private/Kismet2/BlueprintEditorUtils.cpp b/Engine/Source/Editor/UnrealEd/Private/Kismet2/BlueprintEditorUtils.cpp
--- a/Engine/Source/Editor/UnrealEd/Private/Kismet2/BlueprintEditorUtils.cpp
+++ b/Engine/Source/Editor/UnrealEd/Private/Kismet2/BlueprintEditorUtils.cpp
@@ -315,6 +315,25 @@ static bool RenameVariableReferencesInGraph(UBlueprint* InBlueprint, UClass* InV
 	return bFoundReference;
 }
 
+/**
+ * Gathers all variable nodes from all graph's subgraph nodes
+ *
+ * @param Graph			        The Graph to search
+ * @param OutVariableNodes		variable node array to write to
+ */
+static void GetAllChildGraphVariables(UEdGraph* Graph, TArray<UK2Node_Variable*>& OutVariableNodes)
+{
+	for (UEdGraph* SubGraph : Graph->SubGraphs)
+	{
+		check(SubGraph != nullptr);
+		SubGraph->GetNodesOfClass<UK2Node_Variable>(OutVariableNodes);
+		if (!SubGraph->SubGraphs.IsEmpty())
+		{
+			GetAllChildGraphVariables(SubGraph, OutVariableNodes);
+		}
+	}
+}
+
 FBlueprintEditorUtils::FOnRenameVariableReferences FBlueprintEditorUtils::OnRenameVariableReferencesEvent;
 
 void FBlueprintEditorUtils::RenameVariableReferences(UBlueprint* Blueprint, UClass* VariableClass, const FName& OldVarName, const FName& NewVarName)
@@ -2848,18 +2867,20 @@ void FBlueprintEditorUtils::RenameGraph(UEdGraph* Graph, const FString& NewNameS
 		// Find all variable nodes in this graph.
 		TArray<UK2Node_Variable*> VariableNodes;
 		Graph->GetNodesOfClass<UK2Node_Variable>(VariableNodes);
-		for (const UEdGraph* SubGraph : Graph->SubGraphs)
-		{
-			check(SubGraph != nullptr);
-			SubGraph->GetNodesOfClass<UK2Node_Variable>(VariableNodes);
-		}
+		GetAllChildGraphVariables(Graph, VariableNodes);
 
+		// if it's index is >= 0 we know it was found in the array of functiongraphs
+		bool bGraphIsFunction = (Blueprint->FunctionGraphs.IndexOfByKey(Graph) > -1);
 		// For any nodes that reference a local variable, update the variable's scope to be the graph's new name (which will mirror the UFunction).
 		for (UK2Node_Variable* const VariableNode : VariableNodes)
 		{
 			if (VariableNode->VariableReference.IsLocalScope())
 			{
-				VariableNode->VariableReference.SetLocalMember(VariableNode->VariableReference.GetMemberName(), NewNameStr, VariableNode->VariableReference.GetMemberGuid());
+				// if the rename is the function set the local variable scope to the new name otherwise we leave it with the same scope (Ex: subgraphs in a function)
+				if (bGraphIsFunction)
+				{
+					VariableNode->VariableReference.SetLocalMember(VariableNode->VariableReference.GetMemberName(), NewNameStr, VariableNode->VariableReference.GetMemberGuid());
+				}
 			}
 		}
 
