WorldPartitionBuilderCommandlet: remove unused class member.

#rnx

diff --git a/Engine/Source/Editor/UnrealEd/Classes/Commandlets/WorldPartitionBuilderCommandlet.h b/Engine/Source/Editor/UnrealEd/Classes/Commandlets/WorldPartitionBuilderCommandlet.h
--- a/Engine/Source/Editor/UnrealEd/Classes/Commandlets/WorldPartitionBuilderCommandlet.h
+++ b/Engine/Source/Editor/UnrealEd/Classes/Commandlets/WorldPartitionBuilderCommandlet.h
@@ -18,8 +18,4 @@ class UWorldPartitionBuilderCommandlet : public UCommandlet
 	//~ End UCommandlet Interface
 
 	UWorldPartitionBuilder* CreateBuilder(const FString& WorldConfigFilename);
-
-private:
-	UPROPERTY()
-	TArray<TObjectPtr<UWorldPartitionBuilder>> Builders;
 };
\ No newline at end of file
