World Partition Convert Commandlet:
* Avoid deleting World Settings from retained sublevels
* Delete HLOD actors from retained sublevels

diff --git a/Engine/Source/Editor/UnrealEd/Private/Commandlets/WorldPartitionConvertCommandlet.cpp b/Engine/Source/Editor/UnrealEd/Private/Commandlets/WorldPartitionConvertCommandlet.cpp
--- a/Engine/Source/Editor/UnrealEd/Private/Commandlets/WorldPartitionConvertCommandlet.cpp
+++ b/Engine/Source/Editor/UnrealEd/Private/Commandlets/WorldPartitionConvertCommandlet.cpp
@@ -1068,15 +1068,26 @@ int32 UWorldPartitionConvertCommandlet::Main(const FString& Params)
 
 		TArray<AInstancedFoliageActor*> IFAs;
 		TSet<ULandscapeInfo*> LandscapeInfos;
-		for (AActor* Actor : Actors)
+
+		for (auto Iter = Actors.CreateIterator(); Iter; ++Iter)
 		{
+			AActor* Actor = *Iter;
+
 			if (Actor && IsValidChecked(Actor))
 			{
 				check(Actor->GetLevel() == Level);
 
 				if (ShouldDeleteActor(Actor, bMainLevel))
 				{
-					Level->GetWorld()->DestroyActor(Actor);
+					// Delete actor if processing main level, otherwise just ignore them
+					if (bMainLevel)
+					{
+						Level->GetWorld()->DestroyActor(Actor);
+					}
+					else
+					{
+						Iter.RemoveCurrent();
+					}
 				}
 				else 
 				{
@@ -1241,16 +1252,25 @@ int32 UWorldPartitionConvertCommandlet::Main(const FString& Params)
 				{
 					if(IsValid(Actor))
 					{
-						TSet<AActor*> ActorReferences;
-						ActorReferences.Append(ActorsReferencesUtils::GetActorReferences(Actor, RF_NoFlags, true));
-
-						for (AActor* ActorReference : ActorReferences)
+						// Since we'll keep this level around, pass bMainLevel true here to ensure that we 
+						// delete only unwanted actors, and keep level specific actors (world settings, brush and level script)
+						if (ShouldDeleteActor(Actor, /*bMainLevel=*/true))
+						{
+							SubLevel->GetWorld()->DestroyActor(Actor);
+						}
+						else
 						{
-							if (LevelScriptActorReferences.Find(ActorReference))
+							TSet<AActor*> ActorReferences;
+							ActorReferences.Append(ActorsReferencesUtils::GetActorReferences(Actor, RF_NoFlags, true));
+
+							for (AActor* ActorReference : ActorReferences)
 							{
-								LevelScriptActorReferences.Add(Actor);
-								LevelScriptActorReferences.Append(ActorReferences);
-								break;
+								if (LevelScriptActorReferences.Find(ActorReference))
+								{
+									LevelScriptActorReferences.Add(Actor);
+									LevelScriptActorReferences.Append(ActorReferences);
+									break;
+								}
 							}
 						}
 					}
