StaticLighting: Prevent UnrealLightmass from being compiled in configurations other than Debug and Development. Also check the Lightmass executable with Development configuration specifically since the editor can default to DebugGame

diff --git a/Engine/Source/Editor/UnrealEd/Private/StaticLightingSystem/StaticLightingSystem.cpp b/Engine/Source/Editor/UnrealEd/Private/StaticLightingSystem/StaticLightingSystem.cpp
--- a/Engine/Source/Editor/UnrealEd/Private/StaticLightingSystem/StaticLightingSystem.cpp
+++ b/Engine/Source/Editor/UnrealEd/Private/StaticLightingSystem/StaticLightingSystem.cpp
@@ -326,7 +326,7 @@ void FStaticLightingManager::CreateStaticLightingSystem(const FLightingBuildOpti
 		{
 			if (FEngineBuildSettings::IsSourceDistribution())
 			{
-				FStaticLightingManager::Get()->FailLightingBuild(LOCTEXT("LightmassExecutableOutdatedMessage", "Unreal Lightmass executable is outdated. Recompile UnrealLightmass project in Visual Studio."));
+				FStaticLightingManager::Get()->FailLightingBuild(LOCTEXT("LightmassExecutableOutdatedMessage", "Unreal Lightmass executable is outdated. Recompile UnrealLightmass project with Development configuration in Visual Studio."));
 			}
 			else
 			{
@@ -499,7 +499,7 @@ bool FStaticLightingSystem::CheckLightmassExecutableVersion()
 {
 	FTargetReceipt LightmassReceipt;
 	
-	if (!LightmassReceipt.Read(FTargetReceipt::GetDefaultPath(*FPaths::EngineDir(), TEXT("UnrealLightmass"), FPlatformProcess::GetBinariesSubdirectory(), FApp::GetBuildConfiguration(), nullptr)))
+	if (!LightmassReceipt.Read(FTargetReceipt::GetDefaultPath(*FPaths::EngineDir(), TEXT("UnrealLightmass"), FPlatformProcess::GetBinariesSubdirectory(), EBuildConfiguration::Development, nullptr)))
 	{
 		return false;
 	}
diff --git a/Engine/Source/Programs/UnrealBuildTool/ProjectFiles/VisualStudio/VCProject.cs b/Engine/Source/Programs/UnrealBuildTool/ProjectFiles/VisualStudio/VCProject.cs
--- a/Engine/Source/Programs/UnrealBuildTool/ProjectFiles/VisualStudio/VCProject.cs
+++ b/Engine/Source/Programs/UnrealBuildTool/ProjectFiles/VisualStudio/VCProject.cs
@@ -436,7 +436,7 @@ namespace UnrealBuildTool
 					foreach (Project ProjectTarget in ProjectTargets)
 					{
 						string TargetName = ProjectTargets[0].TargetRules!.Name;
-						if (TargetName == "ShaderCompileWorker" || TargetName == "UnrealLightmass")
+						if (TargetName == "ShaderCompileWorker")
 						{
 							MatchingProjectTargets.Add(ProjectTarget);
 							break;
diff --git a/Engine/Source/Programs/UnrealLightmass/UnrealLightmass.Target.cs b/Engine/Source/Programs/UnrealLightmass/UnrealLightmass.Target.cs
--- a/Engine/Source/Programs/UnrealLightmass/UnrealLightmass.Target.cs
+++ b/Engine/Source/Programs/UnrealLightmass/UnrealLightmass.Target.cs
@@ -4,6 +4,7 @@ using UnrealBuildTool;
 using System.Collections.Generic;
 
 [SupportedPlatforms(UnrealPlatformClass.Desktop)]
+[SupportedConfigurations(UnrealTargetConfiguration.Debug, UnrealTargetConfiguration.Development)]
 public class UnrealLightmassTarget : TargetRules
 {
 	public UnrealLightmassTarget(TargetInfo Target) : base(Target)
