CookPackageSplitter: Change name of ModifiedPackages to KeepReferencedPackages to clarify the meaning. Add and fix documentation for the CookPackageSplitter Populate and PreSave functions.
#rnx

diff --git a/Engine/Source/Editor/UnrealEd/Private/CookOnTheFlyServer.cpp b/Engine/Source/Editor/UnrealEd/Private/CookOnTheFlyServer.cpp
--- a/Engine/Source/Editor/UnrealEd/Private/CookOnTheFlyServer.cpp
+++ b/Engine/Source/Editor/UnrealEd/Private/CookOnTheFlyServer.cpp
@@ -3140,18 +3140,18 @@ UE::Cook::EPollStatus UCookOnTheFlyServer::BeginCacheObjectsToMove(UE::Cook::FGe
 	{
 		bool bPopulateSucceeded = false;
 		TArray<UObject*> ObjectsToMove;
-		TArray<UPackage*> ModifiedPackages;
+		TArray<UPackage*> KeepReferencedPackages;
 		if (Info.IsGenerator())
 		{
 			ConstructGeneratedPackagesForPresave(PackageData, Generator, GeneratedPackagesForPresave);
 			bPopulateSucceeded = Splitter->PopulateGeneratorPackage(Package, SplitDataObject, GeneratedPackagesForPresave,
-				ObjectsToMove, ModifiedPackages);
+				ObjectsToMove, KeepReferencedPackages);
 		}
 		else
 		{
 			ICookPackageSplitter::FGeneratedPackageForPopulate SplitterInfo{ Info.RelativePath, Info.GeneratedRootPath, Package, Info.IsCreateAsMap() };
 			bPopulateSucceeded = Splitter->PopulateGeneratedPackage(Package, SplitDataObject, SplitterInfo,
-				ObjectsToMove, ModifiedPackages);
+				ObjectsToMove, KeepReferencedPackages);
 		}
 
 		if (!bPopulateSucceeded)
@@ -3163,7 +3163,7 @@ UE::Cook::EPollStatus UCookOnTheFlyServer::BeginCacheObjectsToMove(UE::Cook::FGe
 			return EPollStatus::Error;
 		}
 
-		Info.AddModifiedPackages(ModifiedPackages);
+		Info.AddKeepReferencedPackages(KeepReferencedPackages);
 		Info.TakeOverCachedObjectsAndAddMoved(Generator, PackageData.GetCachedObjectsInOuter(), ObjectsToMove);
 		Info.SetSaveStateComplete(FCookGenerationInfo::ESaveState::CallObjectsToMove);
 	}
@@ -3195,15 +3195,15 @@ UE::Cook::EPollStatus UCookOnTheFlyServer::PreSaveGeneratorPackage(UE::Cook::FPa
 		return EPollStatus::Error;
 	}
 
-	TArray<UPackage*> ModifiedPackages;
+	TArray<UPackage*> KeepReferencedPackages;
 	ConstructGeneratedPackagesForPresave(PackageData, Generator, GeneratedPackagesForPresave);
-	if (!Splitter->PreSaveGeneratorPackage(Package, SplitDataObject, GeneratedPackagesForPresave, ModifiedPackages))
+	if (!Splitter->PreSaveGeneratorPackage(Package, SplitDataObject, GeneratedPackagesForPresave, KeepReferencedPackages))
 	{
 		UE_LOG(LogCook, Error, TEXT("PackageSplitter returned false from PreSaveGeneratorPackage. Splitter=%s"),
 			*Generator.GetSplitDataObjectName().ToString());
 		return EPollStatus::Error;
 	}
-	Info.AddModifiedPackages(ModifiedPackages);
+	Info.AddKeepReferencedPackages(KeepReferencedPackages);
 
 	return EPollStatus::Success;
 }
@@ -3379,14 +3379,14 @@ UE::Cook::EPollStatus UCookOnTheFlyServer::TryPopulateGeneratedPackage(UE::Cook:
 	PopulateData.GeneratedRootPath = GeneratedInfo.GeneratedRootPath;
 	PopulateData.Package = GeneratedPackage;
 	PopulateData.bCreatedAsMap = GeneratedInfo.IsCreateAsMap();
-	TArray<UPackage*> ModifiedPackages;
-	if (!Splitter->PreSaveGeneratedPackage(OwnerPackage, OwnerObject, PopulateData, ModifiedPackages))
+	TArray<UPackage*> KeepReferencedPackages;
+	if (!Splitter->PreSaveGeneratedPackage(OwnerPackage, OwnerObject, PopulateData, KeepReferencedPackages))
 	{
 		UE_LOG(LogCook, Error, TEXT("PackageSplitter returned false from PreSaveGeneratedPackage. Splitter=%s, Generated=%s."),
 			*Generator.GetSplitDataObjectName().ToString(), *GeneratedPackageName);
 		return EPollStatus::Error;
 	}
-	GeneratedInfo.AddModifiedPackages(ModifiedPackages);
+	GeneratedInfo.AddKeepReferencedPackages(KeepReferencedPackages);
 	bool bPackageIsMap = GeneratedPackage->ContainsMap();
 	if (bPackageIsMap != GeneratedInfo.IsCreateAsMap())
 	{
diff --git a/Engine/Source/Editor/UnrealEd/Private/Cooker/CookPackageData.cpp b/Engine/Source/Editor/UnrealEd/Private/Cooker/CookPackageData.cpp
--- a/Engine/Source/Editor/UnrealEd/Private/Cooker/CookPackageData.cpp
+++ b/Engine/Source/Editor/UnrealEd/Private/Cooker/CookPackageData.cpp
@@ -1480,7 +1480,7 @@ void FGeneratorPackage::PreGarbageCollect(FCookGenerationInfo& Info, TArray<UObj
 				GCKeepPackages.Add(Package);
 				GCKeepPackageDatas.Add(Info.PackageData);
 			}
-			GCKeepPackages.Append(Info.ModifiedPackages);
+			GCKeepPackages.Append(Info.KeepReferencedPackages);
 			for (FBeginCacheObject& BeginCacheObject : Info.BeginCacheObjects.Objects)
 			{
 				UObject* Object = BeginCacheObject.Object.Get();
@@ -1651,7 +1651,7 @@ void FGeneratorPackage::ResetSaveState(FCookGenerationInfo& Info, UPackage* Pack
 	Info.BeginCacheObjects.Reset();
 	Info.SetHasTakenOverCachedCookedPlatformData(false);
 	Info.SetHasIssuedUndeclaredMovedObjectsWarning(false);
-	Info.ModifiedPackages.Reset();
+	Info.KeepReferencedPackages.Reset();
 }
 
 void FGeneratorPackage::UpdateSaveAfterGarbageCollect(const FPackageData& PackageData, bool& bInOutDemote)
diff --git a/Engine/Source/Editor/UnrealEd/Private/Cooker/CookPackageData.h b/Engine/Source/Editor/UnrealEd/Private/Cooker/CookPackageData.h
--- a/Engine/Source/Editor/UnrealEd/Private/Cooker/CookPackageData.h
+++ b/Engine/Source/Editor/UnrealEd/Private/Cooker/CookPackageData.h
@@ -704,7 +704,10 @@ public:
 	EPollStatus RefreshPackageObjects(FGeneratorPackage& Generator, UPackage* Package, bool& bOutFoundNewObjects,
 		ESaveState DemotionState);
 
-	void AddModifiedPackages(TArray<UPackage*>& InModifiedPackages) { ModifiedPackages.Append(InModifiedPackages); }
+	void AddKeepReferencedPackages(TArray<UPackage*>& InKeepReferencedPackages)
+	{
+		KeepReferencedPackages.Append(InKeepReferencedPackages);
+	}
 
 public:
 	FBeginCacheObjects BeginCacheObjects;
@@ -712,7 +715,7 @@ public:
 	FString GeneratedRootPath;
 	TArray<FName> Dependencies;
 	FPackageData* PackageData = nullptr;
-	TArray<UPackage*> ModifiedPackages;
+	TArray<UPackage*> KeepReferencedPackages;
 private:
 	ESaveState GeneratorSaveState = ESaveState::StartGenerate;
 	bool bCreateAsMap : 1;
diff --git a/Engine/Source/Editor/UnrealEd/Public/CookPackageSplitter.h b/Engine/Source/Editor/UnrealEd/Public/CookPackageSplitter.h
--- a/Engine/Source/Editor/UnrealEd/Public/CookPackageSplitter.h
+++ b/Engine/Source/Editor/UnrealEd/Public/CookPackageSplitter.h
@@ -65,91 +65,135 @@ public:
 	/** Return the list of packages to generate. */
 	virtual TArray<FGeneratedPackage> GetGenerateList(const UPackage* OwnerPackage, const UObject* OwnerObject) = 0;
 
-	struct FGeneratedPackageForPopulate
+	/** Representation of a generated package that is provided when populating the generator package. */
+	struct FGeneratedPackageForPreSave
 	{
 		/** RelativePath returned from GetGenerateList */
 		FString RelativePath;
 		/** Root returned from GetGenerateList */
 		FString GeneratedRootPath;
 		/**
-		 * Non-null, constructed package. May have been previously passed into PreSaveGeneratedPackage,
-		 * or may be completely empty if that previous package was garbage collected before this call.
+		 * Non-null UPackage. Possibly an empty placeholder package. Provided so that the generator package
+		 * can create import references to objects that will be stored in the generated package.
 		 */
 		UPackage* Package = nullptr;
 		/** *GetCreateAsMap returned from GetGenerateList. The package filename extension has already been set based on this. */
 		bool bCreatedAsMap = false;
 	};
-
 	/**
-	 * Try to populate a generated package.
-	 *
-	 * Receive an empty UPackage generated from an element in GetGenerateList and populate it
-	 * Return a list of all the objects that will be moved into the Generated package during its save, so the cooker
-	 * can call BeginCacheForCookedPlatformData on them before the move
-	 * After returning, the given package will be queued for saving into the TargetDomain
-	 * Note that PreSaveGeneratedPackage will not be called for every package on
-	 * an iterative cook; it will only be called for the packages with changed dependencies.
-	 *
-	 * @param OwnerPackage				The parent package being split
+	 * Called before presaving the parent generator package, to give the generator a chance to inform the cooker which objects will
+	 * be moved into the generator package that are not already present in it.
+	 * 
+	 * @param OwnerPackage				The generator package being split
 	 * @param OwnerObject				The SplitDataClass instance that this CookPackageSplitter instance was created for
-	 * @param GeneratedPackage			Pointer and information about the package to populate
-	 * @param OutObjectsToMove			List of all the objects that will be moved into the Generated package during its save
-	 * @return							True if successfully populates,  false on error (this will cause a cook error).
+	 * @param GeneratedPackages			Placeholder UPackage and relative path information for all packages that will be generated
+	 * @param OutObjectsToMove			List of all the objects that will be moved into the Generator package during its save
+	 * @param OutKeepReferencedPackages A list of packages which should be kept referenced until all generated packages for
+	 *                                  the generator have finished saving.
+	 * 
+	 * @return							True if successfully populated, false on error (this will cause a cook error).
 	 */
-	virtual bool PopulateGeneratedPackage(UPackage* OwnerPackage, UObject* OwnerObject,
-		const FGeneratedPackageForPopulate& GeneratedPackage, TArray<UObject*>& OutObjectsToMove,
-		TArray<UPackage*>& OutModifiedPackages) { return true; }
-	
+	virtual bool PopulateGeneratorPackage(UPackage* OwnerPackage, UObject* OwnerObject,
+		const TArray<ICookPackageSplitter::FGeneratedPackageForPreSave>& GeneratedPackages, TArray<UObject*>& OutObjectsToMove,
+		TArray<UPackage*>& OutKeepReferencedPackages)
+	{
+		return true;
+	}
+
 	/**
-	 * Called before saving the parent generator package, which itself occurs before PreSaveGeneratedPackage is called on the generated packages.
-	 * Make any required adjustments to the parent package before it is saved into the target domain.
+	 * Called before saving the parent generator package, after PopulateGeneratorPackage but before PopulateGeneratedPackage for any
+	 * generated packages. Make any required adjustments to the parent package before it is saved into the target domain.
+	 *
+	 * @param OwnerPackage				The generator package being split
+	 * @param OwnerObject				The SplitDataClass instance that this CookPackageSplitter instance was created for
+	 * @param GeneratedPackages			Placeholder UPackage and relative path information for all packages that will be generated
+	 * @param OutKeepReferencedPackages A list of packages which should be kept referenced until all generated packages for
+	 *                                  the generator have finished saving.
+	 *
+	 * @return							True if successfully presaved, false on error (this will cause a cook error).
 	 */
-	virtual bool PreSaveGeneratedPackage(UPackage* OwnerPackage, UObject* OwnerObject,
-		const FGeneratedPackageForPopulate& GeneratedPackage, TArray<UPackage*>& OutModifiedPackages) { return true; };
+	virtual bool PreSaveGeneratorPackage(UPackage* OwnerPackage, UObject* OwnerObject,
+		const TArray<FGeneratedPackageForPreSave>& PlaceholderPackages, TArray<UPackage*>& OutKeepReferencedPackages)
+	{
+		return true;
+	}
+
 	/**
-	 * Called after saving the generated package. Undo any required adjustments to the parent package that
-	 * were made in PreSaveGeneratedPackage, so that the package is once again ready for use in the editor or in
+	 * Called after saving the parent generator package. Undo any required adjustments to the parent package that
+	 * were made in PreSaveGeneratorPackage, so that the package is once again ready for use in the editor or in
 	 * future GetGenerateList or PreSaveGeneratedPackage calls
 	 */
-	virtual void PostSaveGeneratedPackage(UPackage* OwnerPackage, UObject* OwnerObject,
-		const FGeneratedPackageForPopulate& GeneratedPackage) {}
-
+	virtual void PostSaveGeneratorPackage(UPackage* OwnerPackage, UObject* OwnerObject)
+	{
+	}
 
-	struct FGeneratedPackageForPreSave
+	/** Representation of a generated package when it itself is being populated. */
+	struct FGeneratedPackageForPopulate
 	{
 		/** RelativePath returned from GetGenerateList */
 		FString RelativePath;
 		/** Root returned from GetGenerateList */
 		FString GeneratedRootPath;
 		/**
-		 * Non-null, constructed package. May have been previously passed into PreSaveGeneratedPackage,
-		 * or may be completely empty if that previous package was garbage collected before this call.
+		 * The UPackage that has been created for the package. Possibly empty, but may also contain still have modifications
+		 * that were made during PopulateGeneratorPackage.
 		 */
 		UPackage* Package = nullptr;
 		/** *GetCreateAsMap returned from GetGenerateList. The package filename extension has already been set based on this. */
 		bool bCreatedAsMap = false;
 	};
+
 	/**
-	 * Return a list of all the objects that will be moved into the Generator package during its save, so the cooker
+	 * Try to populate a generated package.
+	 *
+	 * Receive an empty UPackage generated from an element in GetGenerateList and populate it.
+	 * Return a list of all the objects that will be moved into the Generated package during its save, so the cooker
 	 * can call BeginCacheForCookedPlatformData on them before the move
+	 * After returning, the given package will be queued for saving into the TargetDomain
+	 *
+	 * @param OwnerPackage				The parent package being split
+	 * @param OwnerObject				The SplitDataClass instance that this CookPackageSplitter instance was created for
+	 * @param GeneratedPackage			Pointer and information about the package to populate
+	 * @param OutObjectsToMove			List of all the objects that will be moved into the generated package during its save
+	 * @param OutKeepReferencedPackages A list of packages which should be kept referenced until all generated packages for
+	 *                                  for the generator have finished saving.
+	 * 
+	 * @return							True if successfully populated, false on error (this will cause a cook error).
 	 */
-	virtual bool PopulateGeneratorPackage(UPackage* OwnerPackage, UObject* OwnerObject,
-		const TArray<ICookPackageSplitter::FGeneratedPackageForPreSave>& GeneratedPackages, TArray<UObject*>& OutObjectsToMove,
-		TArray<UPackage*>& OutModifiedPackages) { return true; }
-
+	virtual bool PopulateGeneratedPackage(UPackage* OwnerPackage, UObject* OwnerObject,
+		const FGeneratedPackageForPopulate& GeneratedPackage, TArray<UObject*>& OutObjectsToMove,
+		TArray<UPackage*>& OutKeepReferencedPackages)
+	{
+		return true;
+	}
+	
 	/**
-	 * Called before saving the parent generator package, which itself occurs before PreSaveGeneratedPackage is called on the generated packages.
-	 * Make any required adjustments to the parent package before it is saved into the target domain.
+	 * Called before saving a generated package, after PopulateGeneratedPackage. Make any required adjustments to the
+	 * generated package before it is saved into the target domain.
+	 * 
+	 * @param OwnerPackage				The parent package being split
+	 * @param OwnerObject				The SplitDataClass instance that this CookPackageSplitter instance was created for
+	 * @param GeneratedPackage			Pointer and information about the package to populate
+	 * @param OutKeepReferencedPackages A list of packages which should be kept referenced until all generated packages for
+	 *                                  for the generator have finished saving.
+	 * 
+	 * @return							True if successfully presaved, false on error (this will cause a cook error).
 	 */
-	virtual bool PreSaveGeneratorPackage(UPackage* OwnerPackage, UObject* OwnerObject, const TArray<FGeneratedPackageForPreSave>& PlaceholderPackages,
-		TArray<UPackage*>& OutModifiedPackages) { return true; }
+	virtual bool PreSaveGeneratedPackage(UPackage* OwnerPackage, UObject* OwnerObject,
+		const FGeneratedPackageForPopulate& GeneratedPackage, TArray<UPackage*>& OutKeepReferencedPackages)
+	{
+		return true;
+	};
 
 	/**
-	 * Called after saving the parent generator package. Undo any required adjustments to the parent package that
-	 * were made in PreSaveGeneratorPackage, so that the package is once again ready for use in the editor or in
-	 * future GetGenerateList or PreSaveGeneratedPackage calls
+	 * Called after saving a generated package. Undo any required adjustments to the parent package that
+	 * were made in PreSaveGeneratedPackage, so that the parent package is once again ready for use in the editor or in
+	 * future PreSaveGeneratedPackage calls.
 	 */
-	virtual void PostSaveGeneratorPackage(UPackage* OwnerPackage, UObject* OwnerObject) {}
+	virtual void PostSaveGeneratedPackage(UPackage* OwnerPackage, UObject* OwnerObject,
+		const FGeneratedPackageForPopulate& GeneratedPackage)
+	{
+	}
 
 	/** Called when the Owner package needs to be reloaded after a garbage collect in order to populate a generated package. */
 	virtual void OnOwnerReloaded(UPackage* OwnerPackage, UObject* OwnerObject) {}
