Back out 17997804 and add some handling for name collisions in InstancedReferenceSubobjectHelper so that contention on an object name is not fatal. Object outer is not a sound method of ordering reinstancing, instead we should continue to rely on and improve the archetype updating logic that runs before the object graph is mass updated.

diff --git a/Engine/Source/Editor/UnrealEd/Private/Kismet2/KismetReinstanceUtilities.cpp b/Engine/Source/Editor/UnrealEd/Private/Kismet2/KismetReinstanceUtilities.cpp
--- a/Engine/Source/Editor/UnrealEd/Private/Kismet2/KismetReinstanceUtilities.cpp
+++ b/Engine/Source/Editor/UnrealEd/Private/Kismet2/KismetReinstanceUtilities.cpp
@@ -1848,12 +1848,7 @@ static void ReplaceObjectHelper(UObject*& OldObject, UClass* OldClass, UObject*&
 		}
 
 		FMakeClassSpawnableOnScope TemporarilySpawnable(NewClass);
-		// Since we have previously ensured that instances are processed in such order
-		// that each instance's outer is processed _before_ itself, we can safely
-		// look up the outer object in OldToNewInstanceMap to determine if a new version exists 
-		UObject* OldOuter = OldObject->GetOuter();
-		UObject* const* NewOuter = OldToNewInstanceMap.Find(OldOuter);
-		NewUObject = NewObject<UObject>(NewOuter ? *NewOuter : OldOuter, NewClass, NewName, RF_NoFlags, NewArchetype);
+		NewUObject = NewObject<UObject>(OldObject->GetOuter(), NewClass, NewName, RF_NoFlags, NewArchetype);
 	}
 
 	check(NewUObject != nullptr);
@@ -1864,7 +1859,6 @@ static void ReplaceObjectHelper(UObject*& OldObject, UClass* OldClass, UObject*&
 	InstancedPropertyUtils::FArchiveInstancedSubObjCollector  InstancedSubObjCollector(OldObject, InstancedPropertyMap);
 	// Copy property values
 	Options.bNotifyObjectReplacement = true;
-	Options.bDontClearReferenceIfNewerClassExists = true;
 	Options.bSkipCompilerGeneratedDefaults = true;
 	UEditorEngine::CopyPropertiesForUnrelatedObjects(OldObject, NewUObject, Options);
 	// Generate new subobjects
@@ -2165,6 +2159,8 @@ void FBlueprintCompileReinstancer::ReplaceInstancesOfClass_Inner(TMap<UClass*, U
 	};
 
 	{
+		TArray<UObject*> ObjectsToReplace;
+
 		BP_SCOPED_COMPILER_EVENT_STAT(EKismetReinstancerStats_ReplaceInstancesOfClass);
 		if(GEditor && GEditor->GetSelectedActors())
 		{
@@ -2173,108 +2169,13 @@ void FBlueprintCompileReinstancer::ReplaceInstancesOfClass_Inner(TMap<UClass*, U
 			SelectedActors->Modify();
 		}
 
-		struct FClassInstances
-		{
-			TPair<UClass*, UClass*> OldToNewClass;
-			TArray<UObject*> OldInstances;
-		};
-
-		// Build array of instances to replace, sorted in order of subobject hierarchy
-		TArray<FClassInstances> ClassInstancesWithOldOuter;
-		TArray<FClassInstances> AllClassInstances;
-		ClassInstancesWithOldOuter.Reserve(InOldToNewClassMap.Num());
-		AllClassInstances.Reserve(InOldToNewClassMap.Num());
-		for (TPair<UClass*, UClass*> OldToNewClass : InOldToNewClassMap)
-		{
-			// We depend on the proper states of this flag in the algorithm that follows.  Reload system will
-			// have old == new if a code change has been detected, but no structural changes were made. 
-			if (OldToNewClass.Key != OldToNewClass.Value)
-			{
-				OldToNewClass.Key->ClassFlags |= CLASS_NewerVersionExists; // This flag should be set at this point so make sure it is
-				ensure(!OldToNewClass.Value->HasAnyClassFlags(CLASS_NewerVersionExists));
-			}
-
-			// Cache all instances of each class being reinstanced
-			FClassInstances ClassInstances;
-			ClassInstances.OldToNewClass = OldToNewClass;
-			const bool bIncludeDerivedClasses = false;
-			GetObjectsOfClass(OldToNewClass.Key, ClassInstances.OldInstances, bIncludeDerivedClasses);
-
-			// Put in one of two lists depending on if any instances have an outdated outer
-			bool bHasOldOuter = false;
-			for (UObject* OldInstance : ClassInstances.OldInstances)
-			{
-				if (OldInstance->GetOuter()->GetClass()->HasAnyClassFlags(CLASS_NewerVersionExists))
-				{
-					bHasOldOuter = true;
-					break;
-				}
-			}
-			if (bHasOldOuter)
-			{
-				ClassInstancesWithOldOuter.Add(MoveTemp(ClassInstances));
-			}
-			else
-			{
-				AllClassInstances.Add(MoveTemp(ClassInstances));
-			}
-		}
-
-		// Make sure all subobjects follow behind their outer
-		// If we find any subobject occurring before its outer, swap it with its outer and repeat
-		// This algorithm has quadratic complexity, though the number of elements in this array is usually very small
-		// and for Blueprint reinstancing at this time of writing this array is always empty
-		// It might be possible to topologically sort this with better time complexity but
-		// this would also significantly complicate the algorithm which is at this point not worth it
-		for (int32 ClassIndex1 = 0; ClassIndex1 < ClassInstancesWithOldOuter.Num(); ++ClassIndex1)
-		{
-		StartOver:
-			FClassInstances& ClassInstances1 = ClassInstancesWithOldOuter[ClassIndex1];
-			for (int32 ClassIndex2 = ClassIndex1 + 1; ClassIndex2 < ClassInstancesWithOldOuter.Num(); ++ClassIndex2)
-			{
-				FClassInstances& ClassInstances2 = ClassInstancesWithOldOuter[ClassIndex2];
-				bool bSwap = false;
-				bool bDontSwap = false;
-				for (UObject* Obj1 : ClassInstances1.OldInstances)
-				{
-					UObject* Outer1 = Obj1->GetOuter();
-					for (UObject* Obj2 : ClassInstances2.OldInstances)
-					{
-						UObject* Outer2 = Obj2->GetOuter();
-						bSwap |= (Outer1 == Obj2);
-						bDontSwap |= (Outer2 == Obj1);
-					}
-				}
-				if (bSwap)
-				{
-					// Found at least one instance of class 1 that is a subobject of an instance of class 2
-					// Therefore, let's swap their order so the owner gets reinstanced before its subobjects
-					ClassInstancesWithOldOuter.Swap(ClassIndex1, ClassIndex2);
-					if (!bDontSwap)
-					{
-						goto StartOver;
-					}
-					// If we get here for a legitimate reason, we should consider refactoring this code to operate
-					// not on a list of classes with instances but on a flat list of instances instead
-					// which would allow us to resolve this circular dependency
-					UE_LOG(LogBlueprint, Error, TEXT("Found an instance of class %s contained in an instance of class %s and also vice versa. Reinstancing will be faulty."), *ClassInstances1.OldToNewClass.Value->GetName(), *ClassInstances2.OldToNewClass.Value->GetName());
-					ensureAlways(false);
-				}
-			}
-		}
-
-		// After sorting, append ClassInstancesWithOuter to AllClassInstances,
-		// The prior elements of AllClassInstances might contain outers of ClassInstancesWithOldOuter
-		// and must therefore be processed before the elements of ClassInstancesWithOldOuter 
-		AllClassInstances.Append(MoveTemp(ClassInstancesWithOldOuter));
-
 		// WARNING: for (TPair<UClass*, UClass*> OldToNewClass : InOldToNewClassMap) duplicated below 
 		// to handle reconstructing actors which need to be reinstanced after their owned components 
 		// have been updated:
-		for (FClassInstances& ClassInstances : AllClassInstances)
+		for (TPair<UClass*, UClass*> OldToNewClass : InOldToNewClassMap)
 		{
-			UClass* OldClass = ClassInstances.OldToNewClass.Key;
-			UClass* NewClass = ClassInstances.OldToNewClass.Value;
+			UClass* OldClass = OldToNewClass.Key;
+			UClass* NewClass = OldToNewClass.Value;
 			check(OldClass && NewClass);
 			check(OldClass != NewClass || IsReloadActive());
 			{
@@ -2302,10 +2203,13 @@ void FBlueprintCompileReinstancer::ReplaceInstancesOfClass_Inner(TMap<UClass*, U
 					bFixupSCS = (NewClass->IsChildOf<USceneComponent>() != OldClass->IsChildOf<USceneComponent>());
 				}
 
+				const bool bIncludeDerivedClasses = false;
+				ObjectsToReplace.Reset();
+				GetObjectsOfClass(OldClass, ObjectsToReplace, bIncludeDerivedClasses);
 				// Then fix 'real' (non archetype) instances of the class
-				for (int32 OldObjIndex = 0; OldObjIndex < ClassInstances.OldInstances.Num(); ++OldObjIndex)
+				for (int32 OldObjIndex = 0; OldObjIndex < ObjectsToReplace.Num(); ++OldObjIndex)
 				{
-					UObject* OldObject = ClassInstances.OldInstances[OldObjIndex];
+					UObject* OldObject = ObjectsToReplace[OldObjIndex];
 					
 					AActor* OldActor = Cast<AActor>(OldObject);
 
@@ -2322,7 +2226,7 @@ void FBlueprintCompileReinstancer::ReplaceInstancesOfClass_Inner(TMap<UClass*, U
 					if (OldActor == nullptr)
 					{
 						UObject* NewUObject = nullptr;
-						ReplaceObjectHelper(OldObject, OldClass, NewUObject, NewClass, OldToNewInstanceMap, OldToNewNameMap, OldObjIndex, ClassInstances.OldInstances, PotentialEditorsForRefreshing, OwnersToRerunConstructionScript, &FDirectAttachChildrenAccessor::Get, bIsComponent, bArchetypesAreUpToDate);
+						ReplaceObjectHelper(OldObject, OldClass, NewUObject, NewClass, OldToNewInstanceMap, OldToNewNameMap, OldObjIndex, ObjectsToReplace, PotentialEditorsForRefreshing, OwnersToRerunConstructionScript, &FDirectAttachChildrenAccessor::Get, bIsComponent, bArchetypesAreUpToDate);
 						UpdateObjectBeingDebugged(OldObject, NewUObject);
 						ObjectsReplaced.Add(OldObject);
 
@@ -2348,18 +2252,22 @@ void FBlueprintCompileReinstancer::ReplaceInstancesOfClass_Inner(TMap<UClass*, U
 		// WARNING: for (TPair<UClass*, UClass*> OldToNewClass : InOldToNewClassMap) duplicated above 
 		// this loop only handles actors - which need to be reconstructed *after* their owned components 
 		// have been reinstanced:
-		for (FClassInstances& ClassInstances : AllClassInstances)
+		for (TPair<UClass*, UClass*> OldToNewClass : InOldToNewClassMap)
 		{
-			UClass* OldClass = ClassInstances.OldToNewClass.Key;
-			UClass* NewClass = ClassInstances.OldToNewClass.Value;
+			UClass* OldClass = OldToNewClass.Key;
+			UClass* NewClass = OldToNewClass.Value;
 			check(OldClass && NewClass);
 
 			{
+				const bool bIncludeDerivedClasses = false;
+				ObjectsToReplace.Reset();
+				GetObjectsOfClass(OldClass, ObjectsToReplace, bIncludeDerivedClasses);
+
 				// store old attachment data before we mess with components, etc:
 				TMap<UObject*, FActorAttachmentData> ActorAttachmentData;
-				for (int32 OldObjIndex = 0; OldObjIndex < ClassInstances.OldInstances.Num(); ++OldObjIndex)
+				for (int32 OldObjIndex = 0; OldObjIndex < ObjectsToReplace.Num(); ++OldObjIndex)
 				{
-					UObject* OldObject = ClassInstances.OldInstances[OldObjIndex];
+					UObject* OldObject = ObjectsToReplace[OldObjIndex];
 					if(!IsValid(OldObject) || 
 						(InstancesThatShouldUseOldClass && InstancesThatShouldUseOldClass->Contains(OldObject)))
 					{
@@ -2373,9 +2281,9 @@ void FBlueprintCompileReinstancer::ReplaceInstancesOfClass_Inner(TMap<UClass*, U
 				}
 
 				// Then fix 'real' (non archetype) instances of the class
-				for (int32 OldObjIndex = 0; OldObjIndex < ClassInstances.OldInstances.Num(); ++OldObjIndex)
+				for (int32 OldObjIndex = 0; OldObjIndex < ObjectsToReplace.Num(); ++OldObjIndex)
 				{
-					UObject* OldObject = ClassInstances.OldInstances[OldObjIndex];
+					UObject* OldObject = ObjectsToReplace[OldObjIndex];
 					AActor* OldActor = Cast<AActor>(OldObject);
 
 					// Skip archetype instances, EXCEPT for child actor templates
@@ -2399,7 +2307,7 @@ void FBlueprintCompileReinstancer::ReplaceInstancesOfClass_Inner(TMap<UClass*, U
 						else
 						{
 							// Actors that are not in a level cannot be reconstructed, sequencer team decided to reinstance these as normal objects:
-							ReplaceObjectHelper(OldObject, OldClass, NewUObject, NewClass, OldToNewInstanceMap, OldToNewNameMap, OldObjIndex, ClassInstances.OldInstances, PotentialEditorsForRefreshing, OwnersToRerunConstructionScript, &FDirectAttachChildrenAccessor::Get, false, bArchetypesAreUpToDate);
+							ReplaceObjectHelper(OldObject, OldClass, NewUObject, NewClass, OldToNewInstanceMap, OldToNewNameMap, OldObjIndex, ObjectsToReplace, PotentialEditorsForRefreshing, OwnersToRerunConstructionScript, &FDirectAttachChildrenAccessor::Get, false, bArchetypesAreUpToDate);
 						}
 						UpdateObjectBeingDebugged(OldObject, NewUObject);
 						ObjectsReplaced.Add(OldObject);
diff --git a/Engine/Source/Runtime/Engine/Private/InstancedReferenceSubobjectHelper.cpp b/Engine/Source/Runtime/Engine/Private/InstancedReferenceSubobjectHelper.cpp
--- a/Engine/Source/Runtime/Engine/Private/InstancedReferenceSubobjectHelper.cpp
+++ b/Engine/Source/Runtime/Engine/Private/InstancedReferenceSubobjectHelper.cpp
@@ -249,11 +249,24 @@ void FFindInstancedReferenceSubobjectHelper::Duplicate(UObject* OldObject, UObje
 					const bool bNotHandledYet = !ReferenceReplacementMap.Contains(Obj);
 					if (bKeptByOld && bNotHandledYet)
 					{
+						// This name may have been taken by an instanced subobject, since we're copying from
+						// an object we want to recreate it here with a duplicate of the old object. I'm 
+						// using a rename call here to reserve the name for ourself. These objects come from
+						// FObjectInitializer::InstanceSubobjects()
+						UObject* ExistingObject = StaticFindObjectFast(UObject::StaticClass(), NewObject, Obj->GetFName());
+						if (ExistingObject)
+						{
+							ExistingObject->Rename(
+								nullptr,
+								GetTransientPackage(),
+								REN_DoNotDirty | REN_DontCreateRedirectors | REN_ForceNoResetLoaders | REN_NonTransactional);
+						}
+
 						UObject* NewEditInlineSubobject = StaticDuplicateObject(Obj, NewObject, Obj->GetFName());
 						ReferenceReplacementMap.Add(Obj, NewEditInlineSubobject);
 
 						// NOTE: we cannot patch OldObject's linker table here, since we don't 
-						//       know the relation between the  two objects (one could be of a 
+						//       know the relation between the two objects (one could be of a 
 						//       super class, and the other a child)
 
 						// We also need to make sure to fixup any properties here
diff --git a/Engine/Source/Runtime/Engine/Private/UnrealEngine.cpp b/Engine/Source/Runtime/Engine/Private/UnrealEngine.cpp
--- a/Engine/Source/Runtime/Engine/Private/UnrealEngine.cpp
+++ b/Engine/Source/Runtime/Engine/Private/UnrealEngine.cpp
@@ -16594,7 +16594,6 @@ UEngine::FCopyPropertiesForUnrelatedObjectsParams::FCopyPropertiesForUnrelatedOb
 	, bSkipCompilerGeneratedDefaults(false)
 	, bNotifyObjectReplacement(false)
 	, bClearReferences(true)
-	, bDontClearReferenceIfNewerClassExists(false)
 	, SourceObjectArchetype(nullptr)
 {}
 UEngine::FCopyPropertiesForUnrelatedObjectsParams::FCopyPropertiesForUnrelatedObjectsParams(const FCopyPropertiesForUnrelatedObjectsParams&) = default;
@@ -16835,8 +16834,7 @@ void UEngine::CopyPropertiesForUnrelatedObjects(UObject* OldObject, UObject* New
 	// Also, if requested, leave pointers to instances of renewed classes intact, assuming we are in the midst of reinstancing, and those will be reinstanced
 	ForEachObjectWithOuter(OldObject, [&ReferenceReplacementMap, &Params](UObject* ObjectInOuter)
 	{
-		if (!ReferenceReplacementMap.Contains(ObjectInOuter)
-			&& (!Params.bDontClearReferenceIfNewerClassExists || !ObjectInOuter->GetClass()->HasAnyClassFlags(CLASS_NewerVersionExists)))
+		if (!ReferenceReplacementMap.Contains(ObjectInOuter))
 		{
 			ReferenceReplacementMap.Add(ObjectInOuter, nullptr);
 		}
