Fix: Used By Level now considers external package objects (Actors)

diff --git a/Engine/Source/Editor/UnrealEd/Private/ObjectTools.cpp b/Engine/Source/Editor/UnrealEd/Private/ObjectTools.cpp
--- a/Engine/Source/Editor/UnrealEd/Private/ObjectTools.cpp
+++ b/Engine/Source/Editor/UnrealEd/Private/ObjectTools.cpp
@@ -4599,26 +4599,23 @@ namespace ObjectTools
 			// Clear all marked flags that could have been tagged in a previous search or by another system.
 			Obj->UnMark(EObjectMark(OBJECTMARK_TagImp | OBJECTMARK_TagExp));
 
-
 			// If the object is not flagged for GC and it is in one of the level packages do an indepth search to see what references it.
-			if( IsValidChecked(Obj) && !Obj->IsUnreachable() && LevelPackages.Find( Obj->GetOutermost() ) != NULL )
+			if( IsValidChecked(Obj) && !Obj->IsUnreachable())
 			{
-				// Determine if the current object is in one of the search levels.  This is the same as UObject::IsIn except that we can
-				// search through many levels at once.
-				for ( UObject* ObjectOuter = Obj->GetOuter(); ObjectOuter; ObjectOuter = ObjectOuter->GetOuter() )
+				// Get Object's Outermost package which isn't the same as the object's package for external objects/actors
+				const UObject* OuterPackage = Obj->IsA<UPackage>() ? Obj : Obj->GetOutermostObject()->GetPackage();
+				if (LevelPackages.Find(OuterPackage) != NULL)
 				{
-					if ( Levels.Find(ObjectOuter) != NULL )
+					if (ULevel* OuterLevel = Obj->GetTypedOuter<ULevel>(); OuterLevel && Levels.Contains(OuterLevel))
 					{
 						// this object was contained within one of our ReferenceRoots
-						ObjectsInLevels.Add( Obj );
+						ObjectsInLevels.Add(Obj);
 
 						// If the object is using a blueprint generated class, also add the blueprint as a reference
-						UBlueprint* const Blueprint = Cast<UBlueprint>(Obj->GetClass()->ClassGeneratedBy);
-						if ( Blueprint )
+						if (UBlueprint* const Blueprint = Cast<UBlueprint>(Obj->GetClass()->ClassGeneratedBy))
 						{
-							ObjectsInLevels.Add( Blueprint );
+							ObjectsInLevels.Add(Blueprint);
 						}
-						break;
 					}
 				}
 			}
@@ -4639,7 +4636,7 @@ namespace ObjectTools
 		// Tag all objects that are referenced by objects in the levels we are searching.
 		FArchiveReferenceMarker Marker( ObjectsInLevels, MarkerFlags);
 	}
-
+	
 	TSharedPtr<SWindow> OpenPropertiesForSelectedObjects( const TArray<UObject*>& SelectedObjects )
 	{
 		TSharedPtr<SWindow> FloatingDetailsView;
