EditorDomain - Add deterministic Guid.
#rnx

diff --git a/Engine/Source/Editor/UnrealEd/Private/EditorDomain/EditorDomainUtils.cpp b/Engine/Source/Editor/UnrealEd/Private/EditorDomain/EditorDomainUtils.cpp
--- a/Engine/Source/Editor/UnrealEd/Private/EditorDomain/EditorDomainUtils.cpp
+++ b/Engine/Source/Editor/UnrealEd/Private/EditorDomain/EditorDomainUtils.cpp
@@ -168,7 +168,7 @@ bool bGGlobalConstructClassesInitialized = false;
 int64 GMaxBulkDataSize = -1;
 
 // Change to a new guid when EditorDomain needs to be invalidated
-const TCHAR* EditorDomainVersion = TEXT("1A71E99C857343C7A6AEB0035E8F5F26");
+const TCHAR* EditorDomainVersion = TEXT("D0A87D5986BD47479D64C08F70CE9BB2");
 
 // Identifier of the CacheBuckets for EditorDomain tables
 const TCHAR* EditorDomainPackageBucketName = TEXT("EditorDomainPackage");
@@ -1617,6 +1617,11 @@ private:
 	uint64 BulkDataSize = 0;
 };
 
+static FGuid IoHashToGuid(const FIoHash& Hash)
+{
+	const uint32* HashInts = reinterpret_cast<const uint32*>(Hash.GetBytes());
+	return FGuid(HashInts[0], HashInts[1], HashInts[2], HashInts[3]);
+}
 
 bool TrySavePackage(UPackage* Package)
 {
@@ -1659,8 +1664,6 @@ bool TrySavePackage(UPackage* Package)
 	uint32 SaveFlags = SAVE_NoError // Do not crash the SaveServer on an error
 		| SAVE_BulkDataByReference	// EditorDomain saves reference bulkdata from the WorkspaceDomain rather than duplicating it
 		| SAVE_Async				// SavePackage support for PackageWriter is only implemented with SAVE_Async
-		// EDITOR_DOMAIN_TODO: Add a a save flag that specifies the creation of a deterministic guid
-		// | SAVE_KeepGUID;			// Prevent indeterminism by keeping the Guid
 		;
 
 	TArray<UObject*> PackageObjects;
@@ -1702,6 +1705,9 @@ bool TrySavePackage(UPackage* Package)
 	SaveArgs.SaveFlags = SaveFlags;
 	SaveArgs.bSlowTask = false;
 	SaveArgs.SavePackageContext = &SavePackageContext;
+	PRAGMA_DISABLE_DEPRECATION_WARNINGS;
+	SaveArgs.OutputPackageGuid.Emplace(IoHashToGuid(PackageDigest.Hash));
+	PRAGMA_ENABLE_DEPRECATION_WARNINGS;
 	FSavePackageResultStruct Result = GEditor->Save(Package, nullptr, TEXT("EditorDomainPackageWriter"), SaveArgs);
 	if (Result.Result != ESavePackageResult::Success)
 	{
diff --git a/Engine/Source/Runtime/CoreUObject/Private/UObject/SavePackage2.cpp b/Engine/Source/Runtime/CoreUObject/Private/UObject/SavePackage2.cpp
--- a/Engine/Source/Runtime/CoreUObject/Private/UObject/SavePackage2.cpp
+++ b/Engine/Source/Runtime/CoreUObject/Private/UObject/SavePackage2.cpp
@@ -1003,9 +1003,18 @@ ESavePackageResult BuildLinker(FSaveContext& SaveContext)
 		SaveContext.UpdatePackageLinkerVersions();
 	}
 
-	PRAGMA_DISABLE_DEPRECATION_WARNINGS
-	SaveContext.GetLinker()->Summary.Guid = SaveContext.IsKeepGuid() ? SaveContext.GetPackage()->GetGuid() : SaveContext.GetPackage()->MakeNewGuid();
-	PRAGMA_ENABLE_DEPRECATION_WARNINGS
+	PRAGMA_DISABLE_DEPRECATION_WARNINGS;
+	if (SaveContext.GetSaveArgs().OutputPackageGuid)
+	{
+		SaveContext.GetLinker()->Summary.Guid = *SaveContext.GetSaveArgs().OutputPackageGuid;
+		
+	}
+	else
+	{
+		SaveContext.GetLinker()->Summary.Guid = SaveContext.IsKeepGuid() ?
+			SaveContext.GetPackage()->GetGuid() : SaveContext.GetPackage()->MakeNewGuid();
+	}
+	PRAGMA_ENABLE_DEPRECATION_WARNINGS;
 #if WITH_EDITORONLY_DATA
 	SaveContext.GetLinker()->Summary.PersistentGuid = SaveContext.GetPackage()->GetPersistentGuid();
 #endif
diff --git a/Engine/Source/Runtime/CoreUObject/Public/UObject/SavePackage.h b/Engine/Source/Runtime/CoreUObject/Public/UObject/SavePackage.h
--- a/Engine/Source/Runtime/CoreUObject/Public/UObject/SavePackage.h
+++ b/Engine/Source/Runtime/CoreUObject/Public/UObject/SavePackage.h
@@ -11,6 +11,7 @@
 #include "Delegates/Delegate.h"
 #include "Logging/LogMacros.h"
 #include "Misc/DateTime.h"
+#include "Misc/Optional.h"
 #include "Misc/OutputDeviceError.h"
 #include "ObjectMacros.h"
 #include "Serialization/ArchiveCookData.h"
@@ -89,6 +90,8 @@ struct FSavePackageArgs
 	FSavePackageContext* SavePackageContext = nullptr;
 	UE_DEPRECATED(5.0, "FArchiveDiffMap is no longer used; it is now implemented by DiffPackageWriter.")
 	FArchiveDiffMap* DiffMap = nullptr;
+	UE_DEPRECATED(4.27, "UPackage::Guid has not been used by the engine for a long time and it will be removed.")
+	TOptional<FGuid> OutputPackageGuid;
 
 	PRAGMA_DISABLE_DEPRECATION_WARNINGS;
 	FSavePackageArgs() = default;
