Added external actors support to EditorAssetUtils::FindAssetDataFromAnyPath, by first trying to resolve the asset data before attempting to get to the owning object.

#rnx

diff --git a/Engine/Source/Editor/UnrealEd/Private/Subsystems/EditorAssetSubsystem.cpp b/Engine/Source/Editor/UnrealEd/Private/Subsystems/EditorAssetSubsystem.cpp
--- a/Engine/Source/Editor/UnrealEd/Private/Subsystems/EditorAssetSubsystem.cpp
+++ b/Engine/Source/Editor/UnrealEd/Private/Subsystems/EditorAssetSubsystem.cpp
@@ -51,18 +51,27 @@ namespace UE::EditorAssetUtils
 
 	TValueOrError<FAssetData, FString> FindAssetDataFromAnyPath(const FString& AnyAssetPath)
 	{
-		FString FailureReason;
-		FString ObjectPath = EditorScriptingHelpers::ConvertAnyPathToObjectPath(AnyAssetPath, FailureReason);
-		if (ObjectPath.IsEmpty())
-		{
-			return MakeError(FailureReason);
-		}
-
 		FAssetRegistryModule& AssetRegistryModule = FModuleManager::LoadModuleChecked<FAssetRegistryModule>("AssetRegistry");
-		FAssetData AssetData = AssetRegistryModule.Get().GetAssetByObjectPath(FSoftObjectPath(ObjectPath));
+
+		// First remove class name from asset path and see if an asset exists for it
+		FString TextPath = FPackageName::ExportTextPathToObjectPath(AnyAssetPath);
+		FAssetData AssetData = AssetRegistryModule.Get().GetAssetByObjectPath(FSoftObjectPath(TextPath));
+
+		// If not, try to convert the path to an object path, potentially resolving the owning object
 		if (!AssetData.IsValid())
 		{
-			return MakeError(FString::Printf(TEXT("The AssetData '%s' could not be found in the Asset Registry."), *ObjectPath));
+			FString FailureReason;
+			FString ObjectPath = EditorScriptingHelpers::ConvertAnyPathToObjectPath(AnyAssetPath, FailureReason);
+			if (ObjectPath.IsEmpty())
+			{
+				return MakeError(FailureReason);
+			}
+
+			AssetData = AssetRegistryModule.Get().GetAssetByObjectPath(FSoftObjectPath(ObjectPath));
+			if (!AssetData.IsValid())
+			{
+				return MakeError(FString::Printf(TEXT("The AssetData '%s' could not be found in the Asset Registry."), *ObjectPath));
+			}
 		}
 		
 		return MakeValue(AssetData);
