UEngine::FindAndPrintStaleReferencesToObject : Fix EPrintStaleReferencesOptions::Log being considered as EPrintStaleReferencesOptions::Fatal because of bad condition

#rnx

diff --git a/Engine/Source/Editor/UnrealEd/Private/EditorLevelUtils.cpp b/Engine/Source/Editor/UnrealEd/Private/EditorLevelUtils.cpp
--- a/Engine/Source/Editor/UnrealEd/Private/EditorLevelUtils.cpp
+++ b/Engine/Source/Editor/UnrealEd/Private/EditorLevelUtils.cpp
@@ -908,7 +908,7 @@ bool UEditorLevelUtils::RemoveLevelsFromWorld(TArray<ULevel*> InLevels, bool bCl
 		PrintStaleReferencesOptions = UObjectBaseUtility::IsPendingKillEnabled() ? EPrintStaleReferencesOptions::Fatal : (EPrintStaleReferencesOptions::Error | EPrintStaleReferencesOptions::Ensure);
 	}
 	bool bFailed = !CheckPackage(PackageNames, PrintStaleReferencesOptions);
-	if (bFailed && !(PrintStaleReferencesOptions & EPrintStaleReferencesOptions::Fatal))
+	if (bFailed && ((PrintStaleReferencesOptions & EPrintStaleReferencesOptions::Fatal) != EPrintStaleReferencesOptions::Fatal))
 	{
 		// We tried avoiding clearing the Transaction buffer but it failed. Plan B.
 		GEditor->Cleanse(bClearSelection, false, TransResetText, true);
diff --git a/Engine/Source/Runtime/Engine/Private/UnrealEngine.cpp b/Engine/Source/Runtime/Engine/Private/UnrealEngine.cpp
--- a/Engine/Source/Runtime/Engine/Private/UnrealEngine.cpp
+++ b/Engine/Source/Runtime/Engine/Private/UnrealEngine.cpp
@@ -15930,7 +15930,7 @@ TArray<FString> UEngine::FindAndPrintStaleReferencesToObjects(TConstArrayView<UO
 		}
 	}
 
-	if (!!(Options & EPrintStaleReferencesOptions::Fatal))
+	if ((Options & EPrintStaleReferencesOptions::Fatal) == EPrintStaleReferencesOptions::Fatal)
 	{
 		UE_LOG(LogLoad, Fatal, TEXT("Fatal world leaks detected. Logging first error, check logs for additional information") LINE_TERMINATOR TEXT("%s"), 
 			Paths.Num() ? *Paths[0] : TEXT("No paths to leaked objects found!"));
