Add defensive code for nullptr.
Fixes issue with a material function not loaded.
Still need to investigate how nullptr can be in asset.
#rnx

diff --git a/Engine/Source/Editor/UnrealEd/Private/MaterialGraph.cpp b/Engine/Source/Editor/UnrealEd/Private/MaterialGraph.cpp
--- a/Engine/Source/Editor/UnrealEd/Private/MaterialGraph.cpp
+++ b/Engine/Source/Editor/UnrealEd/Private/MaterialGraph.cpp
@@ -517,10 +517,10 @@ void UMaterialGraph::LinkGraphNodesFromMaterial()
 					else if (UMaterialExpressionReroute* CompositeReroute = Cast<UMaterialExpressionReroute>(ExpressionInputs[Pin->SourceIndex]->Expression))
 					{
 						// This is an unseen composite reroute expression, find the actual expression output to connect to.
-						UMaterialExpressionComposite* OwningComposite = CastChecked<UMaterialExpressionComposite>(CompositeReroute->SubgraphExpression);
+						UMaterialExpressionComposite* OwningComposite = Cast<UMaterialExpressionComposite>(CompositeReroute->SubgraphExpression);
 						
 						// If the input- and output expressions are valid, look for the output pin in the reroute lists and make a link to the current Pin.
-						if (OwningComposite->InputExpressions && OwningComposite->OutputExpressions)
+						if (OwningComposite && OwningComposite->InputExpressions && OwningComposite->OutputExpressions)
 						{
 							UMaterialGraphNode* OutputGraphNode;
 							int32 OutputPinIndex = OwningComposite->InputExpressions->ReroutePins.FindLastByPredicate(ExpressionMatchesPredicate(CompositeReroute));
