EditorDomain: Now that AssetRegistry GetAncestorClassNames provides fullpaths, use that information to allow inheriting editordomain settings from BlueprintGeneratedClasses.
#rnx

diff --git a/Engine/Source/Editor/UnrealEd/Private/EditorDomain/EditorDomainUtils.cpp b/Engine/Source/Editor/UnrealEd/Private/EditorDomain/EditorDomainUtils.cpp
--- a/Engine/Source/Editor/UnrealEd/Private/EditorDomain/EditorDomainUtils.cpp
+++ b/Engine/Source/Editor/UnrealEd/Private/EditorDomain/EditorDomainUtils.cpp
@@ -560,19 +560,23 @@ FClassDigestData* FPrecacheClassDigest::GetRecursive(FName ClassName, bool bAllo
 		// if its parent classes are not, then we will not be able to propagate information from the parent classes; wait on the class to be parsed
 		AncestorPathNames.Reset();
 		IAssetRegistry::Get()->GetAncestorClassNames(ClassObjectPathName, AncestorPathNames);
-		for (const FTopLevelAssetPath& PathName : AncestorPathNames)
+		if (!AncestorPathNames.IsEmpty())
 		{
-			// TODO_EDITORDOMAIN: For robustness and performance, we need the AssetRegistry to return FullPathNames rather than ShortNames
-			// For now, we lookup each shortname using FindObject, and do not handle propagating data from blueprint classes to child classes
-			if (UStruct* CurrentStruct = FindObject<UStruct>(PathName))
+			NameStringBuffer.Reset();
+			AncestorPathNames[0].AppendString(NameStringBuffer);
+			ParentName = FName(*NameStringBuffer);
+			for (const FTopLevelAssetPath& PathName : AncestorPathNames)
 			{
-				NameStringBuffer.Reset();
-				CurrentStruct->GetPathName(nullptr, NameStringBuffer);
-				if (FPackageName::IsScriptPackage(NameStringBuffer))
+				if (UStruct* CurrentStruct = FindObject<UStruct>(PathName))
 				{
-					ParentName = FName(*NameStringBuffer);
-					Struct = CurrentStruct;
-					break;
+					NameStringBuffer.Reset();
+					CurrentStruct->GetPathName(nullptr, NameStringBuffer);
+					if (FPackageName::IsScriptPackage(NameStringBuffer))
+					{
+						Struct = CurrentStruct;
+						DigestData->ResolvedClosestNative = FName(*NameStringBuffer);
+						break;
+					}
 				}
 			}
 		}
@@ -580,7 +584,10 @@ FClassDigestData* FPrecacheClassDigest::GetRecursive(FName ClassName, bool bAllo
 		{
 			ParentName = FName(*UObject::StaticClass()->GetPathName());
 		}
-		DigestData->ResolvedClosestNative = ParentName;
+		if (DigestData->ResolvedClosestNative.IsNone())
+		{
+			DigestData->ResolvedClosestNative = FName(*UObject::StaticClass()->GetPathName());
+		}
 	}
 	check(!DigestData->ResolvedClosestNative.IsNone());
 
@@ -603,8 +610,8 @@ FClassDigestData* FPrecacheClassDigest::GetRecursive(FName ClassName, bool bAllo
 	TArray<FName> ConstructClasses; // Not a class scratch variable because we use it across recursive calls
 	if (!ParentName.IsNone())
 	{
-		// CoreRedirects are expected to act only on import classes from packages; they are not expected to act on the parent class pointer
-		// of a native class, which is authoritative, so set bAllowRedirects = false
+		// Set bAllowRedirects = false. ParentName already has been CoreRedirected, because it is from the native
+		// struct's parent class pointer or it is from GetAncestorClassNames which gives the redirected ancestors.
 		FClassDigestData* ParentDigest = GetRecursive(ParentName, false /* bAllowRedirects */);
 		// The map has possibly been modified so we need to recalculate the address of ClassName's DigestData
 		DigestData = &ClassDigests.FindChecked(ClassName);
