Fix bad merge

diff --git a/Engine/Source/Editor/UnrealEd/Private/Subsystems/EditorAssetSubsystem.cpp b/Engine/Source/Editor/UnrealEd/Private/Subsystems/EditorAssetSubsystem.cpp
--- a/Engine/Source/Editor/UnrealEd/Private/Subsystems/EditorAssetSubsystem.cpp
+++ b/Engine/Source/Editor/UnrealEd/Private/Subsystems/EditorAssetSubsystem.cpp
@@ -67,9 +67,6 @@ namespace UE::EditorAssetUtils
 		
 		return MakeValue(AssetData);
 	}
-		
-		return MakeValue(AssetData);
-	}
 
 	TValueOrError<UObject*, FString> LoadAssetFromData(const FAssetData& AssetData)
 	{
