Changing PollPendingCookedPlatformDatas logic to always allow a GPollAsyncPeriod "downtime" before polling. This prevents the situation where the poll takes the full time and we saturate the game thread.

diff --git a/Engine/Source/Editor/UnrealEd/Private/Cooker/CookPackageData.cpp b/Engine/Source/Editor/UnrealEd/Private/Cooker/CookPackageData.cpp
--- a/Engine/Source/Editor/UnrealEd/Private/Cooker/CookPackageData.cpp
+++ b/Engine/Source/Editor/UnrealEd/Private/Cooker/CookPackageData.cpp
@@ -2630,12 +2630,13 @@ void FPackageDatas::PollPendingCookedPlatformDatas(bool bForce, double& LastCook
 	if (!bForce)
 	{
 		// ProcessAsyncResults and IsCachedCookedPlatformDataLoaded can be expensive to call
-		// Cap the frequency at which we call them.
+		// Cap the frequency at which we call them. We only update the last poll time at completion
+		// so that we don't suddenly saturate the game thread by making derived data key strings
+		// when the time to do the polls increases to GPollAsyncPeriod.
 		if (CurrentTime < LastPollAsyncTime + GPollAsyncPeriod)
 		{
 			return;
 		}
-		LastPollAsyncTime = CurrentTime;
 	}
 
 	GShaderCompilingManager->ProcessAsyncResults(true /* bLimitExecutionTime */,
@@ -2660,6 +2661,8 @@ void FPackageDatas::PollPendingCookedPlatformDatas(bool bForce, double& LastCook
 			++Index;
 		}
 	}
+
+	LastPollAsyncTime = CurrentTime;
 }
 
 TArray<FPackageData*>::RangedForIteratorType FPackageDatas::begin()
