Added external actors support to EditorAssetUtils.

#rnx

diff --git a/Engine/Source/Editor/UnrealEd/Private/EditorScriptingHelpers.cpp b/Engine/Source/Editor/UnrealEd/Private/EditorScriptingHelpers.cpp
--- a/Engine/Source/Editor/UnrealEd/Private/EditorScriptingHelpers.cpp
+++ b/Engine/Source/Editor/UnrealEd/Private/EditorScriptingHelpers.cpp
@@ -55,6 +55,116 @@ namespace EditorScriptingHelpersInternal
 			return AnyAssetPath.Mid(FoundIndex + 1);
 		}
 	}
+
+	FString ConvertAnyPathToObjectPathInternal(const FString& AnyAssetPath, bool bIncludeSubObject, FString& OutFailureReason)
+	{
+		if (AnyAssetPath.Len() < 2) // minimal length to have /G
+		{
+			OutFailureReason = FString::Printf(TEXT("Can't convert the path '%s' because the Root path need to be specified. ie /Game/"), *AnyAssetPath);
+			return FString();
+		}
+
+		// Remove class name from Reference Path
+		FString TextPath = FPackageName::ExportTextPathToObjectPath(AnyAssetPath);
+
+		// Remove class name Fullname
+		TextPath = EditorScriptingHelpersInternal::RemoveFullName(TextPath, OutFailureReason);
+		if (TextPath.IsEmpty())
+		{
+			return FString();
+		}
+
+		// Extract the subobject path if any
+		FString SubObjectPath;
+		int32 SubObjectDelimiterIdx;
+		if (TextPath.FindChar(SUBOBJECT_DELIMITER_CHAR, SubObjectDelimiterIdx))
+		{
+			SubObjectPath = TextPath.Mid(SubObjectDelimiterIdx + 1);
+			TextPath.LeftInline(SubObjectDelimiterIdx);
+		}
+
+		// Convert \ to /
+		TextPath.ReplaceInline(TEXT("\\"), TEXT("/"), ESearchCase::CaseSensitive);
+		FPaths::RemoveDuplicateSlashes(TextPath);
+
+		// Get asset full name, i.e."PackageName.ObjectName:InnerAssetName.2ndInnerAssetName" from "/Game/Folder/PackageName.ObjectName:InnerAssetName.2ndInnerAssetName"
+		FString AssetFullName;
+		{
+			// Get everything after the last slash
+			int32 IndexOfLastSlash = INDEX_NONE;
+			TextPath.FindLastChar('/', IndexOfLastSlash);
+
+			FString Folders = TextPath.Left(IndexOfLastSlash);
+			// Test for invalid characters
+			if (!EditorScriptingHelpers::IsAValidPath(Folders, INVALID_LONGPACKAGE_CHARACTERS, OutFailureReason))
+			{
+				return FString();
+			}
+
+			AssetFullName = TextPath.Mid(IndexOfLastSlash + 1);
+		}
+
+		// Get the object name
+		FString ObjectName = FPackageName::ObjectPathToObjectName(AssetFullName);
+		if (ObjectName.IsEmpty())
+		{
+			OutFailureReason = FString::Printf(TEXT("Can't convert the path '%s' because it doesn't contain an asset name."), *AnyAssetPath);
+			return FString();
+		}
+
+		// Test for invalid characters
+		if (!EditorScriptingHelpers::IsAValidPath(ObjectName, INVALID_OBJECTNAME_CHARACTERS, OutFailureReason))
+		{
+			return FString();
+		}
+
+		// Confirm that we have a valid Root Package and get the valid PackagePath /Game/MyFolder/MyAsset
+		FString PackagePath;
+		if (!FPackageName::TryConvertFilenameToLongPackageName(TextPath, PackagePath, &OutFailureReason))
+		{
+			return FString();
+		}
+
+		if (PackagePath.Len() == 0)
+		{
+			OutFailureReason = FString::Printf(TEXT("Can't convert path '%s' because the PackagePath is empty."), *AnyAssetPath);
+			return FString();
+		}
+
+		if (PackagePath[0] != TEXT('/'))
+		{
+			OutFailureReason = FString::Printf(TEXT("Can't convert path '%s' because the PackagePath '%s' doesn't start with a '/'."), *AnyAssetPath, *PackagePath);
+			return FString();
+		}
+
+		FString ObjectPath = FString::Printf(TEXT("%s.%s"), *PackagePath, *ObjectName);
+
+		if (bIncludeSubObject && !SubObjectPath.IsEmpty())
+		{
+			ObjectPath += TEXT(":");
+			ObjectPath += SubObjectPath;
+		}
+
+		if (FPackageName::IsScriptPackage(ObjectPath))
+		{
+			OutFailureReason = FString::Printf(TEXT("Can't convert the path '%s' because it start with /Script/"), *AnyAssetPath);
+			return FString();
+		}
+		if (FPackageName::IsMemoryPackage(ObjectPath))
+		{
+			OutFailureReason = FString::Printf(TEXT("Can't convert the path '%s' because it start with /Memory/"), *AnyAssetPath);
+			return FString();
+		}
+
+		// Confirm that the PackagePath starts with a valid root
+		if (!EditorScriptingHelpers::HasValidRoot(PackagePath))
+		{
+			OutFailureReason = FString::Printf(TEXT("Can't convert the path '%s' because it does not map to a root."), *AnyAssetPath);
+			return FString();
+		}
+
+		return ObjectPath;
+	}
 }
 
 bool EditorScriptingHelpers::CheckIfInEditorAndPIE()
@@ -246,104 +356,10 @@ bool EditorScriptingHelpers::IsAValidPathForCreateNewAsset(const FString& Object
 
 FString EditorScriptingHelpers::ConvertAnyPathToObjectPath(const FString& AnyAssetPath, FString& OutFailureReason)
 {
-	if (AnyAssetPath.Len() < 2) // minimal length to have /G
-	{
-		OutFailureReason = FString::Printf(TEXT("Can't convert the path '%s' because the Root path need to be specified. ie /Game/"), *AnyAssetPath);
-		return FString();
-	}
-
-	// Remove class name from Reference Path
-	FString TextPath = FPackageName::ExportTextPathToObjectPath(AnyAssetPath);
-
-	// Remove class name Fullname
-	TextPath = EditorScriptingHelpersInternal::RemoveFullName(TextPath, OutFailureReason);
-	if (TextPath.IsEmpty())
-	{
-		return FString();
-	}
-
-	// Extract the subobject path if any
-	FString SubObjectPath;
-	int32 SubObjectDelimiterIdx;
-	if (TextPath.FindChar(SUBOBJECT_DELIMITER_CHAR, SubObjectDelimiterIdx))
-	{
-		SubObjectPath = TextPath.Mid(SubObjectDelimiterIdx + 1);
-		TextPath.LeftInline(SubObjectDelimiterIdx);
-	}
-
-	// Convert \ to /
-	TextPath.ReplaceInline(TEXT("\\"), TEXT("/"), ESearchCase::CaseSensitive);
-	FPaths::RemoveDuplicateSlashes(TextPath);
-
-	// Get asset full name, i.e."PackageName.ObjectName:InnerAssetName.2ndInnerAssetName" from "/Game/Folder/PackageName.ObjectName:InnerAssetName.2ndInnerAssetName"
-	FString AssetFullName;
-	{
-		// Get everything after the last slash
-		int32 IndexOfLastSlash = INDEX_NONE;
-		TextPath.FindLastChar('/', IndexOfLastSlash);
-
-		FString Folders = TextPath.Left(IndexOfLastSlash);
-		// Test for invalid characters
-		if (!IsAValidPath(Folders, INVALID_LONGPACKAGE_CHARACTERS, OutFailureReason))
-		{
-			return FString();
-		}
-
-		AssetFullName = TextPath.Mid(IndexOfLastSlash + 1);
-	}
-
-	// Get the object name
-	FString ObjectName = FPackageName::ObjectPathToObjectName(AssetFullName);
-	if (ObjectName.IsEmpty())
-	{
-		OutFailureReason = FString::Printf(TEXT("Can't convert the path '%s' because it doesn't contain an asset name."), *AnyAssetPath);
-		return FString();
-	}
-
-	// Test for invalid characters
-	if (!IsAValidPath(ObjectName, INVALID_OBJECTNAME_CHARACTERS, OutFailureReason))
-	{
-		return FString();
-	}
-
-	// Confirm that we have a valid Root Package and get the valid PackagePath /Game/MyFolder/MyAsset
-	FString PackagePath;
-	if (!FPackageName::TryConvertFilenameToLongPackageName(TextPath, PackagePath, &OutFailureReason))
-	{
-		return FString();
-	}
-
-	if (PackagePath.Len() == 0)
-	{
-		OutFailureReason = FString::Printf(TEXT("Can't convert path '%s' because the PackagePath is empty."), *AnyAssetPath);
-		return FString();
-	}
-
-	if (PackagePath[0] != TEXT('/'))
-	{
-		OutFailureReason = FString::Printf(TEXT("Can't convert path '%s' because the PackagePath '%s' doesn't start with a '/'."), *AnyAssetPath, *PackagePath);
-		return FString();
-	}
-
-	FString ObjectPath = FString::Printf(TEXT("%s.%s"), *PackagePath, *ObjectName);
-
-	if (FPackageName::IsScriptPackage(ObjectPath))
-	{
-		OutFailureReason = FString::Printf(TEXT("Can't convert the path '%s' because it start with /Script/"), *AnyAssetPath);
-		return FString();
-	}
-	if (FPackageName::IsMemoryPackage(ObjectPath))
-	{
-		OutFailureReason = FString::Printf(TEXT("Can't convert the path '%s' because it start with /Memory/"), *AnyAssetPath);
-		return FString();
-	}
-
-	// Confirm that the PackagePath starts with a valid root
-	if (!HasValidRoot(PackagePath))
-	{
-		OutFailureReason = FString::Printf(TEXT("Can't convert the path '%s' because it does not map to a root."), *AnyAssetPath);
-		return FString();
-	}
+	return EditorScriptingHelpersInternal::ConvertAnyPathToObjectPathInternal(AnyAssetPath, false, OutFailureReason);
+}
 
-	return ObjectPath;
+FString EditorScriptingHelpers::ConvertAnyPathToSubObjectPath(const FString& AnyAssetPath, FString& OutFailureReason)
+{
+	return EditorScriptingHelpersInternal::ConvertAnyPathToObjectPathInternal(AnyAssetPath, true, OutFailureReason);
 }
\ No newline at end of file
diff --git a/Engine/Source/Editor/UnrealEd/Private/Subsystems/EditorAssetSubsystem.cpp b/Engine/Source/Editor/UnrealEd/Private/Subsystems/EditorAssetSubsystem.cpp
--- a/Engine/Source/Editor/UnrealEd/Private/Subsystems/EditorAssetSubsystem.cpp
+++ b/Engine/Source/Editor/UnrealEd/Private/Subsystems/EditorAssetSubsystem.cpp
@@ -52,7 +52,7 @@ namespace UE::EditorAssetUtils
 	TValueOrError<FAssetData, FString> FindAssetDataFromAnyPath(const FString& AnyAssetPath)
 	{
 		FString FailureReason;
-		FString ObjectPath = EditorScriptingHelpers::ConvertAnyPathToObjectPath(AnyAssetPath, FailureReason);
+		FString ObjectPath = EditorScriptingHelpers::ConvertAnyPathToSubObjectPath(AnyAssetPath, FailureReason);
 		if (ObjectPath.IsEmpty())
 		{
 			return MakeError(FailureReason);
@@ -62,7 +62,17 @@ namespace UE::EditorAssetUtils
 		FAssetData AssetData = AssetRegistryModule.Get().GetAssetByObjectPath(FSoftObjectPath(ObjectPath));
 		if (!AssetData.IsValid())
 		{
-			return MakeError(FString::Printf(TEXT("The AssetData '%s' could not be found in the Asset Registry."), *ObjectPath));
+			ObjectPath = EditorScriptingHelpers::ConvertAnyPathToObjectPath(AnyAssetPath, FailureReason);
+			if (ObjectPath.IsEmpty())
+			{
+				return MakeError(FailureReason);
+			}
+
+			AssetData = AssetRegistryModule.Get().GetAssetByObjectPath(*ObjectPath);
+			if (!AssetData.IsValid())
+			{
+				return MakeError(FString::Printf(TEXT("The AssetData '%s' could not be found in the Asset Registry."), *ObjectPath));
+			}
 		}
 		
 		return MakeValue(AssetData);
@@ -445,7 +455,7 @@ bool UEditorAssetSubsystem::DoesAssetExist(const FString& AssetPath)
 	}
 
 	FString FailureReason;
-	FString ObjectPath = EditorScriptingHelpers::ConvertAnyPathToObjectPath(AssetPath, FailureReason);
+	FString ObjectPath = EditorScriptingHelpers::ConvertAnyPathToSubObjectPath(AssetPath, FailureReason);
 	if (ObjectPath.IsEmpty())
 	{
 		UE_LOG(LogEditorAssetSubsystem, Error, TEXT("DoesAssetExist failed: %s"), *FailureReason);
@@ -456,7 +466,18 @@ bool UEditorAssetSubsystem::DoesAssetExist(const FString& AssetPath)
 	FAssetData AssetData = AssetRegistryModule.Get().GetAssetByObjectPath(FSoftObjectPath(ObjectPath));
 	if (!AssetData.IsValid())
 	{
-		return false;
+		ObjectPath = EditorScriptingHelpers::ConvertAnyPathToObjectPath(AssetPath, FailureReason);
+		if (ObjectPath.IsEmpty())
+		{
+			UE_LOG(LogEditorAssetSubsystem, Error, TEXT("DoesAssetExist failed: %s"), *FailureReason);
+			return false;
+		}
+
+		AssetData = AssetRegistryModule.Get().GetAssetByObjectPath(*ObjectPath);
+		if (!AssetData.IsValid())
+		{
+			return false;
+		}
 	}
 
 	return true;
@@ -476,7 +497,7 @@ bool UEditorAssetSubsystem::DoAssetsExist(const TArray<FString>& AssetPaths)
 
 	for (const FString& Path : AssetPaths)
 	{
-		FString ObjectPath = EditorScriptingHelpers::ConvertAnyPathToObjectPath(Path, FailureReason);
+		FString ObjectPath = EditorScriptingHelpers::ConvertAnyPathToSubObjectPath(Path, FailureReason);
 		if (ObjectPath.IsEmpty())
 		{
 			UE_LOG(LogEditorAssetSubsystem, Error, TEXT("DoAssetsExist failed: %s"), *FailureReason);
@@ -486,7 +507,18 @@ bool UEditorAssetSubsystem::DoAssetsExist(const TArray<FString>& AssetPaths)
 		FAssetData AssetData = AssetRegistryModule.Get().GetAssetByObjectPath(FSoftObjectPath(ObjectPath));
 		if (!AssetData.IsValid())
 		{
-			return false;
+			ObjectPath = EditorScriptingHelpers::ConvertAnyPathToObjectPath(Path, FailureReason);
+			if (ObjectPath.IsEmpty())
+			{
+				UE_LOG(LogEditorAssetSubsystem, Error, TEXT("DoAssetsExist failed: %s"), *FailureReason);
+				return false;
+			}
+
+			AssetData = AssetRegistryModule.Get().GetAssetByObjectPath(*ObjectPath);
+			if (!AssetData.IsValid())
+			{
+				return false;
+			}
 		}
 	}
 	return true;
@@ -503,20 +535,29 @@ TArray<FString> UEditorAssetSubsystem::FindPackageReferencersForAsset(const FStr
 	}
 
 	FString FailureReason;
-	FString AssetPath = EditorScriptingHelpers::ConvertAnyPathToObjectPath(AnyAssetPath, FailureReason);
+	FString AssetPath = EditorScriptingHelpers::ConvertAnyPathToSubObjectPath(AnyAssetPath, FailureReason);
 	if (AssetPath.IsEmpty())
 	{
 		UE_LOG(LogEditorAssetSubsystem, Error, TEXT("FindAssetPackageReferencers failed: %s"), *FailureReason);
 		return Result;
 	}
 
-	// Find the reference in packages. Load them to confirm the reference.
-	TArray<FName> PackageReferencers;
+	FAssetRegistryModule& AssetRegistryModule = FModuleManager::LoadModuleChecked<FAssetRegistryModule>(TEXT("AssetRegistry"));
+	FAssetData AssetData = AssetRegistryModule.Get().GetAssetByObjectPath(*AnyAssetPath);
+	if (!AssetData.IsValid())
 	{
-		FAssetRegistryModule& AssetRegistryModule = FModuleManager::LoadModuleChecked<FAssetRegistryModule>(TEXT("AssetRegistry"));
-		AssetRegistryModule.Get().GetReferencers(*FPackageName::ObjectPathToPackageName(AssetPath), PackageReferencers, UE::AssetRegistry::EDependencyCategory::Package);
+		FString ObjectPath = EditorScriptingHelpers::ConvertAnyPathToObjectPath(AnyAssetPath, FailureReason);
+		if (ObjectPath.IsEmpty())
+		{
+			UE_LOG(LogEditorAssetSubsystem, Error, TEXT("FindAssetPackageReferencers failed: %s"), *FailureReason);
+			return Result;
+		}
 	}
 
+	// Find the reference in packages. Load them to confirm the reference.
+	TArray<FName> PackageReferencers;
+	AssetRegistryModule.Get().GetReferencers(*FPackageName::ObjectPathToPackageName(AssetPath), PackageReferencers, UE::AssetRegistry::EDependencyCategory::Package);
+
 	if (bLoadAssetsToConfirm)
 	{
 		// Load the asset to confirm 
diff --git a/Engine/Source/Editor/UnrealEd/Public/EditorScriptingHelpers.h b/Engine/Source/Editor/UnrealEd/Public/EditorScriptingHelpers.h
--- a/Engine/Source/Editor/UnrealEd/Public/EditorScriptingHelpers.h
+++ b/Engine/Source/Editor/UnrealEd/Public/EditorScriptingHelpers.h
@@ -38,4 +38,11 @@ namespace EditorScriptingHelpers
 	 * @note: Object name is inferred from package name when missing
 	 */
 	FString UNREALED_API ConvertAnyPathToObjectPath(const FString& AssetPath, FString& OutFailureReason);
+
+	/*
+	 * From "AssetClass'/Game/Folder/Package.Asset'", "AssetClass /Game/Folder/Package.Asset", "/Game/Folder/Package.Asset", "/Game/Folder/MyAsset" "/Game/Folder/Package.Asset:InnerAsset.2ndInnerAsset"
+	 * and convert to "/Game/Folder/Package.Asset" or "/Game/Folder/Package.Asset:InnerAsset"
+	 * @note: Object name is inferred from package name when missing
+	 */
+	FString UNREALED_API ConvertAnyPathToSubObjectPath(const FString& AssetPath, FString& OutFailureReason);
 }
\ No newline at end of file
