Fix asset registry deprecation that wasn't present in source stream.

diff --git a/Engine/Source/Editor/UnrealEd/Private/Subsystems/EditorAssetSubsystem.cpp b/Engine/Source/Editor/UnrealEd/Private/Subsystems/EditorAssetSubsystem.cpp
--- a/Engine/Source/Editor/UnrealEd/Private/Subsystems/EditorAssetSubsystem.cpp
+++ b/Engine/Source/Editor/UnrealEd/Private/Subsystems/EditorAssetSubsystem.cpp
@@ -68,7 +68,7 @@ namespace UE::EditorAssetUtils
 				return MakeError(FailureReason);
 			}
 
-			AssetData = AssetRegistryModule.Get().GetAssetByObjectPath(*ObjectPath);
+			AssetData = AssetRegistryModule.Get().GetAssetByObjectPath(FSoftObjectPath(ObjectPath));
 			if (!AssetData.IsValid())
 			{
 				return MakeError(FString::Printf(TEXT("The AssetData '%s' could not be found in the Asset Registry."), *ObjectPath));
@@ -473,7 +473,7 @@ bool UEditorAssetSubsystem::DoesAssetExist(const FString& AssetPath)
 			return false;
 		}
 
-		AssetData = AssetRegistryModule.Get().GetAssetByObjectPath(*ObjectPath);
+		AssetData = AssetRegistryModule.Get().GetAssetByObjectPath(FSoftObjectPath(ObjectPath));
 		if (!AssetData.IsValid())
 		{
 			return false;
@@ -514,7 +514,7 @@ bool UEditorAssetSubsystem::DoAssetsExist(const TArray<FString>& AssetPaths)
 				return false;
 			}
 
-			AssetData = AssetRegistryModule.Get().GetAssetByObjectPath(*ObjectPath);
+			AssetData = AssetRegistryModule.Get().GetAssetByObjectPath(FSoftObjectPath(ObjectPath));
 			if (!AssetData.IsValid())
 			{
 				return false;
@@ -543,7 +543,7 @@ TArray<FString> UEditorAssetSubsystem::FindPackageReferencersForAsset(const FStr
 	}
 
 	FAssetRegistryModule& AssetRegistryModule = FModuleManager::LoadModuleChecked<FAssetRegistryModule>(TEXT("AssetRegistry"));
-	FAssetData AssetData = AssetRegistryModule.Get().GetAssetByObjectPath(*AnyAssetPath);
+	FAssetData AssetData = AssetRegistryModule.Get().GetAssetByObjectPath(FSoftObjectPath(AnyAssetPath));
 	if (!AssetData.IsValid())
 	{
 		FString ObjectPath = EditorScriptingHelpers::ConvertAnyPathToObjectPath(AnyAssetPath, FailureReason);
