Fix Aggressive clearing of Memory References in the Make Assets Private tool

diff --git a/Engine/Source/Editor/UnrealEd/Private/AssetPrivatizeModel.cpp b/Engine/Source/Editor/UnrealEd/Private/AssetPrivatizeModel.cpp
--- a/Engine/Source/Editor/UnrealEd/Private/AssetPrivatizeModel.cpp
+++ b/Engine/Source/Editor/UnrealEd/Private/AssetPrivatizeModel.cpp
@@ -25,8 +25,15 @@ FPendingPrivateAsset::FPendingPrivateAsset(UObject* InObject, UPackage* InOwning
 
 bool FPendingPrivateAsset::IsReferenceIllegal(const FName& InReference)
 {
+	FString InReferenceString = InReference.ToString();
+	// We only care about references that can be saved to disk and committed to source control
+	if (InReferenceString.StartsWith(TEXT("/Engine/Transient")) ||
+		FPackageName::IsMemoryPackage(InReferenceString) || FPackageName::IsTempPackage(InReferenceString))
+	{
+		return false;
+	}
 	// A reference is illegal if it refers to a private asset that is not in the same mount point
-	FName ReferenceMountPointName = FPackageName::GetPackageMountPoint(InReference.ToString());
+	FName ReferenceMountPointName = FPackageName::GetPackageMountPoint(InReferenceString);
 	return !ReferenceMountPointName.IsNone() && (ReferenceMountPointName != OwningPackageMountPointName);
 }
 
@@ -53,13 +60,26 @@ void FPendingPrivateAsset::CheckForIllegalReferences()
 
 	IllegalMemoryReferences.ExternalReferences.RemoveAll([=](const FReferencerInformation& ReferenceInfo)
 		{
+			if (ReferenceInfo.Referencer->HasAnyFlags(RF_Transient))
+			{
+				return true;
+			}
 			return !IsReferenceIllegal(ReferenceInfo.Referencer->GetPackage()->GetFName());
 		});
 
 	IllegalMemoryReferences.InternalReferences.RemoveAll([=](const FReferencerInformation& ReferenceInfo)
 		{
+			if (ReferenceInfo.Referencer->HasAnyFlags(RF_Transient))
+			{
+				return true;
+			}
 			return !IsReferenceIllegal(ReferenceInfo.Referencer->GetPackage()->GetFName());
 		});
+
+	if (IllegalMemoryReferences.ExternalReferences.IsEmpty() && IllegalMemoryReferences.InternalReferences.IsEmpty())
+	{
+		bIsReferencedInMemoryByNonUndo = false;
+	}
 }
 
 FAssetPrivatizeModel::FAssetPrivatizeModel(const TArray<UObject*>& InObjectsToPrivatize)
@@ -175,8 +195,11 @@ bool FAssetPrivatizeModel::DoForcePrivatize()
 		ObjectsPrivatized++;
 	}
 
-    // Null out the illegal references
-	ObjectTools::ForceReplaceReferences(nullptr, ObjectsToPrivatize, ObjectsToPrivatizeWithin);
+	if (!ObjectsToPrivatize.IsEmpty() && !ObjectsToPrivatizeWithin.IsEmpty())
+	{
+		// Null out the illegal references
+		ObjectTools::ForceReplaceReferences(nullptr, ObjectsToPrivatize, ObjectsToPrivatizeWithin);
+	}
 
 	return true;
 }
