Group Actors: make sure to deselect the group actor when ungrouping.

#rnx

diff --git a/Engine/Source/Editor/UnrealEd/Private/GroupActor.cpp b/Engine/Source/Editor/UnrealEd/Private/GroupActor.cpp
--- a/Engine/Source/Editor/UnrealEd/Private/GroupActor.cpp
+++ b/Engine/Source/Editor/UnrealEd/Private/GroupActor.cpp
@@ -818,6 +818,7 @@ void AGroupActor::PostRemove()
 				FScopedRefreshAllBrowsers LevelRefreshAllBrowsers;
 
 				// Destroy group and clear references.
+				GEditor->SelectActor( this, /*bSelected=*/ false, /*bNotify=*/ false );
 				ULayersSubsystem* LayersSubsystem = GEditor->GetEditorSubsystem<ULayersSubsystem>();
 				LayersSubsystem->DisassociateActorFromLayers( this );
 				MyWorld->EditorDestroyActor( this, false );			
