Fix TSAN compilation in editor for some TSAN_ATOMIC globals using TGuard scopes

diff --git a/Engine/Source/Developer/IoStoreUtilities/Private/IoStoreUtilities.cpp b/Engine/Source/Developer/IoStoreUtilities/Private/IoStoreUtilities.cpp
--- a/Engine/Source/Developer/IoStoreUtilities/Private/IoStoreUtilities.cpp
+++ b/Engine/Source/Developer/IoStoreUtilities/Private/IoStoreUtilities.cpp
@@ -4928,8 +4928,8 @@ int32 Describe(
 	{
 		TRACE_CPUPROFILER_EVENT_SCOPE(GenerateReport);
 		TGuardValue<ELogTimes::Type> GuardPrintLogTimes(GPrintLogTimes, ELogTimes::None);
-		TGuardValue<bool> GuardPrintLogCategory(GPrintLogCategory, false);
-		TGuardValue<bool> GuardPrintLogVerbosity(GPrintLogVerbosity, false);
+		TGuardValue GuardPrintLogCategory(GPrintLogCategory, false);
+		TGuardValue GuardPrintLogVerbosity(GPrintLogVerbosity, false);
 
 		auto PackageObjectIndexToString = [&ScriptObjectByGlobalIdMap, &ExportByKeyMap](const FPackageDesc* Package, const FPackageObjectIndex& PackageObjectIndex, bool bIncludeName) -> FString
 		{
diff --git a/Engine/Source/Editor/UnrealEd/Private/Commandlets/PackageUtilities.cpp b/Engine/Source/Editor/UnrealEd/Private/Commandlets/PackageUtilities.cpp
--- a/Engine/Source/Editor/UnrealEd/Private/Commandlets/PackageUtilities.cpp
+++ b/Engine/Source/Editor/UnrealEd/Private/Commandlets/PackageUtilities.cpp
@@ -1686,8 +1686,8 @@ int32 UPkgInfoCommandlet::Main( const FString& Params )
 		{
 			// Turn off log categories etc as it makes diffing hard
 			TGuardValue<ELogTimes::Type> GuardPrintLogTimes(GPrintLogTimes, ELogTimes::None);
-			TGuardValue<bool> GuardPrintLogCategory(GPrintLogCategory, false);
-			TGuardValue<bool> GuardPrintLogVerbosity(GPrintLogVerbosity, false);
+			TGuardValue GuardPrintLogCategory(GPrintLogCategory, false);
+			TGuardValue GuardPrintLogVerbosity(GPrintLogVerbosity, false);
 
 			if (Linker)
 			{
diff --git a/Engine/Source/Editor/UnrealEd/Private/Cooker/DiffPackageWriter.cpp b/Engine/Source/Editor/UnrealEd/Private/Cooker/DiffPackageWriter.cpp
--- a/Engine/Source/Editor/UnrealEd/Private/Cooker/DiffPackageWriter.cpp
+++ b/Engine/Source/Editor/UnrealEd/Private/Cooker/DiffPackageWriter.cpp
@@ -259,8 +259,8 @@ void FDiffPackageWriter::ConditionallyDumpObjList()
 			ObjListExec += DumpObjListParams;
 
 			TGuardValue<ELogTimes::Type> GuardLogTimes(GPrintLogTimes, ELogTimes::None);
-			TGuardValue<bool> GuardLogCategory(GPrintLogCategory, false);
-			TGuardValue<bool> GuardPrintLogVerbosity(GPrintLogVerbosity, false);
+			TGuardValue GuardLogCategory(GPrintLogCategory, false);
+			TGuardValue GuardPrintLogVerbosity(GPrintLogVerbosity, false);
 
 			GEngine->Exec(nullptr, *ObjListExec);
 		}
@@ -284,8 +284,8 @@ void FDiffPackageWriter::ConditionallyDumpObjects()
 			}
 
 			TGuardValue<ELogTimes::Type> GuardLogTimes(GPrintLogTimes, ELogTimes::None);
-			TGuardValue<bool> GuardLogCategory(GPrintLogCategory, false);
-			TGuardValue<bool> GuardPrintLogVerbosity(GPrintLogVerbosity, false);
+			TGuardValue GuardLogCategory(GPrintLogCategory, false);
+			TGuardValue GuardPrintLogVerbosity(GPrintLogVerbosity, false);
 
 			for (const FString& Obj : AllObjects)
 			{
