Selection: ignore pending kill actors that can still be in the selection element bridge container. Normally that shouldn't happen but there's still cases of actors deleted by code in the editor that won't clear the current selection, so better safe than sorry.

#rnx

diff --git a/Engine/Source/Editor/UnrealEd/Private/Selection.cpp b/Engine/Source/Editor/UnrealEd/Private/Selection.cpp
--- a/Engine/Source/Editor/UnrealEd/Private/Selection.cpp
+++ b/Engine/Source/Editor/UnrealEd/Private/Selection.cpp
@@ -28,8 +28,7 @@ class FObjectSelectionElementBridge : public ISelectionElementBridge
 public:
 	virtual bool IsValidObjectType(const UObject* InObject) const override
 	{
-		check(InObject);
-		return true;
+		return IsValidChecked(InObject);
 	}
 
 	virtual FTypedElementHandle GetElementHandleForObject(const UObject* InObject, const bool bAllowCreate = true) const override
@@ -44,8 +43,7 @@ class FActorSelectionElementBridge : public ISelectionElementBridge
 public:
 	virtual bool IsValidObjectType(const UObject* InObject) const override
 	{
-		check(InObject);
-		return InObject->IsA<AActor>();
+		return IsValidChecked(InObject) && InObject->IsA<AActor>();
 	}
 
 	virtual FTypedElementHandle GetElementHandleForObject(const UObject* InObject, const bool bAllowCreate = true) const override
@@ -60,8 +58,7 @@ class FComponentSelectionElementBridge : public ISelectionElementBridge
 public:
 	virtual bool IsValidObjectType(const UObject* InObject) const override
 	{
-		check(InObject);
-		return InObject->IsA<UActorComponent>();
+		return IsValidChecked(InObject) && InObject->IsA<UActorComponent>();
 	}
 
 	virtual FTypedElementHandle GetElementHandleForObject(const UObject* InObject, const bool bAllowCreate = true) const override
