LevelEditor: Expose bCurrentlyEditingThroughMovementWidget to allow for more immediate lighting updates when interactively editing the level.

diff --git a/Engine/Source/Editor/UnrealEd/Private/LevelEditorViewport.cpp b/Engine/Source/Editor/UnrealEd/Private/LevelEditorViewport.cpp
--- a/Engine/Source/Editor/UnrealEd/Private/LevelEditorViewport.cpp
+++ b/Engine/Source/Editor/UnrealEd/Private/LevelEditorViewport.cpp
@@ -5178,4 +5178,10 @@ bool FLevelEditorViewportClient::GetPivotForOrbit(FVector& Pivot) const
 	return false;
 }
 
+void FLevelEditorViewportClient::SetEditingThroughMovementWidget()
+{
+	bCurrentlyEditingThroughMovementWidget = true;
+}
+
+
 #undef LOCTEXT_NAMESPACE
diff --git a/Engine/Source/Editor/UnrealEd/Public/LevelEditorViewport.h b/Engine/Source/Editor/UnrealEd/Public/LevelEditorViewport.h
--- a/Engine/Source/Editor/UnrealEd/Public/LevelEditorViewport.h
+++ b/Engine/Source/Editor/UnrealEd/Public/LevelEditorViewport.h
@@ -648,6 +648,12 @@ public:
 
 	virtual bool UseAppTime() const override { return false; }
 
+	/**
+	 * Informs the renderer that the view is being interactively edited. (ex. rotation/translation gizmo).
+	 * This state is reset on tick.
+	 */
+	void SetEditingThroughMovementWidget();
+
 protected:
 	/**
 	* Sets the state of creating a preview actor in the viewport.
