Add a warning in the log if the user try to re-import a texture that change the uasset type. And do not allow the reimport to be done.
We cannot delete the asset since it is reference by the import module and the texture factory.

#rnx

diff --git a/Engine/Source/Editor/UnrealEd/Private/Factories/EditorFactories.cpp b/Engine/Source/Editor/UnrealEd/Private/Factories/EditorFactories.cpp
--- a/Engine/Source/Editor/UnrealEd/Private/Factories/EditorFactories.cpp
+++ b/Engine/Source/Editor/UnrealEd/Private/Factories/EditorFactories.cpp
@@ -3580,6 +3580,14 @@ UTexture* UTextureFactory::ImportTexture(UClass* Class, UObject* InParent, FName
 		if (HDRImportShouldBeLongLatCubeMap == EAppReturnType::Yes ||
 			HDRImportShouldBeLongLatCubeMap == EAppReturnType::YesAll)
 		{
+			if (Class && !Class->IsChildOf<UTextureCube>())
+			{
+				//The source use to re-import this texture is not compatible with the
+				//existing texture type.
+				Warn->Logf(ELogVerbosity::Error, TEXT("Cannot re-import a [%s] with a UTextureCube source file"), *Class->GetName());
+				return nullptr;
+			}
+
 			TSharedPtr<IImageWrapper> HdrImageWrapper = ImageWrapperModule.CreateImageWrapper(EImageFormat::HDR);
 			if (HdrImageWrapper->SetCompressed(Buffer, Length))
 			{
@@ -3626,6 +3634,14 @@ UTexture* UTextureFactory::ImportTexture(UClass* Class, UObject* InParent, FName
 
 		if(IESConverter.IsValid())
 		{
+			if (Class && !Class->IsChildOf<UTextureLightProfile>())
+			{
+				//The source use to re-import this texture is not compatible with the
+				//existing texture type.
+				Warn->Logf(ELogVerbosity::Error, TEXT("Cannot re-import a [%s] with a UTextureLightProfile source file"), *Class->GetName());
+				return nullptr;
+			}
+
 			UTextureLightProfile* Texture = Cast<UTextureLightProfile>( CreateOrOverwriteAsset(UTextureLightProfile::StaticClass(), InParent, Name, Flags) );
 			if ( Texture )
 			{
@@ -3662,6 +3678,14 @@ UTexture* UTextureFactory::ImportTexture(UClass* Class, UObject* InParent, FName
 	FImportImage Image;
 	if (ImportImage(Buffer, Length, Warn, ImportFlags, Image))
 	{
+		if (Class && !Class->IsChildOf<UTexture2D>())
+		{
+			//The source use to re-import this texture is not compatible with the
+			//existing texture type.
+			Warn->Logf(ELogVerbosity::Error, TEXT("Cannot re-import a [%s] with a UTexture2D source file"), *Class->GetName());
+			return nullptr;
+		}
+
 		UTexture2D* Texture = CreateTexture2D(InParent, Name, Flags);
 		if (Texture)
 		{
diff --git a/Engine/Source/Runtime/Interchange/Engine/Public/InterchangeAssetImportData.h b/Engine/Source/Runtime/Interchange/Engine/Public/InterchangeAssetImportData.h
--- a/Engine/Source/Runtime/Interchange/Engine/Public/InterchangeAssetImportData.h
+++ b/Engine/Source/Runtime/Interchange/Engine/Public/InterchangeAssetImportData.h
@@ -71,7 +71,10 @@ public:
 	{
 		if (NodeContainer->IsNodeUidValid(NodeUniqueID))
 		{
-			NodeContainer->GetNode(NodeUniqueID)->AppendAssetRegistryTags(OutTags);
+			if (const UInterchangeBaseNode* Node = NodeContainer->GetNode(NodeUniqueID))
+			{
+				Node->AppendAssetRegistryTags(OutTags);
+			}
 		}
 		Super::AppendAssetRegistryTags(OutTags);
 	}
