Fixes crash occurring when the user is editing a MaterialInstance and its parent Material at the same time, modifies the parent Material and then manually Reloads the material throug the Content Browser.

The issue was caused by UE closing the Material editor but not the MaterialInstance editor. Since this one depends on the parent Material data, it would try to read invalid/destroyed data from the parent material when after it reloads causing the crash.

The engine now closes the editors related to the reloaded asset as well as those of the referencers of the asset being reloaded (the Material).

diff --git a/Engine/Source/Editor/UnrealEd/Private/Subsystems/AssetEditorSubsystem.cpp b/Engine/Source/Editor/UnrealEd/Private/Subsystems/AssetEditorSubsystem.cpp
--- a/Engine/Source/Editor/UnrealEd/Private/Subsystems/AssetEditorSubsystem.cpp
+++ b/Engine/Source/Editor/UnrealEd/Private/Subsystems/AssetEditorSubsystem.cpp
@@ -150,6 +150,55 @@ int32 UAssetEditorSubsystem::CloseAllEditorsForAsset(UObject* Asset)
 	return EditorInstances.Num();
 }
 
+int32 UAssetEditorSubsystem::CloseAllEditorsForAssetAndReferencedAssets(UObject* Asset, TArray<TWeakObjectPtr<UObject>>& AssetsThatNeedToReopen)
+{
+	// Look for the editor instances for the specified asset or any asset in asset's package
+	TArray<IAssetEditorInstance*> EditorInstances = FindEditorsForAssetAndSubObjects(Asset);
+	AssetsThatNeedToReopen.AddUnique(Asset);
+
+	// Fetch the list of internally and externally referencers to this asset. We only care for the external referencers
+	TArray<FReferencerInformation> AssetInternalReferencers, AssetExternalReferencers;
+	Asset->RetrieveReferencers(&AssetInternalReferencers, &AssetExternalReferencers);
+
+	// Visit each external referencer and get the list of its editor instances
+	for (const FReferencerInformation& Ref : AssetExternalReferencers)
+	{
+		TArray<IAssetEditorInstance*> ReferencedAssetEditorInstances = FindEditorsForAssetAndSubObjects(Ref.Referencer);
+		if (!ReferencedAssetEditorInstances.IsEmpty())
+		{
+			// Merge the results with the list of editor instances to close
+			for (IAssetEditorInstance* ReferencedAssetEditorInstance : ReferencedAssetEditorInstances)
+			{
+				EditorInstances.AddUnique(ReferencedAssetEditorInstance);
+			}
+
+			// And notify the caller that the editor for this asset was also closed
+			AssetsThatNeedToReopen.AddUnique(Ref.Referencer);
+		}
+	}
+
+	// Close all the found editor instances 
+	for (IAssetEditorInstance* EditorInstance : EditorInstances)
+	{
+		if (EditorInstance)
+		{
+			EditorInstance->CloseWindow();
+		}
+	}
+
+	// Broadcast the event to all assets whose editors were closed
+	for (TWeakObjectPtr<UObject>& ClosedAsset : AssetsThatNeedToReopen)
+	{
+		if (ClosedAsset.IsValid())
+		{
+			AssetEditorRequestCloseEvent.Broadcast(ClosedAsset.Get(), EAssetEditorCloseReason::CloseAllEditorsForAsset);
+		}
+	}
+
+	return EditorInstances.Num();
+}
+
+
 void UAssetEditorSubsystem::RemoveAssetFromAllEditors(UObject* Asset)
 {
 	TArray<IAssetEditorInstance*> EditorInstances = FindEditorsForAsset(Asset);
@@ -946,7 +995,7 @@ void UAssetEditorSubsystem::SaveOpenAssetEditors(const bool bOnShutdown, const b
 
 void UAssetEditorSubsystem::HandlePackageReloaded(const EPackageReloadPhase InPackageReloadPhase, FPackageReloadedEvent* InPackageReloadedEvent)
 {
-	static TArray<UObject*> PendingAssetsToOpen;
+	static TArray<TWeakObjectPtr<UObject>> PendingAssetsToOpen;
 
 	if (InPackageReloadPhase == EPackageReloadPhase::PrePackageFixup)
 	{
@@ -980,7 +1029,7 @@ void UAssetEditorSubsystem::HandlePackageReloaded(const EPackageReloadPhase InPa
 		int32 NumAssetEditorsClosed = 0;
 		for (UObject* OldAsset : ObjectsToClose)
 		{
-			NumAssetEditorsClosed += CloseAllEditorsForAsset(OldAsset);
+			NumAssetEditorsClosed += CloseAllEditorsForAssetAndReferencedAssets(OldAsset, PendingAssetsToOpen);
 		}
 
 		if (NumAssetEditorsClosed > 0)
@@ -989,13 +1038,17 @@ void UAssetEditorSubsystem::HandlePackageReloaded(const EPackageReloadPhase InPa
 			// Run a GC now to ensure those are cleaned up before the fix-up phase happens
 			CollectGarbage(GARBAGE_COLLECTION_KEEPFLAGS);
 		}
+
 	}
 
 	if (InPackageReloadPhase == EPackageReloadPhase::PostBatchPostGC)
 	{
-		for (UObject* NewAsset : PendingAssetsToOpen)
+		for (TWeakObjectPtr<UObject>& NewAsset : PendingAssetsToOpen)
 		{
-			OpenEditorForAsset(NewAsset);
+			if (NewAsset.IsValid())
+			{
+				OpenEditorForAsset(NewAsset.Get());
+			}
 		}
 		PendingAssetsToOpen.Reset();
 	}
diff --git a/Engine/Source/Editor/UnrealEd/Public/Subsystems/AssetEditorSubsystem.h b/Engine/Source/Editor/UnrealEd/Public/Subsystems/AssetEditorSubsystem.h
--- a/Engine/Source/Editor/UnrealEd/Public/Subsystems/AssetEditorSubsystem.h
+++ b/Engine/Source/Editor/UnrealEd/Public/Subsystems/AssetEditorSubsystem.h
@@ -117,6 +117,9 @@ public:
 	UFUNCTION(BlueprintCallable, Category = "Editor Scripting | Asset Tools")
 	int32 CloseAllEditorsForAsset(UObject* Asset);
 
+	/** Close all active editors for the supplied asset and its external referencers and return the number of asset editors that were closed */
+	int32 CloseAllEditorsForAssetAndReferencedAssets(UObject* Asset, TArray<TWeakObjectPtr<UObject>>& AssetsThatNeedToReopen);
+
 	/** Close any editor which is not this one */
 	void CloseOtherEditors(UObject* Asset, IAssetEditorInstance* OnlyEditor);
 
