Refactored changes I made in 22427807 afterfurther review suggestions. Change now has same effect with smaller footprint.

diff --git a/Engine/Source/Editor/UnrealEd/Private/Subsystems/AssetEditorSubsystem.cpp b/Engine/Source/Editor/UnrealEd/Private/Subsystems/AssetEditorSubsystem.cpp
--- a/Engine/Source/Editor/UnrealEd/Private/Subsystems/AssetEditorSubsystem.cpp
+++ b/Engine/Source/Editor/UnrealEd/Private/Subsystems/AssetEditorSubsystem.cpp
@@ -150,54 +150,6 @@ int32 UAssetEditorSubsystem::CloseAllEditorsForAsset(UObject* Asset)
 	return EditorInstances.Num();
 }
 
-int32 UAssetEditorSubsystem::CloseAllEditorsForAssetAndReferencedAssets(UObject* Asset, TArray<TWeakObjectPtr<UObject>>& AssetsThatNeedToReopen)
-{
-	// Look for the editor instances for the specified asset or any asset in asset's package
-	TArray<IAssetEditorInstance*> EditorInstances = FindEditorsForAssetAndSubObjects(Asset);
-	AssetsThatNeedToReopen.AddUnique(Asset);
-
-	// Fetch the list of internally and externally referencers to this asset. We only care for the external referencers
-	TArray<FReferencerInformation> AssetInternalReferencers, AssetExternalReferencers;
-	Asset->RetrieveReferencers(&AssetInternalReferencers, &AssetExternalReferencers);
-
-	// Visit each external referencer and get the list of its editor instances
-	for (const FReferencerInformation& Ref : AssetExternalReferencers)
-	{
-		TArray<IAssetEditorInstance*> ReferencedAssetEditorInstances = FindEditorsForAssetAndSubObjects(Ref.Referencer);
-		if (!ReferencedAssetEditorInstances.IsEmpty())
-		{
-			// Merge the results with the list of editor instances to close
-			for (IAssetEditorInstance* ReferencedAssetEditorInstance : ReferencedAssetEditorInstances)
-			{
-				EditorInstances.AddUnique(ReferencedAssetEditorInstance);
-			}
-
-			// And notify the caller that the editor for this asset was also closed
-			AssetsThatNeedToReopen.AddUnique(Ref.Referencer);
-		}
-	}
-
-	// Close all the found editor instances 
-	for (IAssetEditorInstance* EditorInstance : EditorInstances)
-	{
-		if (EditorInstance)
-		{
-			EditorInstance->CloseWindow();
-		}
-	}
-
-	// Broadcast the event to all assets whose editors were closed
-	for (TWeakObjectPtr<UObject>& ClosedAsset : AssetsThatNeedToReopen)
-	{
-		if (ClosedAsset.IsValid())
-		{
-			AssetEditorRequestCloseEvent.Broadcast(ClosedAsset.Get(), EAssetEditorCloseReason::CloseAllEditorsForAsset);
-		}
-	}
-
-	return EditorInstances.Num();
-}
-
 
 void UAssetEditorSubsystem::RemoveAssetFromAllEditors(UObject* Asset)
 {
@@ -1022,14 +974,25 @@ void UAssetEditorSubsystem::HandlePackageReloaded(const EPackageReloadPhase InPa
 					PendingAssetsToOpen.AddUnique(NewAsset);
 				}
 
-				ObjectsToClose.AddUnique(AssetEditorPair.Key);
+				UObject* OldAsset = AssetEditorPair.Key;
+				ObjectsToClose.AddUnique(OldAsset);
+
+				// The asset being reloaded might have other assets that depend on it. Find the list of
+				// external referencers to this asset and mark them to be closed and reopened as well.
+				TArray<FReferencerInformation> AssetInternalReferencers, AssetExternalReferencers;
+				OldAsset->RetrieveReferencers(&AssetInternalReferencers, &AssetExternalReferencers);
+				for (const FReferencerInformation& Ref : AssetExternalReferencers)
+				{
+					ObjectsToClose.AddUnique(Ref.Referencer);
+					PendingAssetsToOpen.AddUnique(Ref.Referencer);
+				}
 			}
 		}
 
 		int32 NumAssetEditorsClosed = 0;
 		for (UObject* OldAsset : ObjectsToClose)
 		{
-			NumAssetEditorsClosed += CloseAllEditorsForAssetAndReferencedAssets(OldAsset, PendingAssetsToOpen);
+			NumAssetEditorsClosed += CloseAllEditorsForAsset(OldAsset);
 		}
 
 		if (NumAssetEditorsClosed > 0)
diff --git a/Engine/Source/Editor/UnrealEd/Public/Subsystems/AssetEditorSubsystem.h b/Engine/Source/Editor/UnrealEd/Public/Subsystems/AssetEditorSubsystem.h
--- a/Engine/Source/Editor/UnrealEd/Public/Subsystems/AssetEditorSubsystem.h
+++ b/Engine/Source/Editor/UnrealEd/Public/Subsystems/AssetEditorSubsystem.h
@@ -117,9 +117,6 @@ public:
 	UFUNCTION(BlueprintCallable, Category = "Editor Scripting | Asset Tools")
 	int32 CloseAllEditorsForAsset(UObject* Asset);
 
-	/** Close all active editors for the supplied asset and its external referencers and return the number of asset editors that were closed */
-	int32 CloseAllEditorsForAssetAndReferencedAssets(UObject* Asset, TArray<TWeakObjectPtr<UObject>>& AssetsThatNeedToReopen);
-
 	/** Close any editor which is not this one */
 	void CloseOtherEditors(UObject* Asset, IAssetEditorInstance* OnlyEditor);
 
