GraphConvert: Silently ignore edges to non-vertices rather than failing an assertion. Some edge construction callers will add edges first and verify the vertex after, and removing the edges when vertification detects the vertex is invalid would be expensive. GraphConvert should support these cases via silent removal. We have not yet seen any examples of cases where the caller really wants the edge but forgot to add it to UniqueKeys, so the assertion was not useful. This fixes an intermittent failure in Multiprocess Cook on CitySample; the graph from CookRequestCluster can include edges to packages not in the requestcluster.
#rnx

diff --git a/Engine/Source/Editor/UnrealEd/Private/Cooker/Algo/GraphConvert.h b/Engine/Source/Editor/UnrealEd/Private/Cooker/Algo/GraphConvert.h
--- a/Engine/Source/Editor/UnrealEd/Private/Cooker/Algo/GraphConvert.h
+++ b/Engine/Source/Editor/UnrealEd/Private/Cooker/Algo/GraphConvert.h
@@ -43,8 +43,8 @@ ENUM_CLASS_FLAGS(EConvertToGraphOptions);
  * @param GetKeyEdges A callable with prototype that is one of
  *            RangeType<KeyType> GetKeyEdges(const KeyType& Key);
  *            const RangeType<KeyType>& GetKeyEdges(const KeyType& Key);
- *        It must return the TargetKeys that are pointed to from the directed edges from Key. All returned TargetKeys
- *        must be elements of UniqueKeys. RangeType must support ranged-for (begin() and end()).
+ *        It must return the TargetKeys that are pointed to from the directed edges from Key. TargetKeys that are 
+ *        not elements of UniqueKeys will be discarded. RangeType must support ranged-for (begin() and end()).
  * @param OutGraphBuffer Output value that holds the memory used by the graph. The array must remain allocated and
  *        unmodified until OutGraph is no longer referenced.
  * @param OutGraph Output value that holds the graph. OutGraph[i] is a TConstArrayView of edges from Vertex i, which
@@ -87,7 +87,10 @@ inline void ConvertToGraph(const RangeType& UniqueKeys, GetKeyEdgesType GetKeyEd
 		for (const KeyType& Dependency : Invoke(GetKeyEdges, Element))
 		{
 			FVertex* TargetVertex = VertexOfKey.Find(Dependency);
-			checkf(TargetVertex, TEXT("ConvertToGraph: GetKeyEdges returned a TargetKey that is not a member of UniqueKeys"));
+			if (!TargetVertex)
+			{
+				continue;
+			}
 
 			// Normalize Step 1: remove edges to self
 			if (*TargetVertex == Vertex)
