Fix static analysis

#rnx

diff --git a/Engine/Source/Editor/UnrealEd/Private/ObjectTools.cpp b/Engine/Source/Editor/UnrealEd/Private/ObjectTools.cpp
--- a/Engine/Source/Editor/UnrealEd/Private/ObjectTools.cpp
+++ b/Engine/Source/Editor/UnrealEd/Private/ObjectTools.cpp
@@ -2981,11 +2981,11 @@ namespace ObjectTools
 		if (GEditor)
 		{
 			GEditor->GetSelectedObjects()->Deselect(ObjectToDelete);
-		}
-
-		if (ObjectToDelete->IsAsset())
-		{
-			GEditor->GetEditorSubsystem<UAssetEditorSubsystem>()->CloseAllEditorsForAsset(ObjectToDelete);
+		
+			if (ObjectToDelete->IsAsset())
+			{
+				GEditor->GetEditorSubsystem<UAssetEditorSubsystem>()->CloseAllEditorsForAsset(ObjectToDelete);
+			}
 		}
 
 		{
