EditorDomain: Add static reflected UObject method DeclareConstructedClasses. Classes that can construct other classes during PostLoad need to declare those classes.
UObject statics: Improve the feedback in the compile error for static reflected UObject methods when the function prototype is incorrect.
UObject delegates: Add an OnObjectConstructed delegate, and add the Object being PostLoaded to the PackageAccessTracking UE_TRACK_REFERENCING_PACKAGE_SCOPED, so we can give a warning when objects are constructed during PostLoad that have not been declared via DeclareConstructedClasses.
EditorDomain: Improve performance of CalculatePackageDigest by precomputing the inclusive schema for every class instead of needing to calculate it each time that class is used.
EditorDomain: Fix bug in UObject::DeclareCustomVersions that was causing array properties not to be parsed.
EditorDomainUtils: change ClassDigests to identify ClassPaths using TopLevelAssetPath to match the changes made by ANY_PACKAGE.
CoreRedirects: Add FCoreRedirectObjectName costructor from FTopLevelAssetPath to reduce boilerplate for detecting class redirects.
#rnx

diff --git a/Engine/Source/Editor/UnrealEd/Private/Commandlets/DerivedDataCacheCommandlet.cpp b/Engine/Source/Editor/UnrealEd/Private/Commandlets/DerivedDataCacheCommandlet.cpp
--- a/Engine/Source/Editor/UnrealEd/Private/Commandlets/DerivedDataCacheCommandlet.cpp
+++ b/Engine/Source/Editor/UnrealEd/Private/Commandlets/DerivedDataCacheCommandlet.cpp
@@ -291,7 +291,7 @@ bool UDerivedDataCacheCommandlet::ProcessCachingObjects()
 				}
 
 				{
-					UE_TRACK_REFERENCING_PACKAGE_SCOPED(Object->GetPackage(), PackageAccessTrackingOps::NAME_CookerBuildObject);
+					UE_TRACK_REFERENCING_PACKAGE_SCOPED(Object, PackageAccessTrackingOps::NAME_CookerBuildObject);
 					for (int32 PlatformIndex = 0; PlatformIndex < Platforms.Num(); ++PlatformIndex)
 					{
 						// IsCachedCookedPlatformDataLoaded can be quite slow for some objects
diff --git a/Engine/Source/Editor/UnrealEd/Private/CookOnTheFlyServer.cpp b/Engine/Source/Editor/UnrealEd/Private/CookOnTheFlyServer.cpp
--- a/Engine/Source/Editor/UnrealEd/Private/CookOnTheFlyServer.cpp
+++ b/Engine/Source/Editor/UnrealEd/Private/CookOnTheFlyServer.cpp
@@ -4330,7 +4330,7 @@ void UCookOnTheFlyServer::TickPrecacheObjectsForPlatforms(const float TimeSlice,
 			continue;
 		}
 
-		UE_TRACK_REFERENCING_PACKAGE_SCOPED(Material->GetPackage(), PackageAccessTrackingOps::NAME_CookerBuildObject);
+		UE_TRACK_REFERENCING_PACKAGE_SCOPED(Material, PackageAccessTrackingOps::NAME_CookerBuildObject);
 		for (const ITargetPlatform* TargetPlatform : TargetPlatforms)
 		{
 			if (!TargetPlatform)
@@ -4371,7 +4371,7 @@ void UCookOnTheFlyServer::TickPrecacheObjectsForPlatforms(const float TimeSlice,
 			continue;
 		}
 
-		UE_TRACK_REFERENCING_PACKAGE_SCOPED(Texture->GetPackage(), PackageAccessTrackingOps::NAME_CookerBuildObject);
+		UE_TRACK_REFERENCING_PACKAGE_SCOPED(Texture, PackageAccessTrackingOps::NAME_CookerBuildObject);
 		for (const ITargetPlatform* TargetPlatform : TargetPlatforms)
 		{
 			if (!TargetPlatform)
@@ -10441,7 +10441,7 @@ uint32 UCookOnTheFlyServer::CookFullLoadAndSave()
 			{
 				UObject* Obj = ObjectsToWaitForCookedPlatformData[ObjIdx];
 				bool bAllPlatformDataLoaded = true;
-				UE_TRACK_REFERENCING_PACKAGE_SCOPED(Obj->GetPackage(), PackageAccessTrackingOps::NAME_CookerBuildObject);
+				UE_TRACK_REFERENCING_PACKAGE_SCOPED(Obj, PackageAccessTrackingOps::NAME_CookerBuildObject);
 				for (const ITargetPlatform* TargetPlatform : TargetPlatforms)
 				{
 					if (!RouteIsCachedCookedPlatformDataLoaded(Obj, TargetPlatform))
diff --git a/Engine/Source/Editor/UnrealEd/Private/Cooker/CookPackageData.cpp b/Engine/Source/Editor/UnrealEd/Private/Cooker/CookPackageData.cpp
--- a/Engine/Source/Editor/UnrealEd/Private/Cooker/CookPackageData.cpp
+++ b/Engine/Source/Editor/UnrealEd/Private/Cooker/CookPackageData.cpp
@@ -1907,7 +1907,7 @@ bool FPendingCookedPlatformData::PollIsComplete()
 		Release();
 		return true;
 	}
-	UE_TRACK_REFERENCING_PACKAGE_SCOPED(LocalObject->GetPackage(), PackageAccessTrackingOps::NAME_CookerBuildObject);
+	UE_TRACK_REFERENCING_PACKAGE_SCOPED(LocalObject, PackageAccessTrackingOps::NAME_CookerBuildObject);
 	if (RouteIsCachedCookedPlatformDataLoaded(LocalObject, TargetPlatform))
 	{
 		Release();
diff --git a/Engine/Source/Editor/UnrealEd/Private/EditorDomain/EditorDomain.cpp b/Engine/Source/Editor/UnrealEd/Private/EditorDomain/EditorDomain.cpp
--- a/Engine/Source/Editor/UnrealEd/Private/EditorDomain/EditorDomain.cpp
+++ b/Engine/Source/Editor/UnrealEd/Private/EditorDomain/EditorDomain.cpp
@@ -282,7 +282,18 @@ void FEditorDomain::PrecachePackageDigest(FName PackageName)
 	TOptional<FAssetPackageData> PackageData = AssetRegistry->GetAssetPackageDataCopy(PackageName);
 	if (PackageData)
 	{
-		UE::EditorDomain::PrecacheClassDigests(PackageData->ImportedClasses);
+		TArray<FTopLevelAssetPath> ImportedClassPaths;
+		ImportedClassPaths.Reserve(PackageData->ImportedClasses.Num());
+		for (FName ClassPathName : PackageData->ImportedClasses)
+		{
+			FTopLevelAssetPath ClassPath(WriteToString<256>(ClassPathName).ToView());
+			if (ClassPath.IsValid())
+			{
+				ImportedClassPaths.Add(ClassPath);
+			}
+		}
+
+		UE::EditorDomain::PrecacheClassDigests(ImportedClassPaths);
 	}
 }
 
diff --git a/Engine/Source/Editor/UnrealEd/Private/EditorDomain/EditorDomain.h b/Engine/Source/Editor/UnrealEd/Private/EditorDomain/EditorDomain.h
--- a/Engine/Source/Editor/UnrealEd/Private/EditorDomain/EditorDomain.h
+++ b/Engine/Source/Editor/UnrealEd/Private/EditorDomain/EditorDomain.h
@@ -2,26 +2,26 @@
 
 #pragma once
 
+#include "AssetRegistry/AssetData.h"
 #include "Containers/Array.h"
 #include "Containers/ArrayView.h"
+#include "Containers/Map.h"
 #include "Containers/StringFwd.h"
 #include "Containers/UnrealString.h"
-#include "Misc/EnumClassFlags.h"
-#include "Misc/PackageSegment.h"
-#include "Stats/Stats2.h"
-#include "Tickable.h"
-#include "UObject/UnrealNames.h"
-#include "AssetRegistry/AssetData.h"
-#include "Containers/Map.h"
 #include "HAL/CriticalSection.h"
 #include "HAL/Platform.h"
 #include "IO/IoHash.h"
 #include "Logging/LogMacros.h"
+#include "Misc/EnumClassFlags.h"
+#include "Misc/PackageSegment.h"
+#include "Stats/Stats2.h"
 #include "Templates/RefCounting.h"
 #include "Templates/UniquePtr.h"
+#include "Tickable.h"
 #include "TickableEditorObject.h"
 #include "UObject/NameTypes.h"
 #include "UObject/PackageResourceManager.h"
+#include "UObject/UnrealNames.h"
 
 class FArchive;
 class FEditorDomainSaveClient;
@@ -31,8 +31,8 @@ class FScopeLock;
 class IAssetRegistry;
 class IMappedFileHandle;
 class UPackage;
-namespace UE::DerivedData { class FRequestOwner; }
 struct FEndLoadPackageContext;
+namespace UE::DerivedData { class FRequestOwner; }
 
 namespace UE::EditorDomain
 {
@@ -47,6 +47,7 @@ enum class EDomainUse : uint8
 	SaveEnabled = 0x2,
 };
 ENUM_CLASS_FLAGS(EDomainUse);
+FStringBuilderBase& operator<<(FStringBuilderBase& Writer, UE::EditorDomain::EDomainUse DomainUse);
 
 /** Information about a package's loading from the EditorDomain. */
 struct FPackageDigest
@@ -71,8 +72,6 @@ struct FPackageDigest
 	 * created from properties that change if the saved version would change (package bytes, serialization versions).
 	 */
 	FIoHash Hash;
-	/** Classes used by the resaved package: imports or instances that can be created by PostLoad/PreSave */
-	TArray<FName> ImportedClasses;
 	/** List of CustomVersions used to save the package. */
 	UE::AssetRegistry::FPackageCustomVersionsHandle CustomVersions;
 	/** Allow flags for whether the package can be saved/loaded from EditorDomain. */
diff --git a/Engine/Source/Editor/UnrealEd/Private/EditorDomain/EditorDomainUtils.cpp b/Engine/Source/Editor/UnrealEd/Private/EditorDomain/EditorDomainUtils.cpp
--- a/Engine/Source/Editor/UnrealEd/Private/EditorDomain/EditorDomainUtils.cpp
+++ b/Engine/Source/Editor/UnrealEd/Private/EditorDomain/EditorDomainUtils.cpp
@@ -3,12 +3,11 @@
 #include "EditorDomain/EditorDomainUtils.h"
 
 #include "Algo/AnyOf.h"
+#include "Algo/BinarySearch.h"
 #include "Algo/IsSorted.h"
 #include "Algo/Unique.h"
 #include "AssetRegistry/AssetData.h"
 #include "AssetRegistry/IAssetRegistry.h"
-#include "Containers/Array.h"
-#include "Containers/Map.h"
 #include "DerivedDataBuildDefinition.h"
 #include "DerivedDataCache.h"
 #include "DerivedDataCacheKey.h"
@@ -19,10 +18,13 @@
 #include "Engine/StaticMesh.h"
 #include "HAL/CriticalSection.h"
 #include "HAL/FileManager.h"
+#include "HAL/IConsoleManager.h"
 #include "Hash/Blake3.h"
 #include "Memory/SharedBuffer.h"
 #include "Misc/CoreDelegates.h"
+#include "Misc/DelayedAutoRegister.h"
 #include "Misc/FileHelper.h"
+#include "Misc/PackageAccessTracking.h"
 #include "Misc/PackagePath.h"
 #include "Misc/ScopeRWLock.h"
 #include "Misc/StringBuilder.h"
@@ -44,7 +46,6 @@
 #include "UObject/UObjectHash.h"
 #include "UObject/UE5MainStreamObjectVersion.h"
 
-
 #if !defined(EDITORDOMAINTIMEPROFILERTRACE_ENABLED)
 #if UE_TRACE_ENABLED && !UE_BUILD_SHIPPING
 #define EDITORDOMAINTIMEPROFILERTRACE_ENABLED 1
@@ -88,6 +89,34 @@ static ValueType MapFindRef(const TMap<KeyType, ValueType, SetAllocator, KeyFunc
 	return FoundValue ? *FoundValue : DefaultValue;
 }
 
+static void CoreRedirectClassPath(FTopLevelAssetPath& ClassPath, bool* bWasRedirected = nullptr)
+{
+	FCoreRedirectObjectName ClassRedirect(ClassPath);
+	FCoreRedirectObjectName RedirectedClassRedirect = FCoreRedirects::GetRedirectedName(
+		ECoreRedirectFlags::Type_Class | ECoreRedirectFlags::Type_Struct, ClassRedirect);
+	if (ClassRedirect != RedirectedClassRedirect)
+	{
+		if (RedirectedClassRedirect.OuterName.IsNone())
+		{
+			ClassPath = FTopLevelAssetPath(RedirectedClassRedirect.PackageName, RedirectedClassRedirect.ObjectName);
+			if (bWasRedirected)
+			{
+				*bWasRedirected = true;
+			}
+			return;
+		}
+		else
+		{
+			UE_LOG(LogEditorDomain, Warning, TEXT("CoreRedirect for class %s redirects to %s which is not a TopLevelAssetPath. Ignoring it."),
+				*ClassRedirect.ToString(), *RedirectedClassRedirect.ToString())
+		}
+	}
+	if (bWasRedirected)
+	{
+		*bWasRedirected = false;
+	}
+}
+
 #if ENABLE_COOK_STATS
 namespace UE::EditorDomain::CookStats
 {
@@ -157,14 +186,15 @@ FClassDigestMap& GetClassDigests()
 	return GClassDigests;
 }
 
-TMap<FName, EDomainUse> GClassBlockedUses;
+TMap<FTopLevelAssetPath, EDomainUse> GClassBlockedUses;
 TMap<FName, EDomainUse> GPackageBlockedUses;
-TMultiMap<FName, FName> GConstructClasses;
-TSet<FName> GTargetDomainClassBlockList;
+TMultiMap<FTopLevelAssetPath, FTopLevelAssetPath> GConstructClasses;
+TSet<FTopLevelAssetPath> GTargetDomainClassBlockList;
 bool GTargetDomainClassUseAllowList = true;
 bool GTargetDomainClassEmptyAllowList = false;
-TArray<FName> GGlobalConstructClasses;
+TArray<FTopLevelAssetPath> GGlobalConstructClasses;
 bool bGGlobalConstructClassesInitialized = false;
+FBlake3Hash GGlobalConstructClassesHash;
 int64 GMaxBulkDataSize = -1;
 
 // Change to a new guid when EditorDomain needs to be invalidated
@@ -307,93 +337,71 @@ FPackageDigest CalculatePackageDigest(IAssetRegistry& AssetRegistry, FName Packa
 	FKnownCustomVersions::FindOrAddHandles(CustomVersionHandles,
 		PackageVersions.Num(), [PackageVersions](int32 Index) -> const FGuid& { return PackageVersions[Index].Key;});
 
-	// Loop over all ImportedClasses used by the package on disk, fetch the ClassDigestData for each one, and
-	// extend the list with the transitive closure of classes that can be created by any of those classes.
+	InitializeGlobalConstructClasses();
+	Writer.Update(&GGlobalConstructClassesHash, sizeof(GGlobalConstructClassesHash));
+
+	// Add the InclusiveSchema for all classes imported by the package. The InclusiveSchema includes the schema of
+	// classes that can be constructed by the given class.
 	FClassDigestMap& ClassDigests = GetClassDigests();
-	TArray<FName> ImportedClasses = MoveTemp(PackageData.ImportedClasses);
-	TSet<FName> DiscoveredClasses(ImportedClasses);
-	bool bNeedsResort = false;
 	bool bHasTriedPrecacheClassDigests = false;
-	auto AddClass = [&DiscoveredClasses, &ImportedClasses, &bHasTriedPrecacheClassDigests, &bNeedsResort](FName ClassName)
+	TArray<FTopLevelAssetPath> ImportedClassPaths;
+	ImportedClassPaths.Reserve(PackageData.ImportedClasses.Num());
+	for (FName ClassPathName : PackageData.ImportedClasses)
 	{
-		bool bAlreadyDiscovered;
-		DiscoveredClasses.Add(ClassName, &bAlreadyDiscovered);
-		if (!bAlreadyDiscovered)
+		FTopLevelAssetPath ClassPath(WriteToString<256>(ClassPathName).ToView());
+		if (ClassPath.IsValid())
 		{
-			ImportedClasses.Add(ClassName);
-			bHasTriedPrecacheClassDigests = false;
-			bNeedsResort = true;
+			ImportedClassPaths.Add(ClassPath);
 		}
-	};
-
-	InitializeGlobalConstructClasses();
-	for (FName GlobalAddedClass : GGlobalConstructClasses)
-	{
-		AddClass(GlobalAddedClass);
 	}
+
 	int32 NextClass = 0;
-	while (NextClass < ImportedClasses.Num())
+	while (NextClass < ImportedClassPaths.Num())
 	{
 		{
 			FReadScopeLock ClassDigestsScopeLock(ClassDigests.Lock);
-			while (NextClass < ImportedClasses.Num())
+			while (NextClass < ImportedClassPaths.Num())
 			{
-				FName ClassName = ImportedClasses[NextClass];
-				FClassDigestData* ExistingData = ClassDigests.Map.Find(ClassName);
+				const FTopLevelAssetPath& ClassPath = ImportedClassPaths[NextClass];
+				FClassDigestData* ExistingData = ClassDigests.Map.Find(ClassPath);
 				if (!ExistingData)
 				{
 					break;
 				}
-				if (ExistingData->ResolvedClosestNative == ClassName)
-				{
-					++NextClass;
-				}
-				else
-				{
-					ImportedClasses.RemoveAt(NextClass);
-					AddClass(ExistingData->ResolvedClosestNative);
-				}
-				for (FName ConstructClass : ExistingData->ConstructClasses)
+				NextClass++;
+
+				// We only need to hash the ClosestNative class. EditorDomain packages that contain blueprint class
+				// instances do not use versioned data, and the incremental recook of these packages will be triggered
+				// when their class changes because they have an import dependency on their class's package.
+				FClassDigestData* NativeData = ExistingData;
+				if (ExistingData->ClosestNative != ClassPath)
 				{
-					AddClass(ConstructClass);
+					NativeData = ClassDigests.Map.Find(ExistingData->ClosestNative);
+					checkf(NativeData, TEXT("Classes are only stored in a ClosestNative field if they exist"));
 				}
+				check(NativeData->bNative);
+				Writer.Update(&NativeData->InclusiveSchemaHash, sizeof(NativeData->InclusiveSchemaHash));
+				CustomVersionHandles.Append(NativeData->CustomVersionHandles);
+				EnumSetFlagsAnd(Result.DomainUse, EDomainUse::LoadEnabled | EDomainUse::SaveEnabled,
+					Result.DomainUse, NativeData->EditorDomainUse);
 			}
 		}
 
-		if (NextClass < ImportedClasses.Num())
+		if (NextClass < ImportedClassPaths.Num())
 		{
 			// EDITORDOMAIN_TODO: Remove the clauses !IsInGameThread || GIsSavingPackage once FindObject no longer asserts if GIsSavingPackage
 			if (bHasTriedPrecacheClassDigests || !IsInGameThread() || GIsSavingPackage)
 			{
-				return FPackageDigest(FPackageDigest::EStatus::MissingClass, ImportedClasses[NextClass]);
+				return FPackageDigest(FPackageDigest::EStatus::MissingClass,
+					FName(WriteToString<256>(ImportedClassPaths[NextClass]).ToView()));
 			}
-			TConstArrayView<FName> RemainingClasses(ImportedClasses);
-			RemainingClasses = RemainingClasses.Slice(NextClass, ImportedClasses.Num() - NextClass);
+			TConstArrayView<FTopLevelAssetPath> RemainingClasses =
+				TConstArrayView<FTopLevelAssetPath>(ImportedClassPaths).RightChop(NextClass);
 			PrecacheClassDigests(RemainingClasses);
 			bHasTriedPrecacheClassDigests = true;
 		}
 	}
 
-	// Loop over the sorted list of classes a second time, now that we have all of them,
-	// and add their schemas
-	if (bNeedsResort)
-	{
-		ImportedClasses.Sort(FNameLexicalLess());
-	}
-	{
-		FReadScopeLock ClassDigestsScopeLock(ClassDigests.Lock);
-		for (FName ClassName : ImportedClasses)
-		{
-			FClassDigestData* ExistingData = ClassDigests.Map.Find(ClassName);
-			check(ExistingData); // This should have been added above if not already existing
-			check(ExistingData->bNative); // Non-native classes were removed by the loop above
-			Writer.Update(&ExistingData->SchemaHash, sizeof(ExistingData->SchemaHash));
-			CustomVersionHandles.Append(ExistingData->CustomVersionHandles);
-			EnumSetFlagsAnd(Result.DomainUse, EDomainUse::LoadEnabled | EDomainUse::SaveEnabled,
-				Result.DomainUse, ExistingData->EditorDomainUse);
-		}
-	}
-
 	CustomVersionHandles.Sort();
 	CustomVersionHandles.SetNum(Algo::Unique(CustomVersionHandles));
 
@@ -418,7 +426,6 @@ FPackageDigest CalculatePackageDigest(IAssetRegistry& AssetRegistry, FName Packa
 	Result.Hash = Writer.Finalize();
 	Result.Status = FPackageDigest::EStatus::Successful;
 	Result.CustomVersions = UE::AssetRegistry::FPackageCustomVersionsHandle::FindOrAdd(CustomVersions);
-	Result.ImportedClasses = MoveTemp(ImportedClasses);
 	return Result;
 }
 
@@ -442,7 +449,7 @@ public:
 		ClassDigestsMap.Lock.WriteUnlock();
 	}
 
-	FClassDigestData* GetRecursive(FName ClassName, bool bAllowRedirects);
+	FClassDigestData* GetRecursive(FTopLevelAssetPath ClassPath, bool bAllowRedirects);
 
 	struct FUnlockScope
 	{
@@ -460,166 +467,153 @@ public:
 
 private:
 	FClassDigestMap& ClassDigestsMap;
-	TMap<FName, FClassDigestData>& ClassDigests;
+	TMap<FTopLevelAssetPath, FClassDigestData>& ClassDigests;
 	IAssetRegistry& AssetRegistry;
 
 	// Scratch variables usable during GetRecursive; they are invalidated when a recursive call is made
 	FString NameStringBuffer;
-	TArray<FTopLevelAssetPath> AncestorPathNames;
+	TArray<FTopLevelAssetPath> AncestorClassPaths;
 };
 
-FClassDigestData* FPrecacheClassDigest::GetRecursive(FName ClassName, bool bAllowRedirects)
+FClassDigestData* FPrecacheClassDigest::GetRecursive(FTopLevelAssetPath ClassPath, bool bAllowRedirects)
 {
-	check(!ClassName.IsNone());
+	check(!ClassPath.IsNull());
 	// Called within ClassDigestsMap.Lock.WriteLock()
-	FClassDigestData* DigestData = &ClassDigests.FindOrAdd(ClassName);
+	FClassDigestData* DigestData = &ClassDigests.FindOrAdd(ClassPath);
 	if (DigestData->bConstructed)
 	{
 		return DigestData;
 	}
 	DigestData->bConstructed = true;
 
-	FName RedirectsResolvedName = ClassName;
-	ClassName.ToString(NameStringBuffer);
 	if (bAllowRedirects)
 	{
-		FCoreRedirectObjectName ClassNameRedirect(NameStringBuffer);
-		FCoreRedirectObjectName RedirectedClassNameRedirect = FCoreRedirects::GetRedirectedName(ECoreRedirectFlags::Type_Class | ECoreRedirectFlags::Type_Struct, ClassNameRedirect);
-		if (ClassNameRedirect != RedirectedClassNameRedirect)
-		{
-			NameStringBuffer = RedirectedClassNameRedirect.ToString();
-			RedirectsResolvedName = FName(NameStringBuffer);
+		FTopLevelAssetPath RedirectedClassPath(ClassPath);
+		bool bRedirected;
+		CoreRedirectClassPath(RedirectedClassPath, &bRedirected);
+		if (bRedirected)
+		{ 
+			FClassDigestData* RedirectedDigest = GetRecursive(RedirectedClassPath, false /* bAllowRedirects */);
+			check(RedirectedDigest && RedirectedDigest->bConstructed && RedirectedDigest->bConstructionComplete);
+			// The map has possibly been modified so we need to recalculate the address of ClassName's DigestData
+			DigestData = &ClassDigests.FindChecked(ClassPath);
+			*DigestData = *RedirectedDigest;
+			return DigestData;
 		}
 	}
 
 	// Fill in digest data config-driven flags
 	DigestData->EditorDomainUse = EDomainUse::LoadEnabled | EDomainUse::SaveEnabled;
-	DigestData->EditorDomainUse &= ~MapFindRef(GClassBlockedUses, ClassName, EDomainUse::None);
-	if (RedirectsResolvedName != ClassName)
-	{
-		DigestData->EditorDomainUse &= ~MapFindRef(GClassBlockedUses, RedirectsResolvedName, EDomainUse::None);
-	}
+	DigestData->EditorDomainUse &= ~MapFindRef(GClassBlockedUses, ClassPath, EDomainUse::None);
 	if (!GTargetDomainClassUseAllowList)
 	{
-		DigestData->bTargetIterativeEnabled = !GTargetDomainClassBlockList.Contains(ClassName);
-		if (RedirectsResolvedName != ClassName)
-		{
-			DigestData->bTargetIterativeEnabled &= !GTargetDomainClassBlockList.Contains(RedirectsResolvedName);
-		}
+		DigestData->bTargetIterativeEnabled = !GTargetDomainClassBlockList.Contains(ClassPath);
 	}
 
 	// Fill in native-specific digest data, get the ParentName, and if non-native, get the native ancestor struct
 	UStruct* Struct = nullptr;
-	bool bIsNative = FPackageName::IsScriptPackage(NameStringBuffer);
+	bool bIsNative = FPackageName::IsScriptPackage(WriteToString<256>(ClassPath.GetPackageName()));
 	if (bIsNative)
 	{
-		Struct = FindObject<UStruct>(nullptr, *NameStringBuffer);
+		Struct = FindObject<UStruct>(ClassPath);
 	}
-	FName ParentName;
+	FTopLevelAssetPath ParentClassPath;
+	FBlake3Hash ExclusiveSchemaHash;
 	if (bIsNative)
 	{
 		if (Struct)
 		{
-			DigestData->ResolvedClosestNative = RedirectsResolvedName;
+			DigestData->ClosestNative = ClassPath;
 			DigestData->bNative = true;
-			DigestData->SchemaHash = Struct->GetSchemaHash(false /* bSkipEditorOnly */);
+			ExclusiveSchemaHash = Struct->GetSchemaHash(false /* bSkipEditorOnly */);
 			UStruct* ParentStruct = Struct->GetSuperStruct();
 			if (ParentStruct)
 			{
-				NameStringBuffer.Reset();
-				ParentStruct->GetPathName(nullptr, NameStringBuffer);
-				ParentName = FName(*NameStringBuffer);
+				ParentClassPath = FTopLevelAssetPath(ParentStruct);
 			}
 		}
 		else
 		{
 			UE_LOG(LogEditorDomain, Display, TEXT("Class %s is imported by a package but does not exist in memory. EditorDomain keys for packages using it will be invalid if it still exists.")
 				TEXT("\n\tTo clear this message, resave packages that use the deleted class, or load its module earlier than the packages that use it are referenced."),
-				*NameStringBuffer);
+				*ClassPath.ToString());
 
 			// Create placeholder data that acts as if this class is a leaf off of UObject
-			DigestData->ResolvedClosestNative = FName(*UObject::StaticClass()->GetPathName());
+			DigestData->ClosestNative = FTopLevelAssetPath(UObject::StaticClass());
+			check(DigestData->ClosestNative.IsValid());
 			DigestData->bNative = true;
-			DigestData->SchemaHash.Reset();
-			ParentName = DigestData->ResolvedClosestNative;
+			ExclusiveSchemaHash.Reset();
+			ParentClassPath = DigestData->ClosestNative;
 		}
 	}
 	else
 	{
 		DigestData->bNative = false;
-		DigestData->SchemaHash.Reset();
+		ExclusiveSchemaHash.Reset();
 		DigestData->CustomVersionHandles.Empty();
-		FStringView UnusedClassOfClassName;
-		FStringView ClassPackageName;
-		FStringView ClassObjectName;
-		FStringView ClassSubObjectName;
-		FPackageName::SplitFullObjectPath(NameStringBuffer, UnusedClassOfClassName, ClassPackageName, ClassObjectName, ClassSubObjectName);
-		FTopLevelAssetPath ClassObjectPathName = FTopLevelAssetPath(FName(ClassPackageName), FName(ClassObjectName));
 
 		// TODO_EDITORDOMAIN: If the class is not yet present in the assetregistry, or
 		// if its parent classes are not, then we will not be able to propagate information from the parent classes; wait on the class to be parsed
-		AncestorPathNames.Reset();
-		IAssetRegistry::Get()->GetAncestorClassNames(ClassObjectPathName, AncestorPathNames);
-		if (!AncestorPathNames.IsEmpty())
-		{
-			NameStringBuffer.Reset();
-			AncestorPathNames[0].AppendString(NameStringBuffer);
-			ParentName = FName(*NameStringBuffer);
-			for (const FTopLevelAssetPath& PathName : AncestorPathNames)
+		AncestorClassPaths.Reset();
+		IAssetRegistry::Get()->GetAncestorClassNames(ClassPath, AncestorClassPaths);
+		if (!AncestorClassPaths.IsEmpty())
+		{
+			ParentClassPath = AncestorClassPaths[0];
+			for (const FTopLevelAssetPath& AncestorClassPath : AncestorClassPaths)
 			{
-				if (UStruct* CurrentStruct = FindObject<UStruct>(PathName))
+				if (UStruct* CurrentStruct = FindObject<UStruct>(AncestorClassPath))
 				{
-					NameStringBuffer.Reset();
-					CurrentStruct->GetPathName(nullptr, NameStringBuffer);
-					if (FPackageName::IsScriptPackage(NameStringBuffer))
+					if (CurrentStruct->IsNative())
 					{
-						Struct = CurrentStruct;
-						DigestData->ResolvedClosestNative = FName(*NameStringBuffer);
+						DigestData->ClosestNative = AncestorClassPath;
 						break;
 					}
 				}
 			}
 		}
-		if (ParentName.IsNone())
+		if (ParentClassPath.IsNull())
 		{
-			ParentName = FName(*UObject::StaticClass()->GetPathName());
+			ParentClassPath = FTopLevelAssetPath(UObject::StaticClass());
 		}
-		if (DigestData->ResolvedClosestNative.IsNone())
+		if (DigestData->ClosestNative.IsNull())
 		{
-			DigestData->ResolvedClosestNative = FName(*UObject::StaticClass()->GetPathName());
+			DigestData->ClosestNative = FTopLevelAssetPath(UObject::StaticClass());
 		}
 	}
-	check(!DigestData->ResolvedClosestNative.IsNone());
+	// ParentClassPath can be null for CoreUObject.Object and for structs.
+	check(!DigestData->ClosestNative.IsNull());
+
+	// Propagate values from the parent
+	TArray<FTopLevelAssetPath> ConstructClasses; // Not a class scratch variable because we use it across recursive calls
 
 	// Get the CustomVersions used by the native class; GetCustomVersions already returns all custom versions used
 	// by the parent class so we do not need to copy data from the parent
 	UClass* StructAsClass = Cast<UClass>(Struct);
 	if (StructAsClass)
 	{
-		// GetCustomVersions can create the ClassDefaultObject, which can trigger LoadPackage, which
-		// can reenter this function recursively. We have to drop the lock to prevent a deadlock.
+		// GetCustomVersions and CallDeclareConstructClasses can create the ClassDefaultObject, which can trigger
+		// LoadPackage, which can reenter this function recursively. We have to drop the lock to prevent a deadlock.
 		FUnlockScope UnlockScope(ClassDigestsMap.Lock);
 		FKnownCustomVersions::FindOrAddHandles(DigestData->CustomVersionHandles, GetCustomVersions(*StructAsClass));
+		StructAsClass->CallDeclareConstructClasses(ConstructClasses);
 	}
 	else
 	{
 		DigestData->CustomVersionHandles.Reset();
 	}
 
-	// Propagate values from the parent
-	TArray<FName> ConstructClasses; // Not a class scratch variable because we use it across recursive calls
-	if (!ParentName.IsNone())
+	if (!ParentClassPath.IsNull())
 	{
-		// Set bAllowRedirects = false. ParentName already has been CoreRedirected, because it is from the native
+		// Set bAllowRedirects = false. ParentClassPath already has been CoreRedirected, because it is from the native
 		// struct's parent class pointer or it is from GetAncestorClassNames which gives the redirected ancestors.
-		FClassDigestData* ParentDigest = GetRecursive(ParentName, false /* bAllowRedirects */);
+		FClassDigestData* ParentDigest = GetRecursive(ParentClassPath, false /* bAllowRedirects */);
 		// The map has possibly been modified so we need to recalculate the address of ClassName's DigestData
-		DigestData = &ClassDigests.FindChecked(ClassName);
+		DigestData = &ClassDigests.FindChecked(ClassPath);
 		if (!ParentDigest)
 		{
 			UE_LOG(LogEditorDomain, Display,
 				TEXT("Parent class %s of class %s not found. Allow flags for editordomain and iterative cooking will be invalid."),
-				*ParentName.ToString(), *ClassName.ToString());
+				*ParentClassPath.ToString(), *ClassPath.ToString());
 		}
 		else
 		{
@@ -630,11 +624,11 @@ FClassDigestData* FPrecacheClassDigest::GetRecursive(FName ClassName, bool bAllo
 				// We could fix this case by adding bAllowRedirects to the ClassDigestsMap lookup key, but it's not a
 				// problem for MuticastDelegateProperty and we don't have any other cases where it is a problem, so we
 				// avoid the performance cost of doing so.
-				if (ClassName != FName(TEXT("/Script/CoreUObject.MulticastDelegateProperty")))
+				if (ClassPath != FTopLevelAssetPath(TEXT("/Script/CoreUObject.MulticastInlineDelegateProperty")))
 				{
 					UE_LOG(LogEditorDomain, Display,
 						TEXT("Cycle detected in parents of class %s. Allow flags for editordomain and iterative cooking will be invalid."),
-						*ClassName.ToString());
+						*ClassPath.ToString());
 				}
 			}
 			EnumSetFlagsAnd(DigestData->EditorDomainUse, EDomainUse::LoadEnabled | EDomainUse::SaveEnabled,
@@ -644,52 +638,94 @@ FClassDigestData* FPrecacheClassDigest::GetRecursive(FName ClassName, bool bAllo
 				DigestData->bTargetIterativeEnabled &= ParentDigest->bTargetIterativeEnabled;
 			}
 			ConstructClasses.Append(ParentDigest->ConstructClasses);
+			if (!StructAsClass)
+			{
+				DigestData->CustomVersionHandles = ParentDigest->CustomVersionHandles;
+			}
 		}
 	}
+	// Also add ConstructClasses from ini. MultiFind appends to the out argument rather than overwriting it.
+	GConstructClasses.MultiFind(ClassPath, ConstructClasses);
 
-	// Get and validate the list of ConstructClasses by recursively initializing them
-	TArray<FName> AddedConstructClasses; // Not a class scratch variable because we use it across recursive calls
-	GConstructClasses.MultiFind(ClassName, AddedConstructClasses);
-	if (RedirectsResolvedName != ClassName)
-	{
-		GConstructClasses.MultiFind(RedirectsResolvedName, AddedConstructClasses);
-	}
-	if (!AddedConstructClasses.IsEmpty())
+	// Validate the construct classes, find all the transitive constructclasses, and create the InclusiveSchema
+	TArray<FTopLevelAssetPath> TransitiveConstructClasses;
+	TArray<int32> TransitiveCustomVersionHandles;
+	bool bHasConstructClasses = false;
+	if (ConstructClasses.Num())
 	{
-		for (int32 Index = 0; Index < AddedConstructClasses.Num(); )
+		ConstructClasses.RemoveAll([](FTopLevelAssetPath& ClassPath) { return ClassPath.IsNull(); });
+		// Sort ConstructClasses lexically so the Inclusive schema is deterministic
+		Algo::Sort(ConstructClasses, [](const FTopLevelAssetPath& A, const FTopLevelAssetPath& B) { return A.Compare(B) < 0; });
+		ConstructClasses.SetNum(Algo::Unique(ConstructClasses));
+		FBlake3 Writer;
+		Writer.Update(&ExclusiveSchemaHash, sizeof(ExclusiveSchemaHash));
+		bool bCycleDetected = false;
+
+		for (int32 Index = 0; Index < ConstructClasses.Num(); ++Index)
 		{
-			FName ConstructClass = AddedConstructClasses[Index];
+			FTopLevelAssetPath& ConstructClass = ConstructClasses[Index];
+			if (ConstructClass == ClassPath)
+			{
+				continue;
+			}
 			FClassDigestData* ConstructClassDigest = GetRecursive(ConstructClass, true /* bAllowRedirects */);
 			if (!ConstructClassDigest)
 			{
-				UE_LOG(LogEditorDomain, Warning,
-					TEXT("Construct class %s specified by Editor.ini:[EditorDomain]:PostLoadConstructClasses for class %s is not found. ")
-					TEXT("This is a class that can be constructed by postload upgrades of class %s. ")
-					TEXT("Old packages with class %s will load more slowly."),
-					*ConstructClass.ToString(), *ClassName.ToString(), *ClassName.ToString(), *ClassName.ToString());
-				AddedConstructClasses.RemoveAt(Index);
+				UE_LOG(LogEditorDomain, Verbose,
+					TEXT("Construct class %s specified for class %s is not found."),
+					*ConstructClass.ToString(), *ClassPath.ToString());
+				ConstructClasses.RemoveAt(Index);
+				continue;
 			}
-			else
+			Writer.Update(&ConstructClassDigest->InclusiveSchemaHash, sizeof(ConstructClassDigest->InclusiveSchemaHash));
+			TransitiveCustomVersionHandles.Append(ConstructClassDigest->CustomVersionHandles);
+
+			TransitiveConstructClasses.Add(ConstructClass);
+			for (const FTopLevelAssetPath& TransitiveConstructClass : ConstructClassDigest->ConstructClasses)
 			{
-				if (!ConstructClassDigest->bConstructionComplete)
+				if (TransitiveConstructClass == ConstructClass)
+				{
+					bCycleDetected = true;
+				}
+				else
 				{
-					UE_LOG(LogEditorDomain, Verbose,
-						TEXT("Cycle detected in Editor.ini:[EditorDomain]:PostLoadConstructClasses of class %s. This is unexpected, but not a problem."),
-						*ClassName.ToString());
+					TransitiveConstructClasses.Add(TransitiveConstructClass);
 				}
-				++Index;
 			}
 		}
-		ConstructClasses.Append(AddedConstructClasses);
-
+		if (bCycleDetected)
+		{
+			// Handling cycles is not yet implemented. If we need to handle it, we will have to change the graph
+			// algorithm to keep track of referencers of inprogress classes and set their construct classes to include
+			// all transitive constructclasses of the strongly connected component.
+			UE_LOG(LogEditorDomain, Error,
+				TEXT("Cycle detected in ConstructClasses of class %s. Packages with classes that can construct %s will sometimes fail to update during incremental cooks."),
+				*ClassPath.ToString());
+		}
 		// The map has possibly been modified so we need to recalculate the address of ClassName's DigestData
-		DigestData = &ClassDigests.FindChecked(ClassName);
+		DigestData = &ClassDigests.FindChecked(ClassPath);
+
+		// Store the transitive data on the Digest
+		if (TransitiveConstructClasses.Num())
+		{
+			bHasConstructClasses = true;
+			DigestData->ConstructClasses = MoveTemp(TransitiveConstructClasses);
+			DigestData->CustomVersionHandles.Append(TransitiveCustomVersionHandles);
+
+			// The order of ConstructClasses is allowed to vary between runs; anything that uses it for deterministic data like a schema
+			// hash will need to resort it. We set the order to use CompareFast so we can support fast binary search lookup.
+			Algo::Sort(DigestData->ConstructClasses, [](const FTopLevelAssetPath& A, const FTopLevelAssetPath& B) { return A.CompareFast(B) < 0; });
+			DigestData->ConstructClasses.SetNum(Algo::Unique(DigestData->ConstructClasses));
+			Algo::Sort(DigestData->CustomVersionHandles);
+			DigestData->CustomVersionHandles.SetNum(Algo::Unique(DigestData->CustomVersionHandles));
+			
+			DigestData->InclusiveSchemaHash = Writer.Finalize();
+		}
 	}
-	if (ConstructClasses.Num())
+	if (!bHasConstructClasses)
 	{
-		Algo::Sort(ConstructClasses, FNameLexicalLess());
-		ConstructClasses.SetNum(Algo::Unique(ConstructClasses));
-		DigestData->ConstructClasses = MoveTemp(ConstructClasses);
+		DigestData->ConstructClasses.Reset();
+		DigestData->InclusiveSchemaHash = ExclusiveSchemaHash;
 	}
 
 	DigestData->bConstructionComplete = true;
@@ -697,12 +733,15 @@ FClassDigestData* FPrecacheClassDigest::GetRecursive(FName ClassName, bool bAllo
 }
 
 /** Try to add the FClassDigestData for each given class into the GetClassDigests map */
-void PrecacheClassDigests(TConstArrayView<FName> ClassNames)
+void PrecacheClassDigests(TConstArrayView<FTopLevelAssetPath> ClassPaths)
 {
 	FPrecacheClassDigest Digester;
-	for (FName ClassName : ClassNames)
+	for (const FTopLevelAssetPath& ClassPath : ClassPaths)
 	{
-		Digester.GetRecursive(ClassName, true /* bAllowRedirects */);
+		if (!ClassPath.IsNull())
+		{
+			Digester.GetRecursive(ClassPath, true /* bAllowRedirects */);
+		}
 	}
 }
 
@@ -720,32 +759,47 @@ void InitializeGlobalConstructClasses()
 		GGlobalConstructClasses.Reserve(Lines.Num());
 		for (const FString& Line : Lines)
 		{
-			GGlobalConstructClasses.Add(FName(FStringView(Line).TrimStartAndEnd()));
+			FStringView ClassPathText(FStringView(Line).TrimStartAndEnd());
+			FTopLevelAssetPath ClassPath(ClassPathText);
+			if (ClassPath.IsValid())
+			{
+				CoreRedirectClassPath(ClassPath);
+				GGlobalConstructClasses.Add(ClassPath);
+			}
+			else
+			{
+				UE_LOG(LogEditorDomain, Error, TEXT("Invalid ClassPath %.*s specified by Editor.ini:[EditorDomain]:GlobalCanConstructClasses."),
+					ClassPathText.Len(), ClassPathText.GetData());
+			}
 		}
 	}
 
 	PrecacheClassDigests(GGlobalConstructClasses);
 	FClassDigestMap& ClassDigests = GetClassDigests();
 	FReadScopeLock ClassDigestsScopeLock(ClassDigests.Lock);
+	FBlake3 Writer;
+
+	Algo::Sort(GGlobalConstructClasses, [](const FTopLevelAssetPath& A, const FTopLevelAssetPath& B) { return A.Compare(B) < 0; });
+	GGlobalConstructClasses.SetNum(Algo::Unique(GGlobalConstructClasses));
 	for (int32 Index = 0; Index < GGlobalConstructClasses.Num();)
 	{
-		FName ClassName = GGlobalConstructClasses[Index];
-		FClassDigestData* ExistingData = ClassDigests.Map.Find(ClassName);
+		const FTopLevelAssetPath& ClassPath = GGlobalConstructClasses[Index];
+		FClassDigestData* ExistingData = ClassDigests.Map.Find(ClassPath);
 		if (!ExistingData)
 		{
 			UE_LOG(LogEditorDomain, Display, TEXT("Construct class %s specified by Editor.ini:[EditorDomain]:GlobalCanConstructClasses is not found. ")
 				TEXT("This is a class that can be constructed automatically by SavePackage when saving old packages. ")
 				TEXT("Old packages that do not yet have this class will load more slowly."),
-				*ClassName.ToString());
+				*ClassPath.ToString());
 			GGlobalConstructClasses.RemoveAt(Index);
 		}
 		else
 		{
+			Writer.Update(&ExistingData->InclusiveSchemaHash, sizeof(ExistingData->InclusiveSchemaHash));
 			++Index;
 		}
 	}
-	Algo::Sort(GGlobalConstructClasses, FNameLexicalLess());
-	GGlobalConstructClasses.SetNum(Algo::Unique(GGlobalConstructClasses));
+	GGlobalConstructClassesHash = Writer.Finalize();
 }
 
 /** An archive that just collects custom versions. */
@@ -843,9 +897,9 @@ TMap<FGuid, UObject*> FindCustomVersionCulprits(TConstArrayView<FGuid> UnknownGu
 	return Culprits;
 }
 
-TMap<FName, EDomainUse> ConstructClassBlockedUses()
+TMap<FTopLevelAssetPath, EDomainUse> ConstructClassBlockedUses()
 {
-	TMap<FName, EDomainUse> Result;
+	TMap<FTopLevelAssetPath, EDomainUse> Result;
 	TArray<FString> BlockListArray;
 	TArray<FString> LoadBlockListArray;
 	TArray<FString> SaveBlockListArray;
@@ -857,9 +911,15 @@ TMap<FName, EDomainUse> ConstructClassBlockedUses()
 		EDomainUse BlockedUse = Array == &BlockListArray	?	(EDomainUse::LoadEnabled | EDomainUse::SaveEnabled) : (
 			Array == &LoadBlockListArray					?	EDomainUse::LoadEnabled : (
 																EDomainUse::SaveEnabled));
-		for (const FString& ClassPathName : *Array)
+		for (const FString& ClassPathString : *Array)
 		{
-			Result.FindOrAdd(FName(*ClassPathName), EDomainUse::None) |= BlockedUse;
+			FTopLevelAssetPath ClassPath(ClassPathString);
+			if (!ClassPath.IsValid())
+			{
+				continue;
+			}
+			CoreRedirectClassPath(ClassPath);
+			Result.FindOrAdd(ClassPath, EDomainUse::None) |= BlockedUse;
 		}
 	}
 	return Result;
@@ -874,7 +934,7 @@ TMap<FName, EDomainUse> ConstructPackageNameBlockedUses()
 	GConfig->GetArray(TEXT("EditorDomain"), TEXT("PackageBlockList"), BlockListArray, GEditorIni);
 	GConfig->GetArray(TEXT("EditorDomain"), TEXT("PackageLoadBlockList"), LoadBlockListArray, GEditorIni);
 	GConfig->GetArray(TEXT("EditorDomain"), TEXT("PackageSaveBlockList"), SaveBlockListArray, GEditorIni);
-	FString PackageName;
+	FString PackageNameString;
 	FString ErrorReason;
 	for (TArray<FString>* Array : { &BlockListArray, &LoadBlockListArray, &SaveBlockListArray })
 	{
@@ -883,34 +943,52 @@ TMap<FName, EDomainUse> ConstructPackageNameBlockedUses()
 																EDomainUse::SaveEnabled));
 		for (const FString& PackageNameOrFilename : *Array)
 		{
-			if (!FPackageName::TryConvertFilenameToLongPackageName(PackageNameOrFilename, PackageName, &ErrorReason))
+			if (!FPackageName::TryConvertFilenameToLongPackageName(PackageNameOrFilename, PackageNameString, &ErrorReason))
 			{
 				UE_LOG(LogEditorDomain, Warning, TEXT("Editor.ini:[EditorDomain]:PackageBlocklist: Could not convert %s to a LongPackageName: %s"),
 					*PackageNameOrFilename, *ErrorReason);
 				continue;
 			}
-			Result.FindOrAdd(FName(*PackageName), EDomainUse::None) |= BlockedUse;
+			FName PackageName(*PackageNameString);
+			FCoreRedirectObjectName PackageRedirect(NAME_None, NAME_None, PackageName);
+			FCoreRedirectObjectName RedirectedPackageRedirect = FCoreRedirects::GetRedirectedName(
+				ECoreRedirectFlags::Type_Package, PackageRedirect);
+			if (!RedirectedPackageRedirect.OuterName.IsNone() || !RedirectedPackageRedirect.ObjectName.IsNone())
+			{
+				UE_LOG(LogEditorDomain, Warning, TEXT("CoreRedirect for package %s redirects to %s which is not a package path. Ignoring it."),
+					*PackageRedirect.ToString(), *RedirectedPackageRedirect.ToString())
+			}
+			else
+			{
+				PackageName = RedirectedPackageRedirect.PackageName;
+			}
+			Result.FindOrAdd(PackageName, EDomainUse::None) |= BlockedUse;
 		}
 	}
 	return Result;
 }
 
-TSet<FName> ConstructTargetIterativeClassBlockList()
+TSet<FTopLevelAssetPath> ConstructTargetIterativeClassBlockList()
 {
-	TSet<FName> Result;
+	TSet<FTopLevelAssetPath> Result;
 	TArray<FString> BlockListArray;
 	GConfig->GetArray(TEXT("TargetDomain"), TEXT("IterativeClassBlockList"), BlockListArray, GEditorIni);
-	for (const FString& ClassPathName : BlockListArray)
+	for (const FString& ClassPathString : BlockListArray)
 	{
-		Result.Add(FName(*ClassPathName));
+		FTopLevelAssetPath ClassPath(ClassPathString);
+		if (ClassPath.IsValid())
+		{
+			CoreRedirectClassPath(ClassPath);
+			Result.Add(ClassPath);
+		}
 	}
 	return Result;
 }
 
 void ConstructTargetIterativeClassAllowList()
 {
-	// We're using a allowlist with a blocklist override, so the blocklist is only needed when creating the allowlist
-	TSet<FName> BlockListFNames = ConstructTargetIterativeClassBlockList();
+	// We're using an allowlist with a blocklist override, so the blocklist is only needed when creating the allowlist
+	TSet<FTopLevelAssetPath> BlockListClassPaths = ConstructTargetIterativeClassBlockList();
 
 	// AllowList elements implicitly allow all parent classes, so instead of consulting a list and propagating
 	// from parent classes every time we read a new class, we have to iterate the list for all classes up front and
@@ -919,14 +997,14 @@ void ConstructTargetIterativeClassAllowList()
 
 	// Declare a recursive Visit function. Every class we visit is allowlisted, and we visit its superclasses.
 	// To decide whether a visited class is enabled, we also have to get IsBlockListed recursively from the parent.
-	TSet<FName> EnabledFNames;
+	TSet<FTopLevelAssetPath> EnabledClassPaths;
 	TStringBuilder<256> NameStringBuffer;
-	TMap<FName, TOptional<bool>> Visited;
-	auto EnableClassIfNotBlocked = [&Visited, &EnabledFNames, &BlockListFNames, &NameStringBuffer]
-		(FName PathName, UStruct* Struct, bool& bOutIsBlocked, auto& EnableClassIfNotBlockedRef)
+	TMap<FTopLevelAssetPath, TOptional<bool>> Visited;
+	auto EnableClassIfNotBlocked = [&Visited, &EnabledClassPaths, &BlockListClassPaths, &NameStringBuffer]
+		(const FTopLevelAssetPath& ClassPath, UStruct* Struct, bool& bOutIsBlocked, auto& EnableClassIfNotBlockedRef)
 	{
-		int32 KeyHash = GetTypeHash(PathName);
-		TOptional<bool>& BlockedValue = Visited.FindOrAddByHash(KeyHash, PathName);
+		int32 KeyHash = GetTypeHash(ClassPath);
+		TOptional<bool>& BlockedValue = Visited.FindOrAddByHash(KeyHash, ClassPath);
 		if (BlockedValue.IsSet())
 		{
 			bOutIsBlocked = *BlockedValue;
@@ -938,49 +1016,56 @@ void ConstructTargetIterativeClassAllowList()
 		UStruct* ParentStruct = Struct->GetSuperStruct();
 		if (ParentStruct)
 		{
-			NameStringBuffer.Reset();
-			ParentStruct->GetPathName(nullptr, NameStringBuffer);
-			EnableClassIfNotBlockedRef(FName(*NameStringBuffer), ParentStruct,
-				bParentBlocked, EnableClassIfNotBlockedRef);
+			FTopLevelAssetPath ParentStructPath(ParentStruct);
+			if (ParentStructPath.IsValid())
+			{
+				EnableClassIfNotBlockedRef(ParentStructPath, ParentStruct, bParentBlocked, EnableClassIfNotBlockedRef);
+			}
 		}
 
-		bOutIsBlocked = bParentBlocked || BlockListFNames.Contains(PathName);
+		bOutIsBlocked = bParentBlocked || BlockListClassPaths.Contains(ClassPath);
 		if (bOutIsBlocked)
 		{
 			// Call FindOrAdd again, since the recursive calls may have altered the map and invalidated the BlockedValue reference
-			Visited.FindOrAddByHash(KeyHash, PathName) = bOutIsBlocked;
+			Visited.FindOrAddByHash(KeyHash, ClassPath) = bOutIsBlocked;
 		}
 		else
 		{
-			EnabledFNames.Add(PathName);
+			EnabledClassPaths.Add(ClassPath);
 		}
 	};
 
 	TArray<FString> AllowListLeafNames;
 	GConfig->GetArray(TEXT("TargetDomain"), TEXT("IterativeClassAllowList"), AllowListLeafNames, GEditorIni);
-	for (const FString& ClassPathName : AllowListLeafNames)
+	for (const FString& ClassPathString : AllowListLeafNames)
 	{
-		if (!FPackageName::IsScriptPackage(ClassPathName))
+		FTopLevelAssetPath ClassPath(ClassPathString);
+		if (!ClassPath.IsValid())
 		{
 			continue;
 		}
-		UStruct* Struct = FindObject<UStruct>(nullptr, *ClassPathName);
+		CoreRedirectClassPath(ClassPath);
+		if (!FPackageName::IsScriptPackage(WriteToString<256>(ClassPath.GetPackageName()).ToView()))
+		{
+			continue;
+		}
+		UStruct* Struct = FindObject<UStruct>(ClassPath);
 		if (!Struct)
 		{
 			continue;
 		}
 		bool bUnusedIsBlocked;
-		EnableClassIfNotBlocked(FName(*ClassPathName), Struct, bUnusedIsBlocked, EnableClassIfNotBlocked);
+		EnableClassIfNotBlocked(ClassPath, Struct, bUnusedIsBlocked, EnableClassIfNotBlocked);
 	}
 
-	TArray<FName> EnabledFNamesArray = EnabledFNames.Array();
-	PrecacheClassDigests(EnabledFNamesArray);
+	TArray<FTopLevelAssetPath> EnabledClassPathsArray = EnabledClassPaths.Array();
+	PrecacheClassDigests(EnabledClassPathsArray);
 	FClassDigestMap& ClassDigests = GetClassDigests();
 	{
 		FWriteScopeLock ClassDigestsScopeLock(ClassDigests.Lock);
-		for (FName ClassPathName : EnabledFNamesArray)
+		for (const FTopLevelAssetPath& ClassPath : EnabledClassPathsArray)
 		{
-			FClassDigestData* DigestData = ClassDigests.Map.Find(ClassPathName);
+			FClassDigestData* DigestData = ClassDigests.Map.Find(ClassPath);
 			if (DigestData)
 			{
 				DigestData->bTargetIterativeEnabled = true;
@@ -990,19 +1075,20 @@ void ConstructTargetIterativeClassAllowList()
 }
 
 /** Construct PostLoadCanConstructClasses multimap from config settings and return it */
-TMultiMap<FName, FName> ConstructConstructClasses()
+TMultiMap<FTopLevelAssetPath, FTopLevelAssetPath> ConstructConstructClasses()
 {
 	TArray<FString> Lines;
 	GConfig->GetArray(TEXT("EditorDomain"), TEXT("PostLoadCanConstructClasses"), Lines, GEditorIni);
-	TMultiMap<FName, FName> ConstructClasses;
+	TMultiMap<FTopLevelAssetPath, FTopLevelAssetPath> ConstructClasses;
 	for (const FString& Line : Lines)
 	{
 		int32 NumTokens = 0;
-		FStringView PostLoadClass;
-		FStringView ConstructedClass;
-		UE::String::ParseTokens(Line, TEXT(','), [&NumTokens, &PostLoadClass, &ConstructedClass](FStringView Token)
+		FStringView PostLoadClassString;
+		FStringView ConstructedClassString;
+		UE::String::ParseTokens(Line, TEXT(','),
+			[&NumTokens, &PostLoadClassString, &ConstructedClassString](FStringView Token)
 			{
-				*(NumTokens == 0 ? &PostLoadClass : &ConstructedClass) = Token;
+				*(NumTokens == 0 ? &PostLoadClassString : &ConstructedClassString) = Token;
 				++NumTokens;
 			});
 		if (NumTokens != 2)
@@ -1010,9 +1096,13 @@ TMultiMap<FName, FName> ConstructConstructClasses()
 			UE_LOG(LogEditorDomain, Warning, TEXT("Invalid value %s in config setting Editor.ini:[EditorDomain]:PostLoadCanConstructClasses"), *Line);
 			continue;
 		}
-		PostLoadClass.TrimStartAndEndInline();
-		ConstructedClass.TrimStartAndEndInline();
-		ConstructClasses.Add(FName(PostLoadClass), FName(ConstructedClass));
+		PostLoadClassString.TrimStartAndEndInline();
+		ConstructedClassString.TrimStartAndEndInline();
+		FTopLevelAssetPath PostLoadClass(PostLoadClassString);
+		FTopLevelAssetPath ConstructedClass(ConstructedClassString);
+		CoreRedirectClassPath(PostLoadClass);
+		CoreRedirectClassPath(ConstructedClass);
+		ConstructClasses.Add(PostLoadClass, ConstructedClass);
 	};
 	return ConstructClasses;
 }
@@ -1661,6 +1751,13 @@ bool TrySavePackage(UPackage* Package)
 			*WriteToString<256>(PackageName));
 		return false;
 	}
+	TOptional<FAssetPackageData> PackageData = IAssetRegistry::Get()->GetAssetPackageDataCopy(Package->GetFName());
+	if (!PackageData)
+	{
+		UE_LOG(LogEditorDomain, Verbose, TEXT("Could not save %s to EditorDomain: It does not have An AssetPackageData in the AssetRegistry."),
+			*WriteToString<256>(PackageName));
+		return false;
+	}
 	COOK_STAT(auto Timer = UE::EditorDomain::CookStats::Usage.TimeSyncWork());
 
 	UE_LOG(LogEditorDomain, Verbose, TEXT("Saving to EditorDomain: %s."), *WriteToString<256>(PackageName));
@@ -1729,45 +1826,56 @@ bool TrySavePackage(UPackage* Package)
 	Info.WriteOptions = IPackageWriter::EWriteOptions::Write;
 	PackageWriter->CommitPackage(MoveTemp(Info));
 
-	TArray<FName> SavedImportedClasses;
+	TArray<FName> SavedImportedClassPathNames;
 	FCustomVersionContainer SavedCustomVersions;
-	if (!PackageWriter->TryGetClassesAndVersions(SavedImportedClasses, SavedCustomVersions))
+	if (!PackageWriter->TryGetClassesAndVersions(SavedImportedClassPathNames, SavedCustomVersions))
 	{
 		UE_LOG(LogEditorDomain, Warning, TEXT("Could not save %s to EditorDomain: Could not read the PackageFileSummary from the saved bytes."),
 			*Package->GetName());
 		return false;
 	}
+	TArray<FTopLevelAssetPath> SavedImportedClassPaths;
+	SavedImportedClassPaths.Reserve(SavedImportedClassPathNames.Num());
+	for (FName ClassPathName : SavedImportedClassPathNames)
+	{
+		FTopLevelAssetPath ClassPath(WriteToString<256>(ClassPathName).ToView());
+		if (!ClassPath.IsValid())
+		{
+			UE_LOG(LogEditorDomain, Warning, TEXT("Could not save %s to EditorDomain: It uses an importclass %s that is not a TopLevelAssetPath."),
+				*Package->GetName(), *ClassPathName.ToString());
+			return false;
+		}
+		SavedImportedClassPaths.Add(ClassPath);
+	}
 
 	// Replace any Blueprint classes in SavedImportedClasses with their native base class, since only native classes
 	// have class schemas and customversions that could change the custom versions. EditorDomain packages with blueprint
 	// instances do not use unversioned properties and so do not need to be keyed to changes in the BP class they depend on.
-	// ImportedClasses used that same replacement, so we need to match it here when identifying which classes are new.
 	int32 NextClass = 0;
 	bool bHasTriedPrecacheClassDigests = false;
 	FClassDigestMap& ClassDigests = GetClassDigests();
-	while (NextClass < SavedImportedClasses.Num())
+	while (NextClass < SavedImportedClassPaths.Num())
 	{
 		{
 			FReadScopeLock ClassDigestsScopeLock(ClassDigests.Lock);
-			while (NextClass < SavedImportedClasses.Num())
+			while (NextClass < SavedImportedClassPaths.Num())
 			{
-				FName ClassName = SavedImportedClasses[NextClass];
-				FClassDigestData* ExistingData = ClassDigests.Map.Find(ClassName);
+				FClassDigestData* ExistingData = ClassDigests.Map.Find(SavedImportedClassPaths[NextClass]);
 				if (!ExistingData)
 				{
 					break;
 				}
-				SavedImportedClasses[NextClass] = ExistingData->ResolvedClosestNative;
+				SavedImportedClassPaths[NextClass] = ExistingData->ClosestNative;
 				++NextClass;
 			}
 		}
 
-		if (NextClass < SavedImportedClasses.Num())
+		if (NextClass < SavedImportedClassPaths.Num())
 		{
 			if (bHasTriedPrecacheClassDigests)
 			{
 				UE_LOG(LogEditorDomain, Warning, TEXT("Could not save %s to EditorDomain: It constructed an instance of class %s which which does not exist."),
-					*Package->GetName(), *SavedImportedClasses[NextClass].ToString());
+					*Package->GetName(), *SavedImportedClassPaths[NextClass].ToString());
 				return false;
 			}
 			// EDITORDOMAIN_TODO: Remove this !IsInGameThread check once FindObject no longer asserts if GIsSavingPackage
@@ -1775,10 +1883,11 @@ bool TrySavePackage(UPackage* Package)
 			{
 				UE_LOG(LogEditorDomain, Warning, TEXT("Could not save %s to EditorDomain: It constructed an instance of a class %s which is not yet saved in the EditorDomain, ")
 					TEXT("and we are not on the gamethread so we cannot add it."),
-					*Package->GetName(), *SavedImportedClasses[NextClass].ToString());
+					*Package->GetName(), *SavedImportedClassPaths[NextClass].ToString());
 				return false;
 			}
-			TConstArrayView<FName> RemainingClasses = TConstArrayView<FName>(SavedImportedClasses).RightChop(NextClass);
+			TConstArrayView<FTopLevelAssetPath> RemainingClasses =
+				TConstArrayView<FTopLevelAssetPath>(SavedImportedClassPaths).RightChop(NextClass);
 			PrecacheClassDigests(RemainingClasses);
 			bHasTriedPrecacheClassDigests = true;
 		}
@@ -1786,8 +1895,45 @@ bool TrySavePackage(UPackage* Package)
 
 	ESaveStorageResult StorageResult = ESaveStorageResult::Valid;
 
-	TSet<FName> KnownImportedClasses(PackageDigest.ImportedClasses);
-	for (FName ImportedClass : SavedImportedClasses)
+	TArray<FTopLevelAssetPath> ImportedClassPaths;
+	ImportedClassPaths.Reserve(PackageData->ImportedClasses.Num());
+	for (FName ClassPathName : PackageData->ImportedClasses)
+	{
+		FTopLevelAssetPath ClassPath(WriteToString<256>(ClassPathName).ToView());
+		if (ClassPath.IsValid())
+		{
+			ImportedClassPaths.Add(ClassPath);
+		}
+	}
+
+	TSet<FTopLevelAssetPath> KnownImportedClasses;
+	TSet<FGuid> KnownGuids;
+	TArray<FGuid> ClassCustomVersions;
+	{
+		FReadScopeLock ClassDigestsScopeLock(ClassDigests.Lock);
+		for (TArray<FTopLevelAssetPath>* ClassPathArray : { &ImportedClassPaths, &GGlobalConstructClasses })
+		{
+			for (const FTopLevelAssetPath ImportedClassPath : *ClassPathArray)
+			{
+				FClassDigestData* ExistingData = ClassDigests.Map.Find(ImportedClassPath);
+				if (ExistingData)
+				{
+					// As above with SavedImportedClasses, we add the ClosestNative of each imported class instead of
+					// the (possibly BlueprintGeneratedClass) itself
+					KnownImportedClasses.Add(ExistingData->ClosestNative);
+					for (const FTopLevelAssetPath& ConstructedClass : ExistingData->ConstructClasses)
+					{
+						KnownImportedClasses.Add(ConstructedClass);
+					}
+					ClassCustomVersions.Reset();
+					FKnownCustomVersions::FindGuidsChecked(ClassCustomVersions, ExistingData->CustomVersionHandles);
+					KnownGuids.Append(ClassCustomVersions);
+				}
+			}
+		}
+	}
+
+	for (const FTopLevelAssetPath& ImportedClass : SavedImportedClassPaths)
 	{
 		if (!KnownImportedClasses.Contains(ImportedClass))
 		{
@@ -1795,50 +1941,50 @@ bool TrySavePackage(UPackage* Package)
 			// at the start of LoadPackageInternal. After it gets hit, add a breakpoint in the constructor
 			// of the ImportedClass.
 			UE_LOG(LogEditorDomain, Display, TEXT("Could not save %s to EditorDomain: It uses an unexpected class. ")
-				TEXT("Optimized loading and iterative cooking will be disabled for this package.\n\t")
-				TEXT("Find the class that added %s and add ")
-				TEXT("Editor.ini:[EditorDomain]:+PostLoadCanConstructClasses=<ConstructingClass>,%s"),
-				*Package->GetName(), *ImportedClass.ToString(), *ImportedClass.ToString());
+				TEXT("Optimized loading and incremental cooking will be disabled for this package.\n\t")
+				TEXT("Find the class that added %s and add it to its DeclareConstructClasses function."),
+				*Package->GetName(), *ImportedClass.ToString());
 			StorageResult = ESaveStorageResult::UnexpectedClass;
 		}
 	}
 
-	TConstArrayView<UE::AssetRegistry::FPackageCustomVersion> KnownCustomVersions = PackageDigest.CustomVersions.Get();
-	TSet<FGuid> KnownGuids;
-	KnownGuids.Reserve(KnownCustomVersions.Num());
-	for (const UE::AssetRegistry::FPackageCustomVersion& CustomVersion : KnownCustomVersions)
-	{
-		KnownGuids.Add(CustomVersion.Key);
-	}
-	TArray<FGuid> UnknownGuids;
-	for (const FCustomVersion& CustomVersion : SavedCustomVersions.GetAllVersions())
+	if (StorageResult == ESaveStorageResult::Valid)
 	{
-		if (!KnownGuids.Contains(CustomVersion.Key))
+		TConstArrayView<UE::AssetRegistry::FPackageCustomVersion> KnownCustomVersions = PackageDigest.CustomVersions.Get();
+		for (const UE::AssetRegistry::FPackageCustomVersion& CustomVersion : KnownCustomVersions)
 		{
-			UnknownGuids.Add(CustomVersion.Key);
+			KnownGuids.Add(CustomVersion.Key);
 		}
-	}
-	if (StorageResult == ESaveStorageResult::Valid && !UnknownGuids.IsEmpty())
-	{
-		TMap<FGuid, UObject*> Culprits = FindCustomVersionCulprits(UnknownGuids, Package);
-		TStringBuilder<256> VersionsText;
-		for (const FGuid& CustomVersionGuid : UnknownGuids)
+		TArray<FGuid> UnknownGuids;
+		for (const FCustomVersion& CustomVersion : SavedCustomVersions.GetAllVersions())
 		{
-			TOptional<FCustomVersion> CustomVersion = FCurrentCustomVersions::Get(CustomVersionGuid);
-			UObject* Culprit = Culprits.FindOrAdd(CustomVersionGuid);
-			VersionsText << TEXT("\n\tCustomVersion(Guid=") << CustomVersionGuid << TEXT(", Name=")
-				<< (CustomVersion ? *CustomVersion->GetFriendlyName().ToString() : TEXT("<Unknown>"))
-				<< TEXT("): Used by ")
-				<< (Culprit ? *Culprit->GetClass()->GetPathName() : TEXT("<CulpritUnknown>"));
+			if (!KnownGuids.Contains(CustomVersion.Key))
+			{
+				UnknownGuids.Add(CustomVersion.Key);
+			}
 		}
+		if (!UnknownGuids.IsEmpty())
+		{
+			TMap<FGuid, UObject*> Culprits = FindCustomVersionCulprits(UnknownGuids, Package);
+			TStringBuilder<256> VersionsText;
+			for (const FGuid& CustomVersionGuid : UnknownGuids)
+			{
+				TOptional<FCustomVersion> CustomVersion = FCurrentCustomVersions::Get(CustomVersionGuid);
+				UObject* Culprit = Culprits.FindOrAdd(CustomVersionGuid);
+				VersionsText << TEXT("\n\tCustomVersion(Guid=") << CustomVersionGuid << TEXT(", Name=")
+					<< (CustomVersion ? *CustomVersion->GetFriendlyName().ToString() : TEXT("<Unknown>"))
+					<< TEXT("): Used by ")
+					<< (Culprit ? *Culprit->GetClass()->GetPathName() : TEXT("<CulpritUnknown>"));
+			}
 
-		// Suggested debugging technique for this message: SetNextStatement back to beginning of the function,
-		// add a conditional breakpoint in FArchive::UsingCustomVersion with Key.A == 0x<FirstHexWordFromGuid>
-		UE_LOG(LogEditorDomain, Display, TEXT("Could not save %s to EditorDomain: It uses an unexpected custom version. ")
-			TEXT("Optimized loading and iterative cooking will be disabled for this package. ")
-			TEXT("Modify the classes/structs to call Ar.UsingCustomVersion(Guid) in Serialize or DeclareCustomVersions.%s"),
-			*Package->GetName(), VersionsText.ToString());
-		StorageResult = ESaveStorageResult::UnexpectedCustomVersion;
+			// Suggested debugging technique for this message: SetNextStatement back to beginning of the function,
+			// add a conditional breakpoint in FArchive::UsingCustomVersion with Key.A == 0x<FirstHexWordFromGuid>
+			UE_LOG(LogEditorDomain, Display, TEXT("Could not save %s to EditorDomain: It uses an unexpected custom version. ")
+				TEXT("Optimized loading and iterative cooking will be disabled for this package. ")
+				TEXT("Modify the classes/structs to call Ar.UsingCustomVersion(Guid) in Serialize or DeclareCustomVersions.%s"),
+				*Package->GetName(), VersionsText.ToString());
+			StorageResult = ESaveStorageResult::UnexpectedCustomVersion;
+		}
 	}
 
 	if (StorageResult == ESaveStorageResult::Valid && 
@@ -2038,4 +2184,264 @@ void PutBulkDataPayloadId(FName PackageName, const FGuid& BulkDataId, FSharedBuf
 	Owner.KeepAlive();
 }
 
+FStringBuilderBase& operator<<(FStringBuilderBase& Writer, UE::EditorDomain::EDomainUse DomainUse)
+{
+	using namespace UE::EditorDomain;
+
+	if (!DomainUse)
+	{
+		Writer << TEXT("None");
+		return Writer;
+	}
+
+	int32 StartLen = Writer.Len();
+	auto AddSeparator = [&Writer, StartLen]()
+	{
+		if (Writer.Len() > StartLen)
+		{
+			Writer << TEXT("|");
+		}
+	};
+	if (EnumHasAnyFlags(DomainUse, EDomainUse::LoadEnabled))
+	{
+		AddSeparator();
+		Writer << TEXT("LoadEnabled");
+	}
+	if (EnumHasAnyFlags(DomainUse, EDomainUse::SaveEnabled))
+	{
+		AddSeparator();
+		Writer << TEXT("SaveEnabled");
+	}
+	if (Writer.Len() == StartLen)
+	{
+		Writer << TEXT("<Invalid>");
+	}
+	return Writer;
+}
+
+static void DumpClassDigests()
+{
+	using namespace UE::EditorDomain;
+
+	TArray<FTopLevelAssetPath> ClassPaths;
+	for (TObjectIterator<UClass> Iter; Iter; ++Iter)
+	{
+		UClass* Class = *Iter;
+		UPackage* Package = Class->GetPackage();
+		if (!Class->HasAnyClassFlags(CLASS_Transient) && Package && !Package->HasAnyFlags(RF_Transient))
+		{
+			FTopLevelAssetPath ClassPath(Class);
+			if (ClassPath.IsValid())
+			{
+				ClassPaths.Add(ClassPath);
+			}
+		}
+	}
+	Algo::Sort(ClassPaths, [](const FTopLevelAssetPath& A, const FTopLevelAssetPath& B) { return A.Compare(B) < 0; });
+	TArray<FString> ClassTexts;
+	ClassTexts.Reserve(ClassPaths.Num());
+	TStringBuilder<1024> ClassText;
+
+	PrecacheClassDigests(ClassPaths);
+
+	UE::EditorDomain::FClassDigestMap& ClassDigests = UE::EditorDomain::GetClassDigests();
+	TArray<FGuid> CustomVersionGuids;
+	{
+		FReadScopeLock ClassDigestsScopeLock(ClassDigests.Lock);
+		for (const FTopLevelAssetPath& ClassPath : ClassPaths)
+		{
+			ClassText.Reset();
+			UE::EditorDomain::FClassDigestData* ExistingData= ClassDigests.Map.Find(ClassPath);
+			if (!ExistingData)
+			{
+				ClassTexts.Add(TEXT("<MissingClass>"));
+				continue;
+			}
+
+			CustomVersionGuids.Reset();
+			FKnownCustomVersions::FindGuidsChecked(CustomVersionGuids, ExistingData->CustomVersionHandles);
+			Algo::Sort(CustomVersionGuids);
+
+			ClassText << TEXT("(");
+			{
+				ClassText << TEXT("EditorDomainUse=") << ExistingData->EditorDomainUse;
+				ClassText << TEXT(",") << TEXT("bNative=") << LexToString(ExistingData->bNative);
+				ClassText << TEXT(",") << TEXT("bTargetIterativeEnabled=")
+					<< LexToString(ExistingData->bTargetIterativeEnabled);
+				ClassText << TEXT(",") << TEXT("ConstructClasses=(");
+				if (!ExistingData->ConstructClasses.IsEmpty())
+				{
+					TArray<FTopLevelAssetPath> Sorted(ExistingData->ConstructClasses);
+					Algo::Sort(Sorted, [](const FTopLevelAssetPath& A, const FTopLevelAssetPath& B) { return A.Compare(B) < 0; });
+					for (const FTopLevelAssetPath& Path : Sorted)
+					{
+						ClassText << Path << TEXT(",");
+					}
+					ClassText.RemoveSuffix(1);
+				}
+				ClassText << TEXT(")");
+				ClassText << TEXT(",") << TEXT("CustomVersions=(");
+				if (CustomVersionGuids.Num() > 0)
+				{
+					for (const FGuid& VersionGuid : CustomVersionGuids)
+					{
+						ClassText << VersionGuid << TEXT(",");
+					}
+					ClassText.RemoveSuffix(1);
+				}
+				ClassText << TEXT(")");
+				ClassText << TEXT(",") << TEXT("InclusiveSchemaHash=") << ExistingData->InclusiveSchemaHash;
+			}
+			ClassText << TEXT(")");
+
+			ClassTexts.Add(FString(ClassText.ToView()));
+		}
+	}
+
+	for (int32 Index = 0; Index < ClassPaths.Num(); ++Index)
+	{
+		UE_LOG(LogEditorDomain, Display, TEXT("ClassDigest: %s=%s"),
+			*WriteToString<256>(ClassPaths[Index]), *ClassTexts[Index]);
+	}
+}
+FAutoConsoleCommand DumpClassDigestsCommand(TEXT("EditorDomain.DumpClassDigests"),
+	TEXT("Write to the log the digest information for each class."),
+	FConsoleCommandDelegate::CreateStatic(&DumpClassDigests));
+
+}
+
+#define CONSTRUCTION_DURING_UPGRADE_TRACKER_ENABLED UE_WITH_PACKAGE_ACCESS_TRACKING && !UE_BUILD_SHIPPING
+
+#if CONSTRUCTION_DURING_UPGRADE_TRACKER_ENABLED
+/**
+ * Utility that validates whether objects constructed during PostLoad have been declared as constructible by the
+ * PostLoading class. This utility is not in the EditorDomain namespace because the warnings are given even when not
+ * using EditorDomain and we want to avoid confusion when looking at the callstack.
+ */
+class FConstructionDuringUpgradeTracker : public FDelayedAutoRegisterHelper
+{
+public:
+	FConstructionDuringUpgradeTracker();
+
+private:
+	void OnObjectSystemReady();
+	void OnObjectConstructed(UObject* Constructed);
+	void OnEnginePreExit();
+};
+FConstructionDuringUpgradeTracker ConstructionDuringUpgradeTracker;
+
+FConstructionDuringUpgradeTracker::FConstructionDuringUpgradeTracker()
+	:FDelayedAutoRegisterHelper(EDelayedRegisterRunPhase::ObjectSystemReady, [this] { OnObjectSystemReady(); })
+{
+}
+
+void FConstructionDuringUpgradeTracker::OnObjectSystemReady()
+{
+	// UtilsInitialize is called during OnPreEngineInit by PackageResourceManager, even when not running EditorDomain,
+	// so we don't need to call it here.
+	FCoreUObjectDelegates::OnObjectConstructed.AddRaw(this, &FConstructionDuringUpgradeTracker::OnObjectConstructed);
+	FCoreDelegates::OnEnginePreExit.AddRaw(this, &FConstructionDuringUpgradeTracker::OnEnginePreExit);
 }
+
+void FConstructionDuringUpgradeTracker::OnObjectConstructed(UObject* Constructed)
+{
+	using namespace UE::EditorDomain;
+
+	check(Constructed);
+	const PackageAccessTracking_Private::FTrackedData* AccumulatedScopeData =
+		PackageAccessTracking_Private::FPackageAccessRefScope::GetCurrentThreadAccumulatedData();
+	if (!AccumulatedScopeData)
+	{
+		return;
+	}
+	if (AccumulatedScopeData->OpName != PackageAccessTrackingOps::NAME_PostLoad)
+	{
+		return;
+	}
+	const UObject* Referencer = AccumulatedScopeData->Object;
+	if (!Referencer)
+	{
+		return;
+	}
+	const UClass* ReferencerClass = AccumulatedScopeData->Object->GetClass();
+	const UClass* ConstructedClass = Constructed->GetClass();
+	if (!ReferencerClass || !ConstructedClass || ReferencerClass == ConstructedClass)
+	{
+		return;
+	}
+	UPackage* Package = Constructed->GetPackage();
+	if (Referencer->GetPackage() != Package)
+	{
+		return;
+	}
+
+	FTopLevelAssetPath ReferencerClassPath(ReferencerClass);
+	FTopLevelAssetPath ConstructedClassPath(ConstructedClass);
+	if (!ReferencerClassPath.IsValid() || !ConstructedClassPath.IsValid())
+	{
+		return;
+	}
+
+	const FClassDigestData* ReferencerDigest = nullptr;
+	const FClassDigestData* ConstructedDigest = nullptr;
+	FClassDigestMap& ClassDigests = UE::EditorDomain::GetClassDigests();
+	{
+		FReadScopeLock ClassDigestsScopeLock(ClassDigests.Lock);
+		ReferencerDigest = ClassDigests.Map.Find(ReferencerClassPath);
+		ConstructedDigest = ClassDigests.Map.Find(ConstructedClassPath);
+	}
+	if (!ReferencerDigest || !ConstructedDigest)
+	{
+		PrecacheClassDigests({ ReferencerClassPath, ConstructedClassPath });
+		FReadScopeLock ClassDigestsScopeLock(ClassDigests.Lock);
+		ReferencerDigest = ClassDigests.Map.Find(ReferencerClassPath);
+		ConstructedDigest = ClassDigests.Map.Find(ConstructedClassPath);
+	}
+	if (!ReferencerDigest || !ConstructedDigest)
+	{
+		return;
+	}
+
+	bool bDeclared = INDEX_NONE != Algo::BinarySearch(ReferencerDigest->ConstructClasses, ConstructedClassPath,
+		[](const FTopLevelAssetPath& A, const FTopLevelAssetPath& B) { return A.CompareFast(B) < 0; });
+	if (bDeclared)
+	{
+		return;
+	}
+	if (ConstructedDigest->ClosestNative != ConstructedClassPath)
+	{
+		bDeclared = INDEX_NONE != Algo::BinarySearch(ReferencerDigest->ConstructClasses, ConstructedDigest->ClosestNative,
+			[](const FTopLevelAssetPath& A, const FTopLevelAssetPath& B) { return A.CompareFast(B) < 0; });
+		if (bDeclared)
+		{
+			return;
+		}
+	}
+
+	InitializeGlobalConstructClasses();
+	if (GGlobalConstructClasses.Contains(ConstructedClassPath))
+	{
+		return;
+	}
+
+	TOptional<FAssetPackageData> PackageData = IAssetRegistry::Get()->GetAssetPackageDataCopy(Package->GetFName());
+	FName ConstructedClassPathName(WriteToString<256>(ConstructedClassPath).ToView());
+	if (!PackageData || PackageData->ImportedClasses.Contains(ConstructedClassPathName))
+	{
+		return;
+	}
+
+	UE_LOG(LogCore, Verbose, TEXT("ConstructedDuringPostLoad: %s constructed %s.")
+		TEXT("\n\tThe ability to construct other objects during PostLoad needs to be declared to support incremental cooks.")
+		TEXT("\n\tCreate the static override function %s%s::DeclareConstructClasses and add %s to OutConstructClasses."),
+		*Referencer->GetFullName(), *Constructed->GetFullName(), ReferencerClass->GetPrefixCPP(), *ReferencerClass->GetName(),
+		*ConstructedClassPath.ToString());
+}
+
+void FConstructionDuringUpgradeTracker::OnEnginePreExit()
+{
+	FCoreUObjectDelegates::OnObjectConstructed.RemoveAll(this);
+	FCoreDelegates::OnEnginePreExit.RemoveAll(this);
+}
+#endif // CONSTRUCTION_DURING_UPGRADE_TRACKER_ENABLED
+
diff --git a/Engine/Source/Editor/UnrealEd/Private/EditorDomain/EditorDomainUtils.h b/Engine/Source/Editor/UnrealEd/Private/EditorDomain/EditorDomainUtils.h
--- a/Engine/Source/Editor/UnrealEd/Private/EditorDomain/EditorDomainUtils.h
+++ b/Engine/Source/Editor/UnrealEd/Private/EditorDomain/EditorDomainUtils.h
@@ -17,6 +17,7 @@
 #include "Serialization/CustomVersion.h"
 #include "Templates/Function.h"
 #include "UObject/NameTypes.h"
+#include "UObject/TopLevelAssetPath.h"
 
 DECLARE_LOG_CATEGORY_EXTERN(LogEditorDomainSave, Log, All);
 
@@ -35,13 +36,14 @@ namespace UE::EditorDomain
 struct FClassDigestData
 {
 public:
-	FBlake3Hash SchemaHash;
-	/** CustomVersions used by the class */
+	/** Inclusive schemahash for the class and all of its transitively constructible classes */
+	FBlake3Hash InclusiveSchemaHash;
+	/** CustomVersions used by the class or its constructed classes*/
 	TArray<int32> CustomVersionHandles;
-	/** Classes that can be created by the class during PostLoad/PreSave (parsed from ini) */
-	TArray<FName> ConstructClasses;
+	/** Classes that can be created by the class during PostLoad/PreSave (reported by DeclareConstructClasses) */
+	TArray<FTopLevelAssetPath> ConstructClasses;
 	/** The closest native parent of (the possibly CoreRedirected target of) the class */
-	FName ResolvedClosestNative;
+	FTopLevelAssetPath ClosestNative;
 	/** EditorDomainEnabled allows everything and uses only a blocklist, so DomainUse by default is enabled. */
 	EDomainUse EditorDomainUse = EDomainUse::LoadEnabled | EDomainUse::SaveEnabled;
 
@@ -56,7 +58,7 @@ public:
 /** Threadsafe cache of ClassName -> Digest data for calculating EditorDomain Digests */
 struct FClassDigestMap
 {
-	TMap<FName, FClassDigestData> Map;
+	TMap<FTopLevelAssetPath, FClassDigestData> Map;
 	FRWLock Lock;
 };
 
@@ -84,7 +86,7 @@ ESaveStorageResult SaveStorageResultFromString(FUtf8StringView Text);
 FPackageDigest CalculatePackageDigest(IAssetRegistry& AssetRegistry, FName PackageName);
 
 /** For any ClassNames not already in ClassDigests, look up their UStruct and add them. */
-void PrecacheClassDigests(TConstArrayView<FName> ClassNames);
+void PrecacheClassDigests(TConstArrayView<FTopLevelAssetPath> ClassNames);
 
 /** Get the CacheRequest for the given package from the EditorDomain cache bucket. */
 void RequestEditorDomainPackage(const FPackagePath& PackagePath,
diff --git a/Engine/Source/Editor/UnrealEd/Private/TargetDomain/TargetDomainUtils.cpp b/Engine/Source/Editor/UnrealEd/Private/TargetDomain/TargetDomainUtils.cpp
--- a/Engine/Source/Editor/UnrealEd/Private/TargetDomain/TargetDomainUtils.cpp
+++ b/Engine/Source/Editor/UnrealEd/Private/TargetDomain/TargetDomainUtils.cpp
@@ -391,7 +391,12 @@ bool IsIterativeEnabled(FName PackageName)
 	FReadScopeLock ClassDigestsScopeLock(ClassDigests.Lock);
 	for (FName ClassName : PackageData.ImportedClasses)
 	{
-		UE::EditorDomain::FClassDigestData* ExistingData = ClassDigests.Map.Find(ClassName);
+		FTopLevelAssetPath ClassPath(WriteToString<256>(ClassName).ToView());
+		UE::EditorDomain::FClassDigestData* ExistingData = nullptr;
+		if (ClassPath.IsValid())
+		{
+			ExistingData = ClassDigests.Map.Find(ClassPath);
+		}
 		if (!ExistingData)
 		{
 			// All allowlisted classes are added to ClassDigests at startup, so if the class is not in ClassDigests,
diff --git a/Engine/Source/Runtime/CoreUObject/Private/Misc/PackageAccessTracking.cpp b/Engine/Source/Runtime/CoreUObject/Private/Misc/PackageAccessTracking.cpp
--- a/Engine/Source/Runtime/CoreUObject/Private/Misc/PackageAccessTracking.cpp
+++ b/Engine/Source/Runtime/CoreUObject/Private/Misc/PackageAccessTracking.cpp
@@ -21,10 +21,12 @@ bool IsBuildOpName(FName OpName)
 		;
 }
 
-FTrackedData::FTrackedData(FName InPackageName, FName InOpName, const ITargetPlatform* InTargetPlatform)
+FTrackedData::FTrackedData(FName InPackageName, FName InOpName, const ITargetPlatform* InTargetPlatform,
+	const UObject* InObject)
 	: PackageName(InPackageName)
 	, OpName(InOpName)
 	, BuildOpName(IsBuildOpName(InOpName) ? InOpName : NAME_None)
+	, Object(InObject)
 	, TargetPlatform(InTargetPlatform)
 {
 }
@@ -40,7 +42,7 @@ FTrackedData::FTrackedData(FTrackedData& DirectData, FTrackedData* OuterAccumula
 		// but we want to record that B depends on C, not that A depends on C.  The InnermostThreadScope will have a package name of B and an OpName of PostLoad
 		// and so we want the package name from the InnermostThreadScope while having searched upwards for a scope with an OpName == Load.
 
-		// Make sure this line matches the PackageName line in the FTrackedData accumulating constructor
+		// The accumulated PackageName is the most recent non-empty PackageName
 		PackageName = !DirectData.PackageName.IsNone() ? DirectData.PackageName : OuterAccumulatedData->PackageName;
 
 		// The accumulated OpName is the most recent non-empty OpName
@@ -57,6 +59,9 @@ FTrackedData::FTrackedData(FTrackedData& DirectData, FTrackedData* OuterAccumula
 			BuildOpName = !DirectData.BuildOpName.IsNone() ? DirectData.BuildOpName : OuterAccumulatedData->BuildOpName;
 		}
 
+		// The accumulated object is the most recent non-null Object
+		Object = DirectData.Object ? DirectData.Object : OuterAccumulatedData->Object;
+
 		// The accumulated TargetPlatform is the most recent non-empty TargetPlatform
 		TargetPlatform = DirectData.TargetPlatform ? DirectData.TargetPlatform : OuterAccumulatedData->TargetPlatform;
 	}
@@ -65,20 +70,22 @@ FTrackedData::FTrackedData(FTrackedData& DirectData, FTrackedData* OuterAccumula
 		PackageName = DirectData.PackageName;
 		OpName = DirectData.OpName;
 		BuildOpName = DirectData.BuildOpName;
+		Object = DirectData.Object;
 		TargetPlatform = DirectData.TargetPlatform;
 	}
 }
 
-FPackageAccessRefScope::FPackageAccessRefScope(FName InPackageName, FName InOpName, const ITargetPlatform* InTargetPlatform)
-	: DirectData(InPackageName, InOpName, InTargetPlatform)
+FPackageAccessRefScope::FPackageAccessRefScope(FName InPackageName, FName InOpName,
+	const ITargetPlatform* InTargetPlatform, const UObject* InObject)
+	: DirectData(InPackageName, InOpName, InTargetPlatform, InObject)
 	, AccumulatedData(DirectData, GetCurrentThreadAccumulatedData())
 	, Outer(CurrentThreadScope)
 {
 	CurrentThreadScope = this;
 }
 
-FPackageAccessRefScope::FPackageAccessRefScope(const UPackage* InPackage, FName InOpName)
-	: FPackageAccessRefScope(InPackage->GetFName(), InOpName, nullptr)
+FPackageAccessRefScope::FPackageAccessRefScope(const UObject* InObject, FName InOpName)
+	: FPackageAccessRefScope(InObject->GetPackage()->GetFName(), InOpName, nullptr, InObject)
 {
 }
 
diff --git a/Engine/Source/Runtime/CoreUObject/Private/UObject/CoreRedirects.cpp b/Engine/Source/Runtime/CoreUObject/Private/UObject/CoreRedirects.cpp
--- a/Engine/Source/Runtime/CoreUObject/Private/UObject/CoreRedirects.cpp
+++ b/Engine/Source/Runtime/CoreUObject/Private/UObject/CoreRedirects.cpp
@@ -23,6 +23,11 @@ DEFINE_LOG_CATEGORY(LogCoreRedirects);
 #	define UE_WITH_CORE_REDIRECTS 1
 #endif
 
+FCoreRedirectObjectName::FCoreRedirectObjectName(const FTopLevelAssetPath& TopLevelAssetPath)
+	: FCoreRedirectObjectName(TopLevelAssetPath.GetAssetName(), NAME_None, TopLevelAssetPath.GetPackageName())
+{
+}
+
 FCoreRedirectObjectName::FCoreRedirectObjectName(const FString& InString)
 {
 	if (!ExpandNames(InString, ObjectName, OuterName, PackageName))
diff --git a/Engine/Source/Runtime/CoreUObject/Private/UObject/Obj.cpp b/Engine/Source/Runtime/CoreUObject/Private/UObject/Obj.cpp
--- a/Engine/Source/Runtime/CoreUObject/Private/UObject/Obj.cpp
+++ b/Engine/Source/Runtime/CoreUObject/Private/UObject/Obj.cpp
@@ -1077,7 +1077,7 @@ void UObject::ConditionalPostLoad()
 	{
 
 		check(IsInGameThread() || HasAnyFlags(RF_ClassDefaultObject|RF_ArchetypeObject) || IsPostLoadThreadSafe() || IsA(UClass::StaticClass()))
-		UE_TRACK_REFERENCING_PACKAGE_SCOPED(GetOutermost(), PackageAccessTrackingOps::NAME_PostLoad);
+		UE_TRACK_REFERENCING_PACKAGE_SCOPED(this, PackageAccessTrackingOps::NAME_PostLoad);
 
 #if !(UE_BUILD_SHIPPING || UE_BUILD_TEST)
 		FUObjectThreadContext& ThreadContext = FUObjectThreadContext::Get();
@@ -1542,7 +1542,7 @@ void UObject::DeclareCustomVersions(FArchive& Ar, const UClass* SpecificSubclass
 		{
 			FProperty* InnerProperty = Property;
 			Property = Property->PropertyLinkNext;
-			FArrayProperty* ArrayProperty = CastField<FArrayProperty>(Property);
+			FArrayProperty* ArrayProperty = CastField<FArrayProperty>(InnerProperty);
 			if (ArrayProperty)
 			{
 				InnerProperty = ArrayProperty->Inner;
@@ -1590,6 +1590,10 @@ void UObject::DeclareCustomVersions(FArchive& Ar, const UClass* SpecificSubclass
 void UObject::AppendToClassSchema(FAppendToClassSchemaContext& Context)
 {
 }
+
+void UObject::DeclareConstructClasses(TArray<FTopLevelAssetPath>& OutConstructClasses, const UClass* SpecificSubclass) 
+{
+}
 #endif
 
 void UObject::SerializeScriptProperties(FArchive& Ar) const
diff --git a/Engine/Source/Runtime/CoreUObject/Private/UObject/UObjectGlobals.cpp b/Engine/Source/Runtime/CoreUObject/Private/UObject/UObjectGlobals.cpp
--- a/Engine/Source/Runtime/CoreUObject/Private/UObject/UObjectGlobals.cpp
+++ b/Engine/Source/Runtime/CoreUObject/Private/UObject/UObjectGlobals.cpp
@@ -164,6 +164,7 @@ FCoreUObjectDelegates::FOnObjectsReplaced FCoreUObjectDelegates::OnObjectsReplac
 FCoreUObjectDelegates::FOnObjectsReplaced FCoreUObjectDelegates::OnObjectsReinstanced;
 FCoreUObjectDelegates::FOnObjectPostCDOCompiled FCoreUObjectDelegates::OnObjectPostCDOCompiled;
 FCoreUObjectDelegates::FOnAssetLoaded FCoreUObjectDelegates::OnAssetLoaded;
+FCoreUObjectDelegates::FOnObjectConstructed FCoreUObjectDelegates::OnObjectConstructed;
 FCoreUObjectDelegates::FOnEndLoadPackage FCoreUObjectDelegates::OnEndLoadPackage;
 PRAGMA_DISABLE_DEPRECATION_WARNINGS
 FCoreUObjectDelegates::FOnObjectSaved FCoreUObjectDelegates::OnObjectSaved;
@@ -4233,6 +4234,10 @@ UObject* StaticConstructObject_Internal(const FStaticConstructObjectParameters&
 		SaveToTransactionBuffer(Result, false);
 		Result->ClearGarbage();
 	}
+
+#if WITH_EDITOR
+	FCoreUObjectDelegates::OnObjectConstructed.Broadcast(Result);
+#endif
 	return Result;
 }
 
diff --git a/Engine/Source/Runtime/CoreUObject/Public/Misc/PackageAccessTracking.h b/Engine/Source/Runtime/CoreUObject/Public/Misc/PackageAccessTracking.h
--- a/Engine/Source/Runtime/CoreUObject/Public/Misc/PackageAccessTracking.h
+++ b/Engine/Source/Runtime/CoreUObject/Public/Misc/PackageAccessTracking.h
@@ -15,7 +15,7 @@ class UPackage;
 #define UE_WITH_PACKAGE_ACCESS_TRACKING UE_WITH_OBJECT_HANDLE_TRACKING
 
 #if UE_WITH_PACKAGE_ACCESS_TRACKING
-#define UE_TRACK_REFERENCING_PACKAGE_SCOPED(Package, OpName) PackageAccessTracking_Private::FPackageAccessRefScope ANONYMOUS_VARIABLE(PackageAccessTracker_)(Package, OpName);
+#define UE_TRACK_REFERENCING_PACKAGE_SCOPED(Object, OpName) PackageAccessTracking_Private::FPackageAccessRefScope ANONYMOUS_VARIABLE(PackageAccessTracker_)(Object, OpName);
 #define UE_TRACK_REFERENCING_PACKAGE_DELAYED_SCOPED(TrackerName, OpName) TOptional<PackageAccessTracking_Private::FPackageAccessRefScope> TrackerName; FName TrackerName##_OpName(OpName);
 #define UE_TRACK_REFERENCING_PACKAGE_DELAYED(TrackerName, Package) if (TrackerName) TrackerName->SetPackageName(Package->GetFName()); else TrackerName.Emplace(Package->GetFName(), TrackerName##_OpName);
 #define UE_TRACK_REFERENCING_PLATFORM_SCOPED(TargetPlatform) PackageAccessTracking_Private::FPackageAccessRefScope ANONYMOUS_VARIABLE(PackageAccessTracker_)(TargetPlatform);
@@ -36,21 +36,24 @@ namespace PackageAccessTracking_Private
 	struct FTrackedData
 	{
 		/** Standard constructor; sets variables from DirectData passed to a FPackageAccessRefScope. */
-		COREUOBJECT_API FTrackedData(FName PackgeName, FName OpName, const ITargetPlatform* InTargetPlatform);
+		COREUOBJECT_API FTrackedData(FName PackageName, FName OpName, const ITargetPlatform* InTargetPlatform,
+			const UObject* InObject);
 		/** Accumulating constructor; sets new AccumulatedData by combining DirectData with Outer's AccumulatedData. */
 		COREUOBJECT_API FTrackedData(FTrackedData& DirectData, FTrackedData* OuterAccumulatedData);
 
 		FName PackageName;
 		FName OpName;
 		FName BuildOpName;
+		const UObject* Object = nullptr;
 		const ITargetPlatform* TargetPlatform = nullptr;
 	};
 
 	class FPackageAccessRefScope
 	{
 	public:
-		COREUOBJECT_API FPackageAccessRefScope(FName InPackageName, FName InOpName, const ITargetPlatform* InTargetPlatform = nullptr);
-		COREUOBJECT_API FPackageAccessRefScope(const UPackage* InPackage, FName InOpName);
+		COREUOBJECT_API FPackageAccessRefScope(FName InPackageName, FName InOpName,
+			const ITargetPlatform* InTargetPlatform=nullptr, const UObject* InObject=nullptr);
+		COREUOBJECT_API FPackageAccessRefScope(const UObject* InObject, FName InOpName);
 		COREUOBJECT_API FPackageAccessRefScope(FName InOpName);
 		COREUOBJECT_API FPackageAccessRefScope(const ITargetPlatform* InTargetPlatform);
 
diff --git a/Engine/Source/Runtime/CoreUObject/Public/UObject/Class.h b/Engine/Source/Runtime/CoreUObject/Public/UObject/Class.h
--- a/Engine/Source/Runtime/CoreUObject/Public/UObject/Class.h
+++ b/Engine/Source/Runtime/CoreUObject/Public/UObject/Class.h
@@ -2636,6 +2636,13 @@ public:
 		check(CppClassStaticFunctions.GetAppendToClassSchema());
 		CppClassStaticFunctions.GetAppendToClassSchema()(Context);
 	}
+
+	/** Calls the c++ class's DeclareConstructClasses static function (from the nearest native parent if this is not native) */
+	void CallDeclareConstructClasses(TArray<FTopLevelAssetPath>& OutConstructClasses) const
+	{
+		check(CppClassStaticFunctions.GetDeclareConstructClasses());
+		CppClassStaticFunctions.GetDeclareConstructClasses()(OutConstructClasses, this);
+	}
 #endif
 
 	/** The class default object; used for delta serialization and object initialization */
@@ -3004,7 +3011,7 @@ public:
 	{
 		if (ClassDefaultObject == nullptr && bCreateIfNeeded)
 		{
-			UE_TRACK_REFERENCING_PACKAGE_SCOPED((GetOutermost()), PackageAccessTrackingOps::NAME_CreateDefaultObject);
+			UE_TRACK_REFERENCING_PACKAGE_SCOPED(this, PackageAccessTrackingOps::NAME_CreateDefaultObject);
 			const_cast<UClass*>(this)->CreateDefaultObject();
 		}
 
diff --git a/Engine/Source/Runtime/CoreUObject/Public/UObject/CoreRedirects.h b/Engine/Source/Runtime/CoreUObject/Public/UObject/CoreRedirects.h
--- a/Engine/Source/Runtime/CoreUObject/Public/UObject/CoreRedirects.h
+++ b/Engine/Source/Runtime/CoreUObject/Public/UObject/CoreRedirects.h
@@ -20,6 +20,7 @@
 
 class IPakFile;
 class UClass;
+struct FTopLevelAssetPath;
 
 #define WITH_COREREDIRECTS_MULTITHREAD_WARNING !UE_BUILD_SHIPPING && !IS_PROGRAM && !WITH_EDITOR
 
@@ -98,6 +99,8 @@ struct COREUOBJECT_API FCoreRedirectObjectName
 
 	}
 
+	FCoreRedirectObjectName(const FTopLevelAssetPath& TopLevelAssetPath);
+
 	/** Construct from a path string, this handles full paths with packages, or partial paths without */
 	FCoreRedirectObjectName(const FString& InString);
 
diff --git a/Engine/Source/Runtime/CoreUObject/Public/UObject/Object.h b/Engine/Source/Runtime/CoreUObject/Public/UObject/Object.h
--- a/Engine/Source/Runtime/CoreUObject/Public/UObject/Object.h
+++ b/Engine/Source/Runtime/CoreUObject/Public/UObject/Object.h
@@ -30,6 +30,7 @@ struct FAppendToClassSchemaContext;
 struct FFrame;
 struct FObjectInstancingGraph;
 struct FPropertyChangedChainEvent;
+struct FTopLevelAssetPath;
 class UClass;
 #if UE_WITH_IRIS
 namespace UE::Net
@@ -391,6 +392,15 @@ public:
 	 * invalidate previous results because serialization changed and no custom version was udpated.
 	 */
 	static void AppendToClassSchema(FAppendToClassSchemaContext& Context);
+	/**
+	 * Declare classes that can be constructed by this class during loading. This declaration is implicitly transitive;
+	 * the caller is responsible for following the graph of ConstructClasses to find all transitive ConstructClasses.
+	 * 
+	 * @param OutConstructClasses Output list of classes that this class can construct.
+	 * @param SpecificSubclass The class on which DeclaredConstructClasses was called. This can differ from the current
+	 *        class because each class calls Super::DeclareConstructClasses.
+	 */
+	static void DeclareConstructClasses(TArray<FTopLevelAssetPath>& OutConstructClasses, const UClass* SpecificSubclass);
 #endif
 
 	/** After a critical error, perform any mission-critical cleanup, such as restoring the video mode orreleasing hardware resources. */
diff --git a/Engine/Source/Runtime/CoreUObject/Public/UObject/ObjectMacros.h b/Engine/Source/Runtime/CoreUObject/Public/UObject/ObjectMacros.h
--- a/Engine/Source/Runtime/CoreUObject/Public/UObject/ObjectMacros.h
+++ b/Engine/Source/Runtime/CoreUObject/Public/UObject/ObjectMacros.h
@@ -15,6 +15,7 @@ class FReferenceCollector;
 struct FAppendToClassSchemaContext;
 struct FFrame;
 struct FClassReloadVersionInfo;
+struct FTopLevelAssetPath;
 
 /** Represents a serializable object pointer in blueprint bytecode. This is always 64-bits, even on 32-bit platforms. */
 typedef	uint64 ScriptPointerType;
@@ -1866,14 +1867,15 @@ public: \
 	DECLARE_WITHIN_INTERNAL( UPackage, true )
 
 #define UOBJECT_CPPCLASS_STATICFUNCTIONS_ALLCONFIGS(TClass) \
-	&TClass::AddReferencedObjects
+	FUObjectCppClassStaticFunctions::AddReferencedObjectsType(&TClass::AddReferencedObjects)
 	/* UObjectCppClassStaticFunctions: Extend this macro with the address of your new static function, if it applies to all configs. */ \
 	/* Order must match the order in the FUObjectCppClassStaticFunctions constructor. */ \
 
 #if WITH_EDITORONLY_DATA
 	#define UOBJECT_CPPCLASS_STATICFUNCTIONS_WITHEDITORONLYDATA(TClass) \
-		, &TClass::DeclareCustomVersions \
-		, &TClass::AppendToClassSchema
+		, FUObjectCppClassStaticFunctions::DeclareCustomVersionsType(&TClass::DeclareCustomVersions) \
+		, FUObjectCppClassStaticFunctions::AppendToClassSchemaType(&TClass::AppendToClassSchema) \
+		, FUObjectCppClassStaticFunctions::DeclareConstructClassesType(&TClass::DeclareConstructClasses)
 		/* UObjectCppClassStaticFunctions: Extend this macro with the address of your new static function, if it is editor-only. */ \
 		/* Order must match the order in the FUObjectCppClassStaticFunctions constructor. */ \
 #else
@@ -1904,6 +1906,7 @@ public:
 #if WITH_EDITORONLY_DATA
 	typedef void (*DeclareCustomVersionsType)   (FArchive& Ar, const UClass* SpecificSubclass);
 	typedef void (*AppendToClassSchemaType)   (FAppendToClassSchemaContext& Context);
+	typedef void (*DeclareConstructClassesType)   (TArray<FTopLevelAssetPath>& OutConstructClasses, const UClass* SpecificSubclass);
 #endif
 	// UObjectCppClassStaticFunctions: Extend this list of types with the type of your new static function.
 
@@ -1911,12 +1914,14 @@ public:
 #if WITH_EDITORONLY_DATA
 		, DeclareCustomVersionsType InDeclareCustomVersions
 		, AppendToClassSchemaType InAppendToClassSchema
+		, DeclareConstructClassesType InDeclareConstructClasses
 #endif
 	)
 		: AddReferencedObjects(InAddReferencedObjects)
 #if WITH_EDITORONLY_DATA
 		, DeclareCustomVersions(InDeclareCustomVersions)
 		, AppendToClassSchema(InAppendToClassSchema)
+		, DeclareConstructClasses(InDeclareConstructClasses)
 #endif
 	{
 		// Null elements are not valid in this constructor
@@ -1924,6 +1929,7 @@ public:
 #if WITH_EDITORONLY_DATA
 		check(InDeclareCustomVersions);
 		check(InAppendToClassSchema);
+		check(InDeclareConstructClasses);
 #endif
 		// UObjectCppClassStaticFunctions: Extend the constructor with initializers for your new static function member.
 		// Order must match the order in UOBJECT_CPPCLASS_STATICFUNCTIONS_FORCLASS.
@@ -1970,6 +1976,15 @@ public:
 		check(InAppendToClassSchema != nullptr); // It is not valid to clear single elements (see IsInitialized). Call Reset to clear all elements.
 		AppendToClassSchema = InAppendToClassSchema;
 	}
+	DeclareConstructClassesType GetDeclareConstructClasses() const
+	{
+		return DeclareConstructClasses;
+	}
+	void SetDeclareConstructClasses(DeclareConstructClassesType InDeclareConstructClasses)
+	{
+		check(InDeclareConstructClasses != nullptr); // It is not valid to clear single elements (see IsInitialized). Call Reset to clear all elements.
+		DeclareConstructClasses = InDeclareConstructClasses;
+	}
 #endif
 	// UObjectCppClassStaticFunctions: Extend the list of accessors for your new static function.
 
@@ -1978,6 +1993,7 @@ private:
 #if WITH_EDITORONLY_DATA
 	DeclareCustomVersionsType DeclareCustomVersions = nullptr;
 	AppendToClassSchemaType AppendToClassSchema = nullptr;
+	DeclareConstructClassesType DeclareConstructClasses = nullptr;
 #endif
 	// UObjectCppClassStaticFunctions: Extend this list of members with the member for your new static function.
 
diff --git a/Engine/Source/Runtime/CoreUObject/Public/UObject/UObjectGlobals.h b/Engine/Source/Runtime/CoreUObject/Public/UObject/UObjectGlobals.h
--- a/Engine/Source/Runtime/CoreUObject/Public/UObject/UObjectGlobals.h
+++ b/Engine/Source/Runtime/CoreUObject/Public/UObject/UObjectGlobals.h
@@ -2718,6 +2718,9 @@ struct COREUOBJECT_API FCoreUObjectDelegates
 	DECLARE_MULTICAST_DELEGATE_OneParam(FOnAssetLoaded, UObject*);
 	static FOnAssetLoaded OnAssetLoaded;
 
+	DECLARE_MULTICAST_DELEGATE_OneParam(FOnObjectConstructed, UObject*);
+	static FOnObjectConstructed OnObjectConstructed;
+
 	/** Callback when packages end loading in LoadPackage or AsyncLoadPackage. All packages loaded recursively due to imports are included in the single call of the explicitly-loaded package. */
 	DECLARE_MULTICAST_DELEGATE_OneParam(FOnEndLoadPackage, const FEndLoadPackageContext&);
 	static FOnEndLoadPackage OnEndLoadPackage;
