Fix crash when transactions reenter themselves - this happens due to FActorTransactionAnnotation/FComponentInstanceDataCache

diff --git a/Engine/Source/Editor/UnrealEd/Private/EditorTransaction.cpp b/Engine/Source/Editor/UnrealEd/Private/EditorTransaction.cpp
--- a/Engine/Source/Editor/UnrealEd/Private/EditorTransaction.cpp
+++ b/Engine/Source/Editor/UnrealEd/Private/EditorTransaction.cpp
@@ -620,14 +620,18 @@ void FTransaction::SaveObject( UObject* Object )
 	check(Object);
 	Object->CheckDefaultSubobjects();
 
-	FObjectRecords& ObjectRecords = ObjectRecordsMap.FindOrAdd(UE::Transaction::FPersistentObjectRef(Object));
-	if (ObjectRecords.Records.Num() == 0)
+	FObjectRecords* ObjectRecords = &ObjectRecordsMap.FindOrAdd(UE::Transaction::FPersistentObjectRef(Object));
+	if (ObjectRecords->Records.Num() == 0)
 	{
 		// Save the object.
-		FObjectRecord* UndoRecord = ObjectRecords.Records.Add_GetRef(new FObjectRecord(this, Object, nullptr, nullptr, 0, 0, 0, 0, 0, nullptr, nullptr, nullptr));
-		Records.Add(UndoRecord);
+		FObjectRecord* Record = new FObjectRecord(this, Object, nullptr, nullptr, 0, 0, 0, 0, 0, nullptr, nullptr, nullptr);
+
+		// Side effects of FObjectRecord() may have grown ObjectRecordsMap
+		ObjectRecords = &ObjectRecordsMap.FindChecked(UE::Transaction::FPersistentObjectRef(Object));
+		ObjectRecords->Records.Add(Record);
+		Records.Add(Record);
 	}
-	++ObjectRecords.SaveCount;
+	++ObjectRecords->SaveCount;
 }
 
 void FTransaction::SaveArray( UObject* Object, FScriptArray* Array, int32 Index, int32 Count, int32 Oper, int32 ElementSize, uint32 ElementAlignment, STRUCT_DC DefaultConstructor, STRUCT_AR Serializer, STRUCT_DTOR Destructor )
