Fix the condition to make sure the texture import asset type do not override a different asset type.

#rnx

diff --git a/Engine/Source/Editor/UnrealEd/Private/Factories/EditorFactories.cpp b/Engine/Source/Editor/UnrealEd/Private/Factories/EditorFactories.cpp
--- a/Engine/Source/Editor/UnrealEd/Private/Factories/EditorFactories.cpp
+++ b/Engine/Source/Editor/UnrealEd/Private/Factories/EditorFactories.cpp
@@ -3580,11 +3580,11 @@ UTexture* UTextureFactory::ImportTexture(UClass* Class, UObject* InParent, FName
 		if (HDRImportShouldBeLongLatCubeMap == EAppReturnType::Yes ||
 			HDRImportShouldBeLongLatCubeMap == EAppReturnType::YesAll)
 		{
-			if (Class && !Class->IsChildOf<UTextureCube>())
+			if (Class && !UTextureCube::StaticClass()->IsChildOf(Class))
 			{
-				//The source use to re-import this texture is not compatible with the
+				//The source use to import this texture is not compatible with the
 				//existing texture type.
-				Warn->Logf(ELogVerbosity::Error, TEXT("Cannot re-import a [%s] with a UTextureCube source file"), *Class->GetName());
+				Warn->Logf(ELogVerbosity::Error, TEXT("Cannot import/reimport a [%s] with a UTextureCube source file"), *Class->GetName());
 				return nullptr;
 			}
 
@@ -3634,11 +3634,11 @@ UTexture* UTextureFactory::ImportTexture(UClass* Class, UObject* InParent, FName
 
 		if(IESConverter.IsValid())
 		{
-			if (Class && !Class->IsChildOf<UTextureLightProfile>())
+			if (Class && !UTextureLightProfile::StaticClass()->IsChildOf(Class))
 			{
-				//The source use to re-import this texture is not compatible with the
+				//The source use to import this texture is not compatible with the
 				//existing texture type.
-				Warn->Logf(ELogVerbosity::Error, TEXT("Cannot re-import a [%s] with a UTextureLightProfile source file"), *Class->GetName());
+				Warn->Logf(ELogVerbosity::Error, TEXT("Cannot import/reimport a [%s] with a UTextureLightProfile source file"), *Class->GetName());
 				return nullptr;
 			}
 
@@ -3678,11 +3678,11 @@ UTexture* UTextureFactory::ImportTexture(UClass* Class, UObject* InParent, FName
 	FImportImage Image;
 	if (ImportImage(Buffer, Length, Warn, ImportFlags, Image))
 	{
-		if (Class && !Class->IsChildOf<UTexture2D>())
+		if (Class && !UTexture2D::StaticClass()->IsChildOf(Class))
 		{
-			//The source use to re-import this texture is not compatible with the
+			//The source use to import this texture is not compatible with the
 			//existing texture type.
-			Warn->Logf(ELogVerbosity::Error, TEXT("Cannot re-import a [%s] with a UTexture2D source file"), *Class->GetName());
+			Warn->Logf(ELogVerbosity::Error, TEXT("Cannot import/reimport a [%s] with a UTexture2D source file"), *Class->GetName());
 			return nullptr;
 		}
 
