Minor revisions to reference replacement after BP reinstancing to allow external CDO references to be replaced during PIE/simulation.

Change summary:
- Added FReplaceInstancesOfClassParameters::bReplaceReferencesToOldCDOs.
- Deprecated FBatchReplaceInstancesOfClassParameters, consolidated flags into FReplaceInstancesOfClassParameters.
- Deprecated FReplaceInstancesOfClassParameters::OldClass and NewClass and the ctor that took them as inputs - these are now consolidated into a single ReplaceInstancesOfClass() API that takes them as inputs along w/ the parameter struct.
- Consolidated/deprecated FBlueprintCompileReinstancer::ReplaceInstancesOfClass()/BatchReplaceInstancesOfClass() APIs; updated existing call sites to match.
- Modified ReplaceInstancesOfClass_Inner() to take params directly as input; also revised the block that handled CDO reference replacement to utilize the new flag instead.
- Added a few new checks/ensures to validate the "legacy" path when certain parameters are in use.
- Modified the FBlueprintCompilationManagerImpl::FlushReinstancingQueueImpl() signature to include a new option flag.
- Added EBlueprintCompileOptions::IncludeCDOInReferenceReplacement (new compiler flag).
- Modified FBlueprintCompilationManagerImpl::CompileSynchronouslyImpl() to test for the new compiler flag and pass through to FlushReinstancingQueueImpl().
- Modified FBlueprintEditor::Compile() to set the new compiler flag when IsPlayInEditorActive() is true.

Notes:
- FReplaceInstancesOfClassParameters::bClassObjectReplaced is a legacy flag that was not currently in use outside of the old hot reload path; it remains in place for backwards-compatibility but I've added an ensure() to guard against it being used incorrectly outside of that context.

diff --git a/Engine/Source/Editor/Kismet/Private/BlueprintCompilationManager.cpp b/Engine/Source/Editor/Kismet/Private/BlueprintCompilationManager.cpp
--- a/Engine/Source/Editor/Kismet/Private/BlueprintCompilationManager.cpp
+++ b/Engine/Source/Editor/Kismet/Private/BlueprintCompilationManager.cpp
@@ -124,7 +124,7 @@ struct FBlueprintCompilationManagerImpl : public FGCObject
 	void CompileSynchronouslyImpl(const FBPCompileRequestInternal& Request);
 	void FlushCompilationQueueImpl(bool bSuppressBroadcastCompiled, TArray<UBlueprint*>* BlueprintsCompiled, TArray<UBlueprint*>* BlueprintsCompiledOrSkeletonCompiled, FUObjectSerializeContext* InLoadContext);
 	void ProcessExtensions(const TArray<FCompilerData>& InCurrentlyCompilingBPs);
-	void FlushReinstancingQueueImpl();
+	void FlushReinstancingQueueImpl(bool bFindAndReplaceCDOReferences = false);
 	bool HasBlueprintsToCompile() const;
 	bool IsGeneratedClassLayoutReady() const;
 	void GetDefaultValue(const UClass* ForClass, const FProperty* Property, FString& OutDefaultValueAsString) const;
@@ -232,13 +232,14 @@ void FBlueprintCompilationManagerImpl::CompileSynchronouslyImpl(const FBPCompile
 
 	Request.UserData.BPToCompile->bQueuedForCompilation = true;
 
-	const bool bIsRegeneratingOnLoad		= (Request.UserData.CompileOptions & EBlueprintCompileOptions::IsRegeneratingOnLoad		) != EBlueprintCompileOptions::None;
-	const bool bRegenerateSkeletonOnly		= (Request.UserData.CompileOptions & EBlueprintCompileOptions::RegenerateSkeletonOnly	) != EBlueprintCompileOptions::None;
-	const bool bSkipGarbageCollection		= (Request.UserData.CompileOptions & EBlueprintCompileOptions::SkipGarbageCollection	) != EBlueprintCompileOptions::None
+	const bool bIsRegeneratingOnLoad		= (Request.UserData.CompileOptions & EBlueprintCompileOptions::IsRegeneratingOnLoad				) != EBlueprintCompileOptions::None;
+	const bool bRegenerateSkeletonOnly		= (Request.UserData.CompileOptions & EBlueprintCompileOptions::RegenerateSkeletonOnly			) != EBlueprintCompileOptions::None;
+	const bool bSkipGarbageCollection		= (Request.UserData.CompileOptions & EBlueprintCompileOptions::SkipGarbageCollection			) != EBlueprintCompileOptions::None
 		|| bRegenerateSkeletonOnly;
-	const bool bBatchCompile				= (Request.UserData.CompileOptions & EBlueprintCompileOptions::BatchCompile				) != EBlueprintCompileOptions::None;
-	const bool bSkipReinstancing			= (Request.UserData.CompileOptions & EBlueprintCompileOptions::SkipReinstancing			) != EBlueprintCompileOptions::None;
-	const bool bSkipSaving					= (Request.UserData.CompileOptions & EBlueprintCompileOptions::SkipSave					) != EBlueprintCompileOptions::None;
+	const bool bBatchCompile				= (Request.UserData.CompileOptions & EBlueprintCompileOptions::BatchCompile						) != EBlueprintCompileOptions::None;
+	const bool bSkipReinstancing			= (Request.UserData.CompileOptions & EBlueprintCompileOptions::SkipReinstancing					) != EBlueprintCompileOptions::None;
+	const bool bSkipSaving					= (Request.UserData.CompileOptions & EBlueprintCompileOptions::SkipSave							) != EBlueprintCompileOptions::None;
+	const bool bFindAndReplaceCDOReferences	= (Request.UserData.CompileOptions & EBlueprintCompileOptions::IncludeCDOInReferenceReplacement	) != EBlueprintCompileOptions::None;
 
 	ensure(!bIsRegeneratingOnLoad); // unexpected code path, compile on load handled with different function call
 	ensure(!bSkipReinstancing); // This is an internal option, should not go through CompileSynchronouslyImpl
@@ -267,7 +268,7 @@ void FBlueprintCompilationManagerImpl::CompileSynchronouslyImpl(const FBPCompile
 	TArray<UBlueprint*> CompiledBlueprints;
 	TArray<UBlueprint*> SkeletonCompiledBlueprints;
 	FlushCompilationQueueImpl(bSuppressBroadcastCompiled, &CompiledBlueprints, &SkeletonCompiledBlueprints, nullptr);
-	FlushReinstancingQueueImpl();
+	FlushReinstancingQueueImpl(bFindAndReplaceCDOReferences);
 	
 	if( Request.UserData.ClientResultsLog && Request.UserData.ClientResultsLog->bLogDetailedResults)
 	{
@@ -1787,7 +1788,7 @@ void FBlueprintCompilationManagerImpl::ProcessExtensions(const TArray<FCompilerD
 	}
 }
 
-void FBlueprintCompilationManagerImpl::FlushReinstancingQueueImpl()
+void FBlueprintCompilationManagerImpl::FlushReinstancingQueueImpl(bool bFindAndReplaceCDOReferences)
 {
 	DECLARE_SCOPE_HIERARCHICAL_COUNTER_FUNC()
 
@@ -1808,8 +1809,9 @@ void FBlueprintCompilationManagerImpl::FlushReinstancingQueueImpl()
 		
 		TGuardValue<bool> ReinstancingGuard(GIsReinstancing, true);
 		
-		FBatchReplaceInstancesOfClassParameters Options;
+		FReplaceInstancesOfClassParameters Options;
 		Options.bArchetypesAreUpToDate = true;
+		Options.bReplaceReferencesToOldCDOs = bFindAndReplaceCDOReferences;
 		FBlueprintCompileReinstancer::BatchReplaceInstancesOfClass(ClassesToReinstance, Options);
 
 		if (IsAsyncLoading())
@@ -2017,7 +2019,7 @@ void FBlueprintCompilationManagerImpl::ReparentHierarchies(const TMap<UClass*, U
 		OldClassToNewClassDerivedTypes.Add(ReinstancingJob.OldToNew);
 	}
 	TGuardValue<bool> ReinstancingGuard(GIsReinstancing, true);
-	FBatchReplaceInstancesOfClassParameters BatchOptions;
+	FReplaceInstancesOfClassParameters BatchOptions;
 	BatchOptions.bArchetypesAreUpToDate = true;
 	BatchOptions.bReplaceReferencesToOldClasses = bReplaceReferencesToOldClasses;
 
@@ -2041,7 +2043,7 @@ void FBlueprintCompilationManagerImpl::ReparentHierarchies(const TMap<UClass*, U
 	BatchOptions.ObjectsThatShouldUseOldStuff = &OldObjects;
 	BatchOptions.InstancesThatShouldUseOldClass = &OldObjects;
 
-	FBlueprintCompileReinstancer::BatchReplaceInstancesOfClass( OldClassToNewClassDerivedTypes, BatchOptions);
+	FBlueprintCompileReinstancer::BatchReplaceInstancesOfClass(OldClassToNewClassDerivedTypes, BatchOptions );
 }
 
 
diff --git a/Engine/Source/Editor/Kismet/Private/BlueprintEditor.cpp b/Engine/Source/Editor/Kismet/Private/BlueprintEditor.cpp
--- a/Engine/Source/Editor/Kismet/Private/BlueprintEditor.cpp
+++ b/Engine/Source/Editor/Kismet/Private/BlueprintEditor.cpp
@@ -4055,6 +4055,13 @@ void FBlueprintEditor::Compile()
 			CompileOptions |= (EBlueprintCompileOptions::UseDeltaSerializationDuringReinstancing | EBlueprintCompileOptions::SkipNewVariableDefaultsDetection);
 		}
 
+		// If compilation is enabled during PIE/simulation, references to the CDO might be held by a script variable.
+		// Thus, we set the flag to direct the compiler to allow those references to be replaced during reinstancing.
+		if (IsPlayInEditorActive())
+		{
+			CompileOptions |= EBlueprintCompileOptions::IncludeCDOInReferenceReplacement;
+		}
+
 		FKismetEditorUtilities::CompileBlueprint(BlueprintObj, CompileOptions, &LogResults);
 
 		LogResults.EndEvent();
diff --git a/Engine/Source/Editor/UnrealEd/Private/EditorServer.cpp b/Engine/Source/Editor/UnrealEd/Private/EditorServer.cpp
--- a/Engine/Source/Editor/UnrealEd/Private/EditorServer.cpp
+++ b/Engine/Source/Editor/UnrealEd/Private/EditorServer.cpp
@@ -1422,7 +1422,7 @@ void UEditorEngine::PostUndo(bool)
 	}
 
 	// Re-instance any actors that need it
-	FBlueprintCompileReinstancer::BatchReplaceInstancesOfClass(OldToNewClassMapToReinstance);
+	FBlueprintCompileReinstancer::BatchReplaceInstancesOfClass(OldToNewClassMapToReinstance, FReplaceInstancesOfClassParameters());
 
 	RedrawLevelEditingViewports();
 }
diff --git a/Engine/Source/Editor/UnrealEd/Private/Kismet2/Kismet2.cpp b/Engine/Source/Editor/UnrealEd/Private/Kismet2/Kismet2.cpp
--- a/Engine/Source/Editor/UnrealEd/Private/Kismet2/Kismet2.cpp
+++ b/Engine/Source/Editor/UnrealEd/Private/Kismet2/Kismet2.cpp
@@ -385,7 +385,7 @@ void FBlueprintUnloader::ReplaceStaleRefs(UBlueprint* NewBlueprint)
 		Referencers.Add(Referencer.Value);
 	}
 
-	FBlueprintCompileReinstancer::ReplaceInstancesOfClass(OldGeneratedClass, NewBlueprint->GeneratedClass);
+	FBlueprintCompileReinstancer::ReplaceInstancesOfClass(OldGeneratedClass, NewBlueprint->GeneratedClass, FReplaceInstancesOfClassParameters());
 
 	for (UObject* Referencer : Referencers)
 	{
diff --git a/Engine/Source/Editor/UnrealEd/Private/Kismet2/KismetReinstanceUtilities.cpp b/Engine/Source/Editor/UnrealEd/Private/Kismet2/KismetReinstanceUtilities.cpp
--- a/Engine/Source/Editor/UnrealEd/Private/Kismet2/KismetReinstanceUtilities.cpp
+++ b/Engine/Source/Editor/UnrealEd/Private/Kismet2/KismetReinstanceUtilities.cpp
@@ -88,7 +88,7 @@ struct FReplaceReferenceHelper
 		}
 	}
 
-	static void FindAndReplaceReferences(const TArray<UObject*>& SourceObjects, TSet<UObject*>* ObjectsThatShouldUseOldStuff, const TArray<UObject*>& ObjectsToReplace, const TMap<UObject*, UObject*>& OldToNewInstanceMap, const TMap<FSoftObjectPath, UObject*>& ReinstancedObjectsWeakReferenceMap)
+	static void FindAndReplaceReferences(const TArray<UObject*>& SourceObjects, const TSet<UObject*>* ObjectsThatShouldUseOldStuff, const TArray<UObject*>& ObjectsToReplace, const TMap<UObject*, UObject*>& OldToNewInstanceMap, const TMap<FSoftObjectPath, UObject*>& ReinstancedObjectsWeakReferenceMap)
 	{
 		if(SourceObjects.Num() == 0 && ObjectsToReplace.Num() == 0 )
 		{
@@ -241,17 +241,6 @@ struct FArchetypeReinstanceHelper
 	}
 };
 
-FReplaceInstancesOfClassParameters::FReplaceInstancesOfClassParameters(UClass* InOldClass, UClass* InNewClass)
-	: OldClass(InOldClass)
-	, NewClass(InNewClass)
-	, OriginalCDO(nullptr)
-	, ObjectsThatShouldUseOldStuff(nullptr)
-	, InstancesThatShouldUseOldClass(nullptr)
-	, bClassObjectReplaced(false)
-	, bPreserveRootComponent(true)
-{
-}
-
 /////////////////////////////////////////////////////////////////////////////////
 // FBlueprintCompileReinstancer
 
@@ -715,7 +704,13 @@ TSharedPtr<FReinstanceFinalizer> FBlueprintCompileReinstancer::ReinstanceInner(b
 		if (bShouldReinstance)
 		{
 			UE_LOG(LogBlueprint, Log, TEXT("BlueprintCompileReinstancer: Doing a full reinstance on class '%s'"), *GetPathNameSafe(ClassToReinstance));
-			ReplaceInstancesOfClass(DuplicatedClass, ClassToReinstance, OriginalCDO, &ObjectsThatShouldUseOldStuff, IsClassObjectReplaced(), ShouldPreserveRootComponentOfReinstancedActor());
+
+			FReplaceInstancesOfClassParameters Params;
+			Params.OriginalCDO = OriginalCDO;
+			Params.ObjectsThatShouldUseOldStuff = &ObjectsThatShouldUseOldStuff;
+			Params.bClassObjectReplaced = IsClassObjectReplaced();
+			Params.bPreserveRootComponent = ShouldPreserveRootComponentOfReinstancedActor();
+			ReplaceInstancesOfClass(DuplicatedClass, ClassToReinstance, Params);
 		}
 	}
 	return Finalizer;
@@ -1178,7 +1173,7 @@ struct FActorReplacementHelper
 	 * Runs construction scripts on the new actor and then finishes it off by
 	 * attaching it to the same attachments that its predecessor was set with. 
 	 */
-	void Finalize(const TMap<UObject*, UObject*>& OldToNewInstanceMap, TSet<UObject*>* ObjectsThatShouldUseOldStuff, const TArray<UObject*>& ObjectsToReplace, const TMap<FSoftObjectPath, UObject*>& ReinstancedObjectsWeakReferenceMap);
+	void Finalize(const TMap<UObject*, UObject*>& OldToNewInstanceMap, const TSet<UObject*>* ObjectsThatShouldUseOldStuff, const TArray<UObject*>& ObjectsToReplace, const TMap<FSoftObjectPath, UObject*>& ReinstancedObjectsWeakReferenceMap);
 
 	/**
 	* Takes the cached child actors, as well as the old AttachParent, and sets
@@ -1187,7 +1182,7 @@ struct FActorReplacementHelper
 	*
 	* @param OldToNewInstanceMap Mapping of reinstanced objects.
 	*/
-	void ApplyAttachments(const TMap<UObject*, UObject*>& OldToNewInstanceMap, TSet<UObject*>* ObjectsThatShouldUseOldStuff, const TArray<UObject*>& ObjectsToReplace, const TMap<FSoftObjectPath, UObject*>& ReinstancedObjectsWeakReferenceMap);
+	void ApplyAttachments(const TMap<UObject*, UObject*>& OldToNewInstanceMap, const TSet<UObject*>* ObjectsThatShouldUseOldStuff, const TArray<UObject*>& ObjectsToReplace, const TMap<FSoftObjectPath, UObject*>& ReinstancedObjectsWeakReferenceMap);
 
 private:
 	/**
@@ -1208,7 +1203,7 @@ private:
 	TMap<FName, UActorComponent*> OldActorComponentNameMap;
 };
 
-void FActorReplacementHelper::Finalize(const TMap<UObject*, UObject*>& OldToNewInstanceMap, TSet<UObject*>* ObjectsThatShouldUseOldStuff, const TArray<UObject*>& ObjectsToReplace, const TMap<FSoftObjectPath, UObject*>& ReinstancedObjectsWeakReferenceMap)
+void FActorReplacementHelper::Finalize(const TMap<UObject*, UObject*>& OldToNewInstanceMap, const TSet<UObject*>* ObjectsThatShouldUseOldStuff, const TArray<UObject*>& ObjectsToReplace, const TMap<FSoftObjectPath, UObject*>& ReinstancedObjectsWeakReferenceMap)
 {
 	if (!IsValid(NewActor))
 	{
@@ -1294,7 +1289,7 @@ void FActorReplacementHelper::Finalize(const TMap<UObject*, UObject*>& OldToNewI
 	}
 }
 
-void FActorReplacementHelper::ApplyAttachments(const TMap<UObject*, UObject*>& OldToNewInstanceMap, TSet<UObject*>* ObjectsThatShouldUseOldStuff, const TArray<UObject*>& ObjectsToReplace, const TMap<FSoftObjectPath, UObject*>& ReinstancedObjectsWeakReferenceMap)
+void FActorReplacementHelper::ApplyAttachments(const TMap<UObject*, UObject*>& OldToNewInstanceMap, const TSet<UObject*>* ObjectsThatShouldUseOldStuff, const TArray<UObject*>& ObjectsToReplace, const TMap<FSoftObjectPath, UObject*>& ReinstancedObjectsWeakReferenceMap)
 {
 	USceneComponent* NewRootComponent = NewActor->GetRootComponent();
 	if (NewRootComponent == nullptr)
@@ -1529,28 +1524,74 @@ namespace InstancedPropertyUtils
 	};
 }
 
+// @todo_deprecated - Remove in a future release.
 void FBlueprintCompileReinstancer::ReplaceInstancesOfClass(UClass* OldClass, UClass* NewClass, UObject*	OriginalCDO, TSet<UObject*>* ObjectsThatShouldUseOldStuff, bool bClassObjectReplaced, bool bPreserveRootComponent)
 {
-	TMap<UClass*, UClass*> OldToNewClassMap;
-	OldToNewClassMap.Add(OldClass, NewClass);
-	ReplaceInstancesOfClass_Inner(OldToNewClassMap, OriginalCDO, ObjectsThatShouldUseOldStuff, bClassObjectReplaced, bPreserveRootComponent);
+	FReplaceInstancesOfClassParameters Options;
+PRAGMA_DISABLE_DEPRECATION_WARNINGS
+	Options.OldClass = OldClass;
+	Options.NewClass = NewClass;
+PRAGMA_ENABLE_DEPRECATION_WARNINGS
+	Options.OriginalCDO = OriginalCDO;
+	Options.ObjectsThatShouldUseOldStuff = ObjectsThatShouldUseOldStuff;
+	Options.bClassObjectReplaced = bClassObjectReplaced;
+	Options.bPreserveRootComponent = bPreserveRootComponent;
+	ReplaceInstancesOfClass(OldClass, NewClass, Options);
 }
 
+// @todo_deprecated - Remove in a future release.
 void FBlueprintCompileReinstancer::ReplaceInstancesOfClassEx(const FReplaceInstancesOfClassParameters& Parameters )
 {
+	ReplaceInstancesOfClass(
+PRAGMA_DISABLE_DEPRECATION_WARNINGS
+		Parameters.OldClass,
+		Parameters.NewClass,
+PRAGMA_ENABLE_DEPRECATION_WARNINGS
+		Parameters
+	);
+}
+
+// @todo_deprecated - Remove in a future release.
+void FBlueprintCompileReinstancer::BatchReplaceInstancesOfClass(
+	TMap<UClass*, UClass*>& InOldToNewClassMap,
+PRAGMA_DISABLE_DEPRECATION_WARNINGS
+	const FBatchReplaceInstancesOfClassParameters& BatchParams
+PRAGMA_ENABLE_DEPRECATION_WARNINGS
+)
+{
+	if (InOldToNewClassMap.Num() == 0)
+	{
+		return;
+	}
+
+	FReplaceInstancesOfClassParameters Params;
+PRAGMA_DISABLE_DEPRECATION_WARNINGS
+	Params.bArchetypesAreUpToDate = BatchParams.bArchetypesAreUpToDate;
+	Params.bReplaceReferencesToOldClasses = BatchParams.bReplaceReferencesToOldClasses;
+PRAGMA_ENABLE_DEPRECATION_WARNINGS
+	ReplaceInstancesOfClass_Inner(InOldToNewClassMap, Params);
+}
+
+void FBlueprintCompileReinstancer::ReplaceInstancesOfClass(UClass* OldClass, UClass* NewClass, const FReplaceInstancesOfClassParameters& Params)
+{
+	ensureMsgf(!Params.bClassObjectReplaced || Params.OriginalCDO != nullptr, TEXT("bClassObjectReplaced is not expected to be set without OriginalCDO"));
+
 	TMap<UClass*, UClass*> OldToNewClassMap;
-	OldToNewClassMap.Add(Parameters.OldClass, Parameters.NewClass);
-	ReplaceInstancesOfClass_Inner(OldToNewClassMap, Parameters.OriginalCDO, Parameters.ObjectsThatShouldUseOldStuff, Parameters.bClassObjectReplaced, Parameters.bPreserveRootComponent, /*bArchetypesAreUpToDate=*/false, Parameters.InstancesThatShouldUseOldClass);
+	OldToNewClassMap.Add(OldClass, NewClass);
+	ReplaceInstancesOfClass_Inner(OldToNewClassMap, Params);
 }
 
-void FBlueprintCompileReinstancer::BatchReplaceInstancesOfClass(TMap<UClass*, UClass*>& InOldToNewClassMap, const FBatchReplaceInstancesOfClassParameters& Options )
+void FBlueprintCompileReinstancer::BatchReplaceInstancesOfClass(const TMap<UClass*, UClass*>& InOldToNewClassMap, const FReplaceInstancesOfClassParameters& Params)
 {
 	if (InOldToNewClassMap.Num() == 0)
 	{
 		return;
 	}
 
-	ReplaceInstancesOfClass_Inner(InOldToNewClassMap, nullptr, Options.ObjectsThatShouldUseOldStuff, false /*bClassObjectReplaced*/, true /*bPreserveRootComponent*/, Options.bArchetypesAreUpToDate, Options.InstancesThatShouldUseOldClass, Options.bReplaceReferencesToOldClasses);
+	checkf(Params.OriginalCDO == nullptr, TEXT("This path requires OriginalCDO to be NULL - use ReplaceInstancesOfClass() if you need to set it"));
+	ensureMsgf(!Params.bClassObjectReplaced, TEXT("bClassObjectReplaced is not expected to be set in this path - use ReplaceInstancesOfClass() instead"));
+
+	ReplaceInstancesOfClass_Inner(InOldToNewClassMap, Params);
 }
 
 bool FBlueprintCompileReinstancer::ReinstancerOrderingFunction(UClass* A, UClass* B)
@@ -1930,7 +1971,7 @@ static void ReplaceObjectHelper(UObject*& OldObject, UClass* OldClass, UObject*&
 	}
 }
 
-static void ReplaceActorHelper(AActor* OldActor, UClass* OldClass, UObject*& NewUObject, UClass* NewClass, TMap<UObject*, UObject*>& OldToNewInstanceMap, TMap<UClass*, UClass*>& InOldToNewClassMap, TMap<FSoftObjectPath, UObject*>& ReinstancedObjectsWeakReferenceMap, TMap<UObject*, FActorAttachmentData>& ActorAttachmentData, TArray<FActorReplacementHelper>& ReplacementActors, bool bPreserveRootComponent, bool& bSelectionChanged)
+static void ReplaceActorHelper(AActor* OldActor, UClass* OldClass, UObject*& NewUObject, UClass* NewClass, TMap<UObject*, UObject*>& OldToNewInstanceMap, const TMap<UClass*, UClass*>& InOldToNewClassMap, TMap<FSoftObjectPath, UObject*>& ReinstancedObjectsWeakReferenceMap, TMap<UObject*, FActorAttachmentData>& ActorAttachmentData, TArray<FActorReplacementHelper>& ReplacementActors, bool bPreserveRootComponent, bool& bSelectionChanged)
 {
 	FVector  Location = FVector::ZeroVector;
 	FRotator Rotation = FRotator::ZeroRotator;
@@ -1959,7 +2000,7 @@ static void ReplaceActorHelper(AActor* OldActor, UClass* OldClass, UObject*& New
 
 	// Spawn the new actor instance, in the same level as the original, but deferring running the construction script until we have transferred modified properties
 	ULevel*  ActorLevel = OldActor->GetLevel();
-	UClass** MappedClass = InOldToNewClassMap.Find(OldActor->GetClass());
+	UClass* const* MappedClass = InOldToNewClassMap.Find(OldActor->GetClass());
 	UClass*  SpawnClass = MappedClass ? *MappedClass : NewClass;
 
 	FActorSpawnParameters SpawnInfo;
@@ -2076,15 +2117,23 @@ static void ReplaceActorHelper(AActor* OldActor, UClass* OldClass, UObject*& New
 	OldToNewInstanceMap.Add(OldActor, NewActor);
 }
 
-void FBlueprintCompileReinstancer::ReplaceInstancesOfClass_Inner(TMap<UClass*, UClass*>& InOldToNewClassMap, UObject* InOriginalCDO, TSet<UObject*>* ObjectsThatShouldUseOldStuff, bool bClassObjectReplaced, bool bPreserveRootComponent, bool bArchetypesAreUpToDate, const TSet<UObject*>* InstancesThatShouldUseOldClass, bool bReplaceReferencesToOldClasses)
+void FBlueprintCompileReinstancer::ReplaceInstancesOfClass_Inner(const TMap<UClass*, UClass*>& InOldToNewClassMap, const FReplaceInstancesOfClassParameters& Params)
 {
-	// If there is an original CDO, we are only reinstancing a single class
-	check((InOriginalCDO != nullptr && InOldToNewClassMap.Num() == 1) || InOriginalCDO == nullptr); // (InOldToNewClassMap.Num() > 1 && InOriginalCDO == nullptr) || (InOldToNewClassMap.Num() == 1 && InOriginalCDO != nullptr));
+	// If there is an original CDO, make sure we are only reinstancing a single class (legacy path, non-batch)
+	UObject* InOriginalCDO = Params.OriginalCDO;
+	check((InOriginalCDO != nullptr && InOldToNewClassMap.Num() == 1) || InOriginalCDO == nullptr);
 
-	if (InOldToNewClassMap.Num() == 0)
-	{
-		return;
-	}
+	// This flag only applies to the legacy (i.e. non-batch) path.
+	const bool bClassObjectReplaced = InOriginalCDO != nullptr && Params.bClassObjectReplaced;
+
+	// If we're in the legacy path, always replace references to the CDO. Otherwise, it must be enabled.
+	const bool bReplaceReferencesToOldCDOs = InOriginalCDO != nullptr || Params.bReplaceReferencesToOldCDOs;
+	
+	TSet<UObject*>* ObjectsThatShouldUseOldStuff = Params.ObjectsThatShouldUseOldStuff;
+	const TSet<UObject*>* InstancesThatShouldUseOldClass = Params.InstancesThatShouldUseOldClass;
+	const bool bPreserveRootComponent = Params.bPreserveRootComponent;
+	const bool bArchetypesAreUpToDate = Params.bArchetypesAreUpToDate;
+	const bool bReplaceReferencesToOldClasses = Params.bReplaceReferencesToOldClasses;
 
 	USelection* SelectedActors = nullptr;
 	TArray<UObject*> ObjectsReplaced;
@@ -2343,9 +2392,8 @@ void FBlueprintCompileReinstancer::ReplaceInstancesOfClass_Inner(TMap<UClass*, U
 	TArray<UObject*> SourceObjects;
 	OldToNewInstanceMap.GenerateKeyArray(SourceObjects);
 	
-	if (InOriginalCDO)
+	if (bReplaceReferencesToOldCDOs)
 	{
-		check(InOldToNewClassMap.Num() == 1);
 		for (TPair<UClass*, UClass*> OldToNewClass : InOldToNewClassMap)
 		{
 			UClass* OldClass = OldToNewClass.Key;
diff --git a/Engine/Source/Editor/UnrealEd/Private/ObjectTools.cpp b/Engine/Source/Editor/UnrealEd/Private/ObjectTools.cpp
--- a/Engine/Source/Editor/UnrealEd/Private/ObjectTools.cpp
+++ b/Engine/Source/Editor/UnrealEd/Private/ObjectTools.cpp
@@ -1585,14 +1585,12 @@ namespace ObjectTools
 						UClass* OldClass = BlueprintToConsolidate->GeneratedClass;
 						UClass* OldSkeletonClass = BlueprintToConsolidate->SkeletonGeneratedClass;
 
-						FReplaceInstancesOfClassParameters ReplaceInstanceParams = FReplaceInstancesOfClassParameters(OldClass, BlueprintToConsolidateTo->GeneratedClass);
-						ReplaceInstanceParams.OriginalCDO = nullptr;
+						FReplaceInstancesOfClassParameters ReplaceInstanceParams;
 						ReplaceInstanceParams.ObjectsThatShouldUseOldStuff = &ObjectsToNotConsolidateWithin;
-						ReplaceInstanceParams.bClassObjectReplaced = true;
 						ReplaceInstanceParams.bPreserveRootComponent = true;
 						ReplaceInstanceParams.InstancesThatShouldUseOldClass = &ObjectsToNotConsolidateWithin;
 
-						FBlueprintCompileReinstancer::ReplaceInstancesOfClassEx(ReplaceInstanceParams);
+						FBlueprintCompileReinstancer::ReplaceInstancesOfClass(OldClass, BlueprintToConsolidateTo->GeneratedClass, ReplaceInstanceParams);
 						BlueprintToConsolidate->GeneratedClass = OldClass;
 						BlueprintToConsolidate->SkeletonGeneratedClass = OldSkeletonClass;
 					}
diff --git a/Engine/Source/Editor/UnrealEd/Private/PackageTools.cpp b/Engine/Source/Editor/UnrealEd/Private/PackageTools.cpp
--- a/Engine/Source/Editor/UnrealEd/Private/PackageTools.cpp
+++ b/Engine/Source/Editor/UnrealEd/Private/PackageTools.cpp
@@ -1073,9 +1073,9 @@ UPackageTools::UPackageTools(const FObjectInitializer& ObjectInitializer)
 							);
 
 							TSet<UObject*> InstancesToLeaveAlone(OldInstances);
-							FReplaceInstancesOfClassParameters ReplaceInstancesParameters(OldBlueprint->GeneratedClass, CastChecked<UBlueprint>(NewObject)->GeneratedClass);
+							FReplaceInstancesOfClassParameters ReplaceInstancesParameters;
 							ReplaceInstancesParameters.InstancesThatShouldUseOldClass = &InstancesToLeaveAlone;
-							FBlueprintCompileReinstancer::ReplaceInstancesOfClassEx(ReplaceInstancesParameters);
+							FBlueprintCompileReinstancer::ReplaceInstancesOfClass(OldBlueprint->GeneratedClass, CastChecked<UBlueprint>(NewObject)->GeneratedClass, ReplaceInstancesParameters);
 						}
 						else
 						{
diff --git a/Engine/Source/Editor/UnrealEd/Public/Kismet2/KismetEditorUtilities.h b/Engine/Source/Editor/UnrealEd/Public/Kismet2/KismetEditorUtilities.h
--- a/Engine/Source/Editor/UnrealEd/Public/Kismet2/KismetEditorUtilities.h
+++ b/Engine/Source/Editor/UnrealEd/Public/Kismet2/KismetEditorUtilities.h
@@ -58,6 +58,8 @@ enum class EBlueprintCompileOptions
 	UseDeltaSerializationDuringReinstancing = 0x400,
 	/** Skips the new variable defaults detection - in some cases we do not want to use the defaults from the generated class such as during a reparent */
 	SkipNewVariableDefaultsDetection = 0x800,
+	/** Directs reinstancing to find and replace external references to the regenerated CDO during reference replacement - in general, this is not needed */
+	IncludeCDOInReferenceReplacement = 0x1000,
 };
 
 ENUM_CLASS_FLAGS(EBlueprintCompileOptions)
diff --git a/Engine/Source/Editor/UnrealEd/Public/Kismet2/KismetReinstanceUtilities.h b/Engine/Source/Editor/UnrealEd/Public/Kismet2/KismetReinstanceUtilities.h
--- a/Engine/Source/Editor/UnrealEd/Public/Kismet2/KismetReinstanceUtilities.h
+++ b/Engine/Source/Editor/UnrealEd/Public/Kismet2/KismetReinstanceUtilities.h
@@ -39,41 +39,67 @@ public:
 
 struct UNREALED_API FReplaceInstancesOfClassParameters
 {
-	FReplaceInstancesOfClassParameters(UClass* InOldClass, UClass* InNewClass);
+	FReplaceInstancesOfClassParameters() = default;
+
+	UE_DEPRECATED(5.2, "Please use the default constructor instead.")
+	FReplaceInstancesOfClassParameters(UClass* InOldClass, UClass* InNewClass) {}
 
 	/** The old class, instances of which will be replaced */
-	UClass* OldClass;
+	UE_DEPRECATED(5.2, "This member is no longer in use.")
+	UClass* OldClass = nullptr;
 
 	/** The new class, used to create new instances to replace instances of the old class */
-	UClass* NewClass;
+	UE_DEPRECATED(5.2, "This member is no longer in use.")
+	UClass* NewClass = nullptr;
 
-	/** OriginalCDO, use if OldClass->ClassDefaultObject has been overwritten */
-	UObject* OriginalCDO; 
+	/** OriginalCDO, use if OldClass->ClassDefaultObject has been overwritten (non-batch only, legacy) */
+	UObject* OriginalCDO = nullptr; 
 
 	/** Set of objects that should not have their references updated if they refer to instances that are replaced */
-	TSet<UObject*>* ObjectsThatShouldUseOldStuff;
+	TSet<UObject*>* ObjectsThatShouldUseOldStuff = nullptr;
 
 	/** Set of objects for which new objects should not be created, useful if client has created new instances themself */
-	const TSet<UObject*>* InstancesThatShouldUseOldClass;
+	const TSet<UObject*>* InstancesThatShouldUseOldClass = nullptr;
+
+	/** Set to true if class object has been replaced (non-batch only, legacy) */
+	bool bClassObjectReplaced = false; 
+
+	/** Defaults to true, indicates whether root components should be preserved */
+	bool bPreserveRootComponent = true;
+
+	bool bArchetypesAreUpToDate = false;
 
-	/** Set to true if class object has been replaced*/
-	bool bClassObjectReplaced; 
+	/**
+	 * Blueprints reuses its UClass* from compile to compile, but it's more
+	 * intuitive to just replace a UClass* with a new instance (e.g. from a
+	 * package reload). This flag tells the reinstancer to replace references
+	 * to old classes with references to new classes.
+	 */
+	bool bReplaceReferencesToOldClasses = false;
 
-	/** Defaults to true, indicates whether root component should be preserved */
-	bool bPreserveRootComponent; 
+	/**
+	 * Indicates whether CDOs should be included in reference replacement.
+	 * Disabled by default since this can increase search cost, because this
+	 * effectively means replacement will require us to do a referencer search
+	 * for at least one instance of every remapped class (i.e. - the CDO). In
+	 * most cases (e.g. load time), incurring this cost is unnecessary because
+	 * the CDO is not expected to be referenced outside of its own class type,
+	 * so we treat it as an opt-in behavior rather than enabling it by default.
+	 * 
+	 * Note: This flag is implied to be 'true' if 'OriginalCDO' is non-NULL. In
+	 * that case, the value of this flag will not be used in order to maintain
+	 * backwards-compatibility with existing code paths.
+	 */
+	bool bReplaceReferencesToOldCDOs = false;
 };
 
-struct UNREALED_API FBatchReplaceInstancesOfClassParameters
+struct UNREALED_API UE_DEPRECATED(5.2, "This type is no longer in use.") FBatchReplaceInstancesOfClassParameters
 {
-	TSet<UObject*>* ObjectsThatShouldUseOldStuff = nullptr;
-	
-	const TSet<UObject*>* InstancesThatShouldUseOldClass = nullptr;
+	bool bArchetypesAreUpToDate = false;
 
-	bool bArchetypesAreUpToDate = false; 
-	
-	/** 
-	 * Blueprints reuses its UClass* from compile to compile, but it's more 
-	 * intuitive to just replace a UClass* with a new instance (e.g. from a 
+	/**
+	 * Blueprints reuses its UClass* from compile to compile, but it's more
+	 * intuitive to just replace a UClass* with a new instance (e.g. from a
 	 * package reload). This flag tells the reinstancer to replace references
 	 * to old classes with references to new classes.
 	 */
@@ -171,12 +197,23 @@ public:
 	/** Consumes the set and map populated by calls to UpdateBytecodeReferences */
 	static void FinishUpdateBytecodeReferences( const TSet<UBlueprint*>& DependentBPs, const TMap<FFieldVariant, FFieldVariant>& FieldMappings);
 
-	/** Worker function to replace all instances of OldClass with a new instance of NewClass */
+	// @todo_deprecated - To be removed in a future release.
+	UE_DEPRECATED(5.2, "Please use the version that takes an FReplaceInstancesOfClassParameters input instead.")
 	static void ReplaceInstancesOfClass(UClass* OldClass, UClass* NewClass, UObject* OriginalCDO = nullptr, TSet<UObject*>* ObjectsThatShouldUseOldStuff = nullptr, bool bClassObjectReplaced = false, bool bPreserveRootComponent = true);
+	UE_DEPRECATED(5.2, "Please use ReplaceInstancesOfClass() instead.")
 	static void ReplaceInstancesOfClassEx(const FReplaceInstancesOfClassParameters& Parameters );
-
-	/** Batch replaces a mapping of one or more classes to their new class by leveraging ReplaceInstancesOfClass */
-	static void BatchReplaceInstancesOfClass(TMap<UClass*, UClass*>& InOldToNewClassMap, const FBatchReplaceInstancesOfClassParameters& Options = FBatchReplaceInstancesOfClassParameters());
+	UE_DEPRECATED(5.2, "Please use the version that takes an FReplaceInstancesOfClassParameters input instead.")
+	static void BatchReplaceInstancesOfClass(TMap<UClass*, UClass*>& InOldToNewClassMap,
+PRAGMA_DISABLE_DEPRECATION_WARNINGS
+		const FBatchReplaceInstancesOfClassParameters& BatchParams = FBatchReplaceInstancesOfClassParameters()
+PRAGMA_ENABLE_DEPRECATION_WARNINGS
+	);
+
+	/** Replace all instances of OldClass with a new instance of NewClass */
+	static void ReplaceInstancesOfClass(UClass* OldClass, UClass* NewClass, const FReplaceInstancesOfClassParameters& Params);
+
+	/** Batch replaces a mapping of one or more classes to their new class */
+	static void BatchReplaceInstancesOfClass(const TMap<UClass*, UClass*>& InOldToNewClassMap, const FReplaceInstancesOfClassParameters& Params);
 	
 	/** Function used to safely discard a CDO, so that the class can have its layout changed, callers must move parent CDOs aside before moving child CDOs aside: */
 	static UClass* MoveCDOToNewClass(UClass* OwnerClass, const TMap<UClass*, UClass*>& OldToNewMap, bool bAvoidCDODuplication);
@@ -261,7 +298,7 @@ protected:
 
 private:
 	/** Handles the work of ReplaceInstancesOfClass, handling both normal replacement of instances and batch */
-	static void ReplaceInstancesOfClass_Inner(TMap<UClass*, UClass*>& InOldToNewClassMap, UObject* InOriginalCDO, TSet<UObject*>* ObjectsThatShouldUseOldStuff = NULL, bool bClassObjectReplaced = false, bool bPreserveRootComponent = true, bool bArchetypesAreUpToDate = false, const TSet<UObject*>* InstancesThatShouldUseOldClass = nullptr, bool bReplaceReferencesToOldClasses = false);
+	static void ReplaceInstancesOfClass_Inner(const TMap<UClass*, UClass*>& InOldToNewClassMap, const FReplaceInstancesOfClassParameters& Params);
 
 	/** Returns true if A is higher up the class hierarchy  */
 	static bool ReinstancerOrderingFunction(UClass* A, UClass* B);
