Added SCC conflict warning dialog

diff --git a/Engine/Source/Developer/ToolWidgets/Private/SRichTextWarningOrErrorBox.cpp b/Engine/Source/Developer/ToolWidgets/Private/SRichTextWarningOrErrorBox.cpp
new file mode 100644
--- /dev/null
+++ b/Engine/Source/Developer/ToolWidgets/Private/SRichTextWarningOrErrorBox.cpp
@@ -0,0 +1,48 @@
+// Copyright Epic Games, Inc. All Rights Reserved.
+
+#include "SRichTextWarningOrErrorBox.h"
+#include "Widgets/Text/SRichTextBlock.h"
+#include "Widgets/SBoxPanel.h"
+#include "Widgets/Images/SImage.h"
+
+void SRichTextWarningOrErrorBox::Construct(const FArguments& InArgs)
+{
+	MessageStyle = InArgs._MessageStyle;
+
+	SBorder::Construct(SBorder::FArguments()
+		.Padding(InArgs._Padding)
+		.ForegroundColor(FAppStyle::Get().GetSlateColor("Colors.White"))
+		.BorderImage_Lambda([this]() { return MessageStyle.Get() == EMessageStyle::Warning ? FAppStyle::Get().GetBrush("RoundedWarning") : FAppStyle::Get().GetBrush("RoundedError"); })
+		[
+			SNew(SHorizontalBox)
+			+ SHorizontalBox::Slot()
+			.HAlign(HAlign_Left)
+			.VAlign(VAlign_Top)
+			.AutoWidth()
+			.Padding(FMargin(0.0f, 0.0f, 16.0f, 0.0f))
+			[
+				SNew(SImage)
+				.DesiredSizeOverride(InArgs._IconSize)
+				.Image_Lambda([this]() { return MessageStyle.Get() == EMessageStyle::Warning ? FAppStyle::Get().GetBrush("Icons.WarningWithColor") : FAppStyle::Get().GetBrush("Icons.ErrorWithColor"); })
+				.ColorAndOpacity(FSlateColor::UseForeground())
+			]
+			+SHorizontalBox::Slot()
+			.VAlign(VAlign_Center)
+			[
+				SNew(SRichTextBlock)
+				.DecoratorStyleSet(&FAppStyle::Get())
+				.Text(InArgs._Message)
+				.AutoWrapText(true)
+			]
+
+			+SHorizontalBox::Slot()
+			.AutoWidth()
+			.VAlign(VAlign_Center)
+			.HAlign(HAlign_Right)
+			[
+				InArgs._Content.Widget
+			]
+
+		]);
+}
+
diff --git a/Engine/Source/Developer/ToolWidgets/Public/SRichTextWarningOrErrorBox.h b/Engine/Source/Developer/ToolWidgets/Public/SRichTextWarningOrErrorBox.h
new file mode 100644
--- /dev/null
+++ b/Engine/Source/Developer/ToolWidgets/Public/SRichTextWarningOrErrorBox.h
@@ -0,0 +1,17 @@
+// Copyright Epic Games, Inc. All Rights Reserved.
+
+#pragma once
+
+#include "CoreMinimal.h"
+#include "Widgets/Layout/SBorder.h"
+#include "SWarningOrErrorBox.h"
+
+class TOOLWIDGETS_API SRichTextWarningOrErrorBox : public SWarningOrErrorBox
+{
+public:
+
+	void Construct(const FArguments& InArgs);
+
+private:
+	TAttribute<EMessageStyle> MessageStyle;
+};
\ No newline at end of file
diff --git a/Engine/Source/Editor/SourceControlWindows/Private/SourceControlWindows.cpp b/Engine/Source/Editor/SourceControlWindows/Private/SourceControlWindows.cpp
--- a/Engine/Source/Editor/SourceControlWindows/Private/SourceControlWindows.cpp
+++ b/Engine/Source/Editor/SourceControlWindows/Private/SourceControlWindows.cpp
@@ -260,18 +260,27 @@ bool FSourceControlWindows::PromptForCheckin(FCheckinResultInfo& OutResultInfo,
 	{
 		SourceControlHelpers::RevertUnchangedFiles(SourceControlProvider, Description.FilesForSubmit);
 
-		// Make sure all files are still checked out
-		for (int32 VerifyIndex = Description.FilesForSubmit.Num()-1; VerifyIndex >= 0; --VerifyIndex)
+		const FString SCCProjectDir = ISourceControlModule::Get().GetSourceControlProjectDir();
+
+		if (SCCProjectDir.IsEmpty())
 		{
-			FSourceControlStatePtr SourceControlState = SourceControlProvider.GetState(Description.FilesForSubmit[VerifyIndex], EStateCacheUsage::Use);
-			if( SourceControlState.IsValid() && !SourceControlState->IsCheckedOut() && !SourceControlState->IsAdded() && !SourceControlState->IsDeleted() )
+			// Make sure all files are still checked out
+			for (int32 VerifyIndex = Description.FilesForSubmit.Num() - 1; VerifyIndex >= 0; --VerifyIndex)
 			{
-				Description.FilesForSubmit.RemoveAt(VerifyIndex);
+				FSourceControlStatePtr SourceControlState = SourceControlProvider.GetState(Description.FilesForSubmit[VerifyIndex], EStateCacheUsage::Use);
+				if (SourceControlState.IsValid() && !SourceControlState->IsCheckedOut() && !SourceControlState->IsAdded() && !SourceControlState->IsDeleted())
+				{
+					Description.FilesForSubmit.RemoveAt(VerifyIndex);
+				}
 			}
 		}
+		else
+		{
+			// For project-based source control, we want to go through with a check in attempt even when 
+			// files are not checked out by the current user, and generate a warning dialog
+		}
 	}
 
-
 	//~~~~~~~~~~~~~~~~~~~~~~~~~~~~~~~~~~~~~~~~~~
 	// Mark files for add as needed
 	
@@ -426,7 +435,9 @@ void FSourceControlWindows::ChoosePackagesToCheckInCompleted(const TArray<UPacka
 		return;
 	}
 
-	TArray<FString> PendingDeletePaths;	
+	TArray<FString> PendingDeletePaths;
+
+	bool bUseSourceControlStateCache = true;
 	FString SourceControlProjectDir = ISourceControlModule::Get().GetSourceControlProjectDir();
 	if (SourceControlProjectDir.IsEmpty())
 	{
@@ -435,9 +446,9 @@ void FSourceControlWindows::ChoosePackagesToCheckInCompleted(const TArray<UPacka
 	else
 	{
 		PendingDeletePaths.Add(SourceControlProjectDir);
+		bUseSourceControlStateCache = false;
 	}
 
-	const bool bUseSourceControlStateCache = true;
 	PromptForCheckin(OutResultInfo, PackageNames, PendingDeletePaths, ConfigFiles, bUseSourceControlStateCache);
 }
 
diff --git a/Engine/Source/Editor/UnrealEd/Private/FileHelpers.cpp b/Engine/Source/Editor/UnrealEd/Private/FileHelpers.cpp
--- a/Engine/Source/Editor/UnrealEd/Private/FileHelpers.cpp
+++ b/Engine/Source/Editor/UnrealEd/Private/FileHelpers.cpp
@@ -4554,6 +4554,8 @@ void FEditorFileUtils::FindAllSubmittablePackageFiles(TMap<FString, FSourceContr
 {
 	ISourceControlProvider& SourceControlProvider = ISourceControlModule::Get().GetProvider();
 
+	const FString SCCProjectDir = ISourceControlModule::Get().GetSourceControlProjectDir();
+
 	TArray<FString> Packages;
 	FEditorFileUtils::FindAllPackageFiles(Packages);
 
@@ -4575,12 +4577,26 @@ void FEditorFileUtils::FindAllSubmittablePackageFiles(TMap<FString, FSourceContr
 		FSourceControlStatePtr SourceControlState = SourceControlProvider.GetState(FPaths::ConvertRelativePathToFull(Filename), EStateCacheUsage::Use);
 
 		// Only include non-map packages that are currently checked out or packages not under source control
-		if (SourceControlState.IsValid() && SourceControlState->IsCurrent() &&
-			(SourceControlState->CanCheckIn() || (!SourceControlState->IsSourceControlled() && SourceControlState->CanAdd())) &&
-			(bIncludeMaps || !IsMapPackageAsset(*Filename)))
+		if (SCCProjectDir.IsEmpty())
+		{
+			if (SourceControlState.IsValid() && SourceControlState->IsCurrent() &&
+				(SourceControlState->CanCheckIn() || (!SourceControlState->IsSourceControlled() && SourceControlState->CanAdd())) &&
+				(bIncludeMaps || !IsMapPackageAsset(*Filename)))
+			{
+				OutPackages.Add(MoveTemp(PackageName), MoveTemp(SourceControlState));
+			}
+		}
+		else
 		{
-			OutPackages.Add(MoveTemp(PackageName), MoveTemp(SourceControlState));
+			if (SourceControlState.IsValid() && 
+				(SourceControlState->CanCheckIn() || (!SourceControlState->IsSourceControlled() && SourceControlState->CanAdd())) &&
+				(bIncludeMaps || !IsMapPackageAsset(*Filename)))
+			{
+				OutPackages.Add(MoveTemp(PackageName), MoveTemp(SourceControlState));
+			}
 		}
+
+		
 	}
 }
 
@@ -4610,12 +4626,12 @@ void FEditorFileUtils::FindAllSubmittableProjectFiles(TMap<FString, FSourceContr
 			}
 		);
 
+		OutProjectFiles.Reserve(SourceControlStates.Num());
 		for (FSourceControlStateRef& SourceControlState : SourceControlStates)
 		{
 			const FString& Filename = SourceControlState->GetFilename();
 
-			if (SourceControlState->IsCurrent() &&
-				(SourceControlState->CanCheckIn() || (!SourceControlState->IsSourceControlled() && SourceControlState->CanAdd())))
+			if (SourceControlState->CanCheckIn() || (!SourceControlState->IsSourceControlled() && SourceControlState->CanAdd()))
 			{
 				FString PackageName;
 				if (!FPackageName::TryConvertFilenameToLongPackageName(Filename, PackageName))
