Fix Rename Folder not persisting assets Public/Private state

diff --git a/Engine/Source/Editor/UnrealEd/Private/ObjectTools.cpp b/Engine/Source/Editor/UnrealEd/Private/ObjectTools.cpp
--- a/Engine/Source/Editor/UnrealEd/Private/ObjectTools.cpp
+++ b/Engine/Source/Editor/UnrealEd/Private/ObjectTools.cpp
@@ -3984,27 +3984,30 @@ namespace ObjectTools
 					bool bFoundCompatibleRedirector = false;
 					UObjectRedirector* Redirector = nullptr;
 
-					if (NewPackage != Object->GetPackage())
+					UPackage* OldPackage = Object->GetPackage();
+					if (NewPackage != OldPackage)
 					{
 						NewPackage->GetOutermost()->FullyLoad();
 
 						// Make sure we copy all the cooked package flags if the asset was already cooked.
-						if (Object->GetOutermost()->HasAnyPackageFlags(PKG_FilterEditorOnly))
+						if (OldPackage->HasAnyPackageFlags(PKG_FilterEditorOnly))
 						{
 							NewPackage->SetPackageFlags(PKG_FilterEditorOnly);
 						}
-						NewPackage->bIsCookedForEditor = Object->GetOutermost()->bIsCookedForEditor;
+						NewPackage->bIsCookedForEditor = OldPackage->bIsCookedForEditor;
 
 						// Renaming an asset should respect the export controls of the original.
-						if (Object->GetOutermost()->HasAnyPackageFlags(PKG_DisallowExport))
+						if (OldPackage->HasAnyPackageFlags(PKG_DisallowExport))
 						{
 							NewPackage->SetPackageFlags(PKG_DisallowExport);
 						}
 
+						NewPackage->SetIsExternallyReferenceable(OldPackage->IsExternallyReferenceable());
+
 						// When renaming a World Composition map, make sure to properly initialize WorldTileInfo
-						if (Object->GetOutermost()->GetWorldTileInfo())
+						if (OldPackage->GetWorldTileInfo())
 						{
-							NewPackage->SetWorldTileInfo(MakeUnique<FWorldTileInfo>(*Object->GetOutermost()->GetWorldTileInfo()));
+							NewPackage->SetWorldTileInfo(MakeUnique<FWorldTileInfo>(*OldPackage->GetWorldTileInfo()));
 						}
 
 						Redirector = Cast<UObjectRedirector>(StaticFindObject(UObjectRedirector::StaticClass(), NewPackage, *NewObjectName));
