Reduce the priority of loose cooked package writer tasks to avoid saturating all worker threads in a way that leads to the game thread becoming stalled.  Reduces UCookOnTheFlyServer::WaitForAsync duration on large cooks and doesn't manifest as increase in UPackage::WaitForAsyncFileWrites at the end of the cook operation.

diff --git a/Engine/Source/Editor/UnrealEd/Private/Cooker/LooseCookedPackageWriter.cpp b/Engine/Source/Editor/UnrealEd/Private/Cooker/LooseCookedPackageWriter.cpp
--- a/Engine/Source/Editor/UnrealEd/Private/Cooker/LooseCookedPackageWriter.cpp
+++ b/Engine/Source/Editor/UnrealEd/Private/Cooker/LooseCookedPackageWriter.cpp
@@ -322,7 +322,8 @@ void FLooseCookedPackageWriter::AsyncSaveOutputFiles(FRecord& Record, FCommitCon
 
 			// This is used to release the game thread to access the hashes
 			UE::SavePackageUtilities::DecrementOutstandingAsyncWrites();
-		}
+		},
+		UE::Tasks::ETaskPriority::BackgroundNormal
 	);
 }
 
