World Partition - Spline Mesh Builder: Fixed deleted files not being submitted

diff --git a/Engine/Source/Editor/UnrealEd/Private/WorldPartition/WorldPartitionBuilder.cpp b/Engine/Source/Editor/UnrealEd/Private/WorldPartition/WorldPartitionBuilder.cpp
--- a/Engine/Source/Editor/UnrealEd/Private/WorldPartition/WorldPartitionBuilder.cpp
+++ b/Engine/Source/Editor/UnrealEd/Private/WorldPartition/WorldPartitionBuilder.cpp
@@ -394,6 +394,13 @@ bool UWorldPartitionBuilder::SavePackages(const TArray<UPackage*>& Packages, FPa
 	return bSuccess;
 }
 
+bool UWorldPartitionBuilder::DeletePackages(const TArray<UPackage*>& Packages, FPackageSourceControlHelper& PackageHelper, bool bErrorsAsWarnings)
+{
+	TArray<FString> PackagesNames;
+	Algo::Transform(Packages, PackagesNames, [](UPackage* InPackage) { return InPackage->GetName(); });
+	return DeletePackages(PackagesNames, PackageHelper, bErrorsAsWarnings);
+}
+
 bool UWorldPartitionBuilder::DeletePackages(const TArray<FString>& PackageNames, FPackageSourceControlHelper& PackageHelper, bool bErrorsAsWarnings)
 {
 	if (PackageNames.Num() > 0)
diff --git a/Engine/Source/Editor/UnrealEd/Private/WorldPartition/WorldPartitionLandscapeSplineMeshesBuilder.cpp b/Engine/Source/Editor/UnrealEd/Private/WorldPartition/WorldPartitionLandscapeSplineMeshesBuilder.cpp
--- a/Engine/Source/Editor/UnrealEd/Private/WorldPartition/WorldPartitionLandscapeSplineMeshesBuilder.cpp
+++ b/Engine/Source/Editor/UnrealEd/Private/WorldPartition/WorldPartitionLandscapeSplineMeshesBuilder.cpp
@@ -488,12 +488,15 @@ bool UWorldPartitionLandscapeSplineMeshesBuilder::RunInternal(UWorld* InWorld, c
 	}
 
 	// Destroy actors that are no longer needed
-	TArray<FString> PackagesToDelete;
+	TArray<UPackage*> PackagesToDelete;
 	for (ALandscapeSplineMeshesActor* Actor : PreviousGeneratedActors)
 	{
 		if (UPackage* Package = Actor->GetExternalPackage())
 		{
-			PackagesToDelete.Add(Package->GetName());
+			PackagesToDelete.Add(Package);
+
+			// Releases file handles so packages can be deleted
+			ResetLoaders(Package);
 		}
 		InWorld->DestroyActor(Actor);
 	}
@@ -510,8 +513,10 @@ bool UWorldPartitionLandscapeSplineMeshesBuilder::RunInternal(UWorld* InWorld, c
 			return false;
 		}
 
-		TArray<FString> ModifiedFiles = MoveTemp(PackagesToDelete);
-		Algo::Transform(PackagesToSave, ModifiedFiles, [](const UPackage* InPackage) { return USourceControlHelpers::PackageFilename(InPackage); });
+		TArray<FString> ModifiedFiles;
+		auto GetPackageFilename = [](const UPackage* InPackage) { return USourceControlHelpers::PackageFilename(InPackage); };
+		Algo::Transform(PackagesToDelete, ModifiedFiles, GetPackageFilename);
+		Algo::Transform(PackagesToSave, ModifiedFiles, GetPackageFilename);
 
 		const FString ChangeDescription = FString::Printf(TEXT("Rebuilt landscape splines for %s"), *InWorld->GetName());
 		if (!AutoSubmitFiles(ModifiedFiles, ChangeDescription))
diff --git a/Engine/Source/Editor/UnrealEd/Public/WorldPartition/WorldPartitionBuilder.h b/Engine/Source/Editor/UnrealEd/Public/WorldPartition/WorldPartitionBuilder.h
--- a/Engine/Source/Editor/UnrealEd/Public/WorldPartition/WorldPartitionBuilder.h
+++ b/Engine/Source/Editor/UnrealEd/Public/WorldPartition/WorldPartitionBuilder.h
@@ -58,6 +58,7 @@ public:
 	virtual bool PreWorldInitialization(FPackageSourceControlHelper& PackageHelper) { return true; }
 
 	static bool SavePackages(const TArray<UPackage*>& Packages, FPackageSourceControlHelper& PackageHelper, bool bErrorsAsWarnings = false);
+	static bool DeletePackages(const TArray<UPackage*>& Packages, FPackageSourceControlHelper& PackageHelper, bool bErrorsAsWarnings = false);
 	static bool DeletePackages(const TArray<FString>& PackageNames, FPackageSourceControlHelper& PackageHelper, bool bErrorsAsWarnings = false);
 
 protected:
