Interchange: Reimport with new file need to use the interchange picker if the file was imported with interchange.

#rnx

diff --git a/Engine/Plugins/Interchange/Runtime/Source/Import/Public/Animation/InterchangeAnimSequenceFactory.h b/Engine/Plugins/Interchange/Runtime/Source/Import/Public/Animation/InterchangeAnimSequenceFactory.h
--- a/Engine/Plugins/Interchange/Runtime/Source/Import/Public/Animation/InterchangeAnimSequenceFactory.h
+++ b/Engine/Plugins/Interchange/Runtime/Source/Import/Public/Animation/InterchangeAnimSequenceFactory.h
@@ -22,6 +22,7 @@ public:
 	// Interchange factory base interface begin
 
 	virtual UClass* GetFactoryClass() const override;
+	virtual EInterchangeFactoryAssetType GetFactoryAssetType() override { return EInterchangeFactoryAssetType::Animations; }
 	virtual UObject* CreateEmptyAsset(const FCreateAssetParams& Arguments) override;
 	virtual UObject* CreateAsset(const FCreateAssetParams& Arguments) override;
 	virtual void PreImportPreCompletedCallback(const FImportPreCompletedCallbackParams& Arguments) override;
diff --git a/Engine/Plugins/Interchange/Runtime/Source/Import/Public/Animation/InterchangeAnimationTrackSetFactory.h b/Engine/Plugins/Interchange/Runtime/Source/Import/Public/Animation/InterchangeAnimationTrackSetFactory.h
--- a/Engine/Plugins/Interchange/Runtime/Source/Import/Public/Animation/InterchangeAnimationTrackSetFactory.h
+++ b/Engine/Plugins/Interchange/Runtime/Source/Import/Public/Animation/InterchangeAnimationTrackSetFactory.h
@@ -19,6 +19,7 @@ public:
 	// Interchange factory base interface begin
 
 	virtual UClass* GetFactoryClass() const override;
+	virtual EInterchangeFactoryAssetType GetFactoryAssetType() override { return EInterchangeFactoryAssetType::Animations; }
 	virtual UObject* CreateEmptyAsset(const FCreateAssetParams& Arguments) override;
 	virtual UObject* CreateAsset(const FCreateAssetParams& Arguments) override;
 	virtual void PreImportPreCompletedCallback(const FImportPreCompletedCallbackParams& Arguments) override;
diff --git a/Engine/Plugins/Interchange/Runtime/Source/Import/Public/Material/InterchangeMaterialFactory.h b/Engine/Plugins/Interchange/Runtime/Source/Import/Public/Material/InterchangeMaterialFactory.h
--- a/Engine/Plugins/Interchange/Runtime/Source/Import/Public/Material/InterchangeMaterialFactory.h
+++ b/Engine/Plugins/Interchange/Runtime/Source/Import/Public/Material/InterchangeMaterialFactory.h
@@ -26,6 +26,7 @@ public:
 	// Interchange factory base interface begin
 
 	virtual UClass* GetFactoryClass() const override;
+	virtual EInterchangeFactoryAssetType GetFactoryAssetType() override { return EInterchangeFactoryAssetType::Materials; }
 	virtual UObject* CreateEmptyAsset(const FCreateAssetParams& Arguments) override;
 	virtual UObject* CreateAsset(const FCreateAssetParams& Arguments) override;
 	virtual void PreImportPreCompletedCallback(const FImportPreCompletedCallbackParams& Arguments) override;
@@ -52,6 +53,7 @@ public:
 	// Interchange factory base interface begin
 
 	virtual UClass* GetFactoryClass() const override;
+	virtual EInterchangeFactoryAssetType GetFactoryAssetType() override { return EInterchangeFactoryAssetType::Materials; }
 	virtual UObject* CreateEmptyAsset(const FCreateAssetParams& Arguments) override;
 	virtual UObject* CreateAsset(const FCreateAssetParams& Arguments) override;
 	virtual void PreImportPreCompletedCallback(const FImportPreCompletedCallbackParams& Arguments) override;
diff --git a/Engine/Plugins/Interchange/Runtime/Source/Import/Public/Mesh/InterchangePhysicsAssetFactory.h b/Engine/Plugins/Interchange/Runtime/Source/Import/Public/Mesh/InterchangePhysicsAssetFactory.h
--- a/Engine/Plugins/Interchange/Runtime/Source/Import/Public/Mesh/InterchangePhysicsAssetFactory.h
+++ b/Engine/Plugins/Interchange/Runtime/Source/Import/Public/Mesh/InterchangePhysicsAssetFactory.h
@@ -19,6 +19,7 @@ public:
 	// Interchange factory base interface begin
 
 	virtual UClass* GetFactoryClass() const override;
+	virtual EInterchangeFactoryAssetType GetFactoryAssetType() override { return EInterchangeFactoryAssetType::Physics; }
 	virtual UObject* CreateEmptyAsset(const FCreateAssetParams& Arguments) override;
 	virtual UObject* CreateAsset(const FCreateAssetParams& Arguments) override;
 	//virtual void PreImportPreCompletedCallback(const FImportPreCompletedCallbackParams& Arguments) const override;
diff --git a/Engine/Plugins/Interchange/Runtime/Source/Import/Public/Mesh/InterchangeSkeletalMeshFactory.h b/Engine/Plugins/Interchange/Runtime/Source/Import/Public/Mesh/InterchangeSkeletalMeshFactory.h
--- a/Engine/Plugins/Interchange/Runtime/Source/Import/Public/Mesh/InterchangeSkeletalMeshFactory.h
+++ b/Engine/Plugins/Interchange/Runtime/Source/Import/Public/Mesh/InterchangeSkeletalMeshFactory.h
@@ -21,6 +21,7 @@ public:
 	// Interchange factory base interface begin
 
 	virtual UClass* GetFactoryClass() const override;
+	virtual EInterchangeFactoryAssetType GetFactoryAssetType() override { return EInterchangeFactoryAssetType::Meshes; }
 	virtual UObject* CreateEmptyAsset(const FCreateAssetParams& Arguments) override;
 	virtual UObject* CreateAsset(const FCreateAssetParams& Arguments) override;
 	virtual void PreImportPreCompletedCallback(const FImportPreCompletedCallbackParams& Arguments) override;
diff --git a/Engine/Plugins/Interchange/Runtime/Source/Import/Public/Mesh/InterchangeSkeletonFactory.h b/Engine/Plugins/Interchange/Runtime/Source/Import/Public/Mesh/InterchangeSkeletonFactory.h
--- a/Engine/Plugins/Interchange/Runtime/Source/Import/Public/Mesh/InterchangeSkeletonFactory.h
+++ b/Engine/Plugins/Interchange/Runtime/Source/Import/Public/Mesh/InterchangeSkeletonFactory.h
@@ -19,6 +19,7 @@ public:
 	// Interchange factory base interface begin
 
 	virtual UClass* GetFactoryClass() const override;
+	virtual EInterchangeFactoryAssetType GetFactoryAssetType() override { return EInterchangeFactoryAssetType::Meshes; }
 	virtual UObject* CreateEmptyAsset(const FCreateAssetParams& Arguments) override;
 	virtual UObject* CreateAsset(const FCreateAssetParams& Arguments) override;
 	//virtual void PreImportPreCompletedCallback(const FImportPreCompletedCallbackParams& Arguments) const override;
diff --git a/Engine/Plugins/Interchange/Runtime/Source/Import/Public/Mesh/InterchangeStaticMeshFactory.h b/Engine/Plugins/Interchange/Runtime/Source/Import/Public/Mesh/InterchangeStaticMeshFactory.h
--- a/Engine/Plugins/Interchange/Runtime/Source/Import/Public/Mesh/InterchangeStaticMeshFactory.h
+++ b/Engine/Plugins/Interchange/Runtime/Source/Import/Public/Mesh/InterchangeStaticMeshFactory.h
@@ -29,6 +29,7 @@ public:
 	// Interchange factory base interface begin
 
 	virtual UClass* GetFactoryClass() const override;
+	virtual EInterchangeFactoryAssetType GetFactoryAssetType() override { return EInterchangeFactoryAssetType::Meshes; }
 	virtual UObject* CreateEmptyAsset(const FCreateAssetParams& Arguments) override;
 	virtual UObject* CreateAsset(const FCreateAssetParams& Arguments) override;
 	virtual void PreImportPreCompletedCallback(const FImportPreCompletedCallbackParams& Arguments) override;
diff --git a/Engine/Plugins/Interchange/Runtime/Source/Import/Public/Texture/InterchangeTextureFactory.h b/Engine/Plugins/Interchange/Runtime/Source/Import/Public/Texture/InterchangeTextureFactory.h
--- a/Engine/Plugins/Interchange/Runtime/Source/Import/Public/Texture/InterchangeTextureFactory.h
+++ b/Engine/Plugins/Interchange/Runtime/Source/Import/Public/Texture/InterchangeTextureFactory.h
@@ -55,6 +55,7 @@ public:
 	// Interchange factory base interface begin
 
 	virtual UClass* GetFactoryClass() const override;
+	virtual EInterchangeFactoryAssetType GetFactoryAssetType() override { return EInterchangeFactoryAssetType::Textures; }
 	virtual UObject* CreateEmptyAsset(const FCreateAssetParams& Arguments) override;
 	virtual UObject* CreateAsset(const FCreateAssetParams& Arguments) override;
 	virtual void PreImportPreCompletedCallback(const FImportPreCompletedCallbackParams& Arguments) override;
diff --git a/Engine/Source/Editor/UnrealEd/Private/Editor.cpp b/Engine/Source/Editor/UnrealEd/Private/Editor.cpp
--- a/Engine/Source/Editor/UnrealEd/Private/Editor.cpp
+++ b/Engine/Source/Editor/UnrealEd/Private/Editor.cpp
@@ -752,10 +752,22 @@ void FReimportManager::GetNewReimportPath(UObject* Obj, TArray<FString>& InOutFi
 
 	TMultiMap<uint32, UFactory*> DummyFilterIndexToFactory;
 
-	// Generate the file types and extensions represented by the selected factories
-	ObjectTools::GenerateFactoryFileExtensions( Factories, FileTypes, AllExtensions, DummyFilterIndexToFactory );
+	//Append the interchange supported translator formats for this object
+	if (UInterchangeManager::IsInterchangeImportEnabled())
+	{
+		//Get the extension interchange can translate for this object
+		TArray<FString> TranslatorFormats = UInterchangeManager::GetInterchangeManager().GetSupportedFormatsForObject(Obj);
+		ObjectTools::AppendFormatsFileExtensions(TranslatorFormats, FileTypes, AllExtensions, DummyFilterIndexToFactory);
+	}
+
+	//If this object was not import with interchange add the legacy formats
+	if(AllExtensions.IsEmpty())
+	{
+		// Generate the file types and extensions represented by the selected factories
+		ObjectTools::GenerateFactoryFileExtensions(Factories, FileTypes, AllExtensions, DummyFilterIndexToFactory);
+	}
 
-	FileTypes = FString::Printf(TEXT("All Files (%s)|%s|%s"),*AllExtensions,*AllExtensions,*FileTypes);
+	
 
 	FString DefaultFolder;
 	FString DefaultFile;
@@ -765,8 +777,21 @@ void FReimportManager::GetNewReimportPath(UObject* Obj, TArray<FString>& InOutFi
 	{
 		DefaultFolder = FPaths::GetPath(ExistingPaths[0]);
 		DefaultFile = FPaths::GetCleanFilename(ExistingPaths[0]);
+		//Make sure we have at least the existing asset source files path extension.
+		//If an asset is import with legacy importer and we switch importer to use interchange (or we go back to legacy)
+		for (const FString& ExistingPath : ExistingPaths)
+		{
+			const FString& Extension = FPaths::GetExtension(ExistingPath);
+			if (!AllExtensions.Contains(Extension))
+			{
+				AllExtensions.Append(TEXT(";*.") + Extension);
+				FileTypes.Append(FString::Printf(TEXT("|Asset file source type (*.%s)|*.%s"), *Extension, *Extension));
+			}
+		}
 	}
 
+	FileTypes = FString::Printf(TEXT("All Files (%s)|%s|%s"), *AllExtensions, *AllExtensions, *FileTypes);
+
 	// Prompt the user for the filenames
 	TArray<FString> OpenFilenames;
 	IDesktopPlatform* DesktopPlatform = FDesktopPlatformModule::Get();
diff --git a/Engine/Source/Editor/UnrealEd/Private/Factories/EditorFactories.cpp b/Engine/Source/Editor/UnrealEd/Private/Factories/EditorFactories.cpp
--- a/Engine/Source/Editor/UnrealEd/Private/Factories/EditorFactories.cpp
+++ b/Engine/Source/Editor/UnrealEd/Private/Factories/EditorFactories.cpp
@@ -6317,8 +6317,10 @@ bool UReimportFbxSkeletalMeshFactory::CanReimport( UObject* Obj, TArray<FString>
 				//This skeletal mesh was import with a scene import, we cannot reimport it here
 				return false;
 			}
-			else if (FPaths::GetExtension(AssetImportData->GetFirstFilename()) == TEXT("abc"))
+			else if (FPaths::GetExtension(AssetImportData->GetFirstFilename()) != TEXT("fbx")
+				  && FPaths::GetExtension(AssetImportData->GetFirstFilename()) != TEXT("obj"))
 			{
+				//Fbx factory only support fbx and obj files
 				return false;
 			}
 			AssetImportData->ExtractFilenames(OutFilenames);
@@ -6793,7 +6795,8 @@ bool UReimportFbxAnimSequenceFactory::CanReimport( UObject* Obj, TArray<FString>
 				//This mesh was import with a scene import, we cannot reimport it
 				return false;
 			}
-			else if (FPaths::GetExtension(AnimSequence->AssetImportData->GetFirstFilename()) == TEXT("abc"))
+			else if (FPaths::GetExtension(AnimSequence->AssetImportData->GetFirstFilename()) != TEXT("fbx")
+				&& FPaths::GetExtension(AnimSequence->AssetImportData->GetFirstFilename()) != TEXT("obj"))
 			{
 				return false;
 			}
diff --git a/Engine/Source/Runtime/Interchange/Core/Public/InterchangeFactoryBase.h b/Engine/Source/Runtime/Interchange/Core/Public/InterchangeFactoryBase.h
--- a/Engine/Source/Runtime/Interchange/Core/Public/InterchangeFactoryBase.h
+++ b/Engine/Source/Runtime/Interchange/Core/Public/InterchangeFactoryBase.h
@@ -17,6 +17,17 @@ class UInterchangePipelineBase;
 class UInterchangeSourceData;
 class UInterchangeTranslatorBase;
 
+UENUM(BlueprintType)
+enum class EInterchangeFactoryAssetType : uint8
+{
+	None = 0,
+	Textures,
+	Materials,
+	Meshes,
+	Animations,
+	Physics
+};
+
 UCLASS(BlueprintType, Blueprintable, Abstract)
 class INTERCHANGECORE_API UInterchangeFactoryBase : public UObject
 {
@@ -32,6 +43,11 @@ public:
 		return nullptr;
 	}
 
+	/**
+	 * return the asset type this factory can create.
+	 */
+	virtual EInterchangeFactoryAssetType GetFactoryAssetType() { return EInterchangeFactoryAssetType::None; }
+
 	/**
 	 * Parameters to pass to CreateAsset function
 	 */
diff --git a/Engine/Source/Runtime/Interchange/Engine/Private/InterchangeManager.cpp b/Engine/Source/Runtime/Interchange/Engine/Private/InterchangeManager.cpp
--- a/Engine/Source/Runtime/Interchange/Engine/Private/InterchangeManager.cpp
+++ b/Engine/Source/Runtime/Interchange/Engine/Private/InterchangeManager.cpp
@@ -758,6 +758,52 @@ TArray<FString> UInterchangeManager::GetSupportedAssetTypeFormats(const EInterch
 	return FileExtensions;
 }
 
+TArray<FString> UInterchangeManager::GetSupportedFormatsForObject(const UObject* Object) const
+{
+	TArray<FString> FileExtensions;
+	if (!IsInterchangeImportEnabled())
+	{
+		return FileExtensions;
+	}
+
+	const UClass* RegisteredFactoryClass = GetRegisteredFactoryClass(Object->GetClass());
+	if (!RegisteredFactoryClass)
+	{
+		return FileExtensions;
+	}
+
+	UInterchangeFactoryBase* Factory = RegisteredFactoryClass->GetDefaultObject<UInterchangeFactoryBase>();
+	TArray<FString> TempFilenames;
+	//GetSourceFilenames verify we have a valid UInterchangeAssetImportData for this Object
+	//This ensure we do not allow re-import
+	if (!Factory->GetSourceFilenames(Object, TempFilenames))
+	{
+		return FileExtensions;
+	}
+
+	switch (Factory->GetFactoryAssetType())
+	{
+	case EInterchangeFactoryAssetType::Animations:
+		FileExtensions = GetSupportedAssetTypeFormats(EInterchangeTranslatorAssetType::Animations);
+		break;
+	case EInterchangeFactoryAssetType::Materials:
+		FileExtensions = GetSupportedAssetTypeFormats(EInterchangeTranslatorAssetType::Materials);
+		break;
+	case EInterchangeFactoryAssetType::Meshes:
+	case EInterchangeFactoryAssetType::Physics:
+		FileExtensions = GetSupportedAssetTypeFormats(EInterchangeTranslatorAssetType::Meshes);
+		break;
+	case EInterchangeFactoryAssetType::Textures:
+		FileExtensions = GetSupportedAssetTypeFormats(EInterchangeTranslatorAssetType::Textures);
+		break;
+	case EInterchangeFactoryAssetType::None: //Actor factories return None
+		FileExtensions = GetSupportedFormats(EInterchangeTranslatorType::Actors);
+		break;
+	}
+
+	return FileExtensions;
+}
+
 bool UInterchangeManager::CanTranslateSourceData(const UInterchangeSourceData* SourceData) const
 {
 	if (!IsInterchangeImportEnabled())
diff --git a/Engine/Source/Runtime/Interchange/Engine/Public/InterchangeManager.h b/Engine/Source/Runtime/Interchange/Engine/Public/InterchangeManager.h
--- a/Engine/Source/Runtime/Interchange/Engine/Public/InterchangeManager.h
+++ b/Engine/Source/Runtime/Interchange/Engine/Public/InterchangeManager.h
@@ -378,6 +378,11 @@ public:
 	 */
 	TArray<FString> GetSupportedAssetTypeFormats(const EInterchangeTranslatorAssetType ForTranslatorAssetType) const;
 
+	/**
+	 * Returns the list of supported formats for a given Object.
+	 */
+	TArray<FString> GetSupportedFormatsForObject(const UObject* Object) const;
+
 	/**
 	 * Look if there is a registered translator for this source data.
 	 * This allow us to by pass the original asset tools system to import supported asset.
