Typed Element Copy Paste

diff --git a/Engine/Source/Editor/LevelEditor/Private/LevelEditorActions.cpp b/Engine/Source/Editor/LevelEditor/Private/LevelEditorActions.cpp
--- a/Engine/Source/Editor/LevelEditor/Private/LevelEditorActions.cpp
+++ b/Engine/Source/Editor/LevelEditor/Private/LevelEditorActions.cpp
@@ -45,6 +45,7 @@
 #include "WorldPartition/IWorldPartitionEditorModule.h"
 #include "WorldBrowserModule.h"
 
+#include "Elements/Framework/TypedElementCommonActions.h"
 #include "Elements/Framework/TypedElementSelectionSet.h"
 #include "Elements/Interfaces/TypedElementObjectInterface.h"
 #include "Subsystems/EditorElementSubsystem.h"
@@ -1750,7 +1751,20 @@ bool FLevelEditorActionCallbacks::Cut_CanExecute()
 	}
 
 	bool bCanCut = false;
-	if (GEditor->GetSelectedComponentCount() > 0)
+	if (TypedElementCommonActionsUtils::IsElementCopyAndPasteEnabled())
+	{
+		static const FName NAME_LevelEditor = "LevelEditor";
+		if (TSharedPtr<ILevelEditor> LevelEditor = FModuleManager::GetModuleChecked<FLevelEditorModule>(NAME_LevelEditor).GetLevelEditorInstance().Pin())
+		{
+			const UTypedElementSelectionSet* SelectionSet = LevelEditor->GetElementSelectionSet();
+			SelectionSet->ForEachSelectedElement<ITypedElementWorldInterface>([&bCanCut](const TTypedElement<ITypedElementWorldInterface>& InWorldElement)
+				{
+					bCanCut |= InWorldElement.CanCopyElement();
+					return !bCanCut;
+				});
+		}
+	}
+	else if (GEditor->GetSelectedComponentCount() > 0)
 	{
 		// Make sure the components can be copied and deleted
 		TArray<UActorComponent*> SelectedComponents;
@@ -1800,7 +1814,20 @@ bool FLevelEditorActionCallbacks::Copy_CanExecute()
 	}
 
 	bool bCanCopy = false;
-	if (GEditor->GetSelectedComponentCount() > 0)
+	if (TypedElementCommonActionsUtils::IsElementCopyAndPasteEnabled())
+	{
+		static const FName NAME_LevelEditor = "LevelEditor";
+		if (TSharedPtr<ILevelEditor> LevelEditor = FModuleManager::GetModuleChecked<FLevelEditorModule>(NAME_LevelEditor).GetLevelEditorInstance().Pin())
+		{
+			const UTypedElementSelectionSet* SelectionSet = LevelEditor->GetElementSelectionSet();
+			SelectionSet->ForEachSelectedElement<ITypedElementWorldInterface>([&bCanCopy](const TTypedElement<ITypedElementWorldInterface>& InWorldElement)
+				{
+					bCanCopy |= InWorldElement.CanCopyElement();
+					return !bCanCopy;
+				});
+		}
+	}
+	else if (GEditor->GetSelectedComponentCount() > 0)
 	{
 		TArray<UActorComponent*> SelectedComponents;
 		for (FSelectionIterator It(GEditor->GetSelectedComponentIterator()); It; ++It)
@@ -1848,20 +1875,31 @@ bool FLevelEditorActionCallbacks::Paste_CanExecute()
 	}
 
 	bool bCanPaste = false;
-	if (GEditor->GetSelectedComponentCount() > 0)
+	if (TypedElementCommonActionsUtils::IsElementCopyAndPasteEnabled())
 	{
-		if(ensureMsgf(GEditor->GetSelectedActorCount() == 1, TEXT("Expected SelectedActorCount to be 1 but was %d"), GEditor->GetSelectedActorCount()))
-		{
-			auto SelectedActor = CastChecked<AActor>(*GEditor->GetSelectedActorIterator());
-			bCanPaste = FComponentEditorUtils::CanPasteComponents(SelectedActor->GetRootComponent());
-		}
+		// Todo Copy and Paste find the right logic for a extensible can paste
+		// but for now just set it to true
+		bCanPaste = true;
 	}
-	else
+
+	// Legacy style copy and paste format
+	if (!bCanPaste)
 	{
-		UWorld* World = GetWorld();
-		if (World)
+		if (GEditor->GetSelectedComponentCount() > 0)
+		{
+			if(ensureMsgf(GEditor->GetSelectedActorCount() == 1, TEXT("Expected SelectedActorCount to be 1 but was %d"), GEditor->GetSelectedActorCount()))
+			{
+				auto SelectedActor = CastChecked<AActor>(*GEditor->GetSelectedActorIterator());
+				bCanPaste = FComponentEditorUtils::CanPasteComponents(SelectedActor->GetRootComponent());
+			}
+		}
+		else
 		{
-			bCanPaste = GUnrealEd->CanPasteSelectedActorsFromClipboard(World);
+			UWorld* World = GetWorld();
+			if (World)
+			{
+				bCanPaste = GUnrealEd->CanPasteSelectedActorsFromClipboard(World);
+			}
 		}
 	}
 
diff --git a/Engine/Source/Editor/UnrealEd/Private/EditorObject.cpp b/Engine/Source/Editor/UnrealEd/Private/EditorObject.cpp
--- a/Engine/Source/Editor/UnrealEd/Private/EditorObject.cpp
+++ b/Engine/Source/Editor/UnrealEd/Private/EditorObject.cpp
@@ -14,6 +14,7 @@
 #include "UObject/UnrealType.h"
 #include "UObject/PropertyPortFlags.h"
 #include "Serialization/ArchiveReplaceObjectRef.h"
+#include "Misc/AsciiSet.h"
 #include "Misc/PackageName.h"
 #include "Components/ActorComponent.h"
 #include "GameFramework/Actor.h"
@@ -182,6 +183,27 @@ static void RemapProperty(FProperty* Property, int32 Index, const TMap<AActor*,
 //	ImportProperties
 //
 
+bool HandleLineNumber(const TCHAR* Str)
+{
+	int32 NewLineNumber;
+	if (FParse::Value(Str, TEXT("linenumber="), NewLineNumber))
+	{
+		if (ContextSupplier != nullptr)
+		{
+			ContextSupplier->CurrentLine = NewLineNumber;
+		}
+
+		return true;
+	}
+
+	return false;
+}
+
+static bool IsEndOfProperties(const TCHAR* Str, int32 Depth)
+{
+	return GetEND(&Str, TEXT("Actor")) || GetEND(&Str, TEXT("DefaultProperties")) || GetEND(&Str, TEXT("structdefaultproperties")) || (GetEND(&Str, TEXT("Object")) && Depth);
+}
+
 /**
  * Parse and import text as property values for the object specified.  This function should never be called directly - use ImportObjectProperties instead.
  * 
@@ -283,13 +305,9 @@ static const TCHAR* ImportProperties(
 
 		const TCHAR* Str = *StrLine;
 
-		int32 NewLineNumber;
-		if( FParse::Value( Str, TEXT("linenumber="), NewLineNumber ) )
+		if (HandleLineNumber(Str))
 		{
-			if ( ContextSupplier != NULL )
-			{
-				ContextSupplier->CurrentLine = NewLineNumber;
-			}
+			continue;
 		}
 		else if( GetBEGIN(&Str,TEXT("Brush")) && ObjectStruct->IsChildOf(ABrush::StaticClass()) )
 		{
@@ -646,7 +664,7 @@ static const TCHAR* ImportProperties(
 
 			SubobjectOuter->ImportCustomProperties(Str, Warn);
 		}
-		else if( GetEND(&Str,TEXT("Actor")) || GetEND(&Str,TEXT("DefaultProperties")) || GetEND(&Str,TEXT("structdefaultproperties")) || (GetEND(&Str,TEXT("Object")) && Depth) )
+		else if (IsEndOfProperties(Str, Depth))
 		{
 			// End of properties.
 			break;
@@ -787,8 +805,6 @@ const TCHAR* ImportObjectProperties( FImportObjectParams& InParams )
 	return NewSourceText;
 }
 
-
-
 	
 /**
  * Parse and import text as property values for the object specified.
@@ -842,3 +858,799 @@ const TCHAR* ImportObjectProperties(
 	return ImportObjectProperties( Params );
 }
 
+
+namespace EditorUtilities
+{
+// Multi step object import
+
+/**
+ * Parse and import text as subobjects for the object specified but doesn't serialize the properties.
+ * 
+ * @param	ObjectStruct				The struct for the data we're importing
+ * @param	SourceText					View on buffer containing the values that should be parsed and imported
+ * @param	SubobjectRoot				When dealing with nested subobjects, corresponds to the top-most outer that
+ *										is not a subobject/template
+ * @param	SubobjectOuter				The outer to use for creating subobjects/components. NULL when importing structdefaultproperties
+ * @param	Warn						Output device to use for log messages
+ * @param	Depth						Current nesting level
+ * @param	InstanceGraph				Contains the mappings of instanced objects and components to their templates
+ * @param	ObjectRemapper				a map of existing Object to new instances, used to replace internal references when a number of actors are copy+pasted
+ *
+ * @return	NULL if the default values couldn't be imported
+ */
+static const TCHAR* ImportCreateSubObjectsStep(
+	FStringView SourceText,
+	UStruct* ObjectStruct,
+	UObject* SubobjectRoot,
+	UObject* SubobjectOuter,
+	FFeedbackContext* Warn,
+	int32 Depth,
+	FObjectInstancingGraph& InstanceGraph,
+	TMap<FString, UObject*>* ObjectRemapper
+	)
+{
+	check(!GIsUCCMakeStandaloneHeaderGenerator);
+	check(ObjectStruct!=nullptr);
+
+	if (SourceText.IsEmpty())
+	{
+		return nullptr;
+	}
+
+	// Cannot create subobjects when importing struct defaults, or if SubobjectOuter (used as the Outer for any subobject declarations encountered) is NULL
+	bool bSubObjectsAllowed = !ObjectStruct->IsA(UScriptStruct::StaticClass()) && SubobjectOuter != nullptr;
+
+	// true when DestData corresponds to a subobject in a class default object
+	bool bSubObject = false;
+
+	UClass* ComponentOwnerClass = nullptr;
+
+	if ( bSubObjectsAllowed )
+	{
+		bSubObject = SubobjectRoot != nullptr && SubobjectRoot->HasAnyFlags(RF_ClassDefaultObject);
+		if ( SubobjectRoot == nullptr )
+		{
+			SubobjectRoot = SubobjectOuter;
+		}
+
+		ComponentOwnerClass = SubobjectOuter != nullptr
+			? SubobjectOuter->IsA(UClass::StaticClass())
+				? CastChecked<UClass>(SubobjectOuter)
+				: SubobjectOuter->GetClass()
+			: nullptr;
+	}
+
+	const TCHAR* CurrentSourceText = SourceText.GetData();
+	const TCHAR* EndOfSourceText = SourceText.GetData() + SourceText.Len();
+
+	FString StrLine;
+
+	// Parse all objects stored in the actor.
+	// Build list of all text properties.
+	bool ImportedBrush = 0;
+	int32 LinesConsumed = 0;
+	while (CurrentSourceText < EndOfSourceText)
+	{
+		if (!FParse::LineExtended(&CurrentSourceText, StrLine, LinesConsumed, true))
+		{
+			break;
+		}
+
+		// remove extra whitespace and optional semicolon from the end of the line
+		{
+			int32 Length = StrLine.Len();
+			while ( Length > 0 &&
+					(StrLine[Length - 1] == TCHAR(';') || StrLine[Length - 1] == TCHAR(' ') || StrLine[Length - 1] == TCHAR('\t')) )
+			{
+				Length--;
+			}
+			if (Length != StrLine.Len())
+			{
+				StrLine.LeftInline(Length, false);
+			}
+		}
+
+		if (ContextSupplier)
+		{
+			ContextSupplier->CurrentLine += LinesConsumed;
+		}
+		if (StrLine.Len() == 0)
+		{
+			continue;
+		}
+
+		const TCHAR* Str = *StrLine;
+
+		if (HandleLineNumber(Str))
+		{
+			continue;
+		}
+		else if(GetBEGIN(&Str,TEXT("Brush")))
+		{
+			Warn->Logf(ELogVerbosity::Error, TEXT("BEGIN Brush: Brushes are not supported in the MultiStepsImport use ImportObjectProperties to import these"));
+			return nullptr;
+		}
+		else if (GetBEGIN(&Str, TEXT("Foliage")))
+		{
+			Warn->Logf(ELogVerbosity::Error, TEXT("BEGIN Foliage: Foliage is not supported in the MultiStepsImport use ImportObjectProperties to import these"));
+			return nullptr;
+		}
+		else if(GetBEGIN(&Str,TEXT("Object")))
+		{
+			// If SubobjectOuter is NULL, we are importing defaults for a UScriptStruct's defaultproperties block
+			if ( !bSubObjectsAllowed )
+			{
+				Warn->Logf(ELogVerbosity::Error, TEXT("BEGIN OBJECT: Subobjects are not allowed in this context"));
+				return nullptr;
+			}
+
+			// Parse subobject default properties.
+			// Note: default properties subobjects have compiled class as their Outer (used for localization).
+			UClass*	TemplateClass = nullptr;
+			bool bInvalidClass = false;
+			ParseObject<UClass>(Str, TEXT("Class="), TemplateClass, nullptr, &bInvalidClass);
+			
+			if (bInvalidClass)
+			{
+				Warn->Logf(ELogVerbosity::Error,TEXT("BEGIN OBJECT: Invalid class specified: %s"), *StrLine);
+				return nullptr;
+			}
+
+			// parse the name of the template
+			FName	TemplateName = NAME_None;
+			FParse::Value(Str,TEXT("Name="),TemplateName);
+			if(TemplateName == NAME_None)
+			{
+				Warn->Logf(ELogVerbosity::Error,TEXT("BEGIN OBJECT: Must specify valid name for subobject/component: %s"), *StrLine);
+				return nullptr;
+			}
+
+			FString ExportedObjectFullName;
+			FParse::Value(Str, TEXT("ExportPath="), ExportedObjectFullName);
+
+			// points to the parent class's template subobject/component, if we are overriding a subobject/component declared in our parent class
+			UObject* BaseTemplate = nullptr;
+			bool bRedefiningSubobject = false;
+			if(!TemplateClass)
+			{
+				// next, verify that a template actually exists in the parent class
+				UClass* ParentClass = ComponentOwnerClass->GetSuperClass();
+				check(ParentClass);
+
+				UObject* ParentCDO = ParentClass->GetDefaultObject();
+				check(ParentCDO);
+
+				BaseTemplate = StaticFindObjectFast(UObject::StaticClass(), SubobjectOuter, TemplateName);
+				bRedefiningSubobject = (BaseTemplate != nullptr);
+
+				if (BaseTemplate == nullptr)
+				{
+					BaseTemplate = StaticFindObjectFast(UObject::StaticClass(), ParentCDO, TemplateName);
+				}
+				
+				if ( BaseTemplate == nullptr )
+				{
+					// wasn't found
+					Warn->Logf(ELogVerbosity::Error, TEXT("BEGIN OBJECT: No base template named %s found in parent class %s: %s"), *TemplateName.ToString(), *ParentClass->GetName(), *StrLine);
+					return nullptr;
+				}
+
+				TemplateClass = BaseTemplate->GetClass();
+			}
+
+			// because the outer won't be a default object
+
+			checkSlow(TemplateClass != nullptr);
+			if (bRedefiningSubobject)
+			{
+				// Simply recreate any missing sub object
+				FMultiStepsImportObjectParams Params;
+				Params.SourceText = FStringView(CurrentSourceText, SourceText.Len() - (CurrentSourceText - SourceText.GetData()));
+				Params.ObjectStruct = TemplateClass;
+				Params.SubobjectRoot = SubobjectRoot;
+				Params.SubobjectOuter = BaseTemplate;
+				Params.Warn = Warn;
+				Params.Depth = Depth + 1;
+				Params.LineNumber = ContextSupplier ? ContextSupplier->CurrentLine : 0;
+				Params.InInstanceGraph = &InstanceGraph;
+				Params.ObjectRemapper = ObjectRemapper;
+
+				CurrentSourceText = ImportCreateObjectsStep(Params);
+				
+				if (!ExportedObjectFullName.IsEmpty() && ObjectRemapper)
+				{
+					ObjectRemapper->Add(MoveTemp(ExportedObjectFullName), BaseTemplate);
+				}
+			}
+			else 
+			{
+				UObject* Archetype = nullptr;
+				UObject* ComponentTemplate = nullptr;
+
+				// Since we are changing the class we can't use the Archetype,
+				// however that is fine since we will have been editing the CDO anyways
+				if (!FBlueprintEditorUtils::IsAnonymousBlueprintClass(SubobjectOuter->GetClass()))
+				{
+					// if an archetype was specified in the Begin Object block, use that as the template for the ConstructObject call.
+					FString ArchetypeName;
+					if (FParse::Value(Str, TEXT("Archetype="), ArchetypeName))
+					{
+						// if given a name, break it up along the ' so separate the class from the name
+						FString ObjectClass;
+						FString ArchetypePath;
+						if ( FPackageName::ParseExportTextPath(ArchetypeName, &ObjectClass, &ArchetypePath) )
+						{
+							// find the class
+							check(FPackageName::IsValidObjectPath(ObjectClass));
+							UClass* ArchetypeClass = (UClass*)StaticFindObject(UClass::StaticClass(), nullptr, *ObjectClass);
+							if (ArchetypeClass)
+							{
+								ArchetypePath = ArchetypePath.TrimQuotes();
+								// if we had the class, find the archetype
+								if (!FPackageName::IsShortPackageName(ArchetypePath))
+								{
+									Archetype = StaticFindObject(ArchetypeClass, nullptr, *ArchetypePath);
+								}
+								else
+								{
+									Archetype = StaticFindFirstObject(ArchetypeClass, *ArchetypePath, EFindFirstObjectOptions::NativeFirst | EFindFirstObjectOptions::EnsureIfAmbiguous);
+								}
+							}
+						}
+					}
+				}
+
+				if (SubobjectOuter->HasAnyFlags(RF_ClassDefaultObject))
+				{
+					if (!Archetype) // if an archetype was specified explicitly, we will stick with that
+					{
+						Archetype = ComponentOwnerClass->GetDefaultSubobjectByName(TemplateName);
+						if(Archetype)
+						{
+							if (!BaseTemplate)
+							{
+								// BaseTemplate should only be NULL if the Begin Object line specified a class
+								Warn->Logf(ELogVerbosity::Error, TEXT("BEGIN OBJECT: The component name %s is already used (if you want to override the component, don't specify a class): %s"), *TemplateName.ToString(), *StrLine);
+								return nullptr;
+							}
+
+							// the component currently in the component template map and the base template should be the same
+							checkf(Archetype==BaseTemplate,TEXT("OverrideComponent: '%s'   BaseTemplate: '%s'"), *Archetype->GetFullName(), *BaseTemplate->GetFullName());
+						}
+					}
+				}
+				else // handle the non-template case (subobjects and non-template components)
+				{
+					ComponentTemplate = FindObject<UObject>(SubobjectOuter, *TemplateName.ToString());
+					if (ComponentTemplate != nullptr)
+					{
+						// if we're overriding a subobject declared in a parent class, we should already have an object with that name that
+						// was instanced when ComponentOwnerClass's CDO was initialized; if so, it's archetype should be the BaseTemplate.  If it
+						// isn't, then there are two unrelated subobject definitions using the same name.
+						if (ComponentTemplate->GetArchetype() != BaseTemplate)
+						{
+						}
+						else if (BaseTemplate == nullptr)
+						{
+							// BaseTemplate should only be NULL if the Begin Object line specified a class
+							Warn->Logf(ELogVerbosity::Error, TEXT("BEGIN OBJECT: A subobject named %s is already declared in a parent class.  If you intended to override that subobject, don't specify a class in the derived subobject definition: %s"), *TemplateName.ToString(), *StrLine);
+							return nullptr;
+						}
+					}
+				}
+
+				// Propagate object flags to the sub object.
+				EObjectFlags NewFlags = SubobjectOuter->GetMaskedFlags( RF_PropagateToSubObjects );
+
+				if (!Archetype) // no override and we didn't find one from the class table, so go with the base
+				{
+					Archetype = BaseTemplate;
+				}
+
+				UObject* OldComponent = nullptr;
+				if (ComponentTemplate)
+				{
+					bool bIsOkToReuse = ComponentTemplate->GetClass() == TemplateClass
+						&& ComponentTemplate->GetOuter() == SubobjectOuter
+						&& ComponentTemplate->GetFName() == TemplateName 
+						&& (ComponentTemplate->GetArchetype() == Archetype || !Archetype);
+
+					if (!bIsOkToReuse)
+					{
+						UE_LOG(LogEditorObject, Log, TEXT("Could not reuse component instance %s, name clash?"), *ComponentTemplate->GetFullName());
+						ComponentTemplate->Rename(nullptr, nullptr, REN_DontCreateRedirectors); // just abandon the existing component, we are going to create
+						OldComponent = ComponentTemplate;
+						ComponentTemplate = nullptr;
+					}
+				}
+				
+				if (!ComponentTemplate)
+				{
+					ComponentTemplate = NewObject<UObject>(
+						SubobjectOuter,
+						TemplateClass,
+						TemplateName,
+						NewFlags,
+						Archetype,
+						!!SubobjectOuter,
+						&InstanceGraph
+						);
+
+					if (!ExportedObjectFullName.IsEmpty() && ObjectRemapper)
+					{
+						ObjectRemapper->Add(MoveTemp(ExportedObjectFullName), ComponentTemplate);
+					}
+				}
+				else
+				{
+					// We do not want to set RF_Transactional for construction script created components, so we have to monkey with things here
+					if (NewFlags & RF_Transactional)
+					{
+						UActorComponent* Component = Cast<UActorComponent>(ComponentTemplate);
+						if (Component && Component->IsCreatedByConstructionScript())
+						{
+							NewFlags &= ~RF_Transactional;
+						}
+					}
+					
+					// Ensure DefaultSubojbect flag persists through the clearing of flags
+					if (ComponentTemplate->HasAllFlags(RF_DefaultSubObject))
+					{
+						NewFlags |= RF_DefaultSubObject;
+					}
+					
+					// Make sure desired flags are set - existing object could be pending kill
+					ComponentTemplate->ClearGarbage();
+					ComponentTemplate->ClearFlags(RF_AllFlags);
+					ComponentTemplate->ClearInternalFlags(EInternalObjectFlags::AllFlags);
+					ComponentTemplate->SetFlags(NewFlags);
+				}
+
+				// replace all properties in this subobject outer' class that point to the original subobject with the new subobject
+				TMap<UObject*, UObject*> ReplacementMap;
+				if (Archetype)
+				{
+					checkSlow(ComponentTemplate->GetArchetype() == Archetype);
+					ReplacementMap.Add(Archetype, ComponentTemplate);
+					InstanceGraph.AddNewInstance(ComponentTemplate, Archetype);
+				}
+				if (OldComponent)
+				{
+					ReplacementMap.Add(OldComponent, ComponentTemplate);
+				}
+				FArchiveReplaceObjectRef<UObject> ReplaceAr(SubobjectOuter, ReplacementMap, EArchiveReplaceObjectFlags::IgnoreArchetypeRef);
+
+				// Simply recreate any missing sub object
+				FMultiStepsImportObjectParams Params;
+				Params.SourceText = FStringView(CurrentSourceText, SourceText.Len() - (CurrentSourceText - SourceText.GetData()));
+				Params.ObjectStruct = TemplateClass;
+				Params.SubobjectRoot = SubobjectRoot;
+				Params.SubobjectOuter = ComponentTemplate;
+				Params.Warn = Warn;
+				Params.Depth = Depth + 1;
+				Params.LineNumber = ContextSupplier ? ContextSupplier->CurrentLine : 0;
+				Params.InInstanceGraph = &InstanceGraph;
+				Params.ObjectRemapper = ObjectRemapper;
+
+				if (!ExportedObjectFullName.IsEmpty() && ObjectRemapper)
+				{
+					ObjectRemapper->Add(MoveTemp(ExportedObjectFullName), ComponentTemplate);
+				}
+
+				// Create the subobjects for the subobject
+				CurrentSourceText = ImportCreateObjectsStep(Params);
+			}
+		}
+		else if(IsEndOfProperties(Str, Depth))
+		{
+			// End of properties.
+			break;
+		}
+	}
+
+	return CurrentSourceText;
+}
+
+/**
+ * Parse and import text as property values for the object specified.  This function should never be called directly - use ImportObjectProperties instead.
+ * 
+ * @param	ObjectStruct				The struct for the data we're importing
+ * @param	DestData					The location to import the property values to
+ * @param	SourceText					Text containing the values that should be parsed and imported
+ * @param	SubobjectRoot				When dealing with nested subobjects, corresponds to the top-most outer that
+ *										is not a subobject/template
+ * @param	SubobjectOuter				The outer to use for creating subobjects/components. NULL when importing structdefaultproperties
+ * @param	Warn						Output device to use for log messages
+ * @param	Depth						Current nesting level
+ * @param	InstanceGraph				Contains the mappings of instanced objects and components to their templates
+ * @param	ObjectRemapper				A map of exported names to new instances, used to replace internal references when a number of object are copy+pasted
+ * @param	PropertiesToSkip				A set containing the properties that shouldn't be imported
+ * 
+ * @return	NULL if the default values couldn't be imported
+ */
+static const TCHAR* ImportPropertiesStep(
+	uint8* DestData,
+	FStringView SourceText,
+	UStruct* ObjectStruct,
+	UObject* SubobjectRoot,
+	UObject* SubobjectOuter,
+	FFeedbackContext* Warn,
+	int32 Depth,
+	FObjectInstancingGraph& InstanceGraph,
+	TMap<FString, UObject*>* ObjectRemapper,
+	TSet<FProperty*>* PropertiesToSkip
+	)
+{
+	check(!GIsUCCMakeStandaloneHeaderGenerator);
+	check(ObjectStruct!=nullptr);
+	check(DestData!=nullptr);
+
+	if (SourceText == nullptr)
+	{
+		return nullptr;
+	}
+
+	// Cannot create subobjects when importing struct defaults, or if SubobjectOuter (used as the Outer for any subobject declarations encountered) is NULL
+	bool bSubObjectsAllowed = !ObjectStruct->IsA(UScriptStruct::StaticClass()) && SubobjectOuter != NULL;
+
+	// true when DestData corresponds to a subobject in a class default object
+	bool bSubObject = false;
+
+	UClass* ComponentOwnerClass = nullptr;
+
+	if ( bSubObjectsAllowed )
+	{
+		bSubObject = SubobjectRoot != nullptr && SubobjectRoot->HasAnyFlags(RF_ClassDefaultObject);
+		if ( SubobjectRoot == nullptr )
+		{
+			SubobjectRoot = SubobjectOuter;
+		}
+
+		ComponentOwnerClass = SubobjectOuter != nullptr
+			? SubobjectOuter->IsA(UClass::StaticClass())
+				? CastChecked<UClass>(SubobjectOuter)
+				: SubobjectOuter->GetClass()
+			: nullptr;
+	}
+
+	// The PortFlags to use for all ImportText calls
+	uint32 PortFlags = PPF_Delimited | PPF_CheckReferences;
+	if (GIsImportingT3D)
+	{
+		PortFlags |= PPF_AttemptNonQualifiedSearch;
+	}
+
+	const TCHAR* CurrentSourceText = SourceText.GetData();
+	const TCHAR* EndOfSourceText = SourceText.GetData() + SourceText.Len();
+
+	FString StrLine;
+
+	TArray<FDefinedProperty> DefinedProperties;
+
+	// Parse all objects stored in the actor.
+	// Build list of all text properties.
+	bool ImportedBrush = 0;
+	int32 LinesConsumed = 0;
+	while (CurrentSourceText < EndOfSourceText)
+	{
+		if (FParse::LineExtended(&CurrentSourceText, StrLine, LinesConsumed, true))
+		{
+			// remove extra whitespace and optional semicolon from the end of the line
+			{
+				int32 Length = StrLine.Len();
+				while ( Length > 0 &&
+						(StrLine[Length - 1] == TCHAR(';') || StrLine[Length - 1] == TCHAR(' ') || StrLine[Length - 1] == TCHAR('\t')) )
+				{
+					Length--;
+				}
+				if (Length != StrLine.Len())
+				{
+					StrLine.LeftInline(Length, false);
+				}
+			}
+
+			if ( ContextSupplier != nullptr )
+			{
+				ContextSupplier->CurrentLine += LinesConsumed;
+			}
+			if (StrLine.Len() == 0)
+			{
+				continue;
+			}
+
+			const TCHAR* Str = *StrLine;
+
+			if (HandleLineNumber(Str))
+			{
+				continue;
+			}
+			else if( GetBEGIN(&Str,TEXT("Object")))
+			{
+				// If SubobjectOuter is NULL, we are importing defaults for a UScriptStruct's defaultproperties block
+				if ( !bSubObjectsAllowed )
+				{
+					Warn->Logf(ELogVerbosity::Error, TEXT("BEGIN OBJECT: Subobjects are not allowed in this context"));
+					return NULL;
+				}
+
+				// parse the name of the template
+				FName	TemplateName = NAME_None;
+				FParse::Value(Str,TEXT("Name="),TemplateName);
+				if(TemplateName == NAME_None)
+				{
+					Warn->Logf(ELogVerbosity::Error,TEXT("BEGIN OBJECT: Must specify valid name for subobject/component: %s"), *StrLine);
+					return nullptr;
+				}
+
+
+				if (UObject* BaseTemplate = StaticFindObjectFast(UObject::StaticClass(), SubobjectOuter, TemplateName))
+				{
+					FMultiStepsImportObjectParams Params;
+					Params.DestData = static_cast<uint8*>(static_cast<void*>(BaseTemplate));
+					Params.SourceText = FStringView(CurrentSourceText, SourceText.Len() - (CurrentSourceText - SourceText.GetData()));
+					Params.ObjectStruct = BaseTemplate->GetClass();
+					Params.SubobjectRoot = SubobjectRoot;
+					Params.SubobjectOuter = BaseTemplate;
+					Params.Warn = Warn;
+					Params.Depth = Depth + 1;
+					Params.LineNumber = ContextSupplier ? ContextSupplier->CurrentLine : 0;
+					Params.InInstanceGraph = &InstanceGraph;
+					Params.ObjectRemapper = ObjectRemapper;
+					Params.PropertiesToSkip = PropertiesToSkip;
+
+					ImportObjectsPropertiesStep(Params);
+				}
+			}
+			else if( FParse::Command(&Str,TEXT("CustomProperties")))
+			{
+				check(SubobjectOuter);
+
+				SubobjectOuter->ImportCustomProperties(Str, Warn);
+			}
+			else if (IsEndOfProperties(Str, Depth))
+			{
+				// End of properties.
+				break;
+			}
+			else
+			{
+				// Property.
+				const TCHAR* TextToImport = Str;
+				
+				FString RedirectedProperty;
+				if (ObjectRemapper || PropertiesToSkip)
+				{
+					constexpr FAsciiSet Delimiters("=([.");
+
+					FStringView FullPropertyText(StrLine);
+					FullPropertyText = FullPropertyText.TrimStartAndEnd();
+
+					// find first delimiter
+					if (const TCHAR* EndOfPropertyName = FAsciiSet::FindFirstOrEnd(FullPropertyText.GetData(), Delimiters))
+					{
+						if (PropertiesToSkip)
+						{
+							FStringView PropertyNameView(FullPropertyText.GetData(), EndOfPropertyName - FullPropertyText.GetData());
+
+							if (!PropertyNameView.IsEmpty())
+							{
+								PropertyNameView.TrimEndInline();
+
+								const FName PropertyName(PropertyNameView);
+								FProperty* Property = FindFProperty<FProperty>(ObjectStruct, PropertyName);
+
+								if (Property == nullptr)
+								{
+									// Check for redirects
+									FName NewPropertyName = FProperty::FindRedirectedPropertyName(ObjectStruct, PropertyName);
+
+									if (NewPropertyName != NAME_None)
+									{
+										Property = FindFProperty<FProperty>(ObjectStruct, NewPropertyName);
+									}
+
+									if (!Property)
+									{
+										Property = ObjectStruct->CustomFindProperty(PropertyName);
+									}
+								}
+
+								if (PropertiesToSkip->Contains(Property))
+								{
+									// Do not import this property
+									continue;
+								}
+							}
+						}
+
+						if (ObjectRemapper)
+						{ 
+							// Get the value
+							constexpr FAsciiSet Equal("=");
+							const TCHAR* PtrToEqual = FAsciiSet::FindFirstOrEnd(EndOfPropertyName, Equal);
+
+							if (PtrToEqual && PtrToEqual + 1 < FullPropertyText.GetData() + FullPropertyText.Len())
+							{ 
+								// Select the next character
+								PtrToEqual++;
+								FString PropertyValue(PtrToEqual, FullPropertyText.Len() - (PtrToEqual - FullPropertyText.GetData()));
+								if (UObject* const* PointerToObject = ObjectRemapper->Find(PropertyValue))
+								{
+
+									FString RedirectedFullName = FObjectPropertyBase::GetExportPath(*PointerToObject, nullptr, nullptr, PortFlags | PPF_Delimited);
+
+									FStringView TextBeforeValue(FullPropertyText.GetData(), PtrToEqual - FullPropertyText.GetData());
+									
+									RedirectedProperty.Reserve(RedirectedFullName.Len() + TextBeforeValue.Len());
+									RedirectedProperty.Append(TextBeforeValue);
+									RedirectedProperty.Append(RedirectedFullName);
+									TextToImport = *RedirectedProperty;
+								}
+							}
+						}
+					}
+				}
+
+				FProperty::ImportSingleProperty(TextToImport, DestData, ObjectStruct, SubobjectOuter, PortFlags, Warn, DefinedProperties);
+			}
+		}
+	}
+
+	/**
+	if (ActorRemapper)
+	{
+		for (const auto& DefinedProperty : DefinedProperties)
+		{
+			RemapProperty(DefinedProperty.Property, DefinedProperty.Index, *ActorRemapper, DestData);
+		}
+	}*/
+
+	return CurrentSourceText;
+}
+
+const TCHAR* ImportCreateObjectsStep(FMultiStepsImportObjectParams& InParams)
+{
+	FDefaultPropertiesContextSupplier Supplier;
+	if (InParams.LineNumber != INDEX_NONE)
+	{
+		if (InParams.SubobjectRoot == nullptr)
+		{
+			Supplier.PackageName = InParams.ObjectStruct->GetOwnerClass() ? InParams.ObjectStruct->GetOwnerClass()->GetOutermost()->GetName() : InParams.ObjectStruct->GetOutermost()->GetName();
+			Supplier.ClassName = InParams.ObjectStruct->GetOwnerClass() ? InParams.ObjectStruct->GetOwnerClass()->GetName() : FName(NAME_None).ToString();
+			Supplier.CurrentLine = InParams.LineNumber;
+
+			ContextSupplier = &Supplier; //-V506
+		}
+		else
+		{
+			if (ContextSupplier != nullptr)
+			{
+				ContextSupplier->CurrentLine = InParams.LineNumber;
+			}
+		}
+		InParams.Warn->SetContext(ContextSupplier);
+	}
+
+	FObjectInstancingGraph TempGraph;
+	FObjectInstancingGraph& InstanceGraph = InParams.InInstanceGraph ? *InParams.InInstanceGraph : TempGraph;
+
+	if (InParams.SubobjectRoot && InParams.SubobjectRoot != UObject::StaticClass()->GetDefaultObject())
+	{
+		InstanceGraph.SetDestinationRoot(InParams.SubobjectRoot);
+	}
+
+	// Parse the object properties.
+	const TCHAR* NewSourceText =
+		ImportCreateSubObjectsStep(
+			InParams.SourceText,
+			InParams.ObjectStruct,
+			InParams.SubobjectRoot,
+			InParams.SubobjectOuter,
+			InParams.Warn,
+			InParams.Depth,
+			InstanceGraph,
+			InParams.ObjectRemapper
+		);
+
+	if (InParams.LineNumber != INDEX_NONE)
+	{
+		if (ContextSupplier == &Supplier)
+		{
+			ContextSupplier = nullptr;
+			InParams.Warn->SetContext(nullptr);
+		}
+	}
+
+	return NewSourceText;
+}
+
+
+const TCHAR* ImportObjectsPropertiesStep(FMultiStepsImportObjectParams& InParams)
+{
+	FDefaultPropertiesContextSupplier Supplier;
+	if (InParams.LineNumber != INDEX_NONE)
+	{
+		if (InParams.SubobjectRoot == nullptr)
+		{
+			Supplier.PackageName = InParams.ObjectStruct->GetOwnerClass() ? InParams.ObjectStruct->GetOwnerClass()->GetOutermost()->GetName() : InParams.ObjectStruct->GetOutermost()->GetName();
+			Supplier.ClassName = InParams.ObjectStruct->GetOwnerClass() ? InParams.ObjectStruct->GetOwnerClass()->GetName() : FName(NAME_None).ToString();
+			Supplier.CurrentLine = InParams.LineNumber;
+
+			ContextSupplier = &Supplier; //-V506
+		}
+		else
+		{
+			if (ContextSupplier != nullptr)
+			{
+				ContextSupplier->CurrentLine = InParams.LineNumber;
+			}
+		}
+		InParams.Warn->SetContext(ContextSupplier);
+	}
+
+	if (InParams.bShouldCallEditChange && InParams.SubobjectOuter != nullptr)
+	{
+		InParams.SubobjectOuter->PreEditChange(nullptr);
+	}
+
+	FObjectInstancingGraph TempGraph;
+	FObjectInstancingGraph& InstanceGraph = InParams.InInstanceGraph ? *InParams.InInstanceGraph : TempGraph;
+
+	if (InParams.SubobjectRoot && InParams.SubobjectRoot != UObject::StaticClass()->GetDefaultObject())
+	{
+		InstanceGraph.SetDestinationRoot(InParams.SubobjectRoot);
+	}
+
+	// Parse the object properties.
+	const TCHAR* NewSourceText =
+		ImportPropertiesStep(
+			InParams.DestData,
+			InParams.SourceText,
+			InParams.ObjectStruct,
+			InParams.SubobjectRoot,
+			InParams.SubobjectOuter,
+			InParams.Warn,
+			InParams.Depth,
+			InstanceGraph,
+			InParams.ObjectRemapper,
+			InParams.PropertiesToSkip
+		);
+
+	if (InParams.SubobjectOuter != nullptr)
+	{
+		check(InParams.SubobjectRoot);
+
+		// Update the object properties to point to the newly imported component objects.
+		// Templates inside classes never need to have components instanced.
+		if (!InParams.SubobjectRoot->HasAnyFlags(RF_ClassDefaultObject))
+		{
+			UObject* SubobjectArchetype = InParams.SubobjectOuter->GetArchetype();
+			InParams.ObjectStruct->InstanceSubobjectTemplates(InParams.DestData, SubobjectArchetype, SubobjectArchetype->GetClass(),
+				InParams.SubobjectOuter, &InstanceGraph);
+		}
+
+		if (InParams.bShouldCallEditChange)
+		{
+			// notify the object that it has just been imported
+			InParams.SubobjectOuter->PostEditImport();
+
+			// notify the object that it has been edited
+			InParams.SubobjectOuter->PostEditChange();
+		}
+		InParams.SubobjectRoot->CheckDefaultSubobjects();
+	}
+
+	if (InParams.LineNumber != INDEX_NONE)
+	{
+		if (ContextSupplier == &Supplier)
+		{
+			ContextSupplier = nullptr;
+			InParams.Warn->SetContext(nullptr);
+		}
+	}
+
+	return NewSourceText;
+}
+
+}	// End namespace EditorUtilities
+
diff --git a/Engine/Source/Editor/UnrealEd/Private/Elements/Actor/ActorElementEditorCopyAndPaste.cpp b/Engine/Source/Editor/UnrealEd/Private/Elements/Actor/ActorElementEditorCopyAndPaste.cpp
new file mode 100644
--- /dev/null
+++ b/Engine/Source/Editor/UnrealEd/Private/Elements/Actor/ActorElementEditorCopyAndPaste.cpp
@@ -0,0 +1,438 @@
+// Copyright Epic Games, Inc. All Rights Reserved.
+
+#include "Elements/Actor/ActorElementEditorCopyAndPaste.h"
+
+#include "CoreGlobals.h"
+#include "Editor.h"
+#include "Editor/GroupActor.h"
+#include "Elements/Framework/EngineElementsLibrary.h"
+#include "Engine/BlueprintGeneratedClass.h"
+#include "Engine/World.h"
+#include "Kismet2/BlueprintEditorUtils.h"
+#include "Kismet2/KismetEditorUtilities.h"
+#include "LevelUtils.h"
+#include "Misc/Parse.h"
+#include "Misc/FeedbackContext.h"
+#include "UObject/PropertyPortFlags.h"
+#include "UObject/UnrealType.h"
+#include "UObject/UObjectHash.h"
+
+#define LOCTEXT_NAMESPACE "ActorElementCopyAndPaste"
+
+UActorElementsExporterT3D::UActorElementsExporterT3D(const FObjectInitializer& ObjectInitializer)
+	: Super(ObjectInitializer)
+{
+	SupportedClass = UActorElementsCopy::StaticClass();
+	bText = true;
+	PreferredFormatIndex = 0;
+	FormatExtension.Add(TEXT("COPY"));
+	FormatDescription.Add(TEXT("Unreal world text"));
+}
+
+bool UActorElementsExporterT3D::ExportText(const FExportObjectInnerContext* Context, UObject* Object, const TCHAR* Type,
+                                           FOutputDevice& Out, FFeedbackContext* Warn, uint32 PortFlags /*= 0*/)
+{
+	UActorElementsCopy* ActorElementCopy = Cast<UActorElementsCopy>(Object);
+	if (!ActorElementCopy)
+	{
+		return false;
+	}
+
+	for (AActor* Actor : ActorElementCopy->ActorsToCopy)
+	{
+		if (!Actor->ShouldExport() || Actor->IsA<AGroupActor>())
+		{
+			continue;
+		}
+
+		// Temporarily unbind dynamic delegates so we don't export the bindings.
+		UBlueprintGeneratedClass::UnbindDynamicDelegates(Actor->GetClass(), Actor);
+
+		ULevel* Level = Actor->GetLevel();
+		AActor* ParentActor = Actor->GetAttachParentActor();
+		FName SocketName = Actor->GetAttachParentSocketName();
+		Actor->DetachFromActor(FDetachmentTransformRules::KeepWorldTransform);
+
+		FString ParentActorString = (ParentActor
+			                             ? FString::Printf(TEXT(" ParentActor=%s"), *ParentActor->GetName())
+			                             : TEXT(""));
+		FString SocketNameString = ((ParentActor && SocketName != NAME_None)
+			                            ? FString::Printf(TEXT(" SocketName=%s"), *SocketName.ToString())
+			                            : TEXT(""));
+		FString GroupActor = (Actor->GroupActor
+			                      ? FString::Printf(TEXT(" GroupActor=%s"), *Actor->GroupActor->GetName())
+			                      : TEXT(""));
+		FString GroupFolder = (Actor->GroupActor
+			                       ? FString::Printf(
+				                       TEXT(" GroupFolder=%s"), *Actor->GroupActor->GetFolderPath().ToString())
+			                       : TEXT(""));
+		FString ActorFolderPath = (Level->IsUsingActorFolders()
+			                           ? FString::Printf(
+				                           TEXT(" ActorFolderPath=%s"), *Actor->GetFolderPath().ToString())
+			                           : TEXT(""));
+		FString ActorPath = FString::Printf(
+			TEXT(" ExportPath=%s"),
+			*FObjectPropertyBase::GetExportPath(Actor, nullptr, nullptr, PortFlags | PPF_Delimited));
+		Out.Logf(TEXT("%sBegin Actor Class=%s Name=%s Archetype=%s'%s'%s%s%s%s%s%s") LINE_TERMINATOR,
+		        FCString::Spc(TextIndent), *Actor->GetClass()->GetPathName(), *Actor->GetName(),
+		        *Actor->GetArchetype()->GetClass()->GetPathName(), *Actor->GetArchetype()->GetPathName(),
+		        *ParentActorString, *SocketNameString, *GroupActor, *GroupFolder, *ActorFolderPath, *ActorPath);
+
+		ExportRootScope = Actor;
+		ExportObjectInner(Context, Actor, Out, PortFlags);
+		ExportRootScope = nullptr;
+
+		Out.Logf(TEXT("%sEnd Actor\r\n"), FCString::Spc(TextIndent));
+		Actor->AttachToActor(ParentActor, FAttachmentTransformRules::KeepWorldTransform, SocketName);
+
+		// Restore dynamic delegate bindings.
+		UBlueprintGeneratedClass::BindDynamicDelegates(Actor->GetClass(), Actor);
+	}
+
+	return false;
+}
+
+void FActorElementEditorPasteImporter::Import(FContext& InContext)
+{
+	Context = InContext;
+	FStringView InText = Context.Text;
+
+	if (!Context.World || InText.IsEmpty())
+	{
+		return;
+	}
+
+	// Parse text of actors
+	const TCHAR* TextStart = InText.GetData();
+	const TCHAR** CurrentPos = &TextStart;
+	const TCHAR* EndPos = InText.GetData() + InText.Len();
+	FStringView Line;
+
+	FImportedActor* CurrentImportedActor = nullptr;
+
+	const TCHAR* StartOfCurrentActorInner = nullptr;
+	const TCHAR* EndOfPreviousLine = *CurrentPos;
+
+	while (*CurrentPos < EndPos && FParse::Line(CurrentPos, Line))
+	{
+		const TCHAR* LineStart = Line.GetData();
+		const TCHAR** LinePos = &LineStart;
+		if (Line.GetData() + Line.Len() < EndPos)
+		{
+			if (CurrentImportedActor)
+			{
+				if (!StartOfCurrentActorInner)
+				{
+					StartOfCurrentActorInner = Line.GetData();
+				}
+
+				if (GetEND(LinePos, TEXT("Actor")))
+				{
+					CurrentImportedActor->InnerText = FStringView(StartOfCurrentActorInner,
+					                                              uint32(EndOfPreviousLine - StartOfCurrentActorInner));
+					CurrentImportedActor = nullptr;
+					StartOfCurrentActorInner = nullptr;
+				}
+			}
+			else if (GetBEGIN(LinePos, TEXT("Actor")))
+			{
+				CurrentImportedActor = &(ImportedActors.AddDefaulted_GetRef());
+				ParseObject<UClass>(*LinePos, TEXT("CLASS="), CurrentImportedActor->Class, nullptr);
+				FParse::Value(*LinePos, TEXT("Name="), CurrentImportedActor->Name);
+				FParse::Value(*LinePos, TEXT("Archetype="), CurrentImportedActor->Archetype);
+				FParse::Value(*LinePos, TEXT("ParentActor="), CurrentImportedActor->ParentActor);
+				FParse::Value(*LinePos, TEXT("SocketName="), CurrentImportedActor->SocketNameString);
+				FParse::Value(*LinePos, TEXT("GroupActor="), CurrentImportedActor->GroupActor);
+				FParse::Value(*LinePos, TEXT("GroupFolder="), CurrentImportedActor->GroupFolder);
+				FParse::Value(*LinePos, TEXT("ActorFolderPath="), CurrentImportedActor->ActorFolderPath);
+				FParse::Value(*LinePos, TEXT("ExportPath="), CurrentImportedActor->ExportedActorFullName);
+			}
+		}
+
+		EndOfPreviousLine = Line.GetData() + Line.Len();
+	}
+
+	if (CurrentImportedActor)
+	{
+		ImportedActors.Pop();
+	}
+
+	bPostImportProcessingNeeded = true;
+}
+
+TArray<FTypedElementHandle> FActorElementEditorPasteImporter::GetImportedElements()
+{
+	if (bPostImportProcessingNeeded)
+	{
+		PostImportProcess();
+	}
+	TArray<FTypedElementHandle> Handles;
+	Handles.Reserve(ImportedActors.Num());
+
+	for (FImportedActor& ImportedActor : ImportedActors)
+	{
+		AActor* Actor = ImportedActor.ActorPtr;
+		if (IsValid(Actor))
+		{
+			if (FTypedElementHandle Handle = UEngineElementsLibrary::AcquireEditorActorElementHandle(Actor))
+			{
+				Handles.Add(MoveTemp(Handle));
+			}
+		}
+	}
+
+	for (TPair<FString, AGroupActor*>& GroupPair : NewGroups)
+	{
+		if (FTypedElementHandle Handle = UEngineElementsLibrary::AcquireEditorActorElementHandle(GroupPair.Value))
+		{
+			Handles.Add(MoveTemp(Handle));
+		}
+	}
+
+	return Handles;
+}
+
+void FActorElementEditorPasteImporter::PostImportProcess()
+{
+	TMap<FString, UObject*> InExportedPathToObjectCreated;
+	// CreateObjects
+	{
+		UWorld* World = Context.World;
+		for (FImportedActor& ImportedActor : ImportedActors)
+		{
+			AActor* Found = nullptr;
+			if (ImportedActor.Name != NAME_None)
+			{
+				// look in the current level for the same named actor
+				Found = FindObjectFast<AActor>(World->GetCurrentLevel(), ImportedActor.Name);
+			}
+
+			// Make sure this name is unique. We need to do this upfront because we also want to potentially create the Associated BP class using the same name.
+			bool bNeedGloballyUniqueName = (World->GetCurrentLevel()->IsUsingExternalActors() && CastChecked<AActor>(
+				ImportedActor.Class->GetDefaultObject())->SupportsExternalPackaging()) || Found;
+			FName ActorUniqueName = FActorSpawnUtils::MakeUniqueActorName(
+				World->GetCurrentLevel(), ImportedActor.Class, FActorSpawnUtils::GetBaseName(ImportedActor.Name),
+				bNeedGloballyUniqueName);
+
+
+			AActor* Archetype = nullptr;
+			{
+				FString ObjectClass;
+				FString ObjectPath;
+
+				if (FPackageName::ParseExportTextPath(ImportedActor.Archetype, &ObjectClass, &ObjectPath))
+				{
+					// find the class
+					UClass* ArchetypeClass = UClass::TryFindTypeSlow<UClass>(
+						ObjectClass, EFindFirstObjectOptions::EnsureIfAmbiguous);
+					if (ArchetypeClass)
+					{
+						if (ArchetypeClass->IsChildOf(AActor::StaticClass()))
+						{
+							// if we had the class, find the archetype
+							Archetype = Cast<AActor>(StaticFindObject(ArchetypeClass, nullptr, *ObjectPath));
+						}
+						else
+						{
+							if (GWarn)
+							{
+								GWarn->Logf(ELogVerbosity::Warning,
+								            TEXT(
+									            "Invalid archetype specified in subobject definition '%s': %s is not a child of Actor"),
+								            *(ImportedActor.Name.ToString()), *ObjectClass);
+							}
+						}
+					}
+				}
+			}
+
+			// Might not be needed need to validate if we can copy the level blueprint via the selection code
+			if (FBlueprintEditorUtils::IsAnonymousBlueprintClass(ImportedActor.Class))
+			{
+				UBlueprint* NewBP = DuplicateObject(CastChecked<UBlueprint>(ImportedActor.Class->ClassGeneratedBy),
+				                                    World->GetCurrentLevel(),
+				                                    *FString::Printf(TEXT("%s_BPClass"), *ActorUniqueName.ToString()));
+				if (NewBP)
+				{
+					NewBP->ClearFlags(RF_Standalone);
+
+					FKismetEditorUtilities::CompileBlueprint(NewBP, EBlueprintCompileOptions::SkipGarbageCollection);
+
+					ImportedActor.Class = NewBP->GeneratedClass;
+
+					// Since we changed the class we can't use an Archetype,
+					// however that is fine since we will have been editing the CDO anyways
+					Archetype = nullptr;
+				}
+			}
+
+			if (FLevelUtils::IsLevelLocked(World->GetCurrentLevel()))
+			{
+				if (GWarn)
+				{
+					// Todo 
+					GWarn->Log(ELogVerbosity::Warning,
+					           TEXT(
+						           "Import actor: The requested operation could not be completed because the level is locked."));
+				}
+				return;
+			}
+			
+			FActorSpawnParameters SpawnInfo;
+			SpawnInfo.Name = ActorUniqueName;
+			SpawnInfo.Template = Archetype;
+			SpawnInfo.SpawnCollisionHandlingOverride = ESpawnActorCollisionHandlingMethod::AlwaysSpawn;
+			ImportedActor.ActorPtr = World->SpawnActor(ImportedActor.Class, nullptr, nullptr, SpawnInfo);
+
+			// Todo investigate the support of moving some actor(or elements) to another level using elements instead of the old level Factory/Exporter code?
+			const bool bIsMoveToStreamingLevel = false;
+
+			if (ImportedActor.ActorPtr->ShouldImport(ImportedActor.InnerText, bIsMoveToStreamingLevel))
+			{
+				if (!ImportedActor.ExportedActorFullName.IsEmpty())
+				{
+					InExportedPathToObjectCreated.Add(ImportedActor.ExportedActorFullName, ImportedActor.ActorPtr);
+				}
+
+				EditorUtilities::FMultiStepsImportObjectParams Params;
+				Params.DestData = static_cast<uint8*>(static_cast<void*>(ImportedActor.ActorPtr));
+				Params.SourceText = ImportedActor.InnerText;
+				Params.ObjectStruct = ImportedActor.ActorPtr->GetClass();
+				Params.SubobjectRoot = ImportedActor.ActorPtr;
+				Params.SubobjectOuter = ImportedActor.ActorPtr;
+				Params.Warn = GWarn;
+				Params.ObjectRemapper = &InExportedPathToObjectCreated;
+
+				if (!EditorUtilities::ImportCreateObjectsStep(Params))
+				{
+					ImportedActor.ActorPtr->Destroy();
+					ImportedActor.ActorPtr = nullptr;
+				}
+			}
+			else
+			{
+				ImportedActor.ActorPtr->Destroy();
+				ImportedActor.ActorPtr = nullptr;
+			}
+		}
+	}
+	// SerializeObjects
+	{
+		for (FImportedActor& ImportedActor : ImportedActors)
+		{
+			if (!ImportedActor.ActorPtr)
+			{
+				continue;
+			}
+			{
+				TArray<UObject*> Objects;
+				GetObjectsWithOuter(ImportedActor.ActorPtr, Objects);
+
+				for (UObject* Object : Objects)
+				{
+					if (IsValid(Object))
+					{
+						Object->PreEditChange(nullptr);
+					}
+				}
+			}
+
+			ImportedActor.ActorPtr->PreEditChange(nullptr);
+
+			EditorUtilities::FMultiStepsImportObjectParams Params;
+			Params.DestData = static_cast<uint8*>(static_cast<void*>(ImportedActor.ActorPtr));
+			Params.SourceText = ImportedActor.InnerText;
+			Params.ObjectStruct = ImportedActor.ActorPtr->GetClass();
+			Params.SubobjectRoot = ImportedActor.ActorPtr;
+			Params.SubobjectOuter = ImportedActor.ActorPtr;
+			Params.Warn = GWarn;
+			// This won't change the map
+			Params.ObjectRemapper = (&InExportedPathToObjectCreated);
+
+			if (!EditorUtilities::ImportObjectsPropertiesStep(Params))
+			{
+				ImportedActor.ActorPtr->Destroy();
+				ImportedActor.ActorPtr = nullptr;
+				continue;
+			}
+
+			if (!ImportedActor.ActorFolderPath.IsEmpty())
+			{
+				ImportedActor.ActorPtr->SetFolderPath(*ImportedActor.ActorFolderPath);
+			}
+
+			if (!ImportedActor.GroupActor.IsEmpty())
+			{
+				if (AGroupActor** ActorGroup = NewGroups.Find(ImportedActor.GroupActor))
+				{
+					(*ActorGroup)->Add(*ImportedActor.ActorPtr);
+				}
+				else
+				{
+					AGroupActor* SpawnedGroupActor = ImportedActor.ActorPtr->GetWorld()->SpawnActor<AGroupActor>();
+
+					if (!ImportedActor.GroupFolder.IsEmpty())
+					{
+						SpawnedGroupActor->SetFolderPath(*ImportedActor.GroupFolder);
+					}
+
+					SpawnedGroupActor->Add(*ImportedActor.ActorPtr);
+					NewGroups.Add(ImportedActor.GroupActor, SpawnedGroupActor);
+					FActorLabelUtilities::SetActorLabelUnique(SpawnedGroupActor, ImportedActor.GroupActor);
+				}
+			}
+			NameToImportedActors.Add(ImportedActor.Name, ImportedActor.ActorPtr);
+		}
+	}
+
+	// NotifyCreatedObjects
+	{
+		for (FImportedActor& ImportedActor : ImportedActors)
+		{
+			AActor* Actor = ImportedActor.ActorPtr;
+			if (IsValid(Actor))
+			{
+				TArray<UObject*> Objects;
+				GetObjectsWithOuter(Actor, Objects);
+
+				for (UObject* Object : Objects)
+				{
+					if (IsValid(Object))
+					{
+						Object->PostEditImport();
+						Object->PostEditChange();
+					}
+				}
+
+				// Let the actor deal with having been imported, if desired.
+				Actor->PostEditImport();
+
+				// Notify actor its properties have changed.
+				Actor->PostEditChange();
+			}
+		}
+
+		// Do the attachment of actors
+		for (FImportedActor& ImportedActor : ImportedActors)
+		{
+			if (ImportedActor.ActorPtr)
+			{
+				if (AActor** Parent = NameToImportedActors.Find(ImportedActor.ParentActor))
+				{
+					GEditor->ParentActors(*Parent, ImportedActor.ActorPtr, ImportedActor.SocketNameString);
+				}
+			}
+		}
+
+		for (TPair<FString, AGroupActor*>& GroupPair : NewGroups)
+		{
+			GroupPair.Value->CenterGroupLocation();
+			GroupPair.Value->Lock();
+		}
+
+		NameToImportedActors.Empty();
+		NewGroups.Empty();
+	}
+}
+
+#undef LOCTEXT_NAMESPACE
diff --git a/Engine/Source/Editor/UnrealEd/Private/Elements/Actor/ActorElementEditorWorldInterface.cpp b/Engine/Source/Editor/UnrealEd/Private/Elements/Actor/ActorElementEditorWorldInterface.cpp
--- a/Engine/Source/Editor/UnrealEd/Private/Elements/Actor/ActorElementEditorWorldInterface.cpp
+++ b/Engine/Source/Editor/UnrealEd/Private/Elements/Actor/ActorElementEditorWorldInterface.cpp
@@ -1,7 +1,9 @@
 // Copyright Epic Games, Inc. All Rights Reserved.
 
 #include "Elements/Actor/ActorElementEditorWorldInterface.h"
+
 #include "Elements/Actor/ActorElementData.h"
+#include "Elements/Actor/ActorElementEditorCopyAndPaste.h"
 #include "GameFramework/Actor.h"
 
 #include "Elements/Framework/EngineElementsLibrary.h"
@@ -92,7 +94,9 @@ void UActorElementEditorWorldInterface::DuplicateElements(TArrayView<const FType
 
 		TArray<AActor*> NewActors;
 		ABrush::SetSuppressBSPRegeneration(true);
-		GUnrealEd->DuplicateActors(ActorsToDuplicate, NewActors, InWorld->GetCurrentLevel(), InLocationOffset);
+		ULevel* Level = InWorld->GetCurrentLevel();
+		TGuardValue DisablePrompt(Level->bPromptWhenAddingToLevelOutsideBounds, false);
+		GUnrealEd->DuplicateActors(ActorsToDuplicate, NewActors, Level, InLocationOffset);
 		ABrush::SetSuppressBSPRegeneration(false);
 
 		FEditorDelegates::OnDuplicateActorsEnd.Broadcast();
@@ -114,6 +118,32 @@ void UActorElementEditorWorldInterface::DuplicateElements(TArrayView<const FType
 	}
 }
 
+bool UActorElementEditorWorldInterface::CanCopyElement(const FTypedElementHandle& InElementHandle)
+{
+	return true;
+}
+
+void UActorElementEditorWorldInterface::CopyElements(TArrayView<const FTypedElementHandle> InElementHandles, FOutputDevice& Out)
+{
+	TArray<AActor*> Actors = ActorElementDataUtil::GetActorsFromHandles(InElementHandles);
+
+	if (Actors.IsEmpty())
+	{
+		return;
+	}
+	
+	UActorElementsCopy* ActorElementsCopy = NewObject<UActorElementsCopy>();
+	ActorElementsCopy->ActorsToCopy = Actors;
+
+	constexpr int32 Indent = 3;
+	UExporter::ExportToOutputDevice(nullptr, ActorElementsCopy, nullptr, Out, TEXT("copy"), Indent, PPF_DeepCompareInstances);
+}
+
+TSharedPtr<FWorldElementPasteImporter> UActorElementEditorWorldInterface::GetPasteImporter()
+{
+	return MakeShared<FActorElementEditorPasteImporter>();
+}
+
 bool UActorElementEditorWorldInterface::IsElementInConvexVolume(const FTypedElementHandle& Handle, const FConvexVolume& InVolume, bool bMustEncompassEntireElement)
 {
 	if (TTypedElement<ITypedElementHierarchyInterface> HierarchyElement = GetRegistry().GetElement<ITypedElementHierarchyInterface>(Handle))
diff --git a/Engine/Source/Editor/UnrealEd/Private/Elements/Component/ComponentElementEditorCopyAndPaste.cpp b/Engine/Source/Editor/UnrealEd/Private/Elements/Component/ComponentElementEditorCopyAndPaste.cpp
new file mode 100644
--- /dev/null
+++ b/Engine/Source/Editor/UnrealEd/Private/Elements/Component/ComponentElementEditorCopyAndPaste.cpp
@@ -0,0 +1,301 @@
+﻿// Copyright Epic Games, Inc. All Rights Reserved.
+
+#include "Elements/Component/ComponentElementEditorCopyAndPaste.h"
+
+#include "Components/InstancedStaticMeshComponent.h"
+#include "Editor.h"
+#include "Editor/UnrealEdEngine.h"
+#include "Elements/Actor/ActorElementData.h"
+#include "Elements/Component/ComponentElementData.h"
+#include "Elements/Framework/EngineElementsLibrary.h"
+#include "Elements/Framework/TypedElementList.h"
+#include "Elements/SMInstance/SMInstanceElementData.h"
+#include "HAL/PlatformApplicationMisc.h"
+#include "UnrealEdGlobals.h"
+
+namespace UE::ComponentEditorUtils
+{
+	USceneComponent* FindClosestParent(UActorComponent* ChildComponent, const TArray<TObjectPtr<UActorComponent>>& ComponentList)
+	{
+		// Find the most recent parent that is part of the ComponentList
+		if (USceneComponent* ChildAsScene = Cast<USceneComponent>(ChildComponent))
+		{
+			for (USceneComponent* Parent = ChildAsScene->GetAttachParent(); Parent != nullptr; Parent = Parent->GetAttachParent())
+			{
+				if (ComponentList.Contains(Parent))
+				{
+					return Parent;
+				}
+			}
+		}
+		return nullptr;
+	}
+}
+
+namespace UE::ElementCopyPasteUtil
+{
+	
+FString ParseObjectText(FStringView Text, TArray<FTransform>* OutTransforms = nullptr)
+{
+	FString Result;
+	FString FoundStaticMesh;
+	TArray<FTransform> Transforms;
+
+	const TCHAR* Buffer = Text.GetData();
+	FString Line;
+
+	while (FParse::Line(&Buffer, Line))
+	{
+		const TCHAR* Str = GetData(Line);
+		if (GetBEGIN(&Str, TEXT("Object")))
+		{
+			FTransform Transform;
+			bool bHasStaticMesh = false;
+			do
+			{
+				FParse::Line(&Buffer, Line);
+				Str = GetData(Line);
+				FString StaticMesh;
+				if (FParse::Value(Str, TEXT("StaticMesh="), StaticMesh))
+				{
+					if (FoundStaticMesh.IsEmpty())
+					{
+						FoundStaticMesh = MoveTemp(StaticMesh);
+						bHasStaticMesh = true;
+					}
+					// will not paste partially matching selections
+					else if (FoundStaticMesh != StaticMesh)
+					{
+						return Result;
+					}
+				}
+				if (OutTransforms)
+				{
+					FString Value;
+					if (FParse::Value(Str, TEXT("RelativeLocation="), Value, /*bShouldStopOnSeparator*/false))
+					{
+						FVector Location;
+						if (Location.InitFromString(Value))
+						{
+							Transform.SetLocation(Location);
+						}
+					}
+					if (FParse::Value(Str, TEXT("RelativeRotation="), Value))
+					{
+						FRotator Rotation;
+						FParse::Value(Str, TEXT("Pitch="), Rotation.Pitch);
+						FParse::Value(Str, TEXT("Yaw="), Rotation.Yaw);
+						FParse::Value(Str, TEXT("Roll="), Rotation.Roll);
+						Transform.SetRotation(FQuat::MakeFromRotator(Rotation));
+					}
+					if (FParse::Value(Str, TEXT("RelativeScale3D="), Value, /*bShouldStopOnSeparator*/false))
+					{
+						FVector Scale3D;
+						if (Scale3D.InitFromString(Value))
+						{
+							Transform.SetScale3D(Scale3D);
+						}
+					}
+				}
+			}
+			while (!GetEND(&Str, TEXT("Object")));
+			if (bHasStaticMesh && OutTransforms)
+			{
+				Transforms.Add(MoveTemp(Transform));
+			}
+		}
+	}
+
+	OutTransforms->Append(MoveTemp(Transforms));
+	Result = FoundStaticMesh;
+
+	return Result;
+}
+
+bool TryPasteAsInstances(FStringView Text, UActorComponent* InComponent) 
+{
+	if (UInstancedStaticMeshComponent* ISMComponent = Cast<UInstancedStaticMeshComponent>(InComponent))
+	{
+		TArray<FTransform> Transforms;
+		// What static mesh (if any) is uniform among all pasted components
+		FString PastedStaticMeshExportPath = ParseObjectText(Text, &Transforms);
+		if (PastedStaticMeshExportPath.Len())
+		{
+			FString PastedStaticMeshPath = FPackageName::ExportTextPathToObjectPath(PastedStaticMeshExportPath).TrimQuotes();
+			FString ComponentStaticMeshPath = ISMComponent->GetStaticMesh().GetPath();
+			// check if static mesh of ISMC matches static mesh of pasted components
+			if (PastedStaticMeshPath == ComponentStaticMeshPath)
+			{
+				for (const FTransform& Transform : Transforms)
+				{
+					// if so, create new instances on the selected component with the transforms of the pasted components
+					ISMComponent->AddInstance(Transform.GetRelativeTransform(ISMComponent->GetComponentTransform()));
+				}
+				return true;
+			}
+		}
+	}
+	return false;
+}
+
+struct FPotentialSelection
+{
+	AActor* Actor = nullptr;
+	TArray<UActorComponent*> ISMComponents;
+};
+
+FPotentialSelection GetPotentialSelection(FTypedElementListConstPtr InSelection) 
+{
+	FPotentialSelection PotentialSelection;
+	if (!InSelection->IsValidIndex(0))
+	{
+		return PotentialSelection;
+	}
+	FTypedElementHandle FirstSelectedElement = InSelection->GetElementHandleAt(0);
+	// ad hoc pattern matching the element type
+	if (AActor* Actor = ActorElementDataUtil::GetActorFromHandle(FirstSelectedElement, true))
+	{
+		PotentialSelection.ISMComponents.Append(TInlineComponentArray<UInstancedStaticMeshComponent*>(Actor));
+		PotentialSelection.Actor = Actor;
+	}
+	else if (UActorComponent* Component = ComponentElementDataUtil::GetComponentFromHandle(FirstSelectedElement, true))
+	{
+		if (UInstancedStaticMeshComponent* ISMComponent = Cast<UInstancedStaticMeshComponent>(Component))
+		{
+			PotentialSelection.ISMComponents.Add(ISMComponent);
+		}
+		PotentialSelection.Actor = Component->GetOwner();
+	}
+	else if (FSMInstanceManager SMInstance = SMInstanceElementDataUtil::GetSMInstanceFromHandle(FirstSelectedElement, true))
+	{
+		PotentialSelection.ISMComponents.Add(SMInstance.GetISMComponent());
+	}
+	return PotentialSelection;
+}
+	
+}
+
+UComponentElementsExporterT3D::UComponentElementsExporterT3D(const FObjectInitializer& ObjectInitializer)
+	: Super(ObjectInitializer)
+{
+	SupportedClass = UComponentElementsCopy::StaticClass();
+	bText = true;
+	FormatExtension.Add(TEXT("COPY"));
+	FormatDescription.Add(TEXT("Unreal world text"));
+}
+
+bool UComponentElementsExporterT3D::ExportText(const FExportObjectInnerContext* Context, UObject* Object, const TCHAR* Type, FOutputDevice& Out, FFeedbackContext* Warn, uint32 PortFlags)
+{
+	UComponentElementsCopy* ComponentElementsCopy = Cast<UComponentElementsCopy>(Object);
+	if (!ComponentElementsCopy)
+	{
+		return false;
+	}
+
+	// Clear the mark state for saving.
+	UnMarkAllObjects(EObjectMark(OBJECTMARK_TagExp | OBJECTMARK_TagImp));
+
+	// Duplicate the selected component templates into temporary objects that we can modify
+	TMap<FName, FName> ParentMap;
+	TMap<FName, UActorComponent*> ObjectMap;
+	for (UActorComponent* Component : ComponentElementsCopy->ComponentsToCopy)
+	{
+		// Duplicate the component into a temporary object
+		UObject* DuplicatedComponent = StaticDuplicateObject(Component, GetTransientPackage(), Component->GetFName());
+		if (DuplicatedComponent)
+		{
+			// If the duplicated component is a scene component, wipe its attach parent (to prevent log warnings for referencing a private object in an external package)
+			if (USceneComponent* DuplicatedCompAsSceneComp = Cast<USceneComponent>(DuplicatedComponent))
+			{
+				DuplicatedCompAsSceneComp->SetupAttachment(nullptr);
+
+				AActor* Owner = Component->GetOwner();
+				if (Owner && Component == Owner->GetRootComponent())
+				{
+					DuplicatedCompAsSceneComp->SetRelativeTransform_Direct(FTransform::Identity);
+				}
+			}
+
+			// Find the closest parent component of the current component within the list of components to copy
+			USceneComponent* ClosestSelectedParent = UE::ComponentEditorUtils::FindClosestParent(Component, ComponentElementsCopy->ComponentsToCopy);
+			if (ClosestSelectedParent)
+			{
+				// If the parent is included in the list, record it into the node->parent map
+				ParentMap.Add(Component->GetFName(), ClosestSelectedParent->GetFName());
+			}
+
+			// Record the temporary object into the name->object map
+			ObjectMap.Add(Component->GetFName(), CastChecked<UActorComponent>(DuplicatedComponent));
+		}
+	}
+
+	// Export the component object(s) to text for copying
+	for (const TPair<FName, UActorComponent*>& ObjectPair : ObjectMap)
+	{
+		// Get the component object to be copied
+		UActorComponent* ComponentToCopy = ObjectPair.Value;
+		check(ComponentToCopy);
+
+		// If this component object had a parent within the selected set
+		if (ParentMap.Contains(ComponentToCopy->GetFName()))
+		{
+			// Get the name of the parent component
+			FName ParentName = ParentMap[ComponentToCopy->GetFName()];
+			if (ObjectMap.Contains(ParentName))
+			{
+				// Ensure that this component is a scene component
+				USceneComponent* SceneComponent = Cast<USceneComponent>(ComponentToCopy);
+				if (SceneComponent)
+				{
+					// Set the attach parent to the matching parent object in the temporary set. This allows us to preserve hierarchy in the copied set.
+					SceneComponent->SetupAttachment(Cast<USceneComponent>(ObjectMap[ParentName]));
+				}
+			}
+		}
+
+		// Export the component object to the given string
+		ExportToOutputDevice(Context, ComponentToCopy, nullptr, Out, TEXT("copy"), TextIndent, PPF_ExportsNotFullyQualified | PPF_Copy | PPF_Delimited, false, ComponentToCopy->GetOuter());
+	}
+
+	return false;
+}
+
+void FComponentElementEditorPasteImporter::Import(FContext& Context)
+{
+	using namespace UE::ElementCopyPasteUtil;
+	FTypedElementHandle FirstSelectedElement = Context.CurrentSelection->GetElementHandleAt(0);
+
+	FPotentialSelection PotentialSelection = GetPotentialSelection(Context.CurrentSelection);
+	
+	if (PotentialSelection.ISMComponents.Num())
+	{
+		for (UActorComponent* ISMComponent : PotentialSelection.ISMComponents)
+		{
+			TryPasteAsInstances(Context.Text, ISMComponent);
+		}
+	}
+	else if (PotentialSelection.Actor)
+	{
+		FString Text(Context.Text);
+		GUnrealEd->PasteComponents(ImportedComponents, PotentialSelection.Actor, true, &Text);
+	}
+}
+
+TArray<FTypedElementHandle> FComponentElementEditorPasteImporter::GetImportedElements()
+{
+	TArray<FTypedElementHandle> Handles;
+	Handles.Reserve(ImportedComponents.Num());
+
+	for (UActorComponent* Component : ImportedComponents)
+	{
+		if (IsValid(Component))
+		{
+			if (FTypedElementHandle Handle = UEngineElementsLibrary::AcquireEditorComponentElementHandle(Component))
+			{
+				Handles.Add(MoveTemp(Handle));
+			}
+		}
+	}
+
+	return Handles;
+}
diff --git a/Engine/Source/Editor/UnrealEd/Private/Elements/Component/ComponentElementEditorWorldInterface.cpp b/Engine/Source/Editor/UnrealEd/Private/Elements/Component/ComponentElementEditorWorldInterface.cpp
--- a/Engine/Source/Editor/UnrealEd/Private/Elements/Component/ComponentElementEditorWorldInterface.cpp
+++ b/Engine/Source/Editor/UnrealEd/Private/Elements/Component/ComponentElementEditorWorldInterface.cpp
@@ -10,6 +10,7 @@
 #include "Editor.h"
 #include "UnrealEdGlobals.h"
 #include "Editor/UnrealEdEngine.h"
+#include "Elements/Component/ComponentElementEditorCopyAndPaste.h"
 #include "Kismet2/ComponentEditorUtils.h"
 
 void UComponentElementEditorWorldInterface::NotifyMovementStarted(const FTypedElementHandle& InElementHandle)
@@ -82,6 +83,32 @@ void UComponentElementEditorWorldInterface::DuplicateElements(TArrayView<const F
 	}
 }
 
+bool UComponentElementEditorWorldInterface::CanCopyElement(const FTypedElementHandle& InElementHandle)
+{
+	return true;
+}
+
+void UComponentElementEditorWorldInterface::CopyElements(TArrayView<const FTypedElementHandle> InElementHandles, FOutputDevice& Out)
+{
+	TArray<UActorComponent*> Components = ComponentElementDataUtil::GetComponentsFromHandles(InElementHandles);
+
+	if (Components.IsEmpty())
+	{
+		return;
+	}
+
+	UComponentElementsCopy* ActorElementsCopy = NewObject<UComponentElementsCopy>();
+	ActorElementsCopy->ComponentsToCopy = Components;
+
+	int32 const Indent = 3;
+	UExporter::ExportToOutputDevice(nullptr, ActorElementsCopy, nullptr, Out, TEXT("copy"), Indent, PPF_DeepCompareInstances);
+}
+
+TSharedPtr<FWorldElementPasteImporter> UComponentElementEditorWorldInterface::GetPasteImporter()
+{
+	return MakeShared<FComponentElementEditorPasteImporter>();
+}
+
 bool UComponentElementEditorWorldInterface::IsElementInConvexVolume(const FTypedElementHandle& Handle, const FConvexVolume& InVolume, bool bMustEncompassEntireElement)
 {
 	if (const UPrimitiveComponent* Component = Cast<UPrimitiveComponent>(ComponentElementDataUtil::GetComponentFromHandle(Handle)))
diff --git a/Engine/Source/Editor/UnrealEd/Private/Factories/EditorFactories.cpp b/Engine/Source/Editor/UnrealEd/Private/Factories/EditorFactories.cpp
--- a/Engine/Source/Editor/UnrealEd/Private/Factories/EditorFactories.cpp
+++ b/Engine/Source/Editor/UnrealEd/Private/Factories/EditorFactories.cpp
@@ -1232,8 +1232,8 @@ UObject* ULevelFactory::FactoryCreateText
 
 		// Import properties if the new actor is 
 		bool		bActorChanged = false;
-		FString&	PropText = ActorMapElement.Value;
-		if ( Actor->ShouldImport(&PropText, bIsMoveToStreamingLevel) )
+		const FString&	PropText = ActorMapElement.Value;
+		if ( Actor->ShouldImport(FStringView(PropText), bIsMoveToStreamingLevel) )
 		{
 			Actor->PreEditChange(nullptr);
 			ImportObjectProperties( (uint8*)Actor, *PropText, Actor->GetClass(), Actor, Actor, Warn, 0, INDEX_NONE, NULL, &ExistingToNewMap );
diff --git a/Engine/Source/Editor/UnrealEd/Private/UnrealEdSrv.cpp b/Engine/Source/Editor/UnrealEd/Private/UnrealEdSrv.cpp
--- a/Engine/Source/Editor/UnrealEd/Private/UnrealEdSrv.cpp
+++ b/Engine/Source/Editor/UnrealEd/Private/UnrealEdSrv.cpp
@@ -107,6 +107,7 @@
 #include "ILevelEditor.h"
 #include "Subsystems/BrushEditingSubsystem.h"
 #include "Subsystems/EditorActorSubsystem.h"
+#include "Elements/Framework/TypedElementCommonActions.h"
 
 DEFINE_LOG_CATEGORY_STATIC(LogUnrealEdSrv, Log, All);
 
@@ -1492,7 +1493,35 @@ bool UUnrealEdEngine::Exec_Edit( UWorld* InWorld, const TCHAR* Str, FOutputDevic
 			return true;
 		}
 
-		if (bComponentsSelected)
+		if (TypedElementCommonActionsUtils::IsElementCopyAndPasteEnabled())
+		{
+			if (GCurrentLevelEditingViewportClient)
+			{
+				if (TSharedPtr<ILevelEditor> LevelEditor = GCurrentLevelEditingViewportClient->ParentLevelEditor.Pin())
+				{
+					if (UTypedElementCommonActions* CommonActions = LevelEditor->GetCommonActions())
+					{
+						const FScopedTransaction Transaction(NSLOCTEXT("UnrealEd", "Cut", "Cut"));
+
+						UTypedElementSelectionSet* SelectionSet = LevelEditor->GetMutableElementSelectionSet();
+
+						if (!bComponentsSelected)
+						{
+							FEditorDelegates::OnEditCutActorsBegin.Broadcast();
+						}
+
+						CommonActions->CopySelectedElements(SelectionSet);
+						CommonActions->DeleteSelectedElements(SelectionSet, InWorld, FTypedElementDeletionOptions());
+
+						if (!bComponentsSelected)
+						{
+							FEditorDelegates::OnEditCutActorsEnd.Broadcast();
+						}
+					}
+				}
+			}
+		}
+		else if (bComponentsSelected)
 		{
 			// Same transaction language used in CopySelectedActorsToClipboard below
 			const FScopedTransaction Transaction(NSLOCTEXT("UnrealEd", "Cut", "Cut"));
@@ -1515,7 +1544,32 @@ bool UUnrealEdEngine::Exec_Edit( UWorld* InWorld, const TCHAR* Str, FOutputDevic
 			return true;
 		}
 
-		if (bComponentsSelected)
+		if (TypedElementCommonActionsUtils::IsElementCopyAndPasteEnabled())
+		{
+			if (GCurrentLevelEditingViewportClient)
+			{
+				if (TSharedPtr<ILevelEditor> LevelEditor = GCurrentLevelEditingViewportClient->ParentLevelEditor.Pin())
+				{
+					if (UTypedElementCommonActions* CommonActions = LevelEditor->GetCommonActions())
+					{
+						UTypedElementSelectionSet* SelectionSet = LevelEditor->GetMutableElementSelectionSet();
+
+						if (!bComponentsSelected)
+						{
+							FEditorDelegates::OnEditCopyActorsBegin.Broadcast();
+						}
+
+						CommonActions->CopySelectedElements(SelectionSet);
+
+						if (!bComponentsSelected)
+						{
+							FEditorDelegates::OnEditCopyActorsEnd.Broadcast();
+						}
+					}
+				}
+			}
+		}
+		else if (bComponentsSelected)
 		{
 			edactCopySelected(InWorld);
 		}
@@ -1533,39 +1587,99 @@ bool UUnrealEdEngine::Exec_Edit( UWorld* InWorld, const TCHAR* Str, FOutputDevic
 			return true;
 		}
 
-		if (bComponentsSelected)
+		bool bImportedElements = false;
+		if (TypedElementCommonActionsUtils::IsElementCopyAndPasteEnabled())
 		{
-			const FScopedTransaction Transaction(NSLOCTEXT("UnrealEd", "PasteComponents", "Paste Components"));
-			edactPasteSelected(InWorld, false, false, true);
-		}
-		else
-		{
-			// How should this paste be handled
-			EPasteTo PasteTo = PT_OriginalLocation;
-			FText TransDescription = NSLOCTEXT("UnrealEd", "Paste", "Paste");
-
-			FString TempStr;
-			if (FParse::Value(Str, TEXT("TO="), TempStr))
+			if (GCurrentLevelEditingViewportClient)
 			{
-				if (!FCString::Strcmp(*TempStr, TEXT("HERE")))
+				if (TSharedPtr<ILevelEditor> LevelEditor = GCurrentLevelEditingViewportClient->ParentLevelEditor.Pin())
 				{
-					PasteTo = PT_Here;
-					TransDescription = NSLOCTEXT("UnrealEd", "PasteHere", "Paste Here");
+					if (UTypedElementCommonActions* CommonActions = LevelEditor->GetCommonActions())
+					{
+						FText TransDescription = NSLOCTEXT("UnrealEd", "Paste", "Paste");
+
+						UTypedElementSelectionSet* SelectionSet = LevelEditor->GetMutableElementSelectionSet();
+						FTypedElementPasteOptions PasteOptions;
+						PasteOptions.SelectionSetToModify = SelectionSet;
+
+						FString TempStr;
+						if (FParse::Value(Str, TEXT("TO="), TempStr))
+						{
+							if (!FCString::Strcmp(*TempStr, TEXT("HERE")))
+							{
+								PasteOptions.bPasteAtLocation = true;
+								const FSnappedPositioningData PositioningData = FSnappedPositioningData(GCurrentLevelEditingViewportClient, ClickLocation, ClickPlane)
+									.AlignToSurfaceRotation(false);
+								PasteOptions.PasteLocation = FActorPositioning::GetSnappedSurfaceAlignedTransform(PositioningData).GetLocation();
+
+								TransDescription = NSLOCTEXT("UnrealEd", "PasteHere", "Paste Here");
+							}
+							else
+							{
+								if (!FCString::Strcmp(*TempStr, TEXT("ORIGIN")))
+								{
+									PasteOptions.bPasteAtLocation = true;
+									PasteOptions.PasteLocation = FVector::ZeroVector;
+									TransDescription = NSLOCTEXT("UnrealEd", "PasteToWorldOrigin", "Paste To World Origin");
+								}
+							}
+						}
+
+						
+						const FScopedTransaction Transaction(TransDescription);
+
+						if (!bComponentsSelected)
+						{
+							FEditorDelegates::OnEditPasteActorsBegin.Broadcast();
+						}
+						
+						bImportedElements = !CommonActions->PasteElements(SelectionSet, InWorld, PasteOptions).IsEmpty();
+
+						if (!bComponentsSelected)
+						{
+							FEditorDelegates::OnEditPasteActorsEnd.Broadcast();
+						}
+					}
 				}
-				else
+			}
+		}
+
+		if (!bImportedElements)
+		{ 
+			if (bComponentsSelected)
+			{
+				const FScopedTransaction Transaction(NSLOCTEXT("UnrealEd", "PasteComponents", "Paste Components"));
+				edactPasteSelected(InWorld, false, false, true);
+			}
+			else
+			{
+				// How should this paste be handled
+				EPasteTo PasteTo = PT_OriginalLocation;
+				FText TransDescription = NSLOCTEXT("UnrealEd", "Paste", "Paste");
+
+				FString TempStr;
+				if (FParse::Value(Str, TEXT("TO="), TempStr))
 				{
-					if (!FCString::Strcmp(*TempStr, TEXT("ORIGIN")))
+					if (!FCString::Strcmp(*TempStr, TEXT("HERE")))
 					{
-						PasteTo = PT_WorldOrigin;
-						TransDescription = NSLOCTEXT("UnrealEd", "PasteToWorldOrigin", "Paste To World Origin");
+						PasteTo = PT_Here;
+						TransDescription = NSLOCTEXT("UnrealEd", "PasteHere", "Paste Here");
+					}
+					else
+					{
+						if (!FCString::Strcmp(*TempStr, TEXT("ORIGIN")))
+						{
+							PasteTo = PT_WorldOrigin;
+							TransDescription = NSLOCTEXT("UnrealEd", "PasteToWorldOrigin", "Paste To World Origin");
+						}
 					}
 				}
-			}
 
-			const FScopedTransaction Transaction(TransDescription);
-			FEditorDelegates::OnEditPasteActorsBegin.Broadcast();
-			PasteSelectedActorsFromClipboard(InWorld, TransDescription, PasteTo);
-			FEditorDelegates::OnEditPasteActorsEnd.Broadcast();
+				const FScopedTransaction Transaction(TransDescription);
+				FEditorDelegates::OnEditPasteActorsBegin.Broadcast();
+				PasteSelectedActorsFromClipboard(InWorld, TransDescription, PasteTo);
+				FEditorDelegates::OnEditPasteActorsEnd.Broadcast();
+			}
 		}
 	}
 
diff --git a/Engine/Source/Editor/UnrealEd/Public/Editor.h b/Engine/Source/Editor/UnrealEd/Public/Editor.h
--- a/Engine/Source/Editor/UnrealEd/Public/Editor.h
+++ b/Engine/Source/Editor/UnrealEd/Public/Editor.h
@@ -501,10 +501,8 @@ struct FImportObjectParams
  *
  * @return	NULL if the default values couldn't be imported
  */
-
 const TCHAR* ImportObjectProperties( FImportObjectParams& InParams );
 
-
 /**
  * Parse and import text as property values for the object specified.
  * 
@@ -781,6 +779,82 @@ namespace EditorUtilities
 	 * @return	The number of properties that were copied over (properties that were filtered out, or were already identical, are not counted.)
 	 */
 	UNREALED_API int32 CopyActorProperties( AActor* SourceActor, AActor* TargetActor, const FCopyOptions& Options = FCopyOptions(ECopyOptions::Default) );
+
+
+	// ==== Multi step import utilities ====
+
+	/**
+	 * Parameters for the ImportCreateObjectStep and ImportObjectPropertiesStep. Used for multi steps import.
+	 */
+	struct FMultiStepsImportObjectParams
+	{
+		/** The location to import the property values to */
+		uint8* DestData = nullptr;
+
+		/** Text buffer containing the values that should be parsed and imported */
+		FStringView SourceText;
+
+		/** The struct for the data we're importing */
+		UStruct* ObjectStruct = nullptr;
+
+		/** The original object that ImportObjectProperties was called for.
+			If SubobjectOuter is a subobject, corresponds to the first object in SubobjectOuter's Outer chain that is not a subobject itself.
+			If SubobjectOuter is not a subobject, should normally be the same value as SubobjectOuter */
+		UObject* SubobjectRoot = nullptr;
+
+		/** The object corresponding to DestData; this is the object that will used as the outer when creating subobjects from definitions contained in SourceText */
+		UObject* SubobjectOuter = nullptr;
+
+		/** Output device to use for log messages */
+		FFeedbackContext*	Warn = nullptr;
+
+		/** Current nesting level */
+		int32 Depth = 0;
+
+		/** Used when importing defaults during script compilation for tracking which line we're currently for the purposes of printing compile errors */
+		int32 LineNumber = INDEX_NONE;
+
+		/** Contains the mappings of instanced objects and components to their templates; used when recursively calling ImportObjectProperties; generally
+			not necessary to specify a value when calling this function from other code */
+		FObjectInstancingGraph* InInstanceGraph = nullptr;
+
+		/** 
+		 * Provides a mapping from an exported path to a new instance to which it should be remapped 
+		 * Imported object will be added to this map when possible during the create objects step.
+		 */
+		TMap<FString, UObject*>* ObjectRemapper = nullptr;
+
+		/** 
+		 * Tell what properties shouldn't be imported when importing the properties
+		 */
+		TSet<FProperty*>* PropertiesToSkip = nullptr;
+
+		/** True if we should call PreEditChange/PostEditChange on the object as it's imported. Pass false here
+			if you're going to do that on your own. */
+		bool bShouldCallEditChange = false;
+	};
+
+	/**
+	 * Parse text and create the objects for the object specified.
+	 * See ImportObjectPropertiesStep for the next step of the multi steps object import
+	 * 
+	 * @param	InParams	Parameters for object the multistep object import; see declaration of FMultiStepsImportObjectParams.
+	 *
+	 * @return	nullptr if the objects couldn't be imported
+	 */
+	UNREALED_API const TCHAR* ImportCreateObjectsStep(FMultiStepsImportObjectParams& InParams);
+
+	/**
+	 * Parse text and import the properties for the object specified and its subobjects.
+	 * Call ImportCreateObjectsStep before calling this function to create the subobjects
+	 * 
+	 * @param	InParams	Parameters for object the multistep object import; see declaration of FMultiStepsImportObjectParams.
+	 *
+	 * @return	nullptr if the values couldn't be imported
+	 */
+	UNREALED_API const TCHAR* ImportObjectsPropertiesStep(FMultiStepsImportObjectParams& InParams);
+
+
 }
 
 
diff --git a/Engine/Source/Editor/UnrealEd/Public/Elements/Actor/ActorElementEditorCopyAndPaste.h b/Engine/Source/Editor/UnrealEd/Public/Elements/Actor/ActorElementEditorCopyAndPaste.h
new file mode 100644
--- /dev/null
+++ b/Engine/Source/Editor/UnrealEd/Public/Elements/Actor/ActorElementEditorCopyAndPaste.h
@@ -0,0 +1,79 @@
+// Copyright Epic Games, Inc. All Rights Reserved.
+
+#pragma once
+
+#include "Elements/Interfaces/TypedElementWorldInterface.h"
+
+#include "Containers/StringFwd.h"
+#include "Containers/StringView.h"
+#include "CoreMinimal.h"
+#include "Exporters/Exporter.h"
+#include "GameFramework/Actor.h"
+#include "UObject/UnrealType.h"
+
+#include "ActorElementEditorCopyAndPaste.generated.h"
+
+class AActor;
+class AGroupActor;
+
+UCLASS(BlueprintType)
+class UNREALED_API UActorElementsCopy : public UObject
+{
+	GENERATED_BODY()
+
+public:
+	UPROPERTY(BlueprintReadWrite, Category = "TypedElement|Actor|Copy")
+	TArray<TObjectPtr<AActor>> ActorsToCopy;
+};
+
+UCLASS()
+class UNREALED_API UActorElementsExporterT3D : public UExporter
+{
+public:
+	GENERATED_BODY()
+
+public:
+	UActorElementsExporterT3D(const FObjectInitializer& ObjectInitializer = FObjectInitializer());
+
+	//~ Begin UExporter Interface
+	virtual bool ExportText(const FExportObjectInnerContext* Context, UObject* Object, const TCHAR* Type, FOutputDevice& Out, FFeedbackContext* Warn, uint32 PortFlags = 0) override;
+	//~ End UExporter Interface
+};
+
+struct UNREALED_API FActorElementEditorPasteImporter : public FWorldElementPasteImporter
+{
+	virtual void Import(FContext& Context) override;
+	
+	virtual TArray<FTypedElementHandle> GetImportedElements() override;
+	
+private:
+	void PostImportProcess();
+	
+	FContext Context;
+
+	struct FImportedActor
+	{
+		AActor* ActorPtr = nullptr;
+		FStringView InnerText;
+
+		UClass* Class;
+		FName Name;
+		FString Archetype;
+		FName ParentActor;
+		FName SocketNameString;
+		FString GroupActor;
+		FString GroupFolder;
+		FString ActorFolderPath;
+		FString ExportedActorFullName;
+	};
+
+	TArray<FImportedActor> ImportedActors;
+
+	// Handle the attachment
+	TMap<FName, AActor*> NameToImportedActors;
+
+	// Handle the groups
+	TMap<FString, AGroupActor*> NewGroups;
+
+	bool bPostImportProcessingNeeded;
+};
diff --git a/Engine/Source/Editor/UnrealEd/Public/Elements/Actor/ActorElementEditorWorldInterface.h b/Engine/Source/Editor/UnrealEd/Public/Elements/Actor/ActorElementEditorWorldInterface.h
--- a/Engine/Source/Editor/UnrealEd/Public/Elements/Actor/ActorElementEditorWorldInterface.h
+++ b/Engine/Source/Editor/UnrealEd/Public/Elements/Actor/ActorElementEditorWorldInterface.h
@@ -20,6 +20,9 @@ public:
 	virtual bool DeleteElements(TArrayView<const FTypedElementHandle> InElementHandles, UWorld* InWorld, UTypedElementSelectionSet* InSelectionSet, const FTypedElementDeletionOptions& InDeletionOptions) override;
 	virtual bool CanDuplicateElement(const FTypedElementHandle& InElementHandle) override;
 	virtual void DuplicateElements(TArrayView<const FTypedElementHandle> InElementHandles, UWorld* InWorld, const FVector& InLocationOffset, TArray<FTypedElementHandle>& OutNewElements) override;
+	virtual bool CanCopyElement(const FTypedElementHandle& InElementHandle) override;
+	virtual void CopyElements(TArrayView<const FTypedElementHandle> InElementHandles, FOutputDevice& Out) override;
+	virtual TSharedPtr<FWorldElementPasteImporter> GetPasteImporter() override;
 	virtual bool IsElementInConvexVolume(const FTypedElementHandle& Handle, const FConvexVolume& InVolume, bool bMustEncompassEntireElement = false);
 	virtual bool IsElementInBox(const FTypedElementHandle& Handle, const FBox& InBox, bool bMustEncompassEntireElement = false);
 	virtual TArray<FTypedElementHandle> GetSelectionElementsFromSelectionFunction(const FTypedElementHandle& InElementHandle, const FWorldSelectionElementArgs& SelectionArgs, const TFunction<bool(const FTypedElementHandle&, const FWorldSelectionElementArgs&)>& SelectionFunction) override;
diff --git a/Engine/Source/Editor/UnrealEd/Public/Elements/Component/ComponentElementEditorCopyAndPaste.h b/Engine/Source/Editor/UnrealEd/Public/Elements/Component/ComponentElementEditorCopyAndPaste.h
new file mode 100644
--- /dev/null
+++ b/Engine/Source/Editor/UnrealEd/Public/Elements/Component/ComponentElementEditorCopyAndPaste.h
@@ -0,0 +1,45 @@
+﻿// Copyright Epic Games, Inc. All Rights Reserved.
+
+#pragma once
+
+#include "Elements/Interfaces/TypedElementWorldInterface.h"
+
+#include "CoreMinimal.h"
+#include "Exporters/Exporter.h"
+#include "Components/ActorComponent.h"
+
+#include "ComponentElementEditorCopyAndPaste.generated.h"
+
+UCLASS(BlueprintType)
+class UNREALED_API UComponentElementsCopy : public UObject
+{
+	GENERATED_BODY()
+
+public:
+	UPROPERTY(BlueprintReadWrite, Category = "TypedElement|Component|Copy")
+	TArray<TObjectPtr<UActorComponent>> ComponentsToCopy;
+};
+
+UCLASS()
+class UNREALED_API UComponentElementsExporterT3D : public UExporter
+{
+public:
+	GENERATED_BODY()
+
+public:
+	explicit UComponentElementsExporterT3D(const FObjectInitializer& ObjectInitializer = FObjectInitializer());
+
+	//~ Begin UExporter Interface
+	virtual bool ExportText(const FExportObjectInnerContext* Context, UObject* Object, const TCHAR* Type, FOutputDevice& Out, FFeedbackContext* Warn, uint32 PortFlags = 0) override;
+	//~ End UExporter Interface
+};
+
+struct UNREALED_API FComponentElementEditorPasteImporter : public FWorldElementPasteImporter
+{
+public:
+	virtual void Import(FContext& Context) override;
+	
+	virtual TArray<FTypedElementHandle> GetImportedElements() override;
+private:
+	TArray<TObjectPtr<UActorComponent>> ImportedComponents;
+};
\ No newline at end of file
diff --git a/Engine/Source/Editor/UnrealEd/Public/Elements/Component/ComponentElementEditorWorldInterface.h b/Engine/Source/Editor/UnrealEd/Public/Elements/Component/ComponentElementEditorWorldInterface.h
--- a/Engine/Source/Editor/UnrealEd/Public/Elements/Component/ComponentElementEditorWorldInterface.h
+++ b/Engine/Source/Editor/UnrealEd/Public/Elements/Component/ComponentElementEditorWorldInterface.h
@@ -18,6 +18,9 @@ public:
 	virtual bool DeleteElements(TArrayView<const FTypedElementHandle> InElementHandles, UWorld* InWorld, UTypedElementSelectionSet* InSelectionSet, const FTypedElementDeletionOptions& InDeletionOptions) override;
 	virtual bool CanDuplicateElement(const FTypedElementHandle& InElementHandle) override;
 	virtual void DuplicateElements(TArrayView<const FTypedElementHandle> InElementHandles, UWorld* InWorld, const FVector& InLocationOffset, TArray<FTypedElementHandle>& OutNewElements) override;
+	virtual bool CanCopyElement(const FTypedElementHandle& InElementHandle) override;
+	virtual void CopyElements(TArrayView<const FTypedElementHandle> InElementHandles, FOutputDevice& Out) override;
+	virtual TSharedPtr<FWorldElementPasteImporter> GetPasteImporter() override;
 	virtual bool IsElementInConvexVolume(const FTypedElementHandle& Handle, const FConvexVolume& InVolume, bool bMustEncompassEntireElement = false);
 	virtual bool IsElementInBox(const FTypedElementHandle& Handle, const FBox& InBox, bool bMustEncompassEntireElement = false);
 	virtual TArray<FTypedElementHandle> GetSelectionElementsFromSelectionFunction(const FTypedElementHandle& InElementHandle, const FWorldSelectionElementArgs& SelectionArgs, const TFunction<bool(const FTypedElementHandle&, const FWorldSelectionElementArgs&)>& SelectionFunction) override;
diff --git a/Engine/Source/Runtime/Core/Private/Misc/Parse.cpp b/Engine/Source/Runtime/Core/Private/Misc/Parse.cpp
--- a/Engine/Source/Runtime/Core/Private/Misc/Parse.cpp
+++ b/Engine/Source/Runtime/Core/Private/Misc/Parse.cpp
@@ -15,6 +15,7 @@
 #include "Misc/StringBuilder.h"
 #include "HAL/IConsoleManager.h"
 #include "Containers/LazyPrintf.h"
+#include "Containers/StringView.h"
 
 #if !UE_BUILD_SHIPPING 
 /**
@@ -1029,12 +1030,21 @@ bool FParse::Line(const TCHAR** Stream, TCHAR* Result, int32 MaxLen, bool bExact
 }
 
 bool FParse::Line(const TCHAR** Stream, FString& Result, bool bExact)
+{
+	FStringView View;
+	bool bReturnValue = Line(Stream, View, bExact);
+	Result = View;
+	return bReturnValue;
+}
+
+bool FParse::Line(const TCHAR** Stream, FStringView& Result, bool bExact)
 {
 	bool bGotStream = false;
 	bool bIsQuoted = false;
 	bool bIgnore = false;
 
 	Result.Reset();
+	const TCHAR* StartOfLine = nullptr;
 
 	while (**Stream != TEXT('\0') && **Stream != TEXT('\n') && **Stream != TEXT('\r'))
 	{
@@ -1055,9 +1065,9 @@ bool FParse::Line(const TCHAR** Stream, FString& Result, bool bExact)
 		bGotStream = true;
 
 		// Got stuff.
-		if (!bIgnore)
+		if (!bIgnore && !StartOfLine)
 		{
-			Result.AppendChar(*((*Stream)++));
+			StartOfLine = (*Stream)++;
 		}
 		else
 		{
@@ -1065,6 +1075,11 @@ bool FParse::Line(const TCHAR** Stream, FString& Result, bool bExact)
 		}
 	}
 
+	if (StartOfLine)
+	{
+		Result = FStringView(StartOfLine, int32((*Stream) - StartOfLine));
+	}
+
 	if (bExact)
 	{
 		// Eat up exactly one CR/LF.
diff --git a/Engine/Source/Runtime/Core/Public/Misc/Parse.h b/Engine/Source/Runtime/Core/Public/Misc/Parse.h
--- a/Engine/Source/Runtime/Core/Public/Misc/Parse.h
+++ b/Engine/Source/Runtime/Core/Public/Misc/Parse.h
@@ -57,9 +57,11 @@ struct CORE_API FParse
 	/** Parses a boolean value. */
 	static bool Bool( const TCHAR* Stream, const TCHAR* Match, bool& OnOff );
 	/** Get a line of Stream (everything up to, but not including, CR/LF. Returns 0 if ok, nonzero if at end of stream and returned 0-length string. */
-	static bool Line( const TCHAR** Stream, TCHAR* Result, int32 MaxLen, bool Exact=0 );
+	static bool Line( const TCHAR** Stream, TCHAR* Result, int32 MaxLen, bool Exact= false );
 	/** Get a line of Stream (everything up to, but not including, CR/LF. Returns 0 if ok, nonzero if at end of stream and returned 0-length string. */
-	static bool Line( const TCHAR** Stream, FString& Resultd, bool Exact=0 );
+	static bool Line( const TCHAR** Stream, FString& Result, bool Exact = false );
+	/** Get a line of Stream (everything up to, but not including, CR/LF. Returns 0 if ok, nonzero if at end of stream and returned 0-length string. */
+	static bool Line( const TCHAR** Stream, FStringView& Result, bool Exact = false );
 	/** Get a line of Stream, with support for extending beyond that line with certain characters, e.g. {} and \
 	 * the out character array will not include the ignored endlines
 	 */
diff --git a/Engine/Source/Runtime/Engine/Classes/GameFramework/Actor.h b/Engine/Source/Runtime/Engine/Classes/GameFramework/Actor.h
--- a/Engine/Source/Runtime/Engine/Classes/GameFramework/Actor.h
+++ b/Engine/Source/Runtime/Engine/Classes/GameFramework/Actor.h
@@ -2321,7 +2321,11 @@ public:
 	virtual bool ShouldExport() { return true; }
 
 	/** Called before editor paste, true allow import */
-	virtual bool ShouldImport(FString* ActorPropString, bool IsMovingLevel) { return true; }
+	UE_DEPRECATED(5.1, "Use the override that take a StringView instead instead")
+	virtual bool ShouldImport(FString* ActorPropString, bool IsMovingLevel) { return ActorPropString ? ShouldImport(FStringView(), IsMovingLevel) : ShouldImport(FStringView(*ActorPropString), IsMovingLevel); }
+
+	/** Called before editor paste, true allow import */
+	virtual bool ShouldImport(FStringView ActorPropString, bool IsMovingLevel) { return true; }
 
 	/** Called by InputKey when an unhandled key is pressed with a selected actor */
 	virtual void EditorKeyPressed(FKey Key, EInputEvent Event) {}
diff --git a/Engine/Source/Runtime/Engine/Private/Elements/Framework/TypedElementCommonActions.cpp b/Engine/Source/Runtime/Engine/Private/Elements/Framework/TypedElementCommonActions.cpp
--- a/Engine/Source/Runtime/Engine/Private/Elements/Framework/TypedElementCommonActions.cpp
+++ b/Engine/Source/Runtime/Engine/Private/Elements/Framework/TypedElementCommonActions.cpp
@@ -5,11 +5,34 @@
 #include "Elements/Framework/TypedElementRegistry.h"
 #include "Elements/Framework/TypedElementUtil.h"
 
+#include "Containers/StringFwd.h"
+#include "Containers/StringView.h"
+#include "Containers/UnrealString.h"
+#include "Exporters/Exporter.h"
+#include "HAL/PlatformApplicationMisc.h"
+#include "Math/Box.h"
+#include "Math/Transform.h"
+#include "Misc/CString.h"
+#include "Misc/Parse.h"
 #include "UObject/GCObjectScopeGuard.h"
 #include "UObject/Stack.h"
 
 #include UE_INLINE_GENERATED_CPP_BY_NAME(TypedElementCommonActions)
 
+namespace TypedElementCommonActionsUtils
+{
+	static int32 GEnableElementsCopyAndPaste = 0;
+	static FAutoConsoleVariableRef CVarEnableElementsCopyAndPaste(
+		TEXT("TypedElements.EnableElementsCopyAndPaste"),
+		GEnableElementsCopyAndPaste,
+		TEXT("Is support for elements copy and paste enabled?")
+	);
+
+	bool IsElementCopyAndPasteEnabled()
+	{
+		return GEnableElementsCopyAndPaste != 0;
+	}
+}
 bool FTypedElementCommonActionsCustomization::DeleteElements(ITypedElementWorldInterface* InWorldInterface, TArrayView<const FTypedElementHandle> InElementHandles, UWorld* InWorld, UTypedElementSelectionSet* InSelectionSet, const FTypedElementDeletionOptions& InDeletionOptions)
 {
 	return InWorldInterface->DeleteElements(InElementHandles, InWorld, InSelectionSet, InDeletionOptions);
@@ -20,6 +43,15 @@ void FTypedElementCommonActionsCustomization::DuplicateElements(ITypedElementWor
 	InWorldInterface->DuplicateElements(InElementHandles, InWorld, InLocationOffset, OutNewElements);
 }
 
+void FTypedElementCommonActionsCustomization::CopyElements(ITypedElementWorldInterface* InWorldInterface, TArrayView<const FTypedElementHandle> InElementHandles, FOutputDevice& Out)
+{
+	InWorldInterface->CopyElements(InElementHandles, Out);
+}
+
+TSharedPtr<FWorldElementPasteImporter> FTypedElementCommonActionsCustomization::GetPasteImporter(ITypedElementWorldInterface* InWorldInterface, const FTypedElementListConstPtr& InSelectedHandles, UWorld* InWorld)
+{
+	return InWorldInterface->GetPasteImporter();
+}
 
 bool UTypedElementCommonActions::DeleteSelectedElements(UTypedElementSelectionSet* SelectionSet, UWorld* World, const FTypedElementDeletionOptions& DeletionOptions)
 {
@@ -96,6 +128,240 @@ TArray<FTypedElementHandle> UTypedElementCommonActions::DuplicateNormalizedEleme
 	return NewElements;
 }
 
+bool UTypedElementCommonActions::CopySelectedElements(UTypedElementSelectionSet* SelectionSet)
+{
+	if (!SelectionSet)
+	{
+		FFrame::KismetExecutionMessage(TEXT("SelectionSet is null."), ELogVerbosity::Error);
+		return false;
+	}
+
+	FTypedElementListRef NormalizedElements = SelectionSet->GetNormalizedSelection(FTypedElementSelectionNormalizationOptions());
+	return CopyNormalizedElements(NormalizedElements);
+}
+
+bool UTypedElementCommonActions::CopySelectedElementsToString(UTypedElementSelectionSet* SelectionSet, FString& OutputString)
+{
+	if (!SelectionSet)
+	{
+		FFrame::KismetExecutionMessage(TEXT("SelectionSet is null."), ELogVerbosity::Error);
+		return false;
+	}
+
+	FTypedElementListRef NormalizedElements = SelectionSet->GetNormalizedSelection(FTypedElementSelectionNormalizationOptions());
+	return CopyNormalizedElements(NormalizedElements, &OutputString);
+}
+
+bool UTypedElementCommonActions::CopyNormalizedElements(const FTypedElementListConstPtr& ElementListPtr, FString* OptionalOutputString)
+{
+	if (!ElementListPtr)
+	{
+		return false;
+	}
+	
+	FStringOutputDevice Out;
+
+	TMap<FTypedHandleTypeId, TArray<FTypedElementHandle>> ElementsToCopyByType;
+	TypedElementUtil::BatchElementsByType(ElementListPtr.ToSharedRef(), ElementsToCopyByType);
+
+	UTypedElementRegistry* Registry = UTypedElementRegistry::GetInstance();
+	for (const TPair<FTypedHandleTypeId, TArray<FTypedElementHandle>>& ElementsByTypePair : ElementsToCopyByType)
+	{
+		// Store the copied elements in a separate string to test if anything was copied before writing element group
+		FStringOutputDevice ElementsOut;
+		
+		FTypedElementCommonActionsCustomization* CommonActionsCustomization = GetInterfaceCustomizationByTypeId(ElementsByTypePair.Key);
+		ITypedElementWorldInterface* WorldInterface = Registry->GetElementInterface<ITypedElementWorldInterface>(ElementsByTypePair.Key);
+		if (CommonActionsCustomization && WorldInterface)
+		{
+			CommonActionsCustomization->CopyElements(WorldInterface, ElementsByTypePair.Value, ElementsOut);
+		}
+		if (ElementsOut.IsEmpty())
+		{
+			continue;
+		}
+		
+		FName ElementName = Registry->GetRegisteredElementTypeName(ElementsByTypePair.Key);
+		Out.Logf(TEXT("Begin ElementGroup Type=%s") LINE_TERMINATOR, *ElementName.ToString());
+		Out.Append(ElementsOut);
+		Out.Logf(TEXT("End ElementGroup") LINE_TERMINATOR);
+	}
+
+	bool bHasExportedSomething = !Out.IsEmpty();
+	
+	if (OptionalOutputString)
+	{
+		*OptionalOutputString = MoveTemp(Out);
+	}
+	else
+	{
+		FPlatformApplicationMisc::ClipboardCopy(*Out);
+	}
+
+	return bHasExportedSomething;
+}
+
+TArray<FTypedElementHandle> UTypedElementCommonActions::PasteElements(UTypedElementSelectionSet* SelectionSet, UWorld* World, const FTypedElementPasteOptions& PasteOption, const FString* OptionalInputString)
+{
+	FTypedElementListRef NormalizedElements = SelectionSet->GetNormalizedSelection(FTypedElementSelectionNormalizationOptions());
+	return PasteNormalizedElements(NormalizedElements, World, PasteOption, OptionalInputString);
+}
+
+TArray<FTypedElementHandle> UTypedElementCommonActions::PasteNormalizedElements(const FTypedElementListConstPtr& ElementListPtr, UWorld* World, const FTypedElementPasteOptions& PasteOptions, const FString* OptionalInputString)
+{
+	check(World)
+
+	FString ClipboardString;
+	const TCHAR* StringToImport;
+	if (OptionalInputString)
+	{
+		StringToImport = **OptionalInputString;
+	}
+	else
+	{
+		FPlatformApplicationMisc::ClipboardPaste(ClipboardString);
+		StringToImport = *ClipboardString;
+	}
+
+	UTypedElementRegistry* Registry = ElementListPtr->GetRegistry();
+	TArray<TPair<TSharedRef<FWorldElementPasteImporter>, FStringView>> ImportersAndText;
+
+	// Break the text into the group for their respective importer
+	{
+		const TCHAR** CurrentPosInString = &StringToImport;
+
+		FStringView* CurrentBlock = nullptr;
+		const TCHAR* StartOfBlock = nullptr;
+		FString LookingForEndTag;
+		FStringView Line;
+		const TCHAR* BeginChars = TEXT("Begin ");
+		const int32 BeginCharsLen = FCString::Strlen(BeginChars);
+		TCHAR Buffer[NAME_SIZE];
+		const TCHAR* TypedElementChar = TEXT("ElementGroup");
+
+		while (FParse::Line(CurrentPosInString, Line))
+		{
+			if (!LookingForEndTag.IsEmpty())
+			{
+				if (Line.StartsWith(LookingForEndTag))
+				{
+					(*CurrentBlock) = MakeStringView(StartOfBlock, int32((Line.GetData() + Line.Len()) - StartOfBlock));
+					LookingForEndTag.Reset();
+				}
+			}
+			else if (Line.StartsWith(BeginChars))
+			{
+				const TCHAR* PositionInLine = Line.GetData() + BeginCharsLen;
+				if (FParse::Token(PositionInLine, Buffer, NAME_SIZE, false))
+				{
+					if (FCString::Strcmp(TypedElementChar, Buffer) == 0)
+					{
+						FName TypeName;
+						if (!FParse::Value(++PositionInLine, TEXT("Type="), TypeName))
+						{
+							continue;
+						}
+						FTypedHandleTypeId TypeID = Registry->GetRegisteredElementTypeId(TypeName);
+						if (TypeID == 0)
+						{
+							continue;
+						}
+						ITypedElementWorldInterface* WorldInterface = Registry->GetElementInterface<ITypedElementWorldInterface>(TypeID);
+						if (!WorldInterface)
+						{
+							continue;
+						}
+						FTypedElementCommonActionsCustomization* CommonActionsCustomization = GetInterfaceCustomizationByTypeId(TypeID);
+						TSharedPtr<FWorldElementPasteImporter> Importer = CommonActionsCustomization->GetPasteImporter(WorldInterface, ElementListPtr, World);
+						if (!Importer)
+						{
+							continue;
+						}
+						StartOfBlock = Line.GetData();
+						TPair<TSharedRef<FWorldElementPasteImporter>, FStringView>& Pair = ImportersAndText.Add_GetRef(
+							TPair<TSharedRef<FWorldElementPasteImporter>, FStringView>(Importer.ToSharedRef(), FStringView())
+							);
+						CurrentBlock = &Pair.Value;
+
+						LookingForEndTag = FString::Printf(TEXT("End %s"), Buffer);
+					}
+					else
+					{
+						FStringView BufferView(Buffer);
+						if (const TFunction<TSharedRef<FWorldElementPasteImporter> ()>* GetImporterFunction = PasteOptions.ExtraCustomImport.Find(BufferView))
+						{
+							StartOfBlock = Line.GetData();
+							TPair<TSharedRef<FWorldElementPasteImporter>, FStringView>& Pair = ImportersAndText.Add_GetRef(
+								TPair<TSharedRef<FWorldElementPasteImporter>, FStringView>((*GetImporterFunction)(), FStringView())
+							);
+							CurrentBlock = &Pair.Value;
+
+							LookingForEndTag = FString::Printf(TEXT("End %s"), Buffer);
+						}
+					}
+				}
+			}
+		}
+	}
+
+	TArray<FTypedElementHandle> PastedHandles; 
+
+	// Import the elements
+	{
+		TMap<FString, UObject*> ExportedPathToObjectCreated;
+		for (TPair<TSharedRef<FWorldElementPasteImporter>, FStringView>& Pair : ImportersAndText)
+		{
+			FWorldElementPasteImporter::FContext Context;
+			Context.CurrentSelection = ElementListPtr;
+			Context.World = World;
+			Context.Text = Pair.Value;
+			Pair.Key->Import(Context);
+		}
+		
+		for (TPair<TSharedRef<FWorldElementPasteImporter>, FStringView>& Pair : ImportersAndText)
+		{
+			PastedHandles.Append(Pair.Key->GetImportedElements());
+		}
+	}
+
+	// Paste at implementation
+	if (PasteOptions.bPasteAtLocation)
+	{
+		FBox BoundingBox;
+
+		TArray<TTypedElement<ITypedElementWorldInterface>> WorldElements;
+		WorldElements.Reserve(PastedHandles.Num());
+		for (const FTypedElementHandle& Handle : PastedHandles)
+		{
+			if (TTypedElement<ITypedElementWorldInterface> WorldElement = Registry->GetElement<ITypedElementWorldInterface>(Handle))
+			{
+				FTransform WorldTransform;
+				WorldElement.GetWorldTransform(WorldTransform);
+				BoundingBox += WorldTransform.GetLocation();
+
+				WorldElements.Add(MoveTemp(WorldElement));
+			}
+		}
+
+		FVector Offset = PasteOptions.PasteLocation - BoundingBox.GetCenter();
+		for (TTypedElement<ITypedElementWorldInterface> WorldElement : WorldElements)
+		{
+			FTransform WorldTransform;
+			WorldElement.GetWorldTransform(WorldTransform);
+			WorldTransform.SetLocation(WorldTransform.GetLocation() + Offset);
+			WorldElement.SetWorldTransform(WorldTransform);
+		}
+	}
+
+	// Set the selection optionally
+	if (PasteOptions.SelectionSetToModify && !PastedHandles.IsEmpty())
+	{
+		PasteOptions.SelectionSetToModify->SetSelection(PastedHandles, FTypedElementSelectionOptions());
+	}
+
+	return PastedHandles;
+}
+
 bool UTypedElementCommonActions::DeleteNormalizedElements(const FScriptTypedElementListProxy ElementList, UWorld* World, UTypedElementSelectionSet* InSelectionSet, const FTypedElementDeletionOptions& DeletionOptions)
 {
 	FTypedElementListPtr NativeList = UE::TypedElementFramework::ConvertToNativeTypedElementList(ElementList.GetElementList());
@@ -149,6 +415,101 @@ TArray<FScriptTypedElementHandle> UTypedElementCommonActions::DuplicateNormalize
 	return TypedElementUtil::ConvertToScriptElementArray(DuplicateNormalizedElements(NativeList, World, LocationOffset), NativeList->GetRegistry());
 }
 
+bool UTypedElementCommonActions::CopyNormalizedElements(const FScriptTypedElementListProxy& ElementList)
+{
+	FTypedElementListPtr NativeList = UE::TypedElementFramework::ConvertToNativeTypedElementList(ElementList.GetElementList());
+	if (!NativeList)
+	{
+		FFrame::KismetExecutionMessage(TEXT("ElementList is in a invalid state."), ELogVerbosity::Error);
+		return {};
+	}
+
+	return CopyNormalizedElements(NativeList);
+}
+
+bool UTypedElementCommonActions::CopyNormalizedElementsToString(const FScriptTypedElementListProxy& ElementList, FString& OutputString)
+{
+	FTypedElementListPtr NativeList = UE::TypedElementFramework::ConvertToNativeTypedElementList(ElementList.GetElementList());
+	if (!NativeList)
+	{
+		FFrame::KismetExecutionMessage(TEXT("ElementList is in a invalid state."), ELogVerbosity::Error);
+		return {};
+	}
+
+	return CopyNormalizedElements(NativeList, &OutputString);
+}
+
+TArray<FScriptTypedElementHandle> UTypedElementCommonActions::K2_PasteElements(UTypedElementSelectionSet* SelectionSet, UWorld* World, const FTypedElementPasteOptions& PasteOption)
+{
+	if (!SelectionSet)
+	{
+		FFrame::KismetExecutionMessage(TEXT("SelectionSet is null."), ELogVerbosity::Error);
+		return {};
+	}
+
+	if (!World)
+	{
+		FFrame::KismetExecutionMessage(TEXT("World is null."), ELogVerbosity::Error);
+		return {};
+	}
+
+	return TypedElementUtil::ConvertToScriptElementArray(PasteElements(SelectionSet, World, PasteOption), SelectionSet->GetElementList()->GetRegistry());
+}
+
+TArray<FScriptTypedElementHandle> UTypedElementCommonActions::PasteElementsFromString(UTypedElementSelectionSet* SelectionSet, UWorld* World, const FTypedElementPasteOptions& PasteOption, const FString& InputString)
+{
+	if (!SelectionSet)
+	{
+		FFrame::KismetExecutionMessage(TEXT("SelectionSet is null."), ELogVerbosity::Error);
+		return {};
+	}
+
+	if (!World)
+	{
+		FFrame::KismetExecutionMessage(TEXT("World is null."), ELogVerbosity::Error);
+		return {};
+	}
+
+	return TypedElementUtil::ConvertToScriptElementArray(PasteElements(SelectionSet, World, PasteOption, &InputString), SelectionSet->GetElementList()->GetRegistry());
+}
+
+TArray<FScriptTypedElementHandle> UTypedElementCommonActions::K2_PasteNormalizedElements(const FScriptTypedElementListProxy& ElementList, UWorld* World, const FTypedElementPasteOptions& PasteOption)
+{
+	FTypedElementListPtr NativeList = UE::TypedElementFramework::ConvertToNativeTypedElementList(ElementList.GetElementList());
+	if (!NativeList)
+	{
+		FFrame::KismetExecutionMessage(TEXT("ElementList is in a invalid state."), ELogVerbosity::Error);
+		return {};
+	}
+
+	if (!World)
+	{
+		FFrame::KismetExecutionMessage(TEXT("World is null."), ELogVerbosity::Error);
+		return {};
+	}
+
+	return TypedElementUtil::ConvertToScriptElementArray(PasteNormalizedElements(NativeList, World, PasteOption), NativeList->GetRegistry());
+}
+
+TArray<FScriptTypedElementHandle> UTypedElementCommonActions::PasteNormalizedElementsFromString(const FScriptTypedElementListProxy& ElementList, UWorld* World, const FTypedElementPasteOptions& PasteOption, const FString& InputString)
+{
+	FTypedElementListPtr NativeList = UE::TypedElementFramework::ConvertToNativeTypedElementList(ElementList.GetElementList());
+	if (!NativeList)
+	{
+		FFrame::KismetExecutionMessage(TEXT("ElementList is in a invalid state."), ELogVerbosity::Error);
+		return {};
+	}
+
+	if (!World)
+	{
+		FFrame::KismetExecutionMessage(TEXT("World is null."), ELogVerbosity::Error);
+		return {};
+	}
+
+	return TypedElementUtil::ConvertToScriptElementArray(PasteNormalizedElements(NativeList, World, PasteOption, &InputString), NativeList->GetRegistry());
+}
+
+
 FTypedElementCommonActionsElement UTypedElementCommonActions::ResolveCommonActionsElement(const FTypedElementHandle& InElementHandle) const
 {
 	return InElementHandle
diff --git a/Engine/Source/Runtime/Engine/Private/UnrealExporter.cpp b/Engine/Source/Runtime/Engine/Private/UnrealExporter.cpp
--- a/Engine/Source/Runtime/Engine/Private/UnrealExporter.cpp
+++ b/Engine/Source/Runtime/Engine/Private/UnrealExporter.cpp
@@ -476,6 +476,9 @@ void UExporter::EmitBeginObject( FOutputDevice& Ar, UObject* Obj, uint32 PortFla
 		}
 	}
 
+	// Emit the object path
+	Ar.Logf(TEXT(" ExportPath=%s"), *FObjectPropertyBase::GetExportPath(Obj, nullptr, nullptr, PortFlags | (PPF_Delimited  & ~PPF_ExportsNotFullyQualified)));
+
 	// end in a return
 	Ar.Logf(TEXT("\r\n"));
 
@@ -635,6 +638,7 @@ void UExporter::ExportObjectInner(const FExportObjectInnerContext* Context, UObj
 
 		if (AActor* Actor = Cast<AActor>(Object))
 		{
+			// Todo PlacementMode consider removing that code when we it will replace the foliage
 			// Export anything extra for the components. Used for instanced foliage.
 			// This is done after the actor properties so these are set when regenerating the extra data objects.
 			TArray<UActorComponent*> Components;
diff --git a/Engine/Source/Runtime/Engine/Public/Elements/Framework/TypedElementCommonActions.h b/Engine/Source/Runtime/Engine/Public/Elements/Framework/TypedElementCommonActions.h
--- a/Engine/Source/Runtime/Engine/Public/Elements/Framework/TypedElementCommonActions.h
+++ b/Engine/Source/Runtime/Engine/Public/Elements/Framework/TypedElementCommonActions.h
@@ -5,6 +5,10 @@
 #include "Elements/Framework/TypedElementListProxy.h"
 #include "Elements/Framework/TypedElementInterfaceCustomization.h"
 #include "Elements/Interfaces/TypedElementWorldInterface.h"
+
+#include "Templates/Function.h"
+#include "Templates/SharedPointer.h"
+
 #include "TypedElementCommonActions.generated.h"
 
 class UTypedElementSelectionSet;
@@ -20,6 +24,8 @@ public:
 	//~ See UTypedElementCommonActions for API docs
 	virtual bool DeleteElements(ITypedElementWorldInterface* InWorldInterface, TArrayView<const FTypedElementHandle> InElementHandles, UWorld* InWorld, UTypedElementSelectionSet* InSelectionSet, const FTypedElementDeletionOptions& InDeletionOptions);
 	virtual void DuplicateElements(ITypedElementWorldInterface* InWorldInterface, TArrayView<const FTypedElementHandle> InElementHandles, UWorld* InWorld, const FVector& InLocationOffset, TArray<FTypedElementHandle>& OutNewElements);
+	virtual void CopyElements(ITypedElementWorldInterface* InWorldInterface,TArrayView<const FTypedElementHandle> InElementHandles, FOutputDevice& Out);
+	virtual TSharedPtr<FWorldElementPasteImporter> GetPasteImporter(ITypedElementWorldInterface* InWorldInterface, const FTypedElementListConstPtr& InSelectedHandles, UWorld* InWorld);
 };
 
 /**
@@ -60,6 +66,29 @@ private:
 	FTypedElementCommonActionsCustomization* CommonActionsCustomization = nullptr;
 };
 
+USTRUCT(BlueprintType)
+struct FTypedElementPasteOptions
+{
+	GENERATED_BODY()
+
+	// Todo Copy And Paste should we consider supporting pasting with surface snapping?
+
+	// Todo Copy And Paste should we add optional options to handle where to paste (Like some sort of per type paste context for example instance can go under existing actor/components or the partition mode)?
+
+	// If provided the SelectionSet selection will only contains the newly added elements
+	UPROPERTY(BlueprintReadWrite, Category = "TypedElementInterfaces|World|PasteOptions")
+	TObjectPtr<UTypedElementSelectionSet> SelectionSetToModify;
+
+	UPROPERTY(BlueprintReadWrite, Category = "TypedElementInterfaces|World|PasteOptions")
+	bool bPasteAtLocation = false;
+
+	UPROPERTY(BlueprintReadWrite, Category = "TypedElementInterfaces|World|PasteOptions")
+	FVector PasteLocation = FVector::ZeroVector;
+
+	// C++ Only. Allow for custom callbacks for some custom top level object in the T3D file
+	TMap<FStringView, TFunction<TSharedRef<FWorldElementPasteImporter> ()>> ExtraCustomImport;
+};
+
 /**
  * A utility to handle higher-level common actions, but default via UTypedElementWorldInterface,
  * but asset editors can customize this behavior via FTypedElementCommonActionsCustomization.
@@ -93,8 +122,39 @@ public:
 	 * Duplicate any elements from the given list that can be duplicated.
 	 * @note This list should have been pre-normalized via UTypedElementSelectionSet::GetNormalizedSelection or UTypedElementSelectionSet::GetNormalizedElementList.
 	 */
-	TArray<FTypedElementHandle> DuplicateNormalizedElements(const FTypedElementListConstPtr& ElementListPtr, UWorld* World, const FVector& LocationOffset);
+	virtual TArray<FTypedElementHandle> DuplicateNormalizedElements(const FTypedElementListConstPtr& ElementListPtr, UWorld* World, const FVector& LocationOffset);
 
+	/**
+	 * Copy any elements from the given selection set that can be copied into the clipboard
+	 * @note Internally this just calls CopyNormalizedElements on the result of UTypedElementSelectionSet::GetNormalizedSelection.
+	 */
+	UFUNCTION(BlueprintCallable, Category = "TypedElementFramework|Common")
+	bool CopySelectedElements(UTypedElementSelectionSet* SelectionSet);
+
+	/**
+	 * Copy any elements from the given selection set that can be copied into the string
+	 * @note Internally this just calls CopyNormalizedElements on the result of UTypedElementSelectionSet::GetNormalizedSelection.
+	 */
+	UFUNCTION(BlueprintCallable, Category = "TypedElementFramework|Common")
+	bool CopySelectedElementsToString(UTypedElementSelectionSet* SelectionSet, FString& OutputString);
+
+	/*
+	 * Copy any elements from the given selection set that can be copied into the clipboard.
+	 * @note This list should have been pre-normalized via UTypedElementSelectionSet::GetNormalizedSelection or UTypedElementSelectionSet::GetNormalizedElementList.
+	 */
+	virtual bool CopyNormalizedElements(const FTypedElementListConstPtr& ElementListPtr, FString* OptionalOutputString = nullptr);
+
+	/**
+	 * Paste any elements from the given string or from the clipboard
+	 * @note Internally this just calls PasteNormalizedElements on the result of UTypedElementSelectionSet::GetNormalizedSelection.
+	 */
+	TArray<FTypedElementHandle> PasteElements(UTypedElementSelectionSet* SelectionSet, UWorld* World, const FTypedElementPasteOptions& PasteOption, const FString* OptionalInputString = nullptr);
+
+	/**
+	 * Paste any elements from the given string or from the clipboard
+	  * @note This list should have been pre-normalized via UTypedElementSelectionSet::GetNormalizedSelection or UTypedElementSelectionSet::GetNormalizedElementList.
+	 */
+	virtual TArray<FTypedElementHandle> PasteNormalizedElements(const FTypedElementListConstPtr& ElementListPtr, UWorld* World, const FTypedElementPasteOptions& PasteOptions, const FString* OptionalInputString = nullptr);
 
 	/**
 	 * Script Api
@@ -120,6 +180,48 @@ public:
 	 */
 	UFUNCTION(BlueprintCallable, Category = "TypedElementFramework|Common")
 	TArray<FScriptTypedElementHandle> DuplicateNormalizedElements(const FScriptTypedElementListProxy ElementList, UWorld* World, const FVector& LocationOffset);
+	
+	/*
+	 * Copy any elements from the given selection set that can be copied into the clipboard.
+	 * @note This list should have been pre-normalized via UTypedElementSelectionSet::GetNormalizedSelection or UTypedElementSelectionSet::GetNormalizedElementList.
+	 */
+	UFUNCTION(BlueprintCallable, Category = "TypedElementFramework|Common")
+	bool CopyNormalizedElements(const FScriptTypedElementListProxy& ElementList);
+
+	/*
+	 * Copy any elements from the given selection set that can be copied into the clipboard.
+	 * @note This list should have been pre-normalized via UTypedElementSelectionSet::GetNormalizedSelection or UTypedElementSelectionSet::GetNormalizedElementList.
+	 */
+	UFUNCTION(BlueprintCallable, Category = "TypedElementFramework|Common")
+	bool CopyNormalizedElementsToString(const FScriptTypedElementListProxy& ElementList, FString& OutputString);
+
+	/**
+	 * Paste any elements from the clipboard
+	 * @note Internally this just calls PasteNormalizedElements on the result of UTypedElementSelectionSet::GetNormalizedSelection.
+	 */
+	UFUNCTION(BlueprintCallable, DisplayName="Paste Elements", Category = "TypedElementFramework|Common", meta=(ScriptName="PasteElements"))
+	TArray<FScriptTypedElementHandle> K2_PasteElements(UTypedElementSelectionSet* SelectionSet, UWorld* World, const FTypedElementPasteOptions& PasteOption);
+
+	/**
+	 * Paste any elements from the given string
+	 * @note Internally this just calls PasteNormalizedElements on the result of UTypedElementSelectionSet::GetNormalizedSelection.
+	 */
+	UFUNCTION(BlueprintCallable, Category = "TypedElementFramework|Common")
+	TArray<FScriptTypedElementHandle> PasteElementsFromString(UTypedElementSelectionSet* SelectionSet, UWorld* World, const FTypedElementPasteOptions& PasteOption, const FString& InputString);
+
+	/**
+	 * Paste any elements from the clipboard
+	  * @note This list should have been pre-normalized via UTypedElementSelectionSet::GetNormalizedSelection or UTypedElementSelectionSet::GetNormalizedElementList.
+	 */
+	UFUNCTION(BlueprintCallable, DisplayName="Paste Normalized Elements", Category = "TypedElementFramework|Common", meta=(ScriptName="PasteNormalizedElements"))
+	TArray<FScriptTypedElementHandle> K2_PasteNormalizedElements(const FScriptTypedElementListProxy& ElementList, UWorld* World, const FTypedElementPasteOptions& PasteOption);
+
+	/**
+	 * Paste any elements from the given string
+	  * @note This list should have been pre-normalized via UTypedElementSelectionSet::GetNormalizedSelection or UTypedElementSelectionSet::GetNormalizedElementList.
+	 */
+	UFUNCTION(BlueprintCallable, Category = "TypedElementFramework|Common")
+	TArray<FScriptTypedElementHandle> PasteNormalizedElementsFromString(const FScriptTypedElementListProxy& ElementList, UWorld* World, const FTypedElementPasteOptions& PasteOption, const FString& InputString);
 
 private:
 	/**
@@ -127,3 +229,11 @@ private:
 	 */
 	FTypedElementCommonActionsElement ResolveCommonActionsElement(const FTypedElementHandle& InElementHandle) const;
 };
+
+namespace TypedElementCommonActionsUtils
+{
+	/**
+	 * Is the elements Copy and paste currently enabled?
+	 */
+	ENGINE_API bool IsElementCopyAndPasteEnabled();
+}
diff --git a/Engine/Source/Runtime/Engine/Public/Elements/Interfaces/TypedElementWorldInterface.h b/Engine/Source/Runtime/Engine/Public/Elements/Interfaces/TypedElementWorldInterface.h
--- a/Engine/Source/Runtime/Engine/Public/Elements/Interfaces/TypedElementWorldInterface.h
+++ b/Engine/Source/Runtime/Engine/Public/Elements/Interfaces/TypedElementWorldInterface.h
@@ -8,6 +8,7 @@
 
 #include "CoreFwd.h"
 #include "UObject/Interface.h"
+#include "Templates/SharedPointer.h"
 
 #include "TypedElementWorldInterface.generated.h"
 
@@ -52,6 +53,22 @@ private:
 	bool bWarnAboutSoftReferences = true;
 };
 
+struct FWorldElementPasteImporter
+{
+	virtual ~FWorldElementPasteImporter() = default;
+
+	struct FContext
+	{
+		FTypedElementListConstPtr CurrentSelection;
+		UWorld* World;
+		FStringView Text;
+	};
+	
+	virtual void Import(FContext& Context) {}
+	
+	virtual TArray<FTypedElementHandle> GetImportedElements() { return {}; }
+};
+
 UINTERFACE(MinimalAPI, BlueprintType, meta = (CannotImplementInterfaceInBlueprint))
 class UTypedElementWorldInterface : public UInterface
 {
@@ -199,6 +216,31 @@ public:
 	{
 	}
 
+	virtual bool CanCopyElement(const FTypedElementHandle& InElementHandle) { return false; }
+
+	/**
+	 * Copy the given element into a object to export
+	 * 
+	 * @note Default version calls CopyElements with a single element.
+	 */
+	virtual void CopyElement(const FTypedElementHandle& InElementHandle, FOutputDevice& Out)
+	{
+		CopyElements(MakeArrayView(&InElementHandle, 1), Out);
+	}
+
+	/**
+	 * Copy the given set of elements into a object to export.
+	 * @note If you want to copy an array of elements that are potentially different types, you probably want to use the higher-level UTypedElementCommonActions::CopyNormalizedElements function instead.
+	 */
+	virtual void CopyElements(TArrayView<const FTypedElementHandle> InElementHandles, FOutputDevice& Out)
+	{
+	}
+	
+	virtual TSharedPtr<FWorldElementPasteImporter> GetPasteImporter()
+	{
+		return {};
+	}
+
 	/**
 	 * Can the element be promoted
 	 * Generally available when the element is lighter representation of another element.
@@ -419,6 +461,8 @@ struct TTypedElement<ITypedElementWorldInterface> : public TTypedElementBase<ITy
 	bool DeleteElement(UWorld* InWorld, UTypedElementSelectionSet* InSelectionSet, const FTypedElementDeletionOptions& InDeletionOptions) const { return InterfacePtr->DeleteElement(*this, InWorld, InSelectionSet, InDeletionOptions); }
 	bool CanDuplicateElement() const { return InterfacePtr->CanDuplicateElement(*this); }
 	FTypedElementHandle DuplicateElement(UWorld* InWorld, const FVector& InLocationOffset) const { return InterfacePtr->DuplicateElement(*this, InWorld, InLocationOffset); }
+	bool CanCopyElement() const { return InterfacePtr->CanCopyElement(*this); }
+	void CopyElement(FOutputDevice& Out) const { InterfacePtr->CopyElement(*this, Out); }
 	bool CanPromoteElement() const { return InterfacePtr->CanPromoteElement(*this); }
 	FTypedElementHandle PromoteElement(UWorld* OverrideWorld = nullptr) const { return InterfacePtr->PromoteElement(*this, OverrideWorld); }
 	bool IsElementInConvexVolume(const FConvexVolume& InVolume, bool bMustEncompassEntireElement = false) const { return InterfacePtr->IsElementInConvexVolume(*this, InVolume, bMustEncompassEntireElement); }
diff --git a/Engine/Source/Runtime/Engine/Public/WorldPartition/DataLayer/WorldDataLayers.h b/Engine/Source/Runtime/Engine/Public/WorldPartition/DataLayer/WorldDataLayers.h
--- a/Engine/Source/Runtime/Engine/Public/WorldPartition/DataLayer/WorldDataLayers.h
+++ b/Engine/Source/Runtime/Engine/Public/WorldPartition/DataLayer/WorldDataLayers.h
@@ -52,7 +52,7 @@ public:
 
 #if WITH_EDITOR
 	virtual bool ShouldLevelKeepRefIfExternal() const override { return true; }
-	virtual bool ShouldImport(FString* ActorPropString, bool IsMovingLevel) override { return false; }
+	virtual bool ShouldImport(FStringView ActorPropString, bool IsMovingLevel) override { return false; }
 	virtual bool IsLockLocation() const { return true; }
 	virtual bool IsUserManaged() const override { return false; }
 	virtual TUniquePtr<class FWorldPartitionActorDesc> CreateClassActorDesc() const override;
diff --git a/Engine/Source/Runtime/Engine/Public/WorldPartition/HLOD/HLODActor.h b/Engine/Source/Runtime/Engine/Public/WorldPartition/HLOD/HLODActor.h
--- a/Engine/Source/Runtime/Engine/Public/WorldPartition/HLOD/HLODActor.h
+++ b/Engine/Source/Runtime/Engine/Public/WorldPartition/HLOD/HLODActor.h
@@ -80,7 +80,7 @@ protected:
 	virtual void GetActorBounds(bool bOnlyCollidingComponents, FVector& Origin, FVector& BoxExtent, bool bIncludeFromChildActors) const override;
 	virtual FBox GetStreamingBounds() const override;
 
-	virtual bool ShouldImport(FString* ActorPropString, bool IsMovingLevel) override { return false; }
+	virtual bool ShouldImport(FStringView ActorPropString, bool IsMovingLevel) override { return false; }
 	virtual bool IsLockLocation() const { return true; }
 	virtual bool IsUserManaged() const override { return false; }
 #endif
diff --git a/Engine/Source/Runtime/Foliage/Private/InstancedFoliage.cpp b/Engine/Source/Runtime/Foliage/Private/InstancedFoliage.cpp
--- a/Engine/Source/Runtime/Foliage/Private/InstancedFoliage.cpp
+++ b/Engine/Source/Runtime/Foliage/Private/InstancedFoliage.cpp
@@ -4094,7 +4094,7 @@ bool AInstancedFoliageActor::ShouldExport()
 	return false;
 }
 
-bool AInstancedFoliageActor::ShouldImport(FString* ActorPropString, bool IsMovingLevel)
+bool AInstancedFoliageActor::ShouldImport(FStringView ActorPropString, bool IsMovingLevel)
 {
 	return false;
 }
diff --git a/Engine/Source/Runtime/Foliage/Public/InstancedFoliageActor.h b/Engine/Source/Runtime/Foliage/Public/InstancedFoliageActor.h
--- a/Engine/Source/Runtime/Foliage/Public/InstancedFoliageActor.h
+++ b/Engine/Source/Runtime/Foliage/Public/InstancedFoliageActor.h
@@ -182,7 +182,7 @@ public:
 	virtual void PostEditUndo() override;
 	virtual void PostDuplicate(bool bDuplicateForPIE) override;
 	virtual bool ShouldExport() override;
-	virtual bool ShouldImport(FString* ActorPropString, bool IsMovingLevel) override;
+	virtual bool ShouldImport(FStringView ActorPropString, bool IsMovingLevel) override;
 	virtual FBox GetStreamingBounds() const override;
 
 	// Called in response to BSP rebuilds to migrate foliage from obsolete to new components.
diff --git a/Engine/Source/Runtime/Landscape/Classes/Landscape.h b/Engine/Source/Runtime/Landscape/Classes/Landscape.h
--- a/Engine/Source/Runtime/Landscape/Classes/Landscape.h
+++ b/Engine/Source/Runtime/Landscape/Classes/Landscape.h
@@ -246,7 +246,7 @@ public:
 	virtual void PostEditUndo() override;
 	virtual void PostRegisterAllComponents() override;
 	virtual void PostActorCreated() override;
-	virtual bool ShouldImport(FString* ActorPropString, bool IsMovingLevel) override;
+	virtual bool ShouldImport(FStringView ActorPropString, bool IsMovingLevel) override;
 	virtual void PostEditImport() override;
 	virtual void PostDuplicate(bool bDuplicateForPIE) override;
 	virtual bool CanDeleteSelectedActor(FText& OutReason) const override;
diff --git a/Engine/Source/Runtime/Landscape/Classes/LandscapeProxy.h b/Engine/Source/Runtime/Landscape/Classes/LandscapeProxy.h
--- a/Engine/Source/Runtime/Landscape/Classes/LandscapeProxy.h
+++ b/Engine/Source/Runtime/Landscape/Classes/LandscapeProxy.h
@@ -811,7 +811,7 @@ public:
 	virtual void EditorApplyScale(const FVector& DeltaScale, const FVector* PivotLocation, bool bAltDown, bool bShiftDown, bool bCtrlDown) override;
 	virtual void EditorApplyMirror(const FVector& MirrorScale, const FVector& PivotLocation) override;
 	virtual void PostEditMove(bool bFinished) override;
-	virtual bool ShouldImport(FString* ActorPropString, bool IsMovingLevel) override;
+	virtual bool ShouldImport(FStringView ActorPropString, bool IsMovingLevel) override;
 	virtual bool ShouldExport() override;
 	//~ End AActor Interface
 	virtual uint32 GetDefaultGridSize(UWorld* InWorld) const override { return 1; }
diff --git a/Engine/Source/Runtime/Landscape/Private/LandscapeEdit.cpp b/Engine/Source/Runtime/Landscape/Private/LandscapeEdit.cpp
--- a/Engine/Source/Runtime/Landscape/Private/LandscapeEdit.cpp
+++ b/Engine/Source/Runtime/Landscape/Private/LandscapeEdit.cpp
@@ -4263,7 +4263,7 @@ void ALandscape::PostActorCreated()
 	bIncludeGridSizeInNameForLandscapeActors = true;
 }
 
-bool ALandscape::ShouldImport(FString* ActorPropString, bool IsMovingLevel)
+bool ALandscape::ShouldImport(FStringView ActorPropString, bool IsMovingLevel)
 {
 	return GetWorld() != nullptr && !GetWorld()->IsGameWorld();
 }
@@ -6597,14 +6597,14 @@ bool ALandscapeProxy::ShouldExport()
 	return true;
 }
 
-bool ALandscapeProxy::ShouldImport(FString* ActorPropString, bool IsMovingToLevel)
+bool ALandscapeProxy::ShouldImport(FStringView ActorPropString, bool IsMovingToLevel)
 {
 	bIsMovingToLevel = IsMovingToLevel;
-	if (!bIsMovingToLevel && ActorPropString && ActorPropString->Len() > MAX_LANDSCAPE_PROP_TEXT_LENGTH)
+	if (!bIsMovingToLevel && ActorPropString.Len() > MAX_LANDSCAPE_PROP_TEXT_LENGTH)
 	{
 		// Prompt to save startup packages
 		if (EAppReturnType::Yes == FMessageDialog::Open(EAppMsgType::YesNo, FText::Format(
-			NSLOCTEXT("UnrealEd", "LandscapeImport_Warning", "Landscape is about to import large amount memory ({0}MB) from the clipboard, which will take some time. Do you want to proceed?"), FText::AsNumber(ActorPropString->Len() >> 20))))
+			NSLOCTEXT("UnrealEd", "LandscapeImport_Warning", "Landscape is about to import large amount memory ({0}MB) from the clipboard, which will take some time. Do you want to proceed?"), FText::AsNumber(ActorPropString.Len() >> 20))))
 		{
 			return true;
 		}
