MPCook: missing export warnings: EDL replication needs to happen even in iterative and incremental builds; the export warnings are active in those cooks.
#rnx

diff --git a/Engine/Source/Editor/UnrealEd/Private/CookOnTheFlyServer.cpp b/Engine/Source/Editor/UnrealEd/Private/CookOnTheFlyServer.cpp
--- a/Engine/Source/Editor/UnrealEd/Private/CookOnTheFlyServer.cpp
+++ b/Engine/Source/Editor/UnrealEd/Private/CookOnTheFlyServer.cpp
@@ -9331,7 +9331,7 @@ void UCookOnTheFlyServer::StartCookAsCookWorker()
 
 	// Initialize systems that nothing in StartCookAsCookWorker references
 	// Functions in this section are not dependent upon each other and can be ordered arbitrarily or for async performance	
-	BeginCookEDLCookInfo(BeginContext); // MPCOOKTODO: Need to send information from EDLCookInfo to director
+	BeginCookEDLCookInfo(BeginContext);
 	BeginCookPackageWriters(BeginContext);
 	InitializePollables();
 	if (IsDirectorCookByTheBook())
@@ -9764,13 +9764,6 @@ void UCookOnTheFlyServer::BeginCookEDLCookInfo(FBeginCookContext& BeginContext)
 	{
 		return;
 	}
-	if (!Algo::AllOf(BeginContext.PlatformContexts, [this](const FBeginCookContextPlatform& PlatformContext)
-		{
-			return PlatformContext.PlatformData->bFullBuild;
-		}))
-	{
-		return;
-	}
 	UE::SavePackageUtilities::StartSavingEDLCookInfoForVerification();
 	if (CookDirector)
 	{
