CookOnTheFlyServer: Add more feedback for out of memory.
#rnx

diff --git a/Engine/Source/Editor/UnrealEd/Classes/CookOnTheSide/CookOnTheFlyServer.h b/Engine/Source/Editor/UnrealEd/Classes/CookOnTheSide/CookOnTheFlyServer.h
--- a/Engine/Source/Editor/UnrealEd/Classes/CookOnTheSide/CookOnTheFlyServer.h
+++ b/Engine/Source/Editor/UnrealEd/Classes/CookOnTheSide/CookOnTheFlyServer.h
@@ -798,7 +798,8 @@ public:
 	}
 
 	bool HasExceededMaxMemory() const;
-	void EvaluateGarbageCollectionResults(int32 NumObjectsBeforeGC, const FPlatformMemoryStats& MemStatsBeforeGC);
+	void EvaluateGarbageCollectionResults(int32 NumObjectsBeforeGC, const FPlatformMemoryStats& MemStatsBeforeGC,
+		int32 NumObjectsAfterGC, const FPlatformMemoryStats& MemStatsAfterGC);
 
 	/**
 	 * RequestPackage to be cooked
diff --git a/Engine/Source/Editor/UnrealEd/Private/Commandlets/CookCommandlet.cpp b/Engine/Source/Editor/UnrealEd/Private/Commandlets/CookCommandlet.cpp
--- a/Engine/Source/Editor/UnrealEd/Private/Commandlets/CookCommandlet.cpp
+++ b/Engine/Source/Editor/UnrealEd/Private/Commandlets/CookCommandlet.cpp
@@ -19,6 +19,7 @@
 #include "GlobalShader.h"
 #include "HAL/FileManager.h"
 #include "HAL/IConsoleManager.h"
+#include "HAL/MemoryBase.h"
 #include "HAL/MemoryMisc.h"
 #include "HAL/PlatformApplicationMisc.h"
 #include "HAL/PlatformFileManager.h"
@@ -1024,15 +1025,36 @@ void UCookCommandlet::ConditionalCollectGarbage(uint32 TickResults, UCookOnTheFl
 
 	int32 NumObjectsAfterGC = GUObjectArray.GetObjectArrayNumMinusAvailable();
 	int32 NumObjectsAvailableAfterGC = GUObjectArray.GetObjectArrayEstimatedAvailable();
+	FPlatformMemoryStats MemStatsAfterGC = FPlatformMemory::GetStats();
 	if (!COTFS.IsCookOnTheFlyMode())
 	{
-		UE_LOG(LogCookCommandlet, Display, TEXT("%s GC before %d available %d after %d available %d"),
-			(bPartialGC ? TEXT("Partial") : TEXT("Full")), NumObjectsBeforeGC, NumObjectsAvailableBeforeGC, NumObjectsAfterGC, NumObjectsAvailableAfterGC);
+		int64 VirtualMemBeforeGC = MemStatsBeforeGC.UsedVirtual;
+		int64 VirtualMemAfterGC = MemStatsAfterGC.UsedVirtual;
+		int64 VirtualMemFreed = MemStatsBeforeGC.UsedVirtual - MemStatsAfterGC.UsedVirtual;
+		constexpr int32 BytesPerMeg = 1000000;
+		UE_LOG(LogCookCommandlet, Display, TEXT("GarbageCollection Results:\n")
+			TEXT("\tType: %s\n")
+			TEXT("\tNumObjects:\n")
+			TEXT("\t\tBefore GC:        %10d\n")
+			TEXT("\t\tAvailable Before: %10d\n")
+			TEXT("\t\tAfter GC:         %10d\n")
+			TEXT("\t\tAvailable After:  %10d\n")
+			TEXT("\t\tFreed by GC:      %10d\n")
+			TEXT("\tVirtual Memory:\n")
+			TEXT("\t\tBefore GC:        %10" INT64_FMT " MB\n")
+			TEXT("\t\tAfter GC:         %10" INT64_FMT " MB\n")
+			TEXT("\t\tFreed by GC:      %10" INT64_FMT " MB\n"),
+			(bPartialGC ? TEXT("Partial") : TEXT("Full")),
+			NumObjectsBeforeGC, NumObjectsAvailableBeforeGC, NumObjectsAfterGC, NumObjectsAvailableAfterGC,
+			NumObjectsBeforeGC - NumObjectsAfterGC,
+			VirtualMemBeforeGC / BytesPerMeg, VirtualMemAfterGC / BytesPerMeg, VirtualMemFreed / BytesPerMeg
+		);
+
 		DumpMemStats();
 	}
 
 	if (TickResults & UCookOnTheFlyServer::COSR_RequiresGC_OOM)
 	{
-		COTFS.EvaluateGarbageCollectionResults(NumObjectsBeforeGC, MemStatsBeforeGC);
+		COTFS.EvaluateGarbageCollectionResults(NumObjectsBeforeGC, MemStatsBeforeGC, NumObjectsAfterGC, MemStatsAfterGC);
 	} 
 }
diff --git a/Engine/Source/Editor/UnrealEd/Private/CookOnTheFlyServer.cpp b/Engine/Source/Editor/UnrealEd/Private/CookOnTheFlyServer.cpp
--- a/Engine/Source/Editor/UnrealEd/Private/CookOnTheFlyServer.cpp
+++ b/Engine/Source/Editor/UnrealEd/Private/CookOnTheFlyServer.cpp
@@ -4499,7 +4499,8 @@ bool UCookOnTheFlyServer::HasExceededMaxMemory() const
 	}
 }
 
-void UCookOnTheFlyServer::EvaluateGarbageCollectionResults(int32 NumObjectsBeforeGC, const FPlatformMemoryStats& MemStatsBeforeGC)
+void UCookOnTheFlyServer::EvaluateGarbageCollectionResults(int32 NumObjectsBeforeGC,
+	const FPlatformMemoryStats& MemStatsBeforeGC, int32 NumObjectsAfterGC, const FPlatformMemoryStats& MemStatsAfterGC)
 {
 	using namespace UE::Cook;
 
@@ -4511,9 +4512,7 @@ void UCookOnTheFlyServer::EvaluateGarbageCollectionResults(int32 NumObjectsBefor
 	int64 NumObjectsMin = NumObjectsHistory.GetMinimum();
 	int64 NumObjectsMax = NumObjectsHistory.GetMaximum();
 	int64 NumObjectsSpread = NumObjectsMax - NumObjectsMin;
-	int64 NumObjectsAfterGC = GUObjectArray.GetObjectArrayNumMinusAvailable();
 	int64 NumObjectsFreed = NumObjectsBeforeGC - NumObjectsAfterGC;
-	FPlatformMemoryStats MemStatsAfterGC = FPlatformMemory::GetStats();
 	int64 VirtualMemMin = VirtualMemoryHistory.GetMinimum();
 	int64 VirtualMemMax = VirtualMemoryHistory.GetMaximum();
 	int64 VirtualMemSpread = VirtualMemMax - VirtualMemMin;
@@ -4555,35 +4554,51 @@ void UCookOnTheFlyServer::EvaluateGarbageCollectionResults(int32 NumObjectsBefor
 	constexpr int32 BytesPerMeg = 1000000;
 	UE_LOG(LogCook, Display, TEXT("Garbage Collection was not very impactful.\n")
 		TEXT("\tNumObjects:\n")
-		TEXT("\t\tProcess Min:    %10" INT64_FMT "\n")
-		TEXT("\t\tProcess Max:    %10" INT64_FMT "\n")
-		TEXT("\t\tProcess Spread: %10" INT64_FMT "\n")
-		TEXT("\t\tBefore GC:      %10" INT64_FMT "\n")
-		TEXT("\t\tAfter GC:       %10" INT64_FMT "\n")
-		TEXT("\t\tFreed by GC:    %10" INT64_FMT "\n")
+		TEXT("\t\tProcess Min:      %10" INT64_FMT "\n")
+		TEXT("\t\tProcess Max:      %10" INT64_FMT "\n")
+		TEXT("\t\tProcess Spread:   %10" INT64_FMT "\n")
+		TEXT("\t\tBefore GC:        %10" INT64_FMT "\n")
+		TEXT("\t\tAfter GC:         %10" INT64_FMT "\n")
+		TEXT("\t\tFreed by GC:      %10" INT64_FMT "\n")
 		TEXT("\tVirtual Memory:\n")
-		TEXT("\t\tProcess Min:    %10" INT64_FMT " MB\n")
-		TEXT("\t\tProcess Max:    %10" INT64_FMT " MB\n")
-		TEXT("\t\tProcess Spread: %10" INT64_FMT " MB\n")
-		TEXT("\t\tBefore GC:      %10" INT64_FMT " MB\n")
-		TEXT("\t\tAfter GC:       %10" INT64_FMT " MB\n")
-		TEXT("\t\tFreed by GC:    %10" INT64_FMT " MB\n")
+		TEXT("\t\tProcess Min:      %10" INT64_FMT " MB\n")
+		TEXT("\t\tProcess Max:      %10" INT64_FMT " MB\n")
+		TEXT("\t\tProcess Spread:   %10" INT64_FMT " MB\n")
+		TEXT("\t\tBefore GC:        %10" INT64_FMT " MB\n")
+		TEXT("\t\tAfter GC:         %10" INT64_FMT " MB\n")
+		TEXT("\t\tFreed by GC:      %10" INT64_FMT " MB\n")
 		TEXT("\tReferences:\n")
 		TEXT("\t\tCooker direct packages:         %10" INT64_FMT "\n")
 		TEXT("\t\tCooker transitive packages:     %10" INT64_FMT "\n")
 		TEXT("\t\tCooker direct package size:     %10" INT64_FMT " MB\n")
 		TEXT("\t\tCooker transitive package size: %10" INT64_FMT " MB\n"),
 		NumObjectsMin, NumObjectsMax, NumObjectsSpread,
-		(int64)NumObjectsBeforeGC, NumObjectsAfterGC, NumObjectsFreed,
+		(int64)NumObjectsBeforeGC, (int64)NumObjectsAfterGC, NumObjectsFreed,
 		VirtualMemMin / BytesPerMeg, VirtualMemMax / BytesPerMeg, VirtualMemSpread / BytesPerMeg,
 		VirtualMemBeforeGC / BytesPerMeg, VirtualMemAfterGC / BytesPerMeg, VirtualMemFreed / BytesPerMeg,
 		NumDirectPackages, NumTransitivePackages,
 		DirectResourceSize.GetTotalMemoryBytes() / BytesPerMeg,
 		TransitiveResourceSize.GetTotalMemoryBytes() / BytesPerMeg
 	);
+
+	if (IsCookByTheBookMode())
+	{
+		UE_LOG(LogCook, Display, TEXT("See log for memory use information for UObject classes and LLM tags."));
+		UE::Cook::DumpObjClassList(CookByTheBookOptions->SessionStartupObjects);
 #if ENABLE_LOW_LEVEL_MEM_TRACKER
-	FLowLevelMemTracker::Get().DumpToLog();
+		if (FLowLevelMemTracker::Get().IsEnabled())
+		{
+			GLog->Logf(TEXT("LLM Tags:"));
+			FLowLevelMemTracker::Get().UpdateStatsPerFrame();
+			FLowLevelMemTracker::Get().DumpToLog();
+		}
+		else
 #endif
+		{
+			GLog->Logf(TEXT("LLM Tags are not displayed because llm is disabled. Run with -llm or -trace=memtag to see llm tags."));
+		}
+		GLog->Flush();
+	}
 }
 
 void UCookOnTheFlyServer::GetDirectAndTransitiveResourceSize(FResourceSizeEx& OutDirectSize,
@@ -9271,6 +9286,17 @@ FBeginCookContext UCookOnTheFlyServer::CreateBeginCookByTheBookContext(const FCo
 		// PlatformContext.PlatformData is currently null and is set in SelectSessionPlatforms
 	}
 
+	if (!IsCookingInEditor())
+	{
+		TArray<FWeakObjectPtr>& SessionStartupObjects = CookByTheBookOptions->SessionStartupObjects;
+		SessionStartupObjects.Reset();
+		for (FThreadSafeObjectIterator Iter; Iter; ++Iter)
+		{
+			SessionStartupObjects.Emplace(*Iter);
+		}
+		SessionStartupObjects.Shrink();
+	}
+
 	return BeginContext;
 }
 
diff --git a/Engine/Source/Editor/UnrealEd/Private/Cooker/CookProfiling.cpp b/Engine/Source/Editor/UnrealEd/Private/Cooker/CookProfiling.cpp
--- a/Engine/Source/Editor/UnrealEd/Private/Cooker/CookProfiling.cpp
+++ b/Engine/Source/Editor/UnrealEd/Private/Cooker/CookProfiling.cpp
@@ -1,7 +1,27 @@
 // Copyright Epic Games, Inc. All Rights Reserved.
 
 #include "CookProfiling.h"
+
+#include "Algo/GraphConvert.h"
+#include "Algo/Sort.h"
+#include "Algo/Unique.h"
+#include "Containers/Array.h"
+#include "Containers/BitArray.h"
+#include "Containers/Map.h"
+#include "Containers/UnrealString.h"
 #include "CookOnTheSide/CookLog.h"
+#include "CoreGlobals.h"
+#include "Misc/OutputDevice.h"
+#include "Misc/StringBuilder.h"
+#include "Serialization/ArchiveUObject.h"
+#include "Templates/Casts.h"
+#include "UObject/GCObject.h"
+#include "UObject/MetaData.h"
+#include "UObject/Package.h"
+#include "UObject/Object.h"
+#include "UObject/UObjectGlobals.h"
+#include "UObject/UObjectIterator.h"
+#include "UObject/WeakObjectPtr.h"
 
 #if OUTPUT_COOKTIMING || ENABLE_COOK_STATS
 #include "ProfilingDebugging/ScopedTimers.h"
@@ -198,3 +218,397 @@ namespace DetailedCookStats
 		});
 }
 #endif
+
+namespace UE::Cook
+{
+
+/** The various ways objects can be referenced that keeps them in memory. */
+enum class EObjectReferencerType : uint8
+{
+	Unknown = 0,
+	Rooted,
+	GCObjectRef,
+	Referenced,
+};
+
+/**
+ * Data for how an object is referenced in the DumpObjClassList graph search,
+ * including the type of reference and the vertex of the referencer.
+ */
+struct FObjectReferencer
+{
+	FObjectReferencer() = default;
+	explicit FObjectReferencer(EObjectReferencerType InLinkType, Algo::Graph::FVertex InVertexArgument = Algo::Graph::InvalidVertex)
+	{
+		Set(InLinkType, InVertexArgument);
+	}
+
+	Algo::Graph::FVertex GetVertexArgument() const
+	{
+		return VertexArgument;
+	}
+	EObjectReferencerType GetLinkType()
+	{
+		return LinkType;
+	}
+	void Set(EObjectReferencerType InLinkType, Algo::Graph::FVertex InVertexArgument = Algo::Graph::InvalidVertex)
+	{
+		switch (InLinkType)
+		{
+		case EObjectReferencerType::GCObjectRef:
+			check(InVertexArgument != Algo::Graph::InvalidVertex);
+			break;
+		case EObjectReferencerType::Referenced:
+			check(InVertexArgument != Algo::Graph::InvalidVertex);
+			break;
+		default:
+			break;
+		}
+
+		VertexArgument = InVertexArgument;
+		LinkType = InLinkType;
+	}
+	void ToString(FStringBuilderBase& Builder, TConstArrayView<UObject*> VertexToObject)
+	{
+		switch (GetLinkType())
+		{
+		case EObjectReferencerType::Unknown:
+			Builder << TEXT("<Unknown>");
+			break;
+		case EObjectReferencerType::Rooted:
+			Builder << TEXT("<Rooted>");
+			break;
+		case EObjectReferencerType::GCObjectRef:
+		{
+			check(VertexArgument != Algo::Graph::InvalidVertex);
+			UObject* Object = VertexToObject[VertexArgument];
+			FString ReferencerName;
+			if (!Object || !FGCObject::GGCObjectReferencer->GetReferencerName(Object, ReferencerName))
+			{
+				ReferencerName = TEXT("<Unknown>");
+			}
+			Builder << TEXT("FGCObject ") << ReferencerName;
+			break;
+		}
+		case EObjectReferencerType::Referenced:
+		{
+			check(VertexArgument != Algo::Graph::InvalidVertex);
+			UObject* Object = VertexToObject[VertexArgument];
+			if (Object)
+			{
+				Object->GetPathName(nullptr, Builder);
+			}
+			else
+			{
+				Builder << TEXT("<UnknownObject>");
+			}
+			break;
+		}
+		default:
+			checkNoEntry();
+			break;
+		}
+	}
+
+private:
+	Algo::Graph::FVertex VertexArgument = Algo::Graph::InvalidVertex;
+	EObjectReferencerType LinkType = EObjectReferencerType::Unknown;
+};
+
+/** An ObjectReferenceCollector to pass to Object->Serialize to collect references into an array. */
+class FArchiveGetReferences : public FArchiveUObject
+{
+public:
+	FArchiveGetReferences(UObject* Object, TArray<UObject*>& OutReferencedObjects)
+		:ReferencedObjects(OutReferencedObjects)
+	{
+		ArIsObjectReferenceCollector = true;
+		ArIgnoreOuterRef = false;
+		SetShouldSkipCompilingAssets(false);
+		Object->Serialize(*this);
+	}
+
+	FArchive& operator<<(UObject*& Object)
+	{
+		if (Object)
+		{
+			ReferencedObjects.Add(Object);
+		}
+
+		return *this;
+	}
+private:
+	TArray<UObject*>& ReferencedObjects;
+};
+
+/**
+ * Given the list of AllObjects from e.g. a TObjectIterator, use serialization and other methods from Garbage Collection
+ * to find all the dependencies of each Object.
+ * Return the dependencies as a normalized graph in the style of GraphConvert.h, with the vertex of each object defined
+ * by AllObjects and ObjectToVertex.
+ */
+void ConstructObjectGraph(TConstArrayView<UObject*> AllObjects,
+	const TMap<UObject*, Algo::Graph::FVertex>& ObjectToVertex, TArray64<Algo::Graph::FVertex>& OutGraphBuffer,
+	TArray<TConstArrayView<Algo::Graph::FVertex>>& OutGraph)
+{
+	using namespace Algo::Graph;
+
+	TArray<TArray<FVertex>> LooseEdges;
+	int32 NumVertices = AllObjects.Num();
+	LooseEdges.SetNum(NumVertices);
+	TArray<UObject*> TargetObjects;
+	int32 NumEdges = 0;
+
+	for (FVertex SourceVertex = 0; SourceVertex < NumVertices; ++SourceVertex)
+	{
+		UObject* SourceObject = AllObjects[SourceVertex];
+		TargetObjects.Reset();
+		{
+			FReferenceFinder Collector(TargetObjects);
+			if (SourceObject == FGCObject::GGCObjectReferencer)
+			{
+				UGCObjectReferencer::AddReferencedObjects(FGCObject::GGCObjectReferencer, Collector);
+			}
+			else
+			{
+				FArchiveGetReferences Ar(SourceObject, TargetObjects);
+				if (SourceObject->GetClass())
+				{
+					SourceObject->GetClass()->CallAddReferencedObjects(SourceObject, Collector);
+				}
+				// TODO: Handle elements in the token stream not covered by serialize, such as UPackage's
+				// Class->EmitObjectReference(STRUCT_OFFSET(UPackage, MetaData), TEXT("MetaData"));
+				// In the meantime we handle MetaData explicitly.
+				if (UPackage* AsPackage = Cast<UPackage>(SourceObject))
+				{
+					TargetObjects.Add(AsPackage->GetMetaData());
+				}
+			}
+		}
+		if (TargetObjects.Num())
+		{
+			Algo::Sort(TargetObjects);
+			TargetObjects.SetNum(Algo::Unique(TargetObjects), false /* bAllowShrinking */);
+			TArray<FVertex>& TargetVertices = LooseEdges[SourceVertex];
+			TargetVertices.Reserve(TargetObjects.Num());
+			for (UObject* TargetObject : TargetObjects)
+			{
+				const FVertex* TargetVertex = ObjectToVertex.Find(TargetObject);
+				if (TargetVertex && *TargetVertex != SourceVertex)
+				{
+					TargetVertices.Add(*TargetVertex);
+				}
+			}
+			NumEdges += TargetVertices.Num();
+		}
+	}
+	OutGraphBuffer.Empty(NumEdges);
+	OutGraph.Empty(NumVertices);
+	for (FVertex SourceVertex = 0; SourceVertex < NumVertices; ++SourceVertex)
+	{
+		TArray<FVertex>& InEdges = LooseEdges[SourceVertex];
+		TConstArrayView<FVertex>& OutEdges = OutGraph.Emplace_GetRef();
+		OutEdges = TConstArrayView<FVertex>(OutGraphBuffer.GetData() + OutGraphBuffer.Num(), InEdges.Num());
+		OutGraphBuffer.Append(InEdges);
+	}
+}
+
+void DumpObjClassList(TConstArrayView<FWeakObjectPtr> InitialObjects)
+{
+	using namespace Algo::Graph;
+
+	FOutputDevice& LogAr = *(GLog);
+
+	// Get the list of Objects
+	TArray<UObject*> AllObjects;
+	for (FThreadSafeObjectIterator Iter; Iter; ++Iter)
+	{
+		UObject* Object = *Iter;
+		if (!Object)
+		{
+			continue;
+		}
+		AllObjects.Add(Object);
+	}
+
+	// Convert Objects to Algo::Graph::FVertex to reduce graph search memory
+	int32 NumVertices = AllObjects.Num();
+	TMap<UObject*, FVertex> VertexOfObject;
+	for (FVertex Vertex = 0; Vertex < NumVertices; ++Vertex)
+	{
+		VertexOfObject.Add(AllObjects[Vertex], Vertex);
+	}
+
+	// Store for each vertex whether the vertex is new - not in InitialObjects
+	TBitArray<> IsNew(true, NumVertices);
+	for (const FWeakObjectPtr& InitialObjectWeak : InitialObjects)
+	{
+		UObject* InitialObject = InitialObjectWeak.Get();
+		if (InitialObject)
+		{
+			FVertex* Vertex = VertexOfObject.Find(InitialObject);
+			if (Vertex)
+			{
+				IsNew[*Vertex] = false;
+			}
+		}
+	}
+
+	// Serialize objects to get dependencies and use them to create the ObjectGraph
+	TArray64<FVertex> ObjectGraphBuffer;
+	TArray<TConstArrayView<FVertex>> ObjectGraph;
+	ConstructObjectGraph(AllObjects, VertexOfObject, ObjectGraphBuffer, ObjectGraph);
+
+	// Mark the objects that are rooted by IsRooted, and find any special vertices
+	FVertex GCObjectReferencerVertex = InvalidVertex;
+	TArray<FObjectReferencer> AliveReason;
+	AliveReason.SetNum(NumVertices);
+	for (FVertex Vertex = 0; Vertex < NumVertices; ++Vertex)
+	{
+		UObject* Object = AllObjects[Vertex];
+		if (Object->IsRooted())
+		{
+			AliveReason[Vertex].Set(EObjectReferencerType::Rooted);
+		}
+		if (Object == FGCObject::GGCObjectReferencer)
+		{
+			GCObjectReferencerVertex = Vertex;
+		}
+	}
+
+	// Mark the objects that are rooted by GCObjectReferencerVertex
+	for (FVertex Vertex : ObjectGraph[GCObjectReferencerVertex])
+	{
+		if (AliveReason[Vertex].GetLinkType() == EObjectReferencerType::Unknown)
+		{
+			AliveReason[Vertex].Set(EObjectReferencerType::GCObjectRef, Vertex);
+		}
+	}
+	check(GCObjectReferencerVertex != InvalidVertex);
+
+	// Do a DFS to mark the referencer and root of all non-rooted objects
+	TArray<FVertex> RootOfVertex;
+	RootOfVertex.SetNumUninitialized(NumVertices);
+	for (FVertex& Root : RootOfVertex)
+	{
+		Root = InvalidVertex;
+	}
+
+	TArray<FVertex> Stack;
+	for (FVertex RootedVertex = 0; RootedVertex < NumVertices; ++RootedVertex)
+	{
+		if (AliveReason[RootedVertex].GetLinkType() == EObjectReferencerType::Unknown ||
+			RootedVertex == GCObjectReferencerVertex)
+		{
+			continue;
+		}
+
+		RootOfVertex[RootedVertex] = RootedVertex;
+		Stack.Reset();
+		Stack.Add(RootedVertex);
+		while (!Stack.IsEmpty())
+		{
+			FVertex SourceVertex = Stack.Pop(false /* bAllowShrinking */);
+			for (FVertex TargetVertex : ObjectGraph[SourceVertex])
+			{
+				if (AliveReason[TargetVertex].GetLinkType() == EObjectReferencerType::Unknown)
+				{
+					AliveReason[TargetVertex].Set(EObjectReferencerType::Referenced, SourceVertex);
+					RootOfVertex[TargetVertex] = RootedVertex;
+					Stack.Add(TargetVertex);
+				}
+			}
+		}
+	}
+
+	// Count how many new objects of each class there are, and store all root objects that keep them in memory
+	struct FClassInfo
+	{
+		TMap<FVertex, int32> Roots;
+		int32 Count = 0;
+		UClass* Class = nullptr;
+	};
+	TMap<UClass*, FClassInfo> ClassInfos;
+	for (FVertex Vertex = 0; Vertex < NumVertices; ++Vertex)
+	{
+		// Ignore non-new objects
+		if (!IsNew[Vertex] || Vertex == GCObjectReferencerVertex)
+		{
+			continue;
+		}
+		FObjectReferencer Link = AliveReason[Vertex];
+		EObjectReferencerType LinkType = Link.GetLinkType();
+		// Ignore objects that have AliveReason unknown. This can occur if the objects were rooted during garbage
+		// collection but then asynchronous work RemovedThemFromRoot in between GC finishing and our call to IsRooted.
+		if (LinkType == EObjectReferencerType::Unknown)
+		{
+			continue;
+		}
+		UClass* Class = AllObjects[Vertex]->GetClass();
+		if (!Class || !Class->IsNative())
+		{
+			continue;
+		}
+		FClassInfo& ClassInfo = ClassInfos.FindOrAdd(Class);
+		ClassInfo.Class = Class;
+		ClassInfo.Roots.FindOrAdd(RootOfVertex[Vertex], 0)++;
+		ClassInfo.Count++;
+	}
+
+	TArray<FClassInfo> ClassInfoArray;
+	ClassInfoArray.Reserve(ClassInfos.Num());
+	for (TPair<UClass*, FClassInfo>& Pair : ClassInfos)
+	{
+		ClassInfoArray.Add(MoveTemp(Pair.Value));
+	}
+	Algo::Sort(ClassInfoArray, [](const FClassInfo& A, const FClassInfo& B)
+		{
+			return FTopLevelAssetPath(A.Class).Compare(FTopLevelAssetPath(B.Class)) < 0;
+		});
+
+
+	LogAr.Logf(TEXT("New Objects of each class and the top roots keeping them alive:"));
+	LogAr.Logf(TEXT("\t%6s %s"), TEXT("Count"), TEXT("ClassPath"));
+	LogAr.Logf(TEXT("\t\t%6s %s"), TEXT("Count"), TEXT("RootObjectAndChain"));
+	TStringBuilder<1024> RootObjectString;
+	constexpr int32 MaxRootCount = 2;
+	TArray<TPair<FVertex, int32>, TInlineAllocator<MaxRootCount + 1>> MaxRoots;
+	for (FClassInfo& ClassInfo : ClassInfoArray)
+	{
+		MaxRoots.Reset();
+		for (TPair<FVertex, int32>& RootPair : ClassInfo.Roots)
+		{
+			for (int32 IndexFromMax = 0; IndexFromMax < MaxRootCount; ++IndexFromMax)
+			{
+				if (MaxRoots.Num() <= IndexFromMax || MaxRoots[IndexFromMax].Value < RootPair.Value)
+				{
+					MaxRoots.Insert(RootPair, IndexFromMax);
+					break;
+				}
+			}
+			if (MaxRoots.Num() > MaxRootCount)
+			{
+				MaxRoots.Pop(false /* bAllowShrinking */);
+			}
+		}
+		LogAr.Logf(TEXT("\t%6d %s"), ClassInfo.Count, *ClassInfo.Class->GetPathName());
+		for (TPair<FVertex, int32>& RootPair : MaxRoots)
+		{
+			RootObjectString.Reset();
+			RootObjectString.Appendf(TEXT("\t\t%6d: "), RootPair.Value);
+			AllObjects[RootPair.Key]->GetFullName(RootObjectString);
+			FObjectReferencer Link = AliveReason[RootPair.Key];
+			RootObjectString << TEXT(" <- ");
+			Link.ToString(RootObjectString, AllObjects);
+			while (Link.GetLinkType() == EObjectReferencerType::Referenced)
+			{
+				Link = AliveReason[Link.GetVertexArgument()];
+				RootObjectString << TEXT(" <- ");
+				Link.ToString(RootObjectString, AllObjects);
+			}
+			LogAr.Logf(TEXT("%s"), *RootObjectString);
+		}
+	}
+}
+
+} // namespace UE::Cook
diff --git a/Engine/Source/Editor/UnrealEd/Private/Cooker/CookProfiling.h b/Engine/Source/Editor/UnrealEd/Private/Cooker/CookProfiling.h
--- a/Engine/Source/Editor/UnrealEd/Private/Cooker/CookProfiling.h
+++ b/Engine/Source/Editor/UnrealEd/Private/Cooker/CookProfiling.h
@@ -2,6 +2,7 @@
 
 #pragma once
 
+#include "Containers/ArrayView.h"
 #include "CoreTypes.h"
 #include "ProfilingDebugging/CookStats.h"
 #include "Stats/Stats.h"
@@ -13,6 +14,11 @@
 #include "Trace/Trace.inl"
 #include "ProfilingDebugging/FormatArgsTrace.h"
 #include "ProfilingDebugging/ScopedTimers.h"
+#endif
+
+struct FWeakObjectPtr;
+
+#if OUTPUT_COOKTIMING
 
 void OutputHierarchyTimers();
 void ClearHierarchyTimers();
@@ -152,3 +158,17 @@ DECLARE_STATS_GROUP(TEXT("Cooking"), STATGROUP_Cooking, STATCAT_Advanced);
 
 DECLARE_CYCLE_STAT(TEXT("Precache Derived data for platform"), STAT_TickPrecacheCooking, STATGROUP_Cooking);
 DECLARE_CYCLE_STAT(TEXT("Tick cooking"), STAT_TickCooker, STATGROUP_Cooking);
+
+namespace UE::Cook
+{
+
+/**
+ * Used for profiling memory after a garbage collect. Find how many instances of each class exist, and find the list
+ * of referencers that are keeping instances of each class in memory. Dump this information to the log (GLog).
+ * 
+ * @param SessionStartupObjects: The list of objects that existed at the start of the cook session. Objects in this
+ *                               list are removed from the results since we do not expect them to be garbage collected.
+ */
+void DumpObjClassList(TConstArrayView<FWeakObjectPtr> SessionStartupObjects);
+
+}
\ No newline at end of file
diff --git a/Engine/Source/Editor/UnrealEd/Private/Cooker/CookTypes.cpp b/Engine/Source/Editor/UnrealEd/Private/Cooker/CookTypes.cpp
--- a/Engine/Source/Editor/UnrealEd/Private/Cooker/CookTypes.cpp
+++ b/Engine/Source/Editor/UnrealEd/Private/Cooker/CookTypes.cpp
@@ -329,7 +329,7 @@ bool LoadFromCompactBinary(FCbFieldView Field, UE::Cook::FBeginCookConfigSetting
 FCbWriter& operator<<(FCbWriter& Writer, const UE::Cook::FCookByTheBookOptions& OutValue)
 {
 	Writer.BeginObject();
-	// StartupPackages are process-specific
+	// StartupPackages and SessionStartupObjects are process-specific
 
 	Writer << "DlcName" << OutValue.DlcName;
 	Writer << "CreateReleaseVersion" << OutValue.CreateReleaseVersion;
@@ -356,7 +356,8 @@ FCbWriter& operator<<(FCbWriter& Writer, const UE::Cook::FCookByTheBookOptions&
 bool LoadFromCompactBinary(FCbFieldView Field, UE::Cook::FCookByTheBookOptions& OutValue)
 {
 	bool bOk = Field.IsObject();
-	OutValue.StartupPackages.Reset();
+	OutValue.StartupPackages.Empty();
+	OutValue.SessionStartupObjects.Empty();
 
 	bOk = LoadFromCompactBinary(Field["DlcName"], OutValue.DlcName) & bOk;
 	bOk = LoadFromCompactBinary(Field["CreateReleaseVersion"], OutValue.CreateReleaseVersion) & bOk;
diff --git a/Engine/Source/Editor/UnrealEd/Private/Cooker/CookTypes.h b/Engine/Source/Editor/UnrealEd/Private/Cooker/CookTypes.h
--- a/Engine/Source/Editor/UnrealEd/Private/Cooker/CookTypes.h
+++ b/Engine/Source/Editor/UnrealEd/Private/Cooker/CookTypes.h
@@ -20,6 +20,7 @@
 class FCbFieldView;
 class FCbWriter;
 class ITargetPlatform;
+struct FWeakObjectPtr;
 namespace UE::DerivedData { class FBuildDefinition; }
 namespace UE::Cook { struct FBeginCookConfigSettings; }
 namespace UE::Cook { struct FCookByTheBookOptions; }
@@ -379,6 +380,12 @@ public:
 	TArray<FName>					StartupPackages;
 
 	// Session-lifetime variables
+	/**
+	 * The list of UObjects that existed at the start of the cook. This is used to tell which UObjects
+	 * were created during the cook.
+	 */
+	TArray<FWeakObjectPtr>			SessionStartupObjects;
+
 	/** DlcName setup if we are cooking dlc will be used as the directory to save cooked files to */
 	FString							DlcName;
 
