Add some logging to track down a crash when running Editor Tests

diff --git a/Engine/Source/Editor/UnrealEd/Private/Kismet2/KismetReinstanceUtilities.cpp b/Engine/Source/Editor/UnrealEd/Private/Kismet2/KismetReinstanceUtilities.cpp
--- a/Engine/Source/Editor/UnrealEd/Private/Kismet2/KismetReinstanceUtilities.cpp
+++ b/Engine/Source/Editor/UnrealEd/Private/Kismet2/KismetReinstanceUtilities.cpp
@@ -57,6 +57,7 @@ static FAutoConsoleVariableRef CVarUseLegacyAnimInstanceReinstancingBehavior(
 	TEXT("Use the legacy re-instancing behavior for anim instances where the instance is destroyed and re-created.")
 );
 
+extern ENGINE_API bool GLogReinstancerReferenceReplacement;
 
 struct FReplaceReferenceHelper
 {
@@ -95,6 +96,16 @@ struct FReplaceReferenceHelper
 			return;
 		}
 
+		if (GLogReinstancerReferenceReplacement)
+		{
+			FString Classes;
+			for (UObject* Obj : SourceObjects)
+			{
+				Classes.Append(Obj->GetClass()->GetName() + " ");
+			}
+			UE_LOG(LogBlueprint, Warning, TEXT("Replacing References to: %s"), *Classes);
+		}
+
 		// Remember what values were in UActorChannel::Actor so we can restore them later (this should only affect reinstancing during PIE)
 		// We need the old actor channel to tear down cleanly without affecting the new actor
 		TMap<UActorChannel*, AActor*> ActorChannelActorRestorationMap;
diff --git a/Engine/Source/Runtime/Engine/Private/SkeletalMeshComponentPhysics.cpp b/Engine/Source/Runtime/Engine/Private/SkeletalMeshComponentPhysics.cpp
--- a/Engine/Source/Runtime/Engine/Private/SkeletalMeshComponentPhysics.cpp
+++ b/Engine/Source/Runtime/Engine/Private/SkeletalMeshComponentPhysics.cpp
@@ -2678,6 +2678,10 @@ void USkeletalMeshComponent::AddClothingBounds(FBoxSphereBounds& InOutBounds, co
 	}
 }
 
+#if WITH_EDITOR
+ENGINE_API bool GLogReinstancerReferenceReplacement = false;
+#endif
+
 void USkeletalMeshComponent::RecreateClothingActors()
 {
 	CSV_SCOPED_TIMING_STAT(Animation, ClothInit);
@@ -2690,6 +2694,10 @@ void USkeletalMeshComponent::RecreateClothingActors()
 		return;
 	}
 
+#if WITH_EDITOR
+	TGuardValue<bool> EnableReinstancingLog(GLogReinstancerReferenceReplacement, true);
+#endif
+
 	if(CVarEnableClothPhysics.GetValueOnGameThread() && (SkelMesh->GetMeshClothingAssets().Num() > 0))
 	{
 		UClass* SimFactoryClass = *ClothingSimulationFactory;
