PinActors: Pin of Folder will consider unloaded actors even if they are hidden in the outliner

#rnx

diff --git a/Engine/Source/Editor/SceneOutliner/Private/ActorBrowsingMode.cpp b/Engine/Source/Editor/SceneOutliner/Private/ActorBrowsingMode.cpp
--- a/Engine/Source/Editor/SceneOutliner/Private/ActorBrowsingMode.cpp
+++ b/Engine/Source/Editor/SceneOutliner/Private/ActorBrowsingMode.cpp
@@ -1929,8 +1929,11 @@ namespace ActorBrowsingModeUtils
 		}
 	}
 
-	static void RecursiveAddItemsToActorGuidList(const TArray<FSceneOutlinerTreeItemPtr>& Items, TArray<FGuid>& List)
+	static void RecursiveAddItemsToActorGuidList(const TArray<FSceneOutlinerTreeItemPtr>& Items, TArray<FGuid>& List, bool bSearchForHiddenUnloadedActors)
 	{
+		// In the case where we want the list of unloaded actors under a folder and the option to hide unloaded actors is enabled, we need to find them through FActorFolders 
+		TMap<UWorld*, TSet<FName>> UnloadedActorsFolderPaths;
+
 		for (const FSceneOutlinerTreeItemPtr& Item : Items)
 		{
 			if (const FActorDescTreeItem* const ActorDescTreeItem = Item->CastTo<FActorDescTreeItem>())
@@ -1944,6 +1947,13 @@ namespace ActorBrowsingModeUtils
 					List.Add(ActorTreeItem->Actor->GetActorGuid());
 				}
 			}
+			else if (const FActorFolderTreeItem* const ActorFolderTreeItem = Item->CastTo<FActorFolderTreeItem>(); ActorFolderTreeItem && bSearchForHiddenUnloadedActors)
+			{
+				if (UWorld* FolderWorld = ActorFolderTreeItem->World.Get())
+				{
+					UnloadedActorsFolderPaths.FindOrAdd(FolderWorld).Add(ActorFolderTreeItem->GetPath());
+				}
+			}
 
 			TArray<FSceneOutlinerTreeItemPtr> ChildrenItems;
 			for (const auto& Child : Item->GetChildren())
@@ -1953,12 +1963,25 @@ namespace ActorBrowsingModeUtils
 					ChildrenItems.Add(Child.Pin());
 				}
 			}
-
+						
 			if (ChildrenItems.Num())
 			{
-				RecursiveAddItemsToActorGuidList(ChildrenItems, List);
+				RecursiveAddItemsToActorGuidList(ChildrenItems, List, bSearchForHiddenUnloadedActors);
 			}
 		}
+
+		
+		for (const TTuple<UWorld*, TSet<FName>>& Pair : UnloadedActorsFolderPaths)
+		{
+			FActorFolders::ForEachActorDescInFolders(*Pair.Key, Pair.Value, [&List](const FWorldPartitionActorDesc* ActorDesc)
+			{
+				if (!ActorDesc->IsLoaded())
+				{
+					List.Add(ActorDesc->GetGuid());
+				}
+				return true;
+			});
+		}
 	};
 }
 
@@ -1996,7 +2019,9 @@ void FActorBrowsingMode::PinItems(const TArray<FSceneOutlinerTreeItemPtr>& InIte
 	}
 
 	TArray<FGuid> ActorsToPin;
-	ActorBrowsingModeUtils::RecursiveAddItemsToActorGuidList(InItems, ActorsToPin);
+	// If Unloaded actors are hidden and we are pinning folders we need to find them through FActorFolders
+	const bool bSearchForHiddenUnloadedActors = bHideUnloadedActors;
+	ActorBrowsingModeUtils::RecursiveAddItemsToActorGuidList(InItems, ActorsToPin, bSearchForHiddenUnloadedActors);
 
 	if (ActorsToPin.Num())
 	{
@@ -2036,7 +2061,9 @@ void FActorBrowsingMode::UnpinItems(const TArray<FSceneOutlinerTreeItemPtr>& InI
 	}
 
 	TArray<FGuid> ActorsToUnpin;
-	ActorBrowsingModeUtils::RecursiveAddItemsToActorGuidList(InItems, ActorsToUnpin);
+	// No need to search for hidden unloaded actors when unloading
+	const bool bSearchForHiddenUnloadedActors = false;
+	ActorBrowsingModeUtils::RecursiveAddItemsToActorGuidList(InItems, ActorsToUnpin, bSearchForHiddenUnloadedActors);
 
 	if (ActorsToUnpin.Num())
 	{
diff --git a/Engine/Source/Editor/SceneOutliner/Private/ActorHierarchy.cpp b/Engine/Source/Editor/SceneOutliner/Private/ActorHierarchy.cpp
--- a/Engine/Source/Editor/SceneOutliner/Private/ActorHierarchy.cpp
+++ b/Engine/Source/Editor/SceneOutliner/Private/ActorHierarchy.cpp
@@ -128,24 +128,6 @@ FActorHierarchy::~FActorHierarchy()
 	}
 }
 
-static FFolder GetFolderPathFromActorDesc(UWorld* InWorld, const FWorldPartitionActorDesc* InActorDesc)
-{
-	if (UWorld* OuterWorld = InActorDesc->GetContainer() ? InActorDesc->GetContainer()->GetTypedOuter<UWorld>() : InWorld)
-	{
-		ULevel* OuterLevel = OuterWorld->PersistentLevel;
-		if (OuterLevel->IsUsingActorFolders())
-		{
-			if (UActorFolder* ActorFolder = OuterLevel->GetActorFolder(InActorDesc->GetFolderGuid()))
-			{
-				return ActorFolder->GetFolder();
-			}
-			return FFolder::GetWorldRootFolder(OuterWorld).GetRootObject();
-		}
-		return FFolder(FFolder::GetWorldRootFolder(OuterWorld).GetRootObject(), InActorDesc->GetFolderPath());
-	}
-	return FFolder::GetInvalidFolder();
-}
-
 FSceneOutlinerTreeItemPtr FActorHierarchy::FindOrCreateParentItem(const ISceneOutlinerTreeItem& Item, const TMap<FSceneOutlinerTreeItemID, FSceneOutlinerTreeItemPtr>& Items, bool bCreate)
 {
 	if (Item.IsA<FWorldTreeItem>())
@@ -291,16 +273,19 @@ FSceneOutlinerTreeItemPtr FActorHierarchy::FindOrCreateParentItem(const ISceneOu
 	{
 		if (const FWorldPartitionActorDesc* ActorDesc = ActorDescItem->ActorDescHandle.Get())
 		{
-			const FFolder ActorDescFolder = GetFolderPathFromActorDesc(RepresentingWorld.Get(), ActorDesc);
-			if (!ActorDescFolder.IsNone())
+			if (UWorld* RepresentingWorldPtr = RepresentingWorld.Get())
 			{
-				if (const FSceneOutlinerTreeItemPtr* ParentItem = Items.Find(ActorDescFolder))
+				const FFolder ActorDescFolder = FActorFolders::GetActorDescFolder(*RepresentingWorldPtr, ActorDesc);
+				if (!ActorDescFolder.IsNone())
 				{
-					return *ParentItem;
-				}
-				else
-				{
-					return bCreate ? Mode->CreateItemFor<FActorFolderTreeItem>(FActorFolderTreeItem(ActorDescFolder, RepresentingWorld), true) : nullptr;
+					if (const FSceneOutlinerTreeItemPtr* ParentItem = Items.Find(ActorDescFolder))
+					{
+						return *ParentItem;
+					}
+					else
+					{
+						return bCreate ? Mode->CreateItemFor<FActorFolderTreeItem>(FActorFolderTreeItem(ActorDescFolder, RepresentingWorldPtr), true) : nullptr;
+					}
 				}
 			}
 
diff --git a/Engine/Source/Editor/UnrealEd/Private/EditorActorFolders.cpp b/Engine/Source/Editor/UnrealEd/Private/EditorActorFolders.cpp
--- a/Engine/Source/Editor/UnrealEd/Private/EditorActorFolders.cpp
+++ b/Engine/Source/Editor/UnrealEd/Private/EditorActorFolders.cpp
@@ -19,6 +19,9 @@
 #include "ActorFolder.h"
 #include "Subsystems/ActorEditorContextSubsystem.h"
 #include "Styling/AppStyle.h"
+#include "WorldPartition/WorldPartitionActorDesc.h"
+#include "WorldPartition/WorldPartitionHelpers.h"
+#include "WorldPartition/ActorDescContainer.h"
 
 #define LOCTEXT_NAMESPACE "FActorFolders"
 
@@ -638,7 +641,58 @@ void FActorFolders::ForEachFolderWithRootObject(UWorld& InWorld, const FFolder::
 	GetOrCreateWorldFolders(InWorld).ForEachFolderWithRootObject(FolderRootObject, Operation);
 }
 
-void FActorFolders::ForEachActorInFolders(UWorld& InWorld, const TArray<FName>& Paths, TFunctionRef<bool(AActor*)> Operation, const FFolder::FRootObject& InFolderRootObject /*= FFolder::GetInvalidRootObject()*/)
+FFolder FActorFolders::GetActorDescFolder(UWorld& InWorld, const FWorldPartitionActorDesc* InActorDesc)
+{
+	if (InActorDesc)
+	{
+		UWorld& OuterWorld = InActorDesc->GetContainer() ? *InActorDesc->GetContainer()->GetTypedOuter<UWorld>() : InWorld;
+		{
+			ULevel* OuterLevel = OuterWorld.PersistentLevel;
+			if (OuterLevel->IsUsingActorFolders())
+			{
+				if (UActorFolder* ActorFolder = OuterLevel->GetActorFolder(InActorDesc->GetFolderGuid()))
+				{
+					return ActorFolder->GetFolder();
+				}
+				return FFolder::GetWorldRootFolder(&OuterWorld).GetRootObject();
+			}
+			return FFolder(FFolder::GetWorldRootFolder(&OuterWorld).GetRootObject(), InActorDesc->GetFolderPath());
+		}
+	}
+	return FFolder::GetInvalidFolder();
+}
+
+void FActorFolders::ForEachActorDescInFolders(UWorld& InWorld, const TSet<FName>& InPaths, TFunctionRef<bool(const FWorldPartitionActorDesc*)> Operation, const FFolder::FRootObject& InFolderRootObject /*= FFolder::GetInvalidRootObject()*/)
+{
+	if (UWorldPartition* WorldPartition = InWorld.GetWorldPartition())
+	{
+		FWorldPartitionHelpers::ForEachActorDesc(WorldPartition, [&](const FWorldPartitionActorDesc* ActorDesc)
+		{
+			FFolder ActorDescFolder = GetActorDescFolder(InWorld, ActorDesc);
+			if (ActorDescFolder == FFolder::GetInvalidFolder())
+			{
+				return true;
+			}
+
+			FName ActorDescPath = ActorDescFolder.GetPath();
+			if (ActorDescPath.IsNone() || !InPaths.Contains(ActorDescPath))
+			{
+				return true;
+			}
+
+			return Operation(ActorDesc);
+		});
+	}
+}
+
+void FActorFolders::ForEachActorInFolders(UWorld& InWorld, const TArray<FName>& InPaths, TFunctionRef<bool(AActor*)> Operation, const FFolder::FRootObject& InFolderRootObject /*= FFolder::GetInvalidRootObject()*/)
+{
+	TSet<FName> Paths;
+	Paths.Append(InPaths);
+	ForEachActorInFolders(InWorld, Paths, Operation, InFolderRootObject);
+}
+
+void FActorFolders::ForEachActorInFolders(UWorld& InWorld, const TSet<FName>& InPaths, TFunctionRef<bool(AActor*)> Operation, const FFolder::FRootObject& InFolderRootObject /*= FFolder::GetInvalidRootObject()*/)
 {
 	const FFolder::FRootObject FolderRootObject = FFolder::IsRootObjectValid(InFolderRootObject) ? InFolderRootObject : GetWorldFolderRootObject(InWorld);
 
@@ -649,7 +703,7 @@ void FActorFolders::ForEachActorInFolders(UWorld& InWorld, const TArray<FName>&
 			continue;
 		}
 		FName ActorPath = ActorIt->GetFolderPath();
-		if (ActorPath.IsNone() || !Paths.Contains(ActorPath))
+		if (ActorPath.IsNone() || !InPaths.Contains(ActorPath))
 		{
 			continue;
 		}
@@ -661,22 +715,37 @@ void FActorFolders::ForEachActorInFolders(UWorld& InWorld, const TArray<FName>&
 	}
 }
 
-void FActorFolders::GetActorsFromFolders(UWorld& InWorld, const TArray<FName>& Paths, TArray<AActor*>& OutActors, const FFolder::FRootObject& InFolderRootObject /*= FFolder::GetInvalidRootObject()*/)
+void FActorFolders::GetActorsFromFolders(UWorld& InWorld, const TArray<FName>& InPaths, TArray<AActor*>& OutActors, const FFolder::FRootObject& InFolderRootObject /*= FFolder::GetInvalidRootObject()*/)
+{
+	TSet<FName> Paths;
+	Paths.Append(InPaths);
+	GetActorsFromFolders(InWorld, Paths, OutActors, InFolderRootObject);
+}
+
+void FActorFolders::GetActorsFromFolders(UWorld& InWorld, const TSet<FName>& InPaths, TArray<AActor*>& OutActors, const FFolder::FRootObject& InFolderRootObject /*= FFolder::GetInvalidRootObject()*/)
 {
 	const FFolder::FRootObject FolderRootObject = FFolder::IsRootObjectValid(InFolderRootObject) ? InFolderRootObject : GetWorldFolderRootObject(InWorld);
 
-	ForEachActorInFolders(InWorld, Paths, [&OutActors](AActor* InActor)
+	ForEachActorInFolders(InWorld, InPaths, [&OutActors](AActor* InActor)
 	{
 		OutActors.Add(InActor);
 		return true;
 	}, FolderRootObject);
 }
 
-void FActorFolders::GetWeakActorsFromFolders(UWorld& InWorld, const TArray<FName>& Paths, TArray<TWeakObjectPtr<AActor>>& OutActors, const FFolder::FRootObject& InFolderRootObject /*= FFolder::GetInvalidRootObject()*/)
+
+void FActorFolders::GetWeakActorsFromFolders(UWorld& InWorld, const TArray<FName>& InPaths, TArray<TWeakObjectPtr<AActor>>& OutActors, const FFolder::FRootObject& InFolderRootObject /*= FFolder::GetInvalidRootObject()*/)
+{
+	TSet<FName> Paths;
+	Paths.Append(InPaths);
+	GetWeakActorsFromFolders(InWorld, Paths, OutActors, InFolderRootObject);
+}
+
+void FActorFolders::GetWeakActorsFromFolders(UWorld& InWorld, const TSet<FName>& InPaths, TArray<TWeakObjectPtr<AActor>>& OutActors, const FFolder::FRootObject& InFolderRootObject /*= FFolder::GetInvalidRootObject()*/)
 {
 	const FFolder::FRootObject FolderRootObject = FFolder::IsRootObjectValid(InFolderRootObject) ? InFolderRootObject : GetWorldFolderRootObject(InWorld);
 
-	ForEachActorInFolders(InWorld, Paths, [&OutActors](AActor* InActor)
+	ForEachActorInFolders(InWorld, InPaths, [&OutActors](AActor* InActor)
 	{
 		OutActors.Add(InActor);
 		return true;
diff --git a/Engine/Source/Editor/UnrealEd/Public/EditorActorFolders.h b/Engine/Source/Editor/UnrealEd/Public/EditorActorFolders.h
--- a/Engine/Source/Editor/UnrealEd/Public/EditorActorFolders.h
+++ b/Engine/Source/Editor/UnrealEd/Public/EditorActorFolders.h
@@ -13,6 +13,7 @@
 class FObjectPostSaveContext;
 class AActor;
 class UActorFolder;
+class FWorldPartitionActorDesc;
 
 /** Multicast delegates for broadcasting various folder events */
 
@@ -97,15 +98,22 @@ struct UNREALED_API FActorFolders : public FGCObject, public IActorEditorContext
 	bool RenameFolderInWorld(UWorld& World, FName OldPath, FName NewPath);
 
 	//~ End Deprecated
+	static FFolder GetActorDescFolder(UWorld& InWorld, const FWorldPartitionActorDesc* InActorDesc);
 
+	/** Apply an operation to each actor desc in the given list of folders. */
+	static void ForEachActorDescInFolders(UWorld& InWorld, const TSet<FName>& InPaths, TFunctionRef<bool(const FWorldPartitionActorDesc*)> Operation, const FFolder::FRootObject& InFolderRootObject = FFolder::GetInvalidRootObject());
+		
 	/** Apply an operation to each actor in the given list of folders. Will stop when operation returns false. */
-	static void ForEachActorInFolders(UWorld& InWorld, const TArray<FName>& Paths, TFunctionRef<bool(AActor*)> Operation, const FFolder::FRootObject& InFolderRootObject = FFolder::GetInvalidRootObject());
+	static void ForEachActorInFolders(UWorld& InWorld, const TArray<FName>& InPaths, TFunctionRef<bool(AActor*)> Operation, const FFolder::FRootObject& InFolderRootObject = FFolder::GetInvalidRootObject());
+	static void ForEachActorInFolders(UWorld& InWorld, const TSet<FName>& InPaths, TFunctionRef<bool(AActor*)> Operation, const FFolder::FRootObject& InFolderRootObject = FFolder::GetInvalidRootObject());
 
 	/** Get an array of actors from a list of folders */
-	static void GetActorsFromFolders(UWorld& InWorld, const TArray<FName>& Paths, TArray<AActor*>& OutActors, const FFolder::FRootObject& InFolderRootObject = FFolder::GetInvalidRootObject());
+	static void GetActorsFromFolders(UWorld& InWorld, const TArray<FName>& InPaths, TArray<AActor*>& OutActors, const FFolder::FRootObject& InFolderRootObject = FFolder::GetInvalidRootObject());
+	static void GetActorsFromFolders(UWorld& InWorld, const TSet<FName>& InPaths, TArray<AActor*>& OutActors, const FFolder::FRootObject& InFolderRootObject = FFolder::GetInvalidRootObject());
 
 	/** Get an array of weak actor pointers from a list of folders */
-	static void GetWeakActorsFromFolders(UWorld& InWorld, const TArray<FName>& Paths, TArray<TWeakObjectPtr<AActor>>& OutActors, const FFolder::FRootObject& InFolderRootObject = FFolder::GetInvalidRootObject());
+	static void GetWeakActorsFromFolders(UWorld& InWorld, const TArray<FName>& InPaths, TArray<TWeakObjectPtr<AActor>>& OutActors, const FFolder::FRootObject& InFolderRootObject = FFolder::GetInvalidRootObject());
+	static void GetWeakActorsFromFolders(UWorld& InWorld, const TSet<FName>& InPaths, TArray<TWeakObjectPtr<AActor>>& OutActors, const FFolder::FRootObject& InFolderRootObject = FFolder::GetInvalidRootObject());
 
 	/** Tests whether a folder container exists for the specified world */
 	bool IsInitializedForWorld(UWorld& InWorld) const;
