Merging with resolved conflicts on behalf of sarah.noonan. Original CL description:

*Make FPlatformMemory::GetStats cache memory stats per frame to prevent multiple costly platform calls per frame ( by default, has option to request non-cached values ).
*Update some calls to FPlatformMemory::GetStats to specifically request non-cached values.

diff --git a/Engine/Platforms/Hololens/Source/Runtime/Core/Private/HoloLensPlatformMemory.cpp b/Engine/Platforms/Hololens/Source/Runtime/Core/Private/HoloLensPlatformMemory.cpp
--- a/Engine/Platforms/Hololens/Source/Runtime/Core/Private/HoloLensPlatformMemory.cpp
+++ b/Engine/Platforms/Hololens/Source/Runtime/Core/Private/HoloLensPlatformMemory.cpp
@@ -78,7 +78,7 @@ FMalloc* FHoloLensPlatformMemory::BaseAllocator()
 	//	_CrtSetAllocHook(HoloLensAllocHook); // Enable to track down windows allocs not handled by our wrapper
 }
 
-FPlatformMemoryStats FHoloLensPlatformMemory::GetStats()
+FPlatformMemoryStats FHoloLensPlatformMemory::GetStatsImmediate()
 {
 	/**
 	*	GetProcessMemoryInfo
diff --git a/Engine/Platforms/Hololens/Source/Runtime/Core/Public/HoloLensPlatformMemory.h b/Engine/Platforms/Hololens/Source/Runtime/Core/Public/HoloLensPlatformMemory.h
--- a/Engine/Platforms/Hololens/Source/Runtime/Core/Public/HoloLensPlatformMemory.h
+++ b/Engine/Platforms/Hololens/Source/Runtime/Core/Public/HoloLensPlatformMemory.h
@@ -72,7 +72,7 @@ struct CORE_API FHoloLensPlatformMemory : public FGenericPlatformMemory
 		return true;
 	}
 	static class FMalloc* BaseAllocator();
-	static FPlatformMemoryStats GetStats();
+	static FPlatformMemoryStats GetStatsImmediate();
 	static void GetStatsForMallocProfiler(FGenericMemoryStats& out_Stats);
 	static const FPlatformMemoryConstants& GetConstants();
 	static void* BinnedAllocFromOS(SIZE_T Size);
diff --git a/Engine/Source/Editor/UnrealEd/Private/Commandlets/CookCommandlet.cpp b/Engine/Source/Editor/UnrealEd/Private/Commandlets/CookCommandlet.cpp
--- a/Engine/Source/Editor/UnrealEd/Private/Commandlets/CookCommandlet.cpp
+++ b/Engine/Source/Editor/UnrealEd/Private/Commandlets/CookCommandlet.cpp
@@ -998,7 +998,8 @@ void UCookCommandlet::ConditionalCollectGarbage(uint32 TickResults, UCookOnTheFl
 #endif
 	UE_SCOPED_COOKTIMER(CookCommandlet_GC);
 
-	const FPlatformMemoryStats MemStatsBeforeGC = FPlatformMemory::GetStats();
+	constexpr bool bShouldAllowCaching = false;
+	const FPlatformMemoryStats MemStatsBeforeGC = FPlatformMemory::GetStats(bShouldAllowCaching);
 	int32 NumObjectsBeforeGC = GUObjectArray.GetObjectArrayNumMinusAvailable();
 	int32 NumObjectsAvailableBeforeGC = GUObjectArray.GetObjectArrayEstimatedAvailable();
 	UE_LOG(LogCookCommandlet, Display, TEXT("GarbageCollection...%s (%s)"), (bPartialGC ? TEXT(" partial gc") : TEXT("")), *GCReason);
@@ -1025,7 +1026,7 @@ void UCookCommandlet::ConditionalCollectGarbage(uint32 TickResults, UCookOnTheFl
 
 	int32 NumObjectsAfterGC = GUObjectArray.GetObjectArrayNumMinusAvailable();
 	int32 NumObjectsAvailableAfterGC = GUObjectArray.GetObjectArrayEstimatedAvailable();
-	FPlatformMemoryStats MemStatsAfterGC = FPlatformMemory::GetStats();
+	FPlatformMemoryStats MemStatsAfterGC = FPlatformMemory::GetStats(bShouldAllowCaching);
 	if (!COTFS.IsCookOnTheFlyMode())
 	{
 		int64 VirtualMemBeforeGC = MemStatsBeforeGC.UsedVirtual;
diff --git a/Engine/Source/Editor/UnrealEd/Private/Commandlets/GatherTextFromAssetsCommandlet.cpp b/Engine/Source/Editor/UnrealEd/Private/Commandlets/GatherTextFromAssetsCommandlet.cpp
--- a/Engine/Source/Editor/UnrealEd/Private/Commandlets/GatherTextFromAssetsCommandlet.cpp
+++ b/Engine/Source/Editor/UnrealEd/Private/Commandlets/GatherTextFromAssetsCommandlet.cpp
@@ -270,7 +270,8 @@ void UGatherTextFromAssetsCommandlet::CalculateDependenciesForPackagesPendingGat
 
 bool UGatherTextFromAssetsCommandlet::HasExceededMemoryLimit(const bool bLog)
 {
-	const FPlatformMemoryStats MemStats = FPlatformMemory::GetStats();
+	constexpr bool bShouldAllowCaching = false;
+	const FPlatformMemoryStats MemStats = FPlatformMemory::GetStats(bShouldAllowCaching);
 
 	const uint64 FreeMemoryBytes = MemStats.AvailablePhysical;
 	if (MinFreeMemoryBytes > 0u && FreeMemoryBytes < MinFreeMemoryBytes)
diff --git a/Engine/Source/Editor/UnrealEd/Private/CookOnTheFlyServer.cpp b/Engine/Source/Editor/UnrealEd/Private/CookOnTheFlyServer.cpp
--- a/Engine/Source/Editor/UnrealEd/Private/CookOnTheFlyServer.cpp
+++ b/Engine/Source/Editor/UnrealEd/Private/CookOnTheFlyServer.cpp
@@ -1652,10 +1652,11 @@ void UCookOnTheFlyServer::UpdateDisplay(UE::Cook::FTickStackData& StackData, boo
 		}
 		if (!IsCookOnTheFlyMode() || bCookOnTheFlyShouldDisplay)
 		{
+			constexpr bool bShouldAllowCaching = false;
 			UE_CLOG(!(StackData.TickFlags & ECookTickFlags::HideProgressDisplay) && (GCookProgressDisplay != (int32)ECookProgressDisplayMode::Nothing),
 				LogCook, Display,
 				TEXT("Cook Diagnostics: OpenFileHandles=%d, VirtualMemory=%dMiB"),
-				OpenFileHandles, FPlatformMemory::GetStats().UsedVirtual / 1024 / 1024);
+				OpenFileHandles, FPlatformMemory::GetStats(bShouldAllowCaching).UsedVirtual / 1024 / 1024);
 		}
 		if (bCookOnTheFlyShouldDisplay)
 		{
@@ -1953,7 +1954,8 @@ TRefCountPtr<UCookOnTheFlyServer::FPollable> UCookOnTheFlyServer::CreatePollable
 void UCookOnTheFlyServer::PollGarbageCollection(UE::Cook::FTickStackData& StackData)
 {
 	NumObjectsHistory.AddInstance(GUObjectArray.GetObjectArrayNumMinusAvailable());
-	VirtualMemoryHistory.AddInstance(FPlatformMemory::GetStats().UsedVirtual);
+	constexpr bool bShouldAllowCaching = false;
+	VirtualMemoryHistory.AddInstance(FPlatformMemory::GetStats(bShouldAllowCaching).UsedVirtual);
 
 	if (IsCookFlagSet(ECookInitializationFlags::TestCook))
 	{
@@ -4415,8 +4417,8 @@ bool UCookOnTheFlyServer::HasExceededMaxMemory() const
 
 	TStringBuilder<256> TriggerMessages;
 
-	const FPlatformMemoryStats MemStats = FPlatformMemory::GetStats();
-	
+	constexpr bool bShouldAllowCaching = false;
+	const FPlatformMemoryStats MemStats = FPlatformMemory::GetStats(bShouldAllowCaching);
 	if (MemoryMinFreeVirtual > 0 || MemoryMinFreePhysical > 0)
 	{
 		++ActiveTriggers;
@@ -4954,7 +4956,8 @@ void UCookOnTheFlyServer::PreGarbageCollect()
 	}
 
 	NumObjectsHistory.AddInstance(GUObjectArray.GetObjectArrayNumMinusAvailable());
-	VirtualMemoryHistory.AddInstance(FPlatformMemory::GetStats().UsedVirtual);
+	constexpr bool bShouldAllowCaching = false;
+	VirtualMemoryHistory.AddInstance(FPlatformMemory::GetStats(bShouldAllowCaching).UsedVirtual);
 	TArray<UPackage*> GCKeepPackages;
 	TArray<UE::Cook::FPackageData*> GCKeepPackageDatas;
 	PreGarbageCollectImpl(GCKeepPackages, GCKeepPackageDatas, GCKeepObjects);
@@ -5130,8 +5133,9 @@ void UCookOnTheFlyServer::CookerAddReferencedObjects(FReferenceCollector& Collec
 
 void UCookOnTheFlyServer::PostGarbageCollect()
 {
+	constexpr bool bShouldAllowCaching = false;
 	NumObjectsHistory.AddInstance(GUObjectArray.GetObjectArrayNumMinusAvailable());
-	VirtualMemoryHistory.AddInstance(FPlatformMemory::GetStats().UsedVirtual);
+	VirtualMemoryHistory.AddInstance(FPlatformMemory::GetStats(bShouldAllowCaching).UsedVirtual);
 	PostGarbageCollectImpl(GCKeepObjects);
 }
 
@@ -5936,7 +5940,8 @@ bool UCookOnTheFlyServer::TryInitializeCookWorker()
 void UCookOnTheFlyServer::InitializeSession()
 {
 	NumObjectsHistory.Initialize(GUObjectArray.GetObjectArrayNumMinusAvailable());
-	VirtualMemoryHistory.Initialize(FPlatformMemory::GetStats().UsedVirtual);
+	constexpr bool bShouldAllowCaching = false;
+	VirtualMemoryHistory.Initialize(FPlatformMemory::GetStats(bShouldAllowCaching).UsedVirtual);
 }
 
 bool UCookOnTheFlyServer::Exec(class UWorld* InWorld, const TCHAR* Cmd, FOutputDevice& Ar)
@@ -8493,7 +8498,8 @@ void UCookOnTheFlyServer::PrintFinishStats()
 	const float TotalCookTime = (float)(FPlatformTime::Seconds() - CookByTheBookOptions->CookStartTime);
 	UE_LOG(LogCook, Display, TEXT("Cook by the book total time in tick %fs total time %f"), CookByTheBookOptions->CookTime, TotalCookTime);
 
-	const FPlatformMemoryStats MemStats = FPlatformMemory::GetStats();
+	constexpr bool bShouldAllowCaching = false;
+	const FPlatformMemoryStats MemStats = FPlatformMemory::GetStats(bShouldAllowCaching);
 	UE_LOG(LogCook, Display, TEXT("Peak Used virtual %u MiB Peak Used physical %u MiB"), MemStats.PeakUsedVirtual / 1024 / 1024, MemStats.PeakUsedPhysical / 1024 / 1024);
 
 	COOK_STAT(UE_LOG(LogCook, Display, TEXT("Packages Cooked: %d, Packages Iteratively Skipped: %d, Total Packages: %d"),
diff --git a/Engine/Source/Runtime/Core/Private/Android/AndroidPlatformMemory.cpp b/Engine/Source/Runtime/Core/Private/Android/AndroidPlatformMemory.cpp
--- a/Engine/Source/Runtime/Core/Private/Android/AndroidPlatformMemory.cpp
+++ b/Engine/Source/Runtime/Core/Private/Android/AndroidPlatformMemory.cpp
@@ -177,7 +177,7 @@ static bool GetPss(uint64& PSSOUT, uint64& SwapPSSOUT)
 	return FieldsSetSuccessfully == 2;
 }
 
-FPlatformMemoryStats FAndroidPlatformMemory::GetStats()
+FPlatformMemoryStats FAndroidPlatformMemory::GetStatsImmediate()
 {
 	FPlatformMemoryStats MemoryStats;	// will init from constants
 
diff --git a/Engine/Source/Runtime/Core/Private/Apple/ApplePlatformMemory.cpp b/Engine/Source/Runtime/Core/Private/Apple/ApplePlatformMemory.cpp
--- a/Engine/Source/Runtime/Core/Private/Apple/ApplePlatformMemory.cpp
+++ b/Engine/Source/Runtime/Core/Private/Apple/ApplePlatformMemory.cpp
@@ -342,7 +342,7 @@ FMalloc* FApplePlatformMemory::BaseAllocator()
 	
 }
 
-FPlatformMemoryStats FApplePlatformMemory::GetStats()
+FPlatformMemoryStats FApplePlatformMemory::GetStatsImmediate()
 {
 	const FPlatformMemoryConstants& MemoryConstants = FPlatformMemory::GetConstants();
 	static FPlatformMemoryStats MemoryStats;
diff --git a/Engine/Source/Runtime/Core/Private/GenericPlatform/GenericPlatformMemory.cpp b/Engine/Source/Runtime/Core/Private/GenericPlatform/GenericPlatformMemory.cpp
--- a/Engine/Source/Runtime/Core/Private/GenericPlatform/GenericPlatformMemory.cpp
+++ b/Engine/Source/Runtime/Core/Private/GenericPlatform/GenericPlatformMemory.cpp
@@ -60,6 +60,9 @@ namespace GenericPlatformMemory
 		     "Where a platform can specifically state it's memory pressure this test maybe ignored.\n"
 		     "0 (default) critical pressure will not use the threshold."),
 		ECVF_Default);
+
+	static uint64 CachedMemoryStatsAtFrame = -1;
+	static FPlatformMemoryStats CachedMemoryStats;
 }
 
 struct TUnalignedTester
@@ -250,9 +253,24 @@ FMalloc* FGenericPlatformMemory::BaseAllocator()
 	return new FMallocAnsi();
 }
 
-FPlatformMemoryStats FGenericPlatformMemory::GetStats()
+FPlatformMemoryStats FGenericPlatformMemory::GetStats(const bool bAllowUseCachedStats)
+{
+	if (IsInGameThread())
+	{
+		if (!bAllowUseCachedStats || (GFrameCounter != GenericPlatformMemory::CachedMemoryStatsAtFrame))
+		{
+			GenericPlatformMemory::CachedMemoryStatsAtFrame = GFrameCounter;
+			GenericPlatformMemory::CachedMemoryStats = FPlatformMemory::GetStatsImmediate();
+		}
+		return GenericPlatformMemory::CachedMemoryStats;
+	}
+
+	return FPlatformMemory::GetStatsImmediate();
+}
+
+FPlatformMemoryStats FGenericPlatformMemory::GetStatsImmediate()
 {
-	UE_LOG(LogMemory, Warning, TEXT("FGenericPlatformMemory::GetStats not implemented on this platform"));
+	UE_LOG(LogMemory, Warning, TEXT("FGenericPlatformMemory::GetStatsImmediate not implemented on this platform"));
 	return FPlatformMemoryStats();
 }
 
diff --git a/Engine/Source/Runtime/Core/Private/HAL/MemoryMisc.cpp b/Engine/Source/Runtime/Core/Private/HAL/MemoryMisc.cpp
--- a/Engine/Source/Runtime/Core/Private/HAL/MemoryMisc.cpp
+++ b/Engine/Source/Runtime/Core/Private/HAL/MemoryMisc.cpp
@@ -46,13 +46,13 @@ namespace MemoryMiscInternal
 #if ENABLE_MEMORY_SCOPE_STATS
 FScopedMemoryStats::FScopedMemoryStats(const TCHAR* Name)
 	: Text(Name)
-	, StartStats(FPlatformMemory::GetStats())
+	, StartStats(FPlatformMemory::GetStatsImmediate())
 {
 }
 
 FScopedMemoryStats::~FScopedMemoryStats()
 {
-	const FPlatformMemoryStats EndStats = FPlatformMemory::GetStats();
+	const FPlatformMemoryStats EndStats = FPlatformMemory::GetStatsImmediate();
 	MemoryMiscInternal::LogMemoryDiff(Text, StartStats, EndStats);
 }
 #endif
diff --git a/Engine/Source/Runtime/Core/Private/Mac/MacPlatformMemory.cpp b/Engine/Source/Runtime/Core/Private/Mac/MacPlatformMemory.cpp
--- a/Engine/Source/Runtime/Core/Private/Mac/MacPlatformMemory.cpp
+++ b/Engine/Source/Runtime/Core/Private/Mac/MacPlatformMemory.cpp
@@ -133,7 +133,7 @@ FMalloc* FMacPlatformMemory::BaseAllocator()
 
 }
 
-FPlatformMemoryStats FMacPlatformMemory::GetStats()
+FPlatformMemoryStats FMacPlatformMemory::GetStatsImmediate()
 {
 	const FPlatformMemoryConstants& MemoryConstants = FPlatformMemory::GetConstants();
 
diff --git a/Engine/Source/Runtime/Core/Private/Unix/UnixPlatformMemory.cpp b/Engine/Source/Runtime/Core/Private/Unix/UnixPlatformMemory.cpp
--- a/Engine/Source/Runtime/Core/Private/Unix/UnixPlatformMemory.cpp
+++ b/Engine/Source/Runtime/Core/Private/Unix/UnixPlatformMemory.cpp
@@ -789,7 +789,7 @@ static uint32 ReadProcFields(const char *FileName, FProcField *ProcFields, uint3
 //   uint64 PeakUsedPhysical;  /** The peak amount of physical memory used by the process, in bytes. */	VmHWM
 //   uint64 UsedVirtual;       /** Total amount of virtual memory used by the process. */					VmSize
 //   uint64 PeakUsedVirtual;   /** The peak amount of virtual memory used by the process. */				VmPeak
-FPlatformMemoryStats FUnixPlatformMemory::GetStats()
+FPlatformMemoryStats FUnixPlatformMemory::GetStatsImmediate()
 {
 	uint64 MemFree = 0;
 	uint64 Cached = 0;
diff --git a/Engine/Source/Runtime/Core/Private/Windows/WindowsPlatformMemory.cpp b/Engine/Source/Runtime/Core/Private/Windows/WindowsPlatformMemory.cpp
--- a/Engine/Source/Runtime/Core/Private/Windows/WindowsPlatformMemory.cpp
+++ b/Engine/Source/Runtime/Core/Private/Windows/WindowsPlatformMemory.cpp
@@ -228,7 +228,7 @@ FMalloc* FWindowsPlatformMemory::BaseAllocator()
 	}
 }
 
-FPlatformMemoryStats FWindowsPlatformMemory::GetStats()
+FPlatformMemoryStats FWindowsPlatformMemory::GetStatsImmediate()
 {
 	// release mimalloc held pages back to OS before getting stats
 	// can't do this, it calls me :
diff --git a/Engine/Source/Runtime/Core/Public/Android/AndroidPlatformMemory.h b/Engine/Source/Runtime/Core/Public/Android/AndroidPlatformMemory.h
--- a/Engine/Source/Runtime/Core/Public/Android/AndroidPlatformMemory.h
+++ b/Engine/Source/Runtime/Core/Public/Android/AndroidPlatformMemory.h
@@ -24,7 +24,7 @@ struct CORE_API FAndroidPlatformMemory : public FGenericPlatformMemory
 {
 	//~ Begin FGenericPlatformMemory Interface
 	static void Init();
-	static FPlatformMemoryStats GetStats();
+	static FPlatformMemoryStats GetStatsImmediate();
 	static uint64 GetMemoryUsedFast();
 	static const FPlatformMemoryConstants& GetConstants();
 	static EPlatformMemorySizeBucket GetMemorySizeBucket();
diff --git a/Engine/Source/Runtime/Core/Public/Apple/ApplePlatformMemory.h b/Engine/Source/Runtime/Core/Public/Apple/ApplePlatformMemory.h
--- a/Engine/Source/Runtime/Core/Public/Apple/ApplePlatformMemory.h
+++ b/Engine/Source/Runtime/Core/Public/Apple/ApplePlatformMemory.h
@@ -64,7 +64,7 @@ struct CORE_API FApplePlatformMemory : public FGenericPlatformMemory
 {
 	//~ Begin FGenericPlatformMemory Interface
 	static void Init();
-	static FPlatformMemoryStats GetStats();
+	static FPlatformMemoryStats GetStatsImmediate();
 	static uint64 GetMemoryUsedFast();
 	static const FPlatformMemoryConstants& GetConstants();
 	static FMalloc* BaseAllocator();
diff --git a/Engine/Source/Runtime/Core/Public/GenericPlatform/GenericPlatformMemory.h b/Engine/Source/Runtime/Core/Public/GenericPlatform/GenericPlatformMemory.h
--- a/Engine/Source/Runtime/Core/Public/GenericPlatform/GenericPlatformMemory.h
+++ b/Engine/Source/Runtime/Core/Public/GenericPlatform/GenericPlatformMemory.h
@@ -334,9 +334,14 @@ struct CORE_API FGenericPlatformMemory
 	static FMalloc* BaseAllocator();
 
 	/**
-	 * @return platform specific current memory statistics.
+	 * @return platform specific current memory statistics - by default caches value for same frame to avoid repeated costly platform calls.
 	 */
-	static FPlatformMemoryStats GetStats();
+	static FPlatformMemoryStats GetStats(const bool bAllowUseCachedStats = true);
+
+	/**
+	 * @return platform specific current memory statistics - no caching.
+	 */
+	static FPlatformMemoryStats GetStatsImmediate();
 
 	/**
 	* @return memory used for platforms that can do it quickly (without affecting stat unit much)
diff --git a/Engine/Source/Runtime/Core/Public/Mac/MacPlatformMemory.h b/Engine/Source/Runtime/Core/Public/Mac/MacPlatformMemory.h
--- a/Engine/Source/Runtime/Core/Public/Mac/MacPlatformMemory.h
+++ b/Engine/Source/Runtime/Core/Public/Mac/MacPlatformMemory.h
@@ -15,7 +15,7 @@
 struct CORE_API FMacPlatformMemory : public FApplePlatformMemory
 {
 	//~ Begin FGenericPlatformMemory Interface
-	static FPlatformMemoryStats GetStats();
+	static FPlatformMemoryStats GetStatsImmediate();
 	static const FPlatformMemoryConstants& GetConstants();
 	static FMalloc* BaseAllocator();
 	//~ End FGenericPlatformMemory Interface
diff --git a/Engine/Source/Runtime/Core/Public/Unix/UnixPlatformMemory.h b/Engine/Source/Runtime/Core/Public/Unix/UnixPlatformMemory.h
--- a/Engine/Source/Runtime/Core/Public/Unix/UnixPlatformMemory.h
+++ b/Engine/Source/Runtime/Core/Public/Unix/UnixPlatformMemory.h
@@ -76,7 +76,7 @@ struct CORE_API FUnixPlatformMemory : public FGenericPlatformMemory
 	//~ Begin FGenericPlatformMemory Interface
 	static void Init();
 	static class FMalloc* BaseAllocator();
-	static FPlatformMemoryStats GetStats();
+	static FPlatformMemoryStats GetStatsImmediate();
 	static FExtendedPlatformMemoryStats GetExtendedStats();
 	static const FPlatformMemoryConstants& GetConstants();
 	static bool PageProtect(void* const Ptr, const SIZE_T Size, const bool bCanRead, const bool bCanWrite);
diff --git a/Engine/Source/Runtime/Core/Public/Windows/WindowsPlatformMemory.h b/Engine/Source/Runtime/Core/Public/Windows/WindowsPlatformMemory.h
--- a/Engine/Source/Runtime/Core/Public/Windows/WindowsPlatformMemory.h
+++ b/Engine/Source/Runtime/Core/Public/Windows/WindowsPlatformMemory.h
@@ -84,7 +84,7 @@ struct CORE_API FWindowsPlatformMemory
 	}
 
 	static class FMalloc* BaseAllocator();
-	static FPlatformMemoryStats GetStats();
+	static FPlatformMemoryStats GetStatsImmediate();
 	static void GetStatsForMallocProfiler( FGenericMemoryStats& out_Stats );
 	static const FPlatformMemoryConstants& GetConstants();
 	static bool PageProtect(void* const Ptr, const SIZE_T Size, const bool bCanRead, const bool bCanWrite);
diff --git a/Engine/Source/Runtime/Engine/Private/WorldPartition/WorldPartitionHelpers.cpp b/Engine/Source/Runtime/Engine/Private/WorldPartition/WorldPartitionHelpers.cpp
--- a/Engine/Source/Runtime/Engine/Private/WorldPartition/WorldPartitionHelpers.cpp
+++ b/Engine/Source/Runtime/Engine/Private/WorldPartition/WorldPartitionHelpers.cpp
@@ -208,9 +208,10 @@ bool FWorldPartitionHelpers::HasExceededMaxMemory()
 
 void FWorldPartitionHelpers::DoCollectGarbage()
 {
-	const FPlatformMemoryStats MemStatsBefore = FPlatformMemory::GetStats();
+	constexpr bool bShouldAllowCaching = false;
+	const FPlatformMemoryStats MemStatsBefore = FPlatformMemory::GetStats(bShouldAllowCaching);
 	CollectGarbage(GARBAGE_COLLECTION_KEEPFLAGS, true);
-	const FPlatformMemoryStats MemStatsAfter = FPlatformMemory::GetStats();
+	const FPlatformMemoryStats MemStatsAfter = FPlatformMemory::GetStats(bShouldAllowCaching);
 
 	UE_LOG(LogWorldPartition, Log, TEXT("GC Performed - Available Physical: %.2fGB, Available Virtual: %.2fGB"),
 		(int64)MemStatsAfter.AvailablePhysical / (1024.0 * 1024.0 * 1024.0),
