Improvements and changes to the shader DDC fill commandlet:

- Display verbose material messages so we can see DDC load/save information
- Keep track of a running list of materials we need to wait for async DDC get results from.
- When IsCachedCookedPlatformDataLoaded returns true for a material that either means two things:
    1. The shader map was found in the DDC.
    2. We didn't find the DDC results and we submitted shader compile jobs.
- When we return true we should remove ourselves from the set of materials to compile, and we log some updates.
- Always call ProcessAsyncResults in this busy loop so we can process results as they come in.
- When the number of outstanding jobs change we print progress results.
- Fix usage of deprecated member variable.

diff --git a/Engine/Source/Editor/UnrealEd/Private/Commandlets/CompileShadersTestBedCommandlet.cpp b/Engine/Source/Editor/UnrealEd/Private/Commandlets/CompileShadersTestBedCommandlet.cpp
--- a/Engine/Source/Editor/UnrealEd/Private/Commandlets/CompileShadersTestBedCommandlet.cpp
+++ b/Engine/Source/Editor/UnrealEd/Private/Commandlets/CompileShadersTestBedCommandlet.cpp
@@ -25,7 +25,7 @@ int32 UCompileShadersTestBedCommandlet::Main(const FString& Params)
 {
 	TRACE_CPUPROFILER_EVENT_SCOPE(UCompileShadersTestBedCommandlet::Main);
 
-	StaticExec(nullptr, TEXT("log LogMaterial Log"));
+	StaticExec(nullptr, TEXT("log LogMaterial Verbose"));
 
 	TArray<FString> Tokens;
 	TArray<FString> Switches;
@@ -62,9 +62,7 @@ int32 UCompileShadersTestBedCommandlet::Main(const FString& Params)
 			Filter.ClassPaths.Add(UMaterial::StaticClass()->GetClassPathName());
 
 			FCollectionManagerModule& CollectionManagerModule = FCollectionManagerModule::GetModule();
-PRAGMA_DISABLE_DEPRECATION_WARNINGS
-			CollectionManagerModule.Get().GetObjectsInCollection(FName(*CollectionName), ECollectionShareType::CST_All, Filter.ObjectPaths, ECollectionRecursionFlags::SelfAndChildren);
-PRAGMA_ENABLE_DEPRECATION_WARNINGS
+			CollectionManagerModule.Get().GetObjectsInCollection(FName(*CollectionName), ECollectionShareType::CST_All, Filter.SoftObjectPaths, ECollectionRecursionFlags::SelfAndChildren);
 
 			AssetRegistry.GetAssets(Filter, MaterialList);
 
@@ -76,6 +74,8 @@ PRAGMA_ENABLE_DEPRECATION_WARNINGS
 		}
 	}
 
+	UE_LOG(LogCompileShadersTestBedCommandlet, Log, TEXT("Found %d/%d Materials."), MaterialList.Num(), Filter.SoftObjectPaths.Num());
+
 	// Process -materials= switches separated by a '+'
 	TArray<FString> CmdLineMaterialEntries;
 	const TCHAR* MaterialsSwitchName = TEXT("Materials");
@@ -146,6 +146,8 @@ PRAGMA_ENABLE_DEPRECATION_WARNINGS
 			GShaderCompilingManager->ProcessAsyncResults(bLimitExecutationTime, true /* bBlockOnGlobalShaderCompilation */);
 		}
 
+		TSet<UMaterialInterface*> MaterialsToCompile;
+
 		// Begin Material Compiles
 		{
 			TRACE_CPUPROFILER_EVENT_SCOPE(BeginCacheForCookedPlatformData);
@@ -159,13 +161,14 @@ PRAGMA_ENABLE_DEPRECATION_WARNINGS
 			{
 				if (UMaterialInterface* MaterialInterface = Cast<UMaterialInterface>(AssetData.GetAsset()))
 				{
-					UE_LOG(LogCompileShadersTestBedCommandlet, Log, TEXT("BeginCache for %s"), *MaterialInterface->GetFullName());
+					UE_LOG(LogCompileShadersTestBedCommandlet, Display, TEXT("BeginCache for %s"), *MaterialInterface->GetFullName());
 					MaterialInterface->BeginCacheForCookedPlatformData(Platform);
+					MaterialsToCompile.Add(MaterialInterface);
 				}
 			}
 		}
 
-		bool bAllDone = false;
+		int32 PreviousOutstandingJobs = 0;
 
 		// Submit all the jobs.
 		{
@@ -173,22 +176,30 @@ PRAGMA_ENABLE_DEPRECATION_WARNINGS
 
 			UE_LOG(LogCompileShadersTestBedCommandlet, Log, TEXT("Submit Jobs"));
 
-			for (const FAssetData& AssetData : MaterialList)
+			while (MaterialsToCompile.Num())
 			{
-				if (UMaterialInterface* MaterialInterface = Cast<UMaterialInterface>(AssetData.GetAsset()))
+				for (auto It = MaterialsToCompile.CreateIterator(); It; ++It)
 				{
-					const bool bMaterialDone = MaterialInterface->IsCachedCookedPlatformDataLoaded(Platform);
-					if (bMaterialDone)
+					UMaterialInterface* MaterialInterface = *It;
+					if (MaterialInterface->IsCachedCookedPlatformDataLoaded(Platform))
 					{
-						UE_LOG(LogCompileShadersTestBedCommandlet, Log, TEXT("Finished cache for %s"), *MaterialInterface->GetFullName());
+						It.RemoveCurrent();
+						UE_LOG(LogCompileShadersTestBedCommandlet, Display, TEXT("Finished cache for %s."), *MaterialInterface->GetFullName());
+						UE_LOG(LogCompileShadersTestBedCommandlet, Display, TEXT("Materials remaining: %d"), MaterialsToCompile.Num());
 					}
 
-					bAllDone &= bMaterialDone;
-				}
+					GShaderCompilingManager->ProcessAsyncResults(bLimitExecutationTime, false /* bBlockOnGlobalShaderCompilation */);
 
-				if (bAllDone)
-				{
-					break;
+					const int32 CurrentOutstandingJobs = GShaderCompilingManager->GetNumOutstandingJobs();
+					if (CurrentOutstandingJobs != PreviousOutstandingJobs)
+					{
+						UE_LOG(LogCompileShadersTestBedCommandlet, Log, TEXT("Outstanding Jobs: %d"), CurrentOutstandingJobs);
+						PreviousOutstandingJobs = CurrentOutstandingJobs;
+					}
+
+					// Flush rendering commands to release any RHI resources (shaders and shader maps).
+					// Delete any FPendingCleanupObjects (shader maps).
+					FlushRenderingCommands();
 				}
 			}
 		}
@@ -202,6 +213,13 @@ PRAGMA_ENABLE_DEPRECATION_WARNINGS
 			while (GShaderCompilingManager->IsCompiling())
 			{
 				GShaderCompilingManager->ProcessAsyncResults(bLimitExecutationTime, false /* bBlockOnGlobalShaderCompilation */);
+
+				const int32 CurrentOutstandingJobs = GShaderCompilingManager->GetNumOutstandingJobs();
+				if (CurrentOutstandingJobs != PreviousOutstandingJobs)
+				{
+					UE_LOG(LogCompileShadersTestBedCommandlet, Display, TEXT("Outstanding Jobs: %d"), CurrentOutstandingJobs);
+					PreviousOutstandingJobs = CurrentOutstandingJobs;
+				}
 				
 				// Flush rendering commands to release any RHI resources (shaders and shader maps).
 				// Delete any FPendingCleanupObjects (shader maps).
@@ -213,7 +231,7 @@ PRAGMA_ENABLE_DEPRECATION_WARNINGS
 		{
 			TRACE_CPUPROFILER_EVENT_SCOPE(ClearCachedCookedPlatformData);
 
-			UE_LOG(LogCompileShadersTestBedCommandlet, Log, TEXT("Clear Cached Cooked Platform Data"));
+			UE_LOG(LogCompileShadersTestBedCommandlet, Display, TEXT("Clear Cached Cooked Platform Data"));
 
 			for (const FAssetData& AssetData : MaterialList)
 			{
