Cooker no longer tries to save split packages that come from a package in the platform ignore list

diff --git a/Engine/Source/Editor/UnrealEd/Classes/CookOnTheSide/CookOnTheFlyServer.h b/Engine/Source/Editor/UnrealEd/Classes/CookOnTheSide/CookOnTheFlyServer.h
--- a/Engine/Source/Editor/UnrealEd/Classes/CookOnTheSide/CookOnTheFlyServer.h
+++ b/Engine/Source/Editor/UnrealEd/Classes/CookOnTheSide/CookOnTheFlyServer.h
@@ -1279,7 +1279,7 @@ private:
 	/** Generates long package names for all files to be cooked */
 	void GenerateLongPackageNames(TArray<FName>& FilesInPath, TMap<FName, UE::Cook::FInstigator>& Instigators);
 
-	UE::Cook::EPollStatus ConditionalCreateGeneratorPackage(UE::Cook::FPackageData& PackageData);
+	UE::Cook::EPollStatus ConditionalCreateGeneratorPackage(UE::Cook::FPackageData& PackageData, bool bPrecaching);
 
 	/** Generate the list of cook-time-created packages created by the Generator. */
 	UE::Cook::EPollStatus QueueGeneratedPackages(UE::Cook::FGeneratorPackage& Generator,
diff --git a/Engine/Source/Editor/UnrealEd/Private/CookOnTheFlyServer.cpp b/Engine/Source/Editor/UnrealEd/Private/CookOnTheFlyServer.cpp
--- a/Engine/Source/Editor/UnrealEd/Private/CookOnTheFlyServer.cpp
+++ b/Engine/Source/Editor/UnrealEd/Private/CookOnTheFlyServer.cpp
@@ -2852,7 +2852,7 @@ void UCookOnTheFlyServer::TickNetwork()
 	}
 }
 
-UE::Cook::EPollStatus UCookOnTheFlyServer::ConditionalCreateGeneratorPackage(UE::Cook::FPackageData& PackageData)
+UE::Cook::EPollStatus UCookOnTheFlyServer::ConditionalCreateGeneratorPackage(UE::Cook::FPackageData& PackageData, bool bPrecaching)
 {
 	using namespace UE::Cook;
 
@@ -2907,6 +2907,11 @@ UE::Cook::EPollStatus UCookOnTheFlyServer::ConditionalCreateGeneratorPackage(UE:
 		return EPollStatus::Success;
 	}
 
+	if (bPrecaching)
+	{
+		return EPollStatus::Incomplete;
+	}
+
 	// TODO: Add support for cooking in the editor. Possibly moot since we plan to deprecate cooking in the editor.
 	if (IsCookingInEditor())
 	{
@@ -3477,7 +3482,7 @@ UE::Cook::EPollStatus UCookOnTheFlyServer::PrepareSaveInternal(UE::Cook::FPackag
 		// Check for whether the Package has a Splitter and initialize its list if so
 		if (!PackageData.HasInitializedGeneratorSave())
 		{
-			Result = ConditionalCreateGeneratorPackage(PackageData);
+			Result = ConditionalCreateGeneratorPackage(PackageData, bPrecaching);
 			if (Result != EPollStatus::Success)
 			{
 				return Result;
@@ -4188,6 +4193,7 @@ void UCookOnTheFlyServer::PumpSaves(UE::Cook::FTickStackData& StackData, uint32
 				{
 					break;
 				}
+
 				--LeftToPrecache;
 				PrepareSave(*NextData, StackData.Timer, /*bPrecaching*/ true);
 			}
@@ -5423,12 +5429,17 @@ void FSaveCookedPackageContext::SetupPlatform(const ITargetPlatform* InTargetPla
 	// check if this package is unsupported for the target platform (typically plugin content)
 	else
 	{
-		TSet<FName>* NeverCookPackages = COTFS.PackageTracker->PlatformSpecificNeverCookPackages.Find(TargetPlatform);
-		if (NeverCookPackages && NeverCookPackages->Find(Package->GetFName()))
+		if (TSet<FName>* NeverCookPackages = COTFS.PackageTracker->PlatformSpecificNeverCookPackages.Find(TargetPlatform))
 		{
-			SavePackageResult = ESavePackageResult::ContainsEditorOnlyData;
-			UE_LOG(LogCook, Display, TEXT("Excluding %s -> %s"), *PackageName, *PlatFilename);
-			return;
+			FGeneratorPackage* Generator = PackageData.IsGenerated() ? PackageData.GetGeneratedOwner() : nullptr;
+
+			if (NeverCookPackages->Find(Package->GetFName()) ||
+				(Generator && NeverCookPackages->Find(Generator->GetOwner().GetPackageName())))
+			{
+				SavePackageResult = ESavePackageResult::ContainsEditorOnlyData;
+				UE_LOG(LogCook, Display, TEXT("Excluding %s -> %s"), *PackageName, *PlatFilename);
+				return;				
+			}
 		}
 	}
 
