Fix skeletal mesh LOD clickering when opening multiple editor viewports

diff --git a/Engine/Source/Editor/UnrealEd/Private/EditorEngine.cpp b/Engine/Source/Editor/UnrealEd/Private/EditorEngine.cpp
--- a/Engine/Source/Editor/UnrealEd/Private/EditorEngine.cpp
+++ b/Engine/Source/Editor/UnrealEd/Private/EditorEngine.cpp
@@ -1971,6 +1971,7 @@ void UEditorEngine::Tick( float DeltaSeconds, bool bIdleMode )
 			if (!bAllWindowsHidden || GCurrentLevelEditingViewportClient->WantsDrawWhenAppIsHidden())
 			{
 				bool bAllowNonRealtimeViewports = true;
+				GCurrentLevelEditingViewportClient->SetIsCurrentLevelEditingFocus(true);
 				bool bWasNonRealtimeViewportDraw = UpdateSingleViewportClient(GCurrentLevelEditingViewportClient, bAllowNonRealtimeViewports, bUpdateLinkedOrthoViewports);
 				if (GCurrentLevelEditingViewportClient->IsLevelEditorClient())
 				{
@@ -1999,6 +2000,7 @@ void UEditorEngine::Tick( float DeltaSeconds, bool bIdleMode )
 					{
 						//if we haven't drawn a non-realtime viewport OR not one of the main viewports
 						bool bAllowNonRealtimeViewports = (!bEditorFrameNonRealtimeViewportDrawn) || !(ViewportClient->IsLevelEditorClient());
+						ViewportClient->SetIsCurrentLevelEditingFocus(false);
 						bool bWasNonRealtimeViewportDrawn = UpdateSingleViewportClient(ViewportClient, bAllowNonRealtimeViewports, bUpdateLinkedOrthoViewports);
 						if (ViewportClient->IsLevelEditorClient())
 						{
diff --git a/Engine/Source/Editor/UnrealEd/Private/EditorViewportClient.cpp b/Engine/Source/Editor/UnrealEd/Private/EditorViewportClient.cpp
--- a/Engine/Source/Editor/UnrealEd/Private/EditorViewportClient.cpp
+++ b/Engine/Source/Editor/UnrealEd/Private/EditorViewportClient.cpp
@@ -3945,6 +3945,9 @@ void FEditorViewportClient::Draw(FViewport* InViewport, FCanvas* Canvas)
 	ViewFamily.EngineShowFlags = UseEngineShowFlags;
 
 	ViewFamily.bIsHDR = Viewport->IsHDRViewport();
+	
+	// The view is in focus if it is currently in editing
+	ViewFamily.SetIsInFocus(GetIsCurrentLevelEditingFocus());
 
 	if( !AllowsCinematicControl() )
 	{
diff --git a/Engine/Source/Editor/UnrealEd/Public/EditorViewportClient.h b/Engine/Source/Editor/UnrealEd/Public/EditorViewportClient.h
--- a/Engine/Source/Editor/UnrealEd/Public/EditorViewportClient.h
+++ b/Engine/Source/Editor/UnrealEd/Public/EditorViewportClient.h
@@ -1391,6 +1391,10 @@ public:
 
 	/** @return true if Override EngineShowFlags are currently enabled */
 	bool IsEngineShowFlagsOverrideEnabled() const { return !! OverrideShowFlagsFunc; }
+
+	inline bool GetIsCurrentLevelEditingFocus() const								{ return bIsCurrentLevelEditingFocus; }
+	inline void SetIsCurrentLevelEditingFocus(bool bInIsCurrentLevelEditingFocus)	{ bIsCurrentLevelEditingFocus = bInIsCurrentLevelEditingFocus; }
+
 protected:
 	/** Invalidates the viewport widget (if valid) to register its active timer */
 	void InvalidateViewportWidget();
@@ -1959,6 +1963,9 @@ private:
 	FSceneViewFamily *DragStartViewFamily;
 
 	TArray<FIntPoint> CapturedMouseMoves;
+
+	/** True if this is the current viewport client in editing */
+	bool bIsCurrentLevelEditingFocus = false;
 };
 
 
diff --git a/Engine/Source/Runtime/Engine/Private/SkeletalRender.cpp b/Engine/Source/Runtime/Engine/Private/SkeletalRender.cpp
--- a/Engine/Source/Runtime/Engine/Private/SkeletalRender.cpp
+++ b/Engine/Source/Runtime/Engine/Private/SkeletalRender.cpp
@@ -91,7 +91,7 @@ FSkeletalMeshObject::~FSkeletalMeshObject()
 void FSkeletalMeshObject::UpdateMinDesiredLODLevel(const FSceneView* View, const FBoxSphereBounds& Bounds, int32 FrameNumber, uint8 CurFirstLODIdx)
 {
 	// Thumbnail rendering doesn't contribute to MinDesiredLODLevel calculation
-	if (View->Family && View->Family->bThumbnailRendering)
+	if (View->Family && (View->Family->bThumbnailRendering || !View->Family->GetIsInFocus()))
 	{
 		return;
 	}
diff --git a/Engine/Source/Runtime/Engine/Public/SceneView.h b/Engine/Source/Runtime/Engine/Public/SceneView.h
--- a/Engine/Source/Runtime/Engine/Public/SceneView.h
+++ b/Engine/Source/Runtime/Engine/Public/SceneView.h
@@ -2008,6 +2008,9 @@ public:
 		return SecondarySpatialUpscalerInterface;
 	}
 
+	inline bool GetIsInFocus() const				{ return bIsInFocus; }
+	inline void SetIsInFocus(bool bInIsInFocus)		{ bIsInFocus = bInIsInFocus; }
+
 private:
 	/** Interface to handle screen percentage of the views of the family. */
 	ISceneViewFamilyScreenPercentage* ScreenPercentageInterface;
@@ -2026,6 +2029,9 @@ private:
 	/** whether the pre DOF translucency are allowed to be rendered in separated target from scene to allow for better composition with distortion.*/
 	bool bAllowStandardTranslucencySeparated;
 
+	/** True if this view is the current editing view or the active game view */
+	bool bIsInFocus = true;
+
 	PRAGMA_DISABLE_DEPRECATION_WARNINGS
 	// Only FSceneRenderer can copy a view family.
 	FSceneViewFamily(const FSceneViewFamily&) = default;
