*Remove per-frame caching of stats from FPlaformMemory::GetStats and move caching into helper function.
*Update common callers of FPlaformMemory::GetStats to use the helper.

diff --git a/Engine/Platforms/Hololens/Source/Runtime/Core/Private/HoloLensPlatformMemory.cpp b/Engine/Platforms/Hololens/Source/Runtime/Core/Private/HoloLensPlatformMemory.cpp
--- a/Engine/Platforms/Hololens/Source/Runtime/Core/Private/HoloLensPlatformMemory.cpp
+++ b/Engine/Platforms/Hololens/Source/Runtime/Core/Private/HoloLensPlatformMemory.cpp
@@ -78,7 +78,7 @@ FMalloc* FHoloLensPlatformMemory::BaseAllocator()
 	//	_CrtSetAllocHook(HoloLensAllocHook); // Enable to track down windows allocs not handled by our wrapper
 }
 
-FPlatformMemoryStats FHoloLensPlatformMemory::GetStatsImmediate()
+FPlatformMemoryStats FHoloLensPlatformMemory::GetStats()
 {
 	/**
 	*	GetProcessMemoryInfo
diff --git a/Engine/Platforms/Hololens/Source/Runtime/Core/Public/HoloLensPlatformMemory.h b/Engine/Platforms/Hololens/Source/Runtime/Core/Public/HoloLensPlatformMemory.h
--- a/Engine/Platforms/Hololens/Source/Runtime/Core/Public/HoloLensPlatformMemory.h
+++ b/Engine/Platforms/Hololens/Source/Runtime/Core/Public/HoloLensPlatformMemory.h
@@ -72,7 +72,7 @@ struct CORE_API FHoloLensPlatformMemory : public FGenericPlatformMemory
 		return true;
 	}
 	static class FMalloc* BaseAllocator();
-	static FPlatformMemoryStats GetStatsImmediate();
+	static FPlatformMemoryStats GetStats();
 	static void GetStatsForMallocProfiler(FGenericMemoryStats& out_Stats);
 	static const FPlatformMemoryConstants& GetConstants();
 	static void* BinnedAllocFromOS(SIZE_T Size);
diff --git a/Engine/Source/Editor/UnrealEd/Private/Commandlets/CookCommandlet.cpp b/Engine/Source/Editor/UnrealEd/Private/Commandlets/CookCommandlet.cpp
--- a/Engine/Source/Editor/UnrealEd/Private/Commandlets/CookCommandlet.cpp
+++ b/Engine/Source/Editor/UnrealEd/Private/Commandlets/CookCommandlet.cpp
@@ -998,8 +998,7 @@ void UCookCommandlet::ConditionalCollectGarbage(uint32 TickResults, UCookOnTheFl
 #endif
 	UE_SCOPED_COOKTIMER(CookCommandlet_GC);
 
-	constexpr bool bShouldAllowCaching = false;
-	const FPlatformMemoryStats MemStatsBeforeGC = FPlatformMemory::GetStats(bShouldAllowCaching);
+	const FPlatformMemoryStats MemStatsBeforeGC = FPlatformMemory::GetStats();
 	int32 NumObjectsBeforeGC = GUObjectArray.GetObjectArrayNumMinusAvailable();
 	int32 NumObjectsAvailableBeforeGC = GUObjectArray.GetObjectArrayEstimatedAvailable();
 	UE_LOG(LogCookCommandlet, Display, TEXT("GarbageCollection...%s (%s)"), (bPartialGC ? TEXT(" partial gc") : TEXT("")), *GCReason);
@@ -1026,7 +1025,7 @@ void UCookCommandlet::ConditionalCollectGarbage(uint32 TickResults, UCookOnTheFl
 
 	int32 NumObjectsAfterGC = GUObjectArray.GetObjectArrayNumMinusAvailable();
 	int32 NumObjectsAvailableAfterGC = GUObjectArray.GetObjectArrayEstimatedAvailable();
-	FPlatformMemoryStats MemStatsAfterGC = FPlatformMemory::GetStats(bShouldAllowCaching);
+	FPlatformMemoryStats MemStatsAfterGC = FPlatformMemory::GetStats();
 	if (!COTFS.IsCookOnTheFlyMode())
 	{
 		int64 VirtualMemBeforeGC = MemStatsBeforeGC.UsedVirtual;
diff --git a/Engine/Source/Editor/UnrealEd/Private/Commandlets/GatherTextFromAssetsCommandlet.cpp b/Engine/Source/Editor/UnrealEd/Private/Commandlets/GatherTextFromAssetsCommandlet.cpp
--- a/Engine/Source/Editor/UnrealEd/Private/Commandlets/GatherTextFromAssetsCommandlet.cpp
+++ b/Engine/Source/Editor/UnrealEd/Private/Commandlets/GatherTextFromAssetsCommandlet.cpp
@@ -270,8 +270,7 @@ void UGatherTextFromAssetsCommandlet::CalculateDependenciesForPackagesPendingGat
 
 bool UGatherTextFromAssetsCommandlet::HasExceededMemoryLimit(const bool bLog)
 {
-	constexpr bool bShouldAllowCaching = false;
-	const FPlatformMemoryStats MemStats = FPlatformMemory::GetStats(bShouldAllowCaching);
+	const FPlatformMemoryStats MemStats = FPlatformMemory::GetStats();
 
 	const uint64 FreeMemoryBytes = MemStats.AvailablePhysical;
 	if (MinFreeMemoryBytes > 0u && FreeMemoryBytes < MinFreeMemoryBytes)
diff --git a/Engine/Source/Editor/UnrealEd/Private/CookOnTheFlyServer.cpp b/Engine/Source/Editor/UnrealEd/Private/CookOnTheFlyServer.cpp
--- a/Engine/Source/Editor/UnrealEd/Private/CookOnTheFlyServer.cpp
+++ b/Engine/Source/Editor/UnrealEd/Private/CookOnTheFlyServer.cpp
@@ -1652,11 +1652,10 @@ void UCookOnTheFlyServer::UpdateDisplay(UE::Cook::FTickStackData& StackData, boo
 		}
 		if (!IsCookOnTheFlyMode() || bCookOnTheFlyShouldDisplay)
 		{
-			constexpr bool bShouldAllowCaching = false;
 			UE_CLOG(!(StackData.TickFlags & ECookTickFlags::HideProgressDisplay) && (GCookProgressDisplay != (int32)ECookProgressDisplayMode::Nothing),
 				LogCook, Display,
 				TEXT("Cook Diagnostics: OpenFileHandles=%d, VirtualMemory=%dMiB"),
-				OpenFileHandles, FPlatformMemory::GetStats(bShouldAllowCaching).UsedVirtual / 1024 / 1024);
+				OpenFileHandles, FPlatformMemory::GetStats().UsedVirtual / 1024 / 1024);
 		}
 		if (bCookOnTheFlyShouldDisplay)
 		{
@@ -1954,8 +1953,7 @@ TRefCountPtr<UCookOnTheFlyServer::FPollable> UCookOnTheFlyServer::CreatePollable
 void UCookOnTheFlyServer::PollGarbageCollection(UE::Cook::FTickStackData& StackData)
 {
 	NumObjectsHistory.AddInstance(GUObjectArray.GetObjectArrayNumMinusAvailable());
-	constexpr bool bShouldAllowCaching = false;
-	VirtualMemoryHistory.AddInstance(FPlatformMemory::GetStats(bShouldAllowCaching).UsedVirtual);
+	VirtualMemoryHistory.AddInstance(FPlatformMemory::GetStats().UsedVirtual);
 
 	if (IsCookFlagSet(ECookInitializationFlags::TestCook))
 	{
@@ -4423,8 +4421,8 @@ bool UCookOnTheFlyServer::HasExceededMaxMemory() const
 
 	TStringBuilder<256> TriggerMessages;
 
-	constexpr bool bShouldAllowCaching = false;
-	const FPlatformMemoryStats MemStats = FPlatformMemory::GetStats(bShouldAllowCaching);
+	const FPlatformMemoryStats MemStats = FPlatformMemory::GetStats();
+	
 	if (MemoryMinFreeVirtual > 0 || MemoryMinFreePhysical > 0)
 	{
 		++ActiveTriggers;
@@ -4962,8 +4960,7 @@ void UCookOnTheFlyServer::PreGarbageCollect()
 	}
 
 	NumObjectsHistory.AddInstance(GUObjectArray.GetObjectArrayNumMinusAvailable());
-	constexpr bool bShouldAllowCaching = false;
-	VirtualMemoryHistory.AddInstance(FPlatformMemory::GetStats(bShouldAllowCaching).UsedVirtual);
+	VirtualMemoryHistory.AddInstance(FPlatformMemory::GetStats().UsedVirtual);
 	TArray<UPackage*> GCKeepPackages;
 	TArray<UE::Cook::FPackageData*> GCKeepPackageDatas;
 	PreGarbageCollectImpl(GCKeepPackages, GCKeepPackageDatas, GCKeepObjects);
@@ -5139,9 +5136,8 @@ void UCookOnTheFlyServer::CookerAddReferencedObjects(FReferenceCollector& Collec
 
 void UCookOnTheFlyServer::PostGarbageCollect()
 {
-	constexpr bool bShouldAllowCaching = false;
 	NumObjectsHistory.AddInstance(GUObjectArray.GetObjectArrayNumMinusAvailable());
-	VirtualMemoryHistory.AddInstance(FPlatformMemory::GetStats(bShouldAllowCaching).UsedVirtual);
+	VirtualMemoryHistory.AddInstance(FPlatformMemory::GetStats().UsedVirtual);
 	PostGarbageCollectImpl(GCKeepObjects);
 }
 
@@ -5951,8 +5947,7 @@ bool UCookOnTheFlyServer::TryInitializeCookWorker()
 void UCookOnTheFlyServer::InitializeSession()
 {
 	NumObjectsHistory.Initialize(GUObjectArray.GetObjectArrayNumMinusAvailable());
-	constexpr bool bShouldAllowCaching = false;
-	VirtualMemoryHistory.Initialize(FPlatformMemory::GetStats(bShouldAllowCaching).UsedVirtual);
+	VirtualMemoryHistory.Initialize(FPlatformMemory::GetStats().UsedVirtual);
 }
 
 bool UCookOnTheFlyServer::Exec(class UWorld* InWorld, const TCHAR* Cmd, FOutputDevice& Ar)
@@ -8509,8 +8504,7 @@ void UCookOnTheFlyServer::PrintFinishStats()
 	const float TotalCookTime = (float)(FPlatformTime::Seconds() - CookByTheBookOptions->CookStartTime);
 	UE_LOG(LogCook, Display, TEXT("Cook by the book total time in tick %fs total time %f"), CookByTheBookOptions->CookTime, TotalCookTime);
 
-	constexpr bool bShouldAllowCaching = false;
-	const FPlatformMemoryStats MemStats = FPlatformMemory::GetStats(bShouldAllowCaching);
+	const FPlatformMemoryStats MemStats = FPlatformMemory::GetStats();
 	UE_LOG(LogCook, Display, TEXT("Peak Used virtual %u MiB Peak Used physical %u MiB"), MemStats.PeakUsedVirtual / 1024 / 1024, MemStats.PeakUsedPhysical / 1024 / 1024);
 
 	COOK_STAT(UE_LOG(LogCook, Display, TEXT("Packages Cooked: %d, Packages Iteratively Skipped: %d, Total Packages: %d"),
diff --git a/Engine/Source/Runtime/Core/Private/Android/AndroidPlatformMemory.cpp b/Engine/Source/Runtime/Core/Private/Android/AndroidPlatformMemory.cpp
--- a/Engine/Source/Runtime/Core/Private/Android/AndroidPlatformMemory.cpp
+++ b/Engine/Source/Runtime/Core/Private/Android/AndroidPlatformMemory.cpp
@@ -177,7 +177,7 @@ static bool GetPss(uint64& PSSOUT, uint64& SwapPSSOUT)
 	return FieldsSetSuccessfully == 2;
 }
 
-FPlatformMemoryStats FAndroidPlatformMemory::GetStatsImmediate()
+FPlatformMemoryStats FAndroidPlatformMemory::GetStats()
 {
 	FPlatformMemoryStats MemoryStats;	// will init from constants
 
diff --git a/Engine/Source/Runtime/Core/Private/Apple/ApplePlatformMemory.cpp b/Engine/Source/Runtime/Core/Private/Apple/ApplePlatformMemory.cpp
--- a/Engine/Source/Runtime/Core/Private/Apple/ApplePlatformMemory.cpp
+++ b/Engine/Source/Runtime/Core/Private/Apple/ApplePlatformMemory.cpp
@@ -342,7 +342,7 @@ FMalloc* FApplePlatformMemory::BaseAllocator()
 	
 }
 
-FPlatformMemoryStats FApplePlatformMemory::GetStatsImmediate()
+FPlatformMemoryStats FApplePlatformMemory::GetStats()
 {
 	const FPlatformMemoryConstants& MemoryConstants = FPlatformMemory::GetConstants();
 	static FPlatformMemoryStats MemoryStats;
diff --git a/Engine/Source/Runtime/Core/Private/GenericPlatform/GenericPlatformMemory.cpp b/Engine/Source/Runtime/Core/Private/GenericPlatform/GenericPlatformMemory.cpp
--- a/Engine/Source/Runtime/Core/Private/GenericPlatform/GenericPlatformMemory.cpp
+++ b/Engine/Source/Runtime/Core/Private/GenericPlatform/GenericPlatformMemory.cpp
@@ -60,9 +60,6 @@ namespace GenericPlatformMemory
 		     "Where a platform can specifically state it's memory pressure this test maybe ignored.\n"
 		     "0 (default) critical pressure will not use the threshold."),
 		ECVF_Default);
-
-	static uint64 CachedMemoryStatsAtFrame = -1;
-	static FPlatformMemoryStats CachedMemoryStats;
 }
 
 struct TUnalignedTester
@@ -253,24 +250,9 @@ FMalloc* FGenericPlatformMemory::BaseAllocator()
 	return new FMallocAnsi();
 }
 
-FPlatformMemoryStats FGenericPlatformMemory::GetStats(const bool bAllowUseCachedStats)
-{
-	if (IsInGameThread())
-	{
-		if (!bAllowUseCachedStats || (GFrameCounter != GenericPlatformMemory::CachedMemoryStatsAtFrame))
-		{
-			GenericPlatformMemory::CachedMemoryStatsAtFrame = GFrameCounter;
-			GenericPlatformMemory::CachedMemoryStats = FPlatformMemory::GetStatsImmediate();
-		}
-		return GenericPlatformMemory::CachedMemoryStats;
-	}
-
-	return FPlatformMemory::GetStatsImmediate();
-}
-
-FPlatformMemoryStats FGenericPlatformMemory::GetStatsImmediate()
+FPlatformMemoryStats FGenericPlatformMemory::GetStats()
 {
-	UE_LOG(LogMemory, Warning, TEXT("FGenericPlatformMemory::GetStatsImmediate not implemented on this platform"));
+	UE_LOG(LogMemory, Warning, TEXT("FGenericPlatformMemory::GetStats not implemented on this platform"));
 	return FPlatformMemoryStats();
 }
 
diff --git a/Engine/Source/Runtime/Core/Private/HAL/MemoryMisc.cpp b/Engine/Source/Runtime/Core/Private/HAL/MemoryMisc.cpp
--- a/Engine/Source/Runtime/Core/Private/HAL/MemoryMisc.cpp
+++ b/Engine/Source/Runtime/Core/Private/HAL/MemoryMisc.cpp
@@ -46,13 +46,13 @@ namespace MemoryMiscInternal
 #if ENABLE_MEMORY_SCOPE_STATS
 FScopedMemoryStats::FScopedMemoryStats(const TCHAR* Name)
 	: Text(Name)
-	, StartStats(FPlatformMemory::GetStatsImmediate())
+	, StartStats(FPlatformMemory::GetStats())
 {
 }
 
 FScopedMemoryStats::~FScopedMemoryStats()
 {
-	const FPlatformMemoryStats EndStats = FPlatformMemory::GetStatsImmediate();
+	const FPlatformMemoryStats EndStats = FPlatformMemory::GetStats();
 	MemoryMiscInternal::LogMemoryDiff(Text, StartStats, EndStats);
 }
 #endif
diff --git a/Engine/Source/Runtime/Core/Private/HAL/PlatformMemoryHelpers.cpp b/Engine/Source/Runtime/Core/Private/HAL/PlatformMemoryHelpers.cpp
new file mode 100644
--- /dev/null
+++ b/Engine/Source/Runtime/Core/Private/HAL/PlatformMemoryHelpers.cpp
@@ -0,0 +1,29 @@
+// Copyright Epic Games, Inc. All Rights Reserved.
+
+#include "HAL/PlatformMemoryHelpers.h"
+
+#include "CoreGlobals.h"
+#include "HAL/PlatformMemory.h"
+
+namespace PlatformMemoryHelpers
+{
+
+static uint64 CachedMemoryStatsAtFrame = -1;
+static FPlatformMemoryStats CachedMemoryStats;
+
+FPlatformMemoryStats GetFrameMemoryStats()
+{
+	if (IsInGameThread())
+	{
+		if (GFrameCounter != CachedMemoryStatsAtFrame)
+		{
+			CachedMemoryStatsAtFrame = GFrameCounter;
+			CachedMemoryStats = FPlatformMemory::GetStats();
+		}
+		return CachedMemoryStats;
+	}
+
+	return FPlatformMemory::GetStats();
+}
+
+} //PlatformMemoryHelpers
\ No newline at end of file
diff --git a/Engine/Source/Runtime/Core/Private/Mac/MacPlatformMemory.cpp b/Engine/Source/Runtime/Core/Private/Mac/MacPlatformMemory.cpp
--- a/Engine/Source/Runtime/Core/Private/Mac/MacPlatformMemory.cpp
+++ b/Engine/Source/Runtime/Core/Private/Mac/MacPlatformMemory.cpp
@@ -133,7 +133,7 @@ FMalloc* FMacPlatformMemory::BaseAllocator()
 
 }
 
-FPlatformMemoryStats FMacPlatformMemory::GetStatsImmediate()
+FPlatformMemoryStats FMacPlatformMemory::GetStats()
 {
 	const FPlatformMemoryConstants& MemoryConstants = FPlatformMemory::GetConstants();
 
diff --git a/Engine/Source/Runtime/Core/Private/Unix/UnixPlatformMemory.cpp b/Engine/Source/Runtime/Core/Private/Unix/UnixPlatformMemory.cpp
--- a/Engine/Source/Runtime/Core/Private/Unix/UnixPlatformMemory.cpp
+++ b/Engine/Source/Runtime/Core/Private/Unix/UnixPlatformMemory.cpp
@@ -789,7 +789,7 @@ static uint32 ReadProcFields(const char *FileName, FProcField *ProcFields, uint3
 //   uint64 PeakUsedPhysical;  /** The peak amount of physical memory used by the process, in bytes. */	VmHWM
 //   uint64 UsedVirtual;       /** Total amount of virtual memory used by the process. */					VmSize
 //   uint64 PeakUsedVirtual;   /** The peak amount of virtual memory used by the process. */				VmPeak
-FPlatformMemoryStats FUnixPlatformMemory::GetStatsImmediate()
+FPlatformMemoryStats FUnixPlatformMemory::GetStats()
 {
 	uint64 MemFree = 0;
 	uint64 Cached = 0;
diff --git a/Engine/Source/Runtime/Core/Private/Windows/WindowsPlatformMemory.cpp b/Engine/Source/Runtime/Core/Private/Windows/WindowsPlatformMemory.cpp
--- a/Engine/Source/Runtime/Core/Private/Windows/WindowsPlatformMemory.cpp
+++ b/Engine/Source/Runtime/Core/Private/Windows/WindowsPlatformMemory.cpp
@@ -228,7 +228,7 @@ FMalloc* FWindowsPlatformMemory::BaseAllocator()
 	}
 }
 
-FPlatformMemoryStats FWindowsPlatformMemory::GetStatsImmediate()
+FPlatformMemoryStats FWindowsPlatformMemory::GetStats()
 {
 	// release mimalloc held pages back to OS before getting stats
 	// can't do this, it calls me :
diff --git a/Engine/Source/Runtime/Core/Public/Android/AndroidPlatformMemory.h b/Engine/Source/Runtime/Core/Public/Android/AndroidPlatformMemory.h
--- a/Engine/Source/Runtime/Core/Public/Android/AndroidPlatformMemory.h
+++ b/Engine/Source/Runtime/Core/Public/Android/AndroidPlatformMemory.h
@@ -24,7 +24,7 @@ struct CORE_API FAndroidPlatformMemory : public FGenericPlatformMemory
 {
 	//~ Begin FGenericPlatformMemory Interface
 	static void Init();
-	static FPlatformMemoryStats GetStatsImmediate();
+	static FPlatformMemoryStats GetStats();
 	static uint64 GetMemoryUsedFast();
 	static const FPlatformMemoryConstants& GetConstants();
 	static EPlatformMemorySizeBucket GetMemorySizeBucket();
diff --git a/Engine/Source/Runtime/Core/Public/Apple/ApplePlatformMemory.h b/Engine/Source/Runtime/Core/Public/Apple/ApplePlatformMemory.h
--- a/Engine/Source/Runtime/Core/Public/Apple/ApplePlatformMemory.h
+++ b/Engine/Source/Runtime/Core/Public/Apple/ApplePlatformMemory.h
@@ -64,7 +64,7 @@ struct CORE_API FApplePlatformMemory : public FGenericPlatformMemory
 {
 	//~ Begin FGenericPlatformMemory Interface
 	static void Init();
-	static FPlatformMemoryStats GetStatsImmediate();
+	static FPlatformMemoryStats GetStats();
 	static uint64 GetMemoryUsedFast();
 	static const FPlatformMemoryConstants& GetConstants();
 	static FMalloc* BaseAllocator();
diff --git a/Engine/Source/Runtime/Core/Public/GenericPlatform/GenericPlatformMemory.h b/Engine/Source/Runtime/Core/Public/GenericPlatform/GenericPlatformMemory.h
--- a/Engine/Source/Runtime/Core/Public/GenericPlatform/GenericPlatformMemory.h
+++ b/Engine/Source/Runtime/Core/Public/GenericPlatform/GenericPlatformMemory.h
@@ -334,14 +334,9 @@ struct CORE_API FGenericPlatformMemory
 	static FMalloc* BaseAllocator();
 
 	/**
-	 * @return platform specific current memory statistics - by default caches value for same frame to avoid repeated costly platform calls.
+	 * @return platform specific current memory statistics.
 	 */
-	static FPlatformMemoryStats GetStats(const bool bAllowUseCachedStats = true);
-
-	/**
-	 * @return platform specific current memory statistics - no caching.
-	 */
-	static FPlatformMemoryStats GetStatsImmediate();
+	static FPlatformMemoryStats GetStats();
 
 	/**
 	* @return memory used for platforms that can do it quickly (without affecting stat unit much)
diff --git a/Engine/Source/Runtime/Core/Public/HAL/PlatformMemoryHelpers.h b/Engine/Source/Runtime/Core/Public/HAL/PlatformMemoryHelpers.h
new file mode 100644
--- /dev/null
+++ b/Engine/Source/Runtime/Core/Public/HAL/PlatformMemoryHelpers.h
@@ -0,0 +1,11 @@
+// Copyright Epic Games, Inc. All Rights Reserved.
+#pragma once
+
+#include "HAL/PlatformMemory.h"
+
+namespace PlatformMemoryHelpers
+{
+
+	//@return platform specific current memory statistics - caches value for same frame to avoid repeated costly platform calls.
+	CORE_API FPlatformMemoryStats GetFrameMemoryStats();
+}
\ No newline at end of file
diff --git a/Engine/Source/Runtime/Core/Public/Mac/MacPlatformMemory.h b/Engine/Source/Runtime/Core/Public/Mac/MacPlatformMemory.h
--- a/Engine/Source/Runtime/Core/Public/Mac/MacPlatformMemory.h
+++ b/Engine/Source/Runtime/Core/Public/Mac/MacPlatformMemory.h
@@ -15,7 +15,7 @@
 struct CORE_API FMacPlatformMemory : public FApplePlatformMemory
 {
 	//~ Begin FGenericPlatformMemory Interface
-	static FPlatformMemoryStats GetStatsImmediate();
+	static FPlatformMemoryStats GetStats();
 	static const FPlatformMemoryConstants& GetConstants();
 	static FMalloc* BaseAllocator();
 	//~ End FGenericPlatformMemory Interface
diff --git a/Engine/Source/Runtime/Core/Public/Unix/UnixPlatformMemory.h b/Engine/Source/Runtime/Core/Public/Unix/UnixPlatformMemory.h
--- a/Engine/Source/Runtime/Core/Public/Unix/UnixPlatformMemory.h
+++ b/Engine/Source/Runtime/Core/Public/Unix/UnixPlatformMemory.h
@@ -76,7 +76,7 @@ struct CORE_API FUnixPlatformMemory : public FGenericPlatformMemory
 	//~ Begin FGenericPlatformMemory Interface
 	static void Init();
 	static class FMalloc* BaseAllocator();
-	static FPlatformMemoryStats GetStatsImmediate();
+	static FPlatformMemoryStats GetStats();
 	static FExtendedPlatformMemoryStats GetExtendedStats();
 	static const FPlatformMemoryConstants& GetConstants();
 	static bool PageProtect(void* const Ptr, const SIZE_T Size, const bool bCanRead, const bool bCanWrite);
diff --git a/Engine/Source/Runtime/Core/Public/Windows/WindowsPlatformMemory.h b/Engine/Source/Runtime/Core/Public/Windows/WindowsPlatformMemory.h
--- a/Engine/Source/Runtime/Core/Public/Windows/WindowsPlatformMemory.h
+++ b/Engine/Source/Runtime/Core/Public/Windows/WindowsPlatformMemory.h
@@ -84,7 +84,7 @@ struct CORE_API FWindowsPlatformMemory
 	}
 
 	static class FMalloc* BaseAllocator();
-	static FPlatformMemoryStats GetStatsImmediate();
+	static FPlatformMemoryStats GetStats();
 	static void GetStatsForMallocProfiler( FGenericMemoryStats& out_Stats );
 	static const FPlatformMemoryConstants& GetConstants();
 	static bool PageProtect(void* const Ptr, const SIZE_T Size, const bool bCanRead, const bool bCanWrite);
diff --git a/Engine/Source/Runtime/Engine/Private/ChartCreation.cpp b/Engine/Source/Runtime/Engine/Private/ChartCreation.cpp
--- a/Engine/Source/Runtime/Engine/Private/ChartCreation.cpp
+++ b/Engine/Source/Runtime/Engine/Private/ChartCreation.cpp
@@ -22,6 +22,7 @@
 #include "ProfilingDebugging/CsvProfiler.h"
 #include "DeviceProfiles/DeviceProfileManager.h"
 #include "DynamicResolutionState.h"
+#include "HAL/PlatformMemoryHelpers.h"
 
 #ifndef FPS_CHART_SUPPORT_CSV_PROFILE
 #define FPS_CHART_SUPPORT_CSV_PROFILE (CSV_PROFILER && !UE_BUILD_SHIPPING)
@@ -807,7 +808,7 @@ void FPerformanceTrackingChart::ProcessFrame(const FFrameData& FrameData)
 		TotalDrawnPrimitives += GNumPrimitivesDrawnRHI[0];
 
 		// track memory
-		FPlatformMemoryStats MemoryStats = FPlatformMemory::GetStats();
+		FPlatformMemoryStats MemoryStats = PlatformMemoryHelpers::GetFrameMemoryStats();
 		MaxPhysicalMemory = FMath::Max(MaxPhysicalMemory, static_cast<uint64>(MemoryStats.UsedPhysical));
 		MaxVirtualMemory = FMath::Max(MaxVirtualMemory, static_cast<uint64>(MemoryStats.UsedVirtual));
 		MinPhysicalMemory = FMath::Min(MinPhysicalMemory, static_cast<uint64>(MemoryStats.UsedPhysical));
@@ -926,7 +927,7 @@ void FPerformanceTrackingChart::DumpChartToAnalyticsParams(const FString& InMapN
 			// Dump some extra non-chart-based stats
 
 			// Get the system memory stats
-			FPlatformMemoryStats Stats = FPlatformMemory::GetStats();
+			FPlatformMemoryStats Stats = PlatformMemoryHelpers::GetFrameMemoryStats();
 			InParamArray.Add(FAnalyticsEventAttribute(TEXT("TotalPhysical"), static_cast<uint64>(Stats.TotalPhysical)));
 			InParamArray.Add(FAnalyticsEventAttribute(TEXT("TotalVirtual"), static_cast<uint64>(Stats.TotalVirtual)));
 			InParamArray.Add(FAnalyticsEventAttribute(TEXT("PeakPhysical"), static_cast<uint64>(Stats.PeakUsedPhysical)));
diff --git a/Engine/Source/Runtime/Engine/Private/UnrealEngine.cpp b/Engine/Source/Runtime/Engine/Private/UnrealEngine.cpp
--- a/Engine/Source/Runtime/Engine/Private/UnrealEngine.cpp
+++ b/Engine/Source/Runtime/Engine/Private/UnrealEngine.cpp
@@ -141,6 +141,7 @@ UnrealEngine.cpp: Implements the UEngine class and helpers.
 #include "Engine/ViewportStatsSubsystem.h"
 #include "GPUSkinCache.h"
 #include "Chaos/TriangleMeshImplicitObject.h"
+#include "HAL/PlatformMemoryHelpers.h"
 
 #if UE_WITH_IRIS
 #include "Iris/IrisConfig.h"
@@ -1656,7 +1657,7 @@ void UEngine::PreGarbageCollect()
 
 float GetLowMemoryGCTimer(const float DefaultTimeBetweenGC)
 {
-	float MBFree = float(FPlatformMemory::GetStats().AvailablePhysical / 1024 / 1024);
+	float MBFree = float(PlatformMemoryHelpers::GetFrameMemoryStats().AvailablePhysical / 1024 / 1024);
 #if !UE_BUILD_SHIPPING
 	MBFree -= float(FPlatformMemory::GetExtraDevelopmentMemorySize() / 1024 / 1024);
 #endif
@@ -1823,7 +1824,7 @@ void UEngine::ConditionalCollectGarbage()
 						float IncGCTime = GIncrementalGCTimePerFrame;
 						if (GLowMemoryMemoryThresholdMB > 0.0)
 						{
-							float MBFree = float(FPlatformMemory::GetStats().AvailablePhysical / 1024 / 1024);
+							float MBFree = float(PlatformMemoryHelpers::GetFrameMemoryStats().AvailablePhysical / 1024 / 1024);
 #if !UE_BUILD_SHIPPING
 							MBFree -= float(FPlatformMemory::GetExtraDevelopmentMemorySize() / 1024 / 1024);
 #endif
diff --git a/Engine/Source/Runtime/Engine/Private/WorldPartition/WorldPartitionHelpers.cpp b/Engine/Source/Runtime/Engine/Private/WorldPartition/WorldPartitionHelpers.cpp
--- a/Engine/Source/Runtime/Engine/Private/WorldPartition/WorldPartitionHelpers.cpp
+++ b/Engine/Source/Runtime/Engine/Private/WorldPartition/WorldPartitionHelpers.cpp
@@ -208,10 +208,9 @@ bool FWorldPartitionHelpers::HasExceededMaxMemory()
 
 void FWorldPartitionHelpers::DoCollectGarbage()
 {
-	constexpr bool bShouldAllowCaching = false;
-	const FPlatformMemoryStats MemStatsBefore = FPlatformMemory::GetStats(bShouldAllowCaching);
+	const FPlatformMemoryStats MemStatsBefore = FPlatformMemory::GetStats();
 	CollectGarbage(GARBAGE_COLLECTION_KEEPFLAGS, true);
-	const FPlatformMemoryStats MemStatsAfter = FPlatformMemory::GetStats(bShouldAllowCaching);
+	const FPlatformMemoryStats MemStatsAfter = FPlatformMemory::GetStats();
 
 	UE_LOG(LogWorldPartition, Log, TEXT("GC Performed - Available Physical: %.2fGB, Available Virtual: %.2fGB"),
 		(int64)MemStatsAfter.AvailablePhysical / (1024.0 * 1024.0 * 1024.0),
diff --git a/Engine/Source/Runtime/Launch/Private/LaunchEngineLoop.cpp b/Engine/Source/Runtime/Launch/Private/LaunchEngineLoop.cpp
--- a/Engine/Source/Runtime/Launch/Private/LaunchEngineLoop.cpp
+++ b/Engine/Source/Runtime/Launch/Private/LaunchEngineLoop.cpp
@@ -33,6 +33,7 @@
 #include "Misc/ScopeExit.h"
 #include "Misc/TrackedActivity.h"
 #include "HAL/PlatformFileManager.h"
+#include "HAL/PlatformMemoryHelpers.h"
 #include "HAL/FileManagerGeneric.h"
 #include "HAL/ExceptionHandling.h"
 #include "HAL/IPlatformFileManagedStorageWrapper.h"
@@ -1447,7 +1448,7 @@ static void UpdateCoreCsvStats_EndFrame()
 	    {
 		    CSV_CUSTOM_STAT_GLOBAL(InputLatencyTime, FPlatformTime::ToMilliseconds64(GInputLatencyTime), ECsvCustomStatOp::Set);
 	    }
-	    FPlatformMemoryStats MemoryStats = FPlatformMemory::GetStats();
+	    FPlatformMemoryStats MemoryStats = PlatformMemoryHelpers::GetFrameMemoryStats();
 	    float PhysicalMBFree = float(MemoryStats.AvailablePhysical / 1024) / 1024.0f;
 #if !UE_BUILD_SHIPPING
 	    // Subtract any extra development memory from physical free. This can result in negative values in cases where we would have crashed OOM
diff --git a/Engine/Source/Runtime/PerfCounters/Private/PerfCounters.cpp b/Engine/Source/Runtime/PerfCounters/Private/PerfCounters.cpp
--- a/Engine/Source/Runtime/PerfCounters/Private/PerfCounters.cpp
+++ b/Engine/Source/Runtime/PerfCounters/Private/PerfCounters.cpp
@@ -6,6 +6,7 @@
 #include "HAL/FileManager.h"
 #include "Misc/CommandLine.h"
 #include "Misc/Paths.h"
+#include "HAL/PlatformMemoryHelpers.h"
 #include "HAL/PlatformTime.h"
 #include "HAL/PlatformProcess.h"
 #include "HAL/RunnableThread.h"
@@ -197,7 +198,7 @@ void FPerfCounters::TickSystemCounters(float DeltaTime)
 		Set(TEXT("ProcessCPUUsageRelativeToCore"), CPUStats.CPUTimePctRelative);
 
 		// memory
-		FPlatformMemoryStats Stats = FPlatformMemory::GetStats();
+		FPlatformMemoryStats Stats = PlatformMemoryHelpers::GetFrameMemoryStats();
 		Set(TEXT("AvailablePhysicalMemoryMB"), static_cast<uint64>(Stats.AvailablePhysical / (1024 * 1024)));
 		Set(TEXT("AvailableVirtualMemoryMB"), static_cast<uint64>(Stats.AvailableVirtual / (1024 * 1024)));
 		Set(TEXT("ProcessPhysicalMemoryMB"), static_cast<uint64>(Stats.UsedPhysical/ (1024 * 1024)));
