Clear class thumbnail scene when its BPGC is being unloaded

diff --git a/Engine/Source/Editor/UnrealEd/Classes/ThumbnailRendering/BlueprintThumbnailRenderer.h b/Engine/Source/Editor/UnrealEd/Classes/ThumbnailRendering/BlueprintThumbnailRenderer.h
--- a/Engine/Source/Editor/UnrealEd/Classes/ThumbnailRendering/BlueprintThumbnailRenderer.h
+++ b/Engine/Source/Editor/UnrealEd/Classes/ThumbnailRendering/BlueprintThumbnailRenderer.h
@@ -7,7 +7,7 @@
 
 #pragma once
 
-#include "CoreMinimal.h"
+#include "CoreTypes.h"
 #include "UObject/ObjectMacros.h"
 #include "ThumbnailRendering/DefaultSizedThumbnailRenderer.h"
 #include "ThumbnailHelpers.h"
@@ -22,8 +22,6 @@ class UBlueprintThumbnailRenderer : public UDefaultSizedThumbnailRenderer
 {
 	GENERATED_UCLASS_BODY()
 
-	UNREALED_API virtual ~UBlueprintThumbnailRenderer();
-
 	// Begin UThumbnailRenderer Object
 	UNREALED_API virtual bool CanVisualizeAsset(UObject* Object) override;
 	UNREALED_API virtual void Draw(UObject* Object, int32 X, int32 Y, uint32 Width, uint32 Height, FRenderTarget* RenderTarget, FCanvas* Canvas, bool bAdditionalViewFamily) override;
diff --git a/Engine/Source/Editor/UnrealEd/Classes/ThumbnailRendering/ClassThumbnailRenderer.h b/Engine/Source/Editor/UnrealEd/Classes/ThumbnailRendering/ClassThumbnailRenderer.h
--- a/Engine/Source/Editor/UnrealEd/Classes/ThumbnailRendering/ClassThumbnailRenderer.h
+++ b/Engine/Source/Editor/UnrealEd/Classes/ThumbnailRendering/ClassThumbnailRenderer.h
@@ -2,12 +2,13 @@
 
 #pragma once
 
-#include "CoreMinimal.h"
+#include "CoreTypes.h"
 #include "UObject/ObjectMacros.h"
 #include "ThumbnailRendering/DefaultSizedThumbnailRenderer.h"
 #include "ThumbnailHelpers.h"
 #include "ClassThumbnailRenderer.generated.h"
 
+class UBlueprintGeneratedClass;
 class FCanvas;
 class FRenderTarget;
 
@@ -25,6 +26,9 @@ class UClassThumbnailRenderer : public UDefaultSizedThumbnailRenderer
 	UNREALED_API virtual void BeginDestroy() override;
 	// End UObject Implementation
 
+private:
+	void OnBlueprintGeneratedClassUnloaded(UBlueprintGeneratedClass* BPGC);
+
 private:
 	TClassInstanceThumbnailScene<FClassThumbnailScene, 100> ThumbnailScenes;
 };
diff --git a/Engine/Source/Editor/UnrealEd/Private/Kismet2/Kismet2.cpp b/Engine/Source/Editor/UnrealEd/Private/Kismet2/Kismet2.cpp
--- a/Engine/Source/Editor/UnrealEd/Private/Kismet2/Kismet2.cpp
+++ b/Engine/Source/Editor/UnrealEd/Private/Kismet2/Kismet2.cpp
@@ -321,7 +321,7 @@ void FBlueprintUnloader::UnloadBlueprint(const bool bResetPackage)
 			}
 		}
 
-		FKismetEditorUtilities::OnBlueprintUnloaded.Broadcast(UnloadingBp);
+		UnloadingBp->ClearEditorReferences();
 
 		// handled in FBlueprintEditor (from the OnBlueprintUnloaded event)
 // 		IAssetEditorInstance* EditorInst = GEditor->GetEditorSubsystem<UAssetEditorSubsystem>()->FindEditorForAsset(UnloadingBp, /*bFocusIfOpen =*/false);
@@ -398,6 +398,7 @@ void FBlueprintUnloader::ReplaceStaleRefs(UBlueprint* NewBlueprint)
 // Static variable definition
 TArray<FString> FKismetEditorUtilities::TrackedBlueprintParentList;
 FKismetEditorUtilities::FOnBlueprintUnloaded FKismetEditorUtilities::OnBlueprintUnloaded;
+FKismetEditorUtilities::FOnBlueprintGeneratedClassUnloaded FKismetEditorUtilities::OnBlueprintGeneratedClassUnloaded;
 TMultiMap<void*, FKismetEditorUtilities::FDefaultEventNodeData> FKismetEditorUtilities::AutoGeneratedDefaultEventsMap;
 TMultiMap<void*, FKismetEditorUtilities::FOnBlueprintCreatedData> FKismetEditorUtilities::OnBlueprintCreatedCallbacks;
 
diff --git a/Engine/Source/Editor/UnrealEd/Private/PackageTools.cpp b/Engine/Source/Editor/UnrealEd/Private/PackageTools.cpp
--- a/Engine/Source/Editor/UnrealEd/Private/PackageTools.cpp
+++ b/Engine/Source/Editor/UnrealEd/Private/PackageTools.cpp
@@ -498,7 +498,11 @@ UPackageTools::UPackageTools(const FObjectInitializer& ObjectInitializer)
 
 						BP->CachedDependents.Reset();
 					}
-					if (UWorld* World = Cast<UWorld>(Obj))
+					else if (UBlueprintGeneratedClass* BPGC = Cast<UBlueprintGeneratedClass>(Obj))
+					{
+						FKismetEditorUtilities::OnBlueprintGeneratedClassUnloaded.Broadcast(BPGC);
+					}
+					else if (UWorld* World = Cast<UWorld>(Obj))
 					{
 						if (World->bIsWorldInitialized)
 						{
diff --git a/Engine/Source/Editor/UnrealEd/Private/ThumbnailRendering/BlueprintThumbnailRenderer.cpp b/Engine/Source/Editor/UnrealEd/Private/ThumbnailRendering/BlueprintThumbnailRenderer.cpp
--- a/Engine/Source/Editor/UnrealEd/Private/ThumbnailRendering/BlueprintThumbnailRenderer.cpp
+++ b/Engine/Source/Editor/UnrealEd/Private/ThumbnailRendering/BlueprintThumbnailRenderer.cpp
@@ -15,11 +15,6 @@ UBlueprintThumbnailRenderer::UBlueprintThumbnailRenderer(const FObjectInitialize
 	FKismetEditorUtilities::OnBlueprintUnloaded.AddUObject(this, &UBlueprintThumbnailRenderer::OnBlueprintUnloaded);
 }
 
-UBlueprintThumbnailRenderer::~UBlueprintThumbnailRenderer()
-{
-	FKismetEditorUtilities::OnBlueprintUnloaded.RemoveAll(this);
-}
-
 bool UBlueprintThumbnailRenderer::CanVisualizeAsset(UObject* Object)
 {
 	UBlueprint* Blueprint = Cast<UBlueprint>(Object);
@@ -117,6 +112,7 @@ void UBlueprintThumbnailRenderer::Draw(UObject* Object, int32 X, int32 Y, uint32
 
 void UBlueprintThumbnailRenderer::BeginDestroy()
 {
+	FKismetEditorUtilities::OnBlueprintUnloaded.RemoveAll(this);
 	ThumbnailScenes.Clear();
 
 	Super::BeginDestroy();
diff --git a/Engine/Source/Editor/UnrealEd/Private/ThumbnailRendering/ClassThumbnailRenderer.cpp b/Engine/Source/Editor/UnrealEd/Private/ThumbnailRendering/ClassThumbnailRenderer.cpp
--- a/Engine/Source/Editor/UnrealEd/Private/ThumbnailRendering/ClassThumbnailRenderer.cpp
+++ b/Engine/Source/Editor/UnrealEd/Private/ThumbnailRendering/ClassThumbnailRenderer.cpp
@@ -1,6 +1,7 @@
 // Copyright Epic Games, Inc. All Rights Reserved.
 
 #include "ThumbnailRendering/ClassThumbnailRenderer.h"
+#include "Kismet2/KismetEditorUtilities.h"
 #include "ShowFlags.h"
 #include "SceneView.h"
 #include "Misc/App.h"
@@ -8,6 +9,7 @@
 UClassThumbnailRenderer::UClassThumbnailRenderer(const FObjectInitializer& ObjectInitializer)
 	: Super(ObjectInitializer)
 {
+	FKismetEditorUtilities::OnBlueprintGeneratedClassUnloaded.AddUObject(this, &UClassThumbnailRenderer::OnBlueprintGeneratedClassUnloaded);
 }
 
 bool UClassThumbnailRenderer::CanVisualizeAsset(UObject* Object)
@@ -53,7 +55,13 @@ void UClassThumbnailRenderer::Draw(UObject* Object, int32 X, int32 Y, uint32 Wid
 
 void UClassThumbnailRenderer::BeginDestroy()
 {
+	FKismetEditorUtilities::OnBlueprintGeneratedClassUnloaded.RemoveAll(this);
 	ThumbnailScenes.Clear();
 
 	Super::BeginDestroy();
 }
+
+void UClassThumbnailRenderer::OnBlueprintGeneratedClassUnloaded(UBlueprintGeneratedClass* BPGC)
+{
+	ThumbnailScenes.RemoveThumbnailScene(BPGC);
+}
diff --git a/Engine/Source/Editor/UnrealEd/Public/Kismet2/KismetEditorUtilities.h b/Engine/Source/Editor/UnrealEd/Public/Kismet2/KismetEditorUtilities.h
--- a/Engine/Source/Editor/UnrealEd/Public/Kismet2/KismetEditorUtilities.h
+++ b/Engine/Source/Editor/UnrealEd/Public/Kismet2/KismetEditorUtilities.h
@@ -126,6 +126,10 @@ public:
 	DECLARE_MULTICAST_DELEGATE_OneParam(FOnBlueprintUnloaded, UBlueprint*);
 	static FOnBlueprintUnloaded OnBlueprintUnloaded;
 
+	/** Event that's broadcast anytime a blueprint generated class is unloaded */
+	DECLARE_MULTICAST_DELEGATE_OneParam(FOnBlueprintGeneratedClassUnloaded, UBlueprintGeneratedClass*);
+	static FOnBlueprintGeneratedClassUnloaded OnBlueprintGeneratedClassUnloaded;
+
 	/** 
 	 * Unloads the specified Blueprint (marking it pending-kill, and removing it 
 	 * from its outer package). Then proceeds to replace all references with a
diff --git a/Engine/Source/Runtime/Engine/Private/Blueprint.cpp b/Engine/Source/Runtime/Engine/Private/Blueprint.cpp
--- a/Engine/Source/Runtime/Engine/Private/Blueprint.cpp
+++ b/Engine/Source/Runtime/Engine/Private/Blueprint.cpp
@@ -1989,6 +1989,10 @@ void UBlueprint::ReplaceDeprecatedNodes()
 void UBlueprint::ClearEditorReferences()
 {
 	FKismetEditorUtilities::OnBlueprintUnloaded.Broadcast(this);
+	if (UBlueprintGeneratedClass* BPGC = Cast<UBlueprintGeneratedClass>(GeneratedClass))
+	{
+		FKismetEditorUtilities::OnBlueprintGeneratedClassUnloaded.Broadcast(BPGC);
+	}
 }
 
 UInheritableComponentHandler* UBlueprint::GetInheritableComponentHandler(bool bCreateIfNecessary)
