UE-160860 - Fix UnrealEd Module truncation warnings
Part 2 of N

diff --git a/Engine/Source/Editor/UnrealEd/Private/AssetThumbnail.cpp b/Engine/Source/Editor/UnrealEd/Private/AssetThumbnail.cpp
--- a/Engine/Source/Editor/UnrealEd/Private/AssetThumbnail.cpp
+++ b/Engine/Source/Editor/UnrealEd/Private/AssetThumbnail.cpp
@@ -321,7 +321,7 @@ public:
 	{
 		if ( WidthLastFrame != AllottedGeometry.Size.X )
 		{
-			WidthLastFrame = AllottedGeometry.Size.X;
+			WidthLastFrame = static_cast<float>(AllottedGeometry.Size.X);
 
 			// The width changed, update the font
 			if ( GenericLabelTextBlock.IsValid() )
@@ -1255,7 +1255,7 @@ FSlateTexture2DRHIRef* FAssetThumbnailPool::AccessTexture( const FAssetData& Ass
 			if( FreeThumbnails.Num() == 0 && ThumbnailToTextureMap.Num() == NumInPool )
 			{
 				// Find the thumbnail which was accessed last and use it for the new thumbnail
-				float LastAccessTime = FLT_MAX;
+				double LastAccessTime = std::numeric_limits<double>::max();
 				const FThumbId* AssetToRemove = NULL;
 				for( TMap< FThumbId, TSharedRef<FThumbnailInfo> >::TConstIterator It(ThumbnailToTextureMap); It; ++It )
 				{
diff --git a/Engine/Source/Editor/UnrealEd/Private/AutoReimport/ReimportFeedbackContext.cpp b/Engine/Source/Editor/UnrealEd/Private/AutoReimport/ReimportFeedbackContext.cpp
--- a/Engine/Source/Editor/UnrealEd/Private/AutoReimport/ReimportFeedbackContext.cpp
+++ b/Engine/Source/Editor/UnrealEd/Private/AutoReimport/ReimportFeedbackContext.cpp
@@ -179,7 +179,7 @@ public:
 			.AutoHeight()
 			[
 				SNew(SBox)
-				.HeightOverride(2)
+				.HeightOverride(2.0f)
 				[
 					SAssignNew(ProgressBar, SProgressBar)
 					.BorderPadding(FVector2D::ZeroVector)
diff --git a/Engine/Source/Editor/UnrealEd/Private/CameraController.cpp b/Engine/Source/Editor/UnrealEd/Private/CameraController.cpp
--- a/Engine/Source/Editor/UnrealEd/Private/CameraController.cpp
+++ b/Engine/Source/Editor/UnrealEd/Private/CameraController.cpp
@@ -298,12 +298,12 @@ void FEditorCameraController::UpdateRotation( const FCameraControllerUserImpulse
 		// This will serve as both our source and destination rotation value
 		FVector::FReal& RotationVelocity = RotationVelocityEuler[ CurRotationAxis ];
 
-		const float RotationImpulse = RotateImpulseEuler[ CurRotationAxis ];
-		const float RotationVelocityModifier = RotateVelocityModifierEuler[ CurRotationAxis ];
+		const FVector::FReal RotationImpulse = RotateImpulseEuler[CurRotationAxis];
+		const FVector::FReal RotationVelocityModifier = RotateVelocityModifierEuler[ CurRotationAxis ];
 
 
 		// Compute acceleration
-		float RotationAcceleration = RotationImpulse * Config.RotationAccelerationRate;
+		FVector::FReal RotationAcceleration = RotationImpulse * Config.RotationAccelerationRate;
 
 		if( Config.bUsePhysicsBasedRotation || Config.bForceRotationalPhysics)
 		{
@@ -337,7 +337,7 @@ void FEditorCameraController::UpdateRotation( const FCameraControllerUserImpulse
 
 
 		// Constrain maximum rotation speed
-		RotationVelocity = FMath::Clamp<float>(RotationVelocity, -Config.MaximumRotationSpeed, Config.MaximumRotationSpeed);
+		RotationVelocity = FMath::Clamp(RotationVelocity, -Config.MaximumRotationSpeed, Config.MaximumRotationSpeed);
 
 		// Clamp velocity to a reasonably small number
 		if( FMath::Abs( RotationVelocity ) < KINDA_SMALL_NUMBER )
diff --git a/Engine/Source/Editor/UnrealEd/Private/Commandlets/DiffAssetRegistriesCommandlet.cpp b/Engine/Source/Editor/UnrealEd/Private/Commandlets/DiffAssetRegistriesCommandlet.cpp
--- a/Engine/Source/Editor/UnrealEd/Private/Commandlets/DiffAssetRegistriesCommandlet.cpp
+++ b/Engine/Source/Editor/UnrealEd/Private/Commandlets/DiffAssetRegistriesCommandlet.cpp
@@ -367,18 +367,6 @@ void UDiffAssetRegistriesCommandlet::FillChangelists(FString Branch, FString CL,
 	}
 }
 
-void rescale(int bytes, double &value, int &exp)
-{
-	value = bytes;
-	value = fabs(value);
-
-	while (value > 1024.0)
-	{
-		value /= 1024.0;
-		exp++;
-	}
-}
-
 void UDiffAssetRegistriesCommandlet::ConsistencyCheck(const FString& OldPath, const FString& NewPath)
 {
 	{
@@ -476,8 +464,8 @@ void UDiffAssetRegistriesCommandlet::ConsistencyCheck(const FString& OldPath, co
 		}
 	}
 
-	int64 changes = 0;
-	int64 change_bytes = 0;
+	int64 Changes = 0;
+	int64 ChangeBytes = 0;
 
 	// find all entries of CookModified that do not exist in GuidModified
 	const TMap<FName, const FAssetPackageData*>& PackageMap = NewState.GetAssetPackageDataMap();
@@ -487,8 +475,8 @@ void UDiffAssetRegistriesCommandlet::ConsistencyCheck(const FString& OldPath, co
 
 		if (!GuidModified.Contains(Package))
 		{
-			++changes;
-			change_bytes += Data->DiskSize;
+			++Changes;
+			ChangeBytes += Data->DiskSize;
 			if (bIsVerbose)
 			{
 				UE_LOG(LogDiffAssets, Display, TEXT("%s : %d bytes"), *Package.ToString(), Data->DiskSize);
@@ -496,13 +484,25 @@ void UDiffAssetRegistriesCommandlet::ConsistencyCheck(const FString& OldPath, co
 		}
 	}
 
-	double change_value = 0.0;
-	int change_exp = 0;
+	double ChangeValue = 0.0;
+	int32 ChangeExp = 0;
+
+	auto Rescale = [](int64 Bytes, double& Value, int32& Exp)
+	{
+		Value = Bytes;
+		Value = fabs(Value);
+
+		while (Value > 1024.0)
+		{
+			Value /= 1024.0;
+			Exp++;
+		}
+	};
 
-	rescale(change_bytes, change_value, change_exp);
+	Rescale(ChangeBytes, ChangeValue, ChangeExp);
 
 	UE_LOG(LogDiffAssets, Display, TEXT("Summary:"));
-	UE_LOG(LogDiffAssets, Display, TEXT("%d nondeterministic cooks, %8.3f %cB"), changes, change_value, " KMGTP"[change_exp]);
+	UE_LOG(LogDiffAssets, Display, TEXT("%d nondeterministic cooks, %8.3f %cB"), Changes, ChangeValue, " KMGTP"[ChangeExp]);
 }
 
 bool	UDiffAssetRegistriesCommandlet::IsInRelevantChunk(FAssetRegistryState& InRegistryState, FName InAssetPath)
diff --git a/Engine/Source/Editor/UnrealEd/Private/Commandlets/SummarizeTraceCommandlet.cpp b/Engine/Source/Editor/UnrealEd/Private/Commandlets/SummarizeTraceCommandlet.cpp
--- a/Engine/Source/Editor/UnrealEd/Private/Commandlets/SummarizeTraceCommandlet.cpp
+++ b/Engine/Source/Editor/UnrealEd/Private/Commandlets/SummarizeTraceCommandlet.cpp
@@ -521,7 +521,7 @@ public:
 			return;
 		}
 
-		const double TotalPopulation = Count + Other.Count;
+		const double TotalPopulation = static_cast<double>(Count + Other.Count);
 		const double MeanDifference = Mean - Other.Mean;
 		const double A = ((Count - 1) * Variance()) + ((Other.Count - 1) * Other.Variance());
 		const double B = (MeanDifference) * (MeanDifference) * (Count * Other.Count / TotalPopulation);
diff --git a/Engine/Source/Editor/UnrealEd/Private/DebugToolExec.cpp b/Engine/Source/Editor/UnrealEd/Private/DebugToolExec.cpp
--- a/Engine/Source/Editor/UnrealEd/Private/DebugToolExec.cpp
+++ b/Engine/Source/Editor/UnrealEd/Private/DebugToolExec.cpp
@@ -246,7 +246,7 @@ bool FDebugToolExec::Exec( UWorld* InWorld, const TCHAR* Cmd, FOutputDevice& Ar
 				{
 					if (IsValid(*It))
 					{
-						float const Dist = (PlayerController && It->GetRootComponent()) ? FVector::Dist(It->GetActorLocation(), PlayerLocation) : 0.f;
+						const double Dist = (PlayerController && It->GetRootComponent()) ? FVector::Dist(It->GetActorLocation(), PlayerLocation) : 0.f;
 						if (Dist < MinDist)
 						{
 							MinDist = Dist;
diff --git a/Engine/Source/Editor/UnrealEd/Private/Dialogs/DlgMoveAssets.cpp b/Engine/Source/Editor/UnrealEd/Private/Dialogs/DlgMoveAssets.cpp
--- a/Engine/Source/Editor/UnrealEd/Private/Dialogs/DlgMoveAssets.cpp
+++ b/Engine/Source/Editor/UnrealEd/Private/Dialogs/DlgMoveAssets.cpp
@@ -90,7 +90,7 @@ public:
 						SNew(SVerticalBox)
 						+SVerticalBox::Slot()
 						.AutoHeight()
-						.Padding(3)
+						.Padding(3.0f)
 						[
 							SNew(SGridPanel)
 
diff --git a/Engine/Source/Editor/UnrealEd/Private/Dialogs/DlgPickAssetPath.cpp b/Engine/Source/Editor/UnrealEd/Private/Dialogs/DlgPickAssetPath.cpp
--- a/Engine/Source/Editor/UnrealEd/Private/Dialogs/DlgPickAssetPath.cpp
+++ b/Engine/Source/Editor/UnrealEd/Private/Dialogs/DlgPickAssetPath.cpp
@@ -50,14 +50,14 @@ void SDlgPickAssetPath::Construct(const FArguments& InArgs)
 
 					+ SVerticalBox::Slot()
 					.FillHeight(1)
-					.Padding(3)
+					.Padding(3.0f)
 					[
 						ContentBrowserModule.Get().CreatePathPicker(PathPickerConfig)
 					]
 
 					+ SVerticalBox::Slot()
 					.AutoHeight()
-					.Padding(3)
+					.Padding(3.0f)
 					[
 						SNew(SHorizontalBox)
 
diff --git a/Engine/Source/Editor/UnrealEd/Private/Dialogs/DlgPickPath.cpp b/Engine/Source/Editor/UnrealEd/Private/Dialogs/DlgPickPath.cpp
--- a/Engine/Source/Editor/UnrealEd/Private/Dialogs/DlgPickPath.cpp
+++ b/Engine/Source/Editor/UnrealEd/Private/Dialogs/DlgPickPath.cpp
@@ -46,7 +46,7 @@ void SDlgPickPath::Construct(const FArguments& InArgs)
 
 					+ SVerticalBox::Slot()
 					.FillHeight(1)
-					.Padding(3)
+					.Padding(3.0f)
 					[
 						ContentBrowserModule.Get().CreatePathPicker(PathPickerConfig)
 					]
diff --git a/Engine/Source/Editor/UnrealEd/Private/Dialogs/SDeleteAssetsDialog.cpp b/Engine/Source/Editor/UnrealEd/Private/Dialogs/SDeleteAssetsDialog.cpp
--- a/Engine/Source/Editor/UnrealEd/Private/Dialogs/SDeleteAssetsDialog.cpp
+++ b/Engine/Source/Editor/UnrealEd/Private/Dialogs/SDeleteAssetsDialog.cpp
@@ -134,7 +134,7 @@ void SDeleteAssetsDialog::Construct( const FArguments& InArgs, TSharedRef<FAsset
 	[
 		SAssignNew(RootContainer, SBorder)
 		.BorderImage( FAppStyle::GetBrush( "AssetDeleteDialog.Background" ) )
-		.Padding(10)
+		.Padding(10.0f)
 	];
 
 	DeleteModel->OnStateChanged().AddRaw(this, &SDeleteAssetsDialog::HandleDeleteModelStateChanged);
@@ -310,7 +310,7 @@ TSharedRef<SWidget> SDeleteAssetsDialog::BuildDeleteDialog()
 	[
 		SNew( SBorder )
 		.BorderImage( FAppStyle::GetBrush( "ToolPanel.GroupBorder" ) )
-		.Padding( 0 )
+		.Padding( 0.0f )
 		[
 			SNew( SVerticalBox )
 
@@ -515,7 +515,7 @@ TSharedRef<SWidget> SDeleteAssetsDialog::BuildReplaceReferencesWidget()
 		.VAlign( EVerticalAlignment::VAlign_Center )
 		.ComboButtonStyle( FAppStyle::Get(), "ToolbarComboButton" )
 		.ForegroundColor( FLinearColor::White )
-		.ContentPadding( 3 )
+		.ContentPadding(3.0f)
 		.MenuPlacement( EMenuPlacement::MenuPlacement_BelowAnchor )
 		.OnGetMenuContent( this, &SDeleteAssetsDialog::MakeConsolidationAssetPicker )
 		.ButtonContent()
diff --git a/Engine/Source/Editor/UnrealEd/Private/DragAndDrop/AssetDragDropOp.cpp b/Engine/Source/Editor/UnrealEd/Private/DragAndDrop/AssetDragDropOp.cpp
--- a/Engine/Source/Editor/UnrealEd/Private/DragAndDrop/AssetDragDropOp.cpp
+++ b/Engine/Source/Editor/UnrealEd/Private/DragAndDrop/AssetDragDropOp.cpp
@@ -109,8 +109,8 @@ TSharedPtr<SWidget> FAssetDragDropOp::GetDefaultDecorator() const
 			.HAlign(HAlign_Left)
 			[
 				SNew(SBox) 
-				.WidthOverride(ThumbnailSize) 
-				.HeightOverride(ThumbnailSize)
+				.WidthOverride(static_cast<float>(ThumbnailSize)) 
+				.HeightOverride(static_cast<float>(ThumbnailSize))
 				.Content()
 				[
 					SNew(SOverlay)
@@ -154,7 +154,7 @@ TSharedPtr<SWidget> FAssetDragDropOp::GetDefaultDecorator() const
 			.VAlign(VAlign_Center)
 			[
 				SNew(SBox)
-				.MinDesiredWidth(80)
+				.MinDesiredWidth(80.0f)
 				.Content()
 				[
 					SNew(SHorizontalBox)
diff --git a/Engine/Source/Editor/UnrealEd/Private/Editor/ActorPositioning.cpp b/Engine/Source/Editor/UnrealEd/Private/Editor/ActorPositioning.cpp
--- a/Engine/Source/Editor/UnrealEd/Private/Editor/ActorPositioning.cpp
+++ b/Engine/Source/Editor/UnrealEd/Private/Editor/ActorPositioning.cpp
@@ -138,7 +138,7 @@ FActorPositionTraceResult FActorPositioning::TraceWorldForPosition(const UWorld&
 		}
 
 		// Go through all hits and find closest
-		float ClosestHitDistanceSqr = TNumericLimits<float>::Max();
+		double ClosestHitDistanceSqr = std::numeric_limits<double>::max();
 
 		for (const FHitResult& Hit : Hits)
 		{
diff --git a/Engine/Source/Editor/UnrealEd/Private/EditorEngine.cpp b/Engine/Source/Editor/UnrealEd/Private/EditorEngine.cpp
--- a/Engine/Source/Editor/UnrealEd/Private/EditorEngine.cpp
+++ b/Engine/Source/Editor/UnrealEd/Private/EditorEngine.cpp
@@ -2225,7 +2225,9 @@ bool UEditorEngine::UpdateSingleViewportClient(FEditorViewportClient* InViewport
 	// otherwise content for editor view can be streamed out if there are other views (ex: thumbnails)
 	if (InViewportClient->IsPerspective())
 	{
-		IStreamingManager::Get().AddViewInformation( InViewportClient->GetViewLocation(), InViewportClient->Viewport->GetSizeXY().X, InViewportClient->Viewport->GetSizeXY().X / FMath::Tan(FMath::DegreesToRadians(InViewportClient->ViewFOV * 0.5f)) );
+		float XSize = static_cast<float>(InViewportClient->Viewport->GetSizeXY().X);
+
+		IStreamingManager::Get().AddViewInformation( InViewportClient->GetViewLocation(), XSize, XSize / FMath::Tan(FMath::DegreesToRadians(InViewportClient->ViewFOV * 0.5f)) );
 	}
 	
 	// Only allow viewports to be drawn if we are not throttling for slate UI responsiveness or if the viewport client requested a redraw
diff --git a/Engine/Source/Editor/UnrealEd/Private/EditorExporters.cpp b/Engine/Source/Editor/UnrealEd/Private/EditorExporters.cpp
--- a/Engine/Source/Editor/UnrealEd/Private/EditorExporters.cpp
+++ b/Engine/Source/Editor/UnrealEd/Private/EditorExporters.cpp
@@ -1418,8 +1418,8 @@ static void ExportPolys( UPolys* Polys, int32 &PolyNum, int32 TotalPolys, FFeedb
 					const FVert& Vert = Model->Verts[TriVertIndices[TriVertexIndex]];
 					const FVector& Vertex = (FVector)Model->Points[Vert.pVertex];
 
-					float U = ((Vertex - TextureBase) | TextureX) / UModel::GetGlobalBSPTexelScale();
-					float V = ((Vertex - TextureBase) | TextureY) / UModel::GetGlobalBSPTexelScale();
+					double U = ((Vertex - TextureBase) | TextureX) / UModel::GetGlobalBSPTexelScale();
+					double V = ((Vertex - TextureBase) | TextureY) / UModel::GetGlobalBSPTexelScale();
 
 					Vertices[TriVertexIndex].Vert = Vertex;
 					Vertices[TriVertexIndex].UV = FVector2D( U, V );
diff --git a/Engine/Source/Editor/UnrealEd/Private/EditorScriptingHelpers.cpp b/Engine/Source/Editor/UnrealEd/Private/EditorScriptingHelpers.cpp
--- a/Engine/Source/Editor/UnrealEd/Private/EditorScriptingHelpers.cpp
+++ b/Engine/Source/Editor/UnrealEd/Private/EditorScriptingHelpers.cpp
@@ -16,7 +16,7 @@ namespace EditorScriptingHelpersInternal
 	FString RemoveFullName(const FString& AnyAssetPath, FString& OutFailureReason)
 	{
 		FString Result = AnyAssetPath.TrimStartAndEnd();
-		int32 NumberOfSpace = Algo::Count(AnyAssetPath, TEXT(' '));
+		SIZE_T NumberOfSpace = Algo::Count(AnyAssetPath, TEXT(' '));
 
 		if (NumberOfSpace == 0)
 		{
diff --git a/Engine/Source/Editor/UnrealEd/Private/Factories/CSVImportFactory.cpp b/Engine/Source/Editor/UnrealEd/Private/Factories/CSVImportFactory.cpp
--- a/Engine/Source/Editor/UnrealEd/Private/Factories/CSVImportFactory.cpp
+++ b/Engine/Source/Editor/UnrealEd/Private/Factories/CSVImportFactory.cpp
@@ -232,7 +232,7 @@ UObject* UCSVImportFactory::FactoryCreateText(UClass* InClass, UObject* InParent
 	if (bDoImport)
 	{
 		// Convert buffer to an FString (will this be slow with big tables?)
-		int32 NumChars = (BufferEnd - Buffer);
+		int32 NumChars = UE_PTRDIFF_TO_INT32(BufferEnd - Buffer);
 		TArray<TCHAR, FString::AllocatorType>& StringChars = ImportSettings.DataToImport.GetCharArray();
 		StringChars.AddUninitialized(NumChars+1);
 		FMemory::Memcpy(StringChars.GetData(), Buffer, NumChars*sizeof(TCHAR));
diff --git a/Engine/Source/Editor/UnrealEd/Private/Factories/DataTableFactory.cpp b/Engine/Source/Editor/UnrealEd/Private/Factories/DataTableFactory.cpp
--- a/Engine/Source/Editor/UnrealEd/Private/Factories/DataTableFactory.cpp
+++ b/Engine/Source/Editor/UnrealEd/Private/Factories/DataTableFactory.cpp
@@ -97,7 +97,7 @@ bool UDataTableFactory::ConfigureProperties()
 
 			return
 				SNew(SBox)
-				.WidthOverride(330)
+				.WidthOverride(330.0f)
 				[
 					SNew(SVerticalBox)
 
diff --git a/Engine/Source/Editor/UnrealEd/Private/Factories/MirrorDataTableFactory.cpp b/Engine/Source/Editor/UnrealEd/Private/Factories/MirrorDataTableFactory.cpp
--- a/Engine/Source/Editor/UnrealEd/Private/Factories/MirrorDataTableFactory.cpp
+++ b/Engine/Source/Editor/UnrealEd/Private/Factories/MirrorDataTableFactory.cpp
@@ -102,7 +102,7 @@ public:
 		Options.StructFilter = MakeShared<FDataTableStructFilter>();
 		return
 			SNew(SBox)
-			.WidthOverride(330)
+			.WidthOverride(330.0f)
 			[
 				SNew(SVerticalBox)
 
@@ -195,7 +195,7 @@ public:
 				[
 					SNew(SExpandableArea)
 					.AreaTitle(LOCTEXT("CreateMirrorDataTableAdvancedOptions", "Advanced Options"))
-					.Padding(2)
+					.Padding(2.0f)
 					.InitiallyCollapsed(true)
 					.BodyContent()
 					[
diff --git a/Engine/Source/Editor/UnrealEd/Private/Fbx/FbxCompareWindow.cpp b/Engine/Source/Editor/UnrealEd/Private/Fbx/FbxCompareWindow.cpp
--- a/Engine/Source/Editor/UnrealEd/Private/Fbx/FbxCompareWindow.cpp
+++ b/Engine/Source/Editor/UnrealEd/Private/Fbx/FbxCompareWindow.cpp
@@ -91,7 +91,7 @@ void SFbxSkeltonConflictWindow::Construct(const FArguments& InArgs)
 					SNew(SVerticalBox)
 					+ SVerticalBox::Slot()
 					.FillHeight(1.0f)
-					.Padding(2)
+					.Padding(2.0f)
 					[
 						SNew(SSplitter)
 						.Orientation(Orient_Vertical)
@@ -114,7 +114,7 @@ void SFbxSkeltonConflictWindow::Construct(const FArguments& InArgs)
 			+ SVerticalBox::Slot()
 			.AutoHeight()
 			.HAlign(HAlign_Right)
-			.Padding(2)
+			.Padding(2.0f)
 			[
 				SNew(SHorizontalBox)
 				+ SHorizontalBox::Slot()
diff --git a/Engine/Source/Editor/UnrealEd/Private/Fbx/FbxExportOptionsWindow.cpp b/Engine/Source/Editor/UnrealEd/Private/Fbx/FbxExportOptionsWindow.cpp
--- a/Engine/Source/Editor/UnrealEd/Private/Fbx/FbxExportOptionsWindow.cpp
+++ b/Engine/Source/Editor/UnrealEd/Private/Fbx/FbxExportOptionsWindow.cpp
@@ -46,13 +46,13 @@ void SFbxExportOptionsWindow::Construct(const FArguments& InArgs)
 			SNew(SVerticalBox)
 			+ SVerticalBox::Slot()
 			.AutoHeight()
-			.Padding(2)
+			.Padding(2.0f)
 			[
 				SAssignNew(HeaderToolBox, SBox)
 			]
 			+SVerticalBox::Slot()
 			.AutoHeight()
-			.Padding(2)
+			.Padding(2.0f)
 			[
 				SNew(SBorder)
 				.Padding(FMargin(3))
@@ -79,17 +79,17 @@ void SFbxExportOptionsWindow::Construct(const FArguments& InArgs)
 			]
 			+ SVerticalBox::Slot()
 			.FillHeight(1.0f)
-			.Padding(2)
+			.Padding(2.0f)
 			[
 				SAssignNew(InspectorBox, SBox)
 			]
 			+ SVerticalBox::Slot()
 			.AutoHeight()
 			.HAlign(HAlign_Right)
-			.Padding(2)
+			.Padding(2.0f)
 			[
 				SNew(SUniformGridPanel)
-				.SlotPadding(2)
+				.SlotPadding(2.0f)
 				+ SUniformGridPanel::Slot(0, 0)
 				[
 					SNew(SButton)
diff --git a/Engine/Source/Editor/UnrealEd/Private/Fbx/FbxMaterialConflictWindow.cpp b/Engine/Source/Editor/UnrealEd/Private/Fbx/FbxMaterialConflictWindow.cpp
--- a/Engine/Source/Editor/UnrealEd/Private/Fbx/FbxMaterialConflictWindow.cpp
+++ b/Engine/Source/Editor/UnrealEd/Private/Fbx/FbxMaterialConflictWindow.cpp
@@ -86,7 +86,7 @@ void SFbxMaterialConflictWindow::Construct(const FArguments& InArgs)
 							SNew(SVerticalBox)
 							+ SVerticalBox::Slot()
 							.AutoHeight()
-							.Padding(2)
+							.Padding(2.0f)
 							[
 								// Material Compare section
 								MaterialCompareSection.ToSharedRef()
diff --git a/Engine/Source/Editor/UnrealEd/Private/Fbx/FbxOptionWindow.cpp b/Engine/Source/Editor/UnrealEd/Private/Fbx/FbxOptionWindow.cpp
--- a/Engine/Source/Editor/UnrealEd/Private/Fbx/FbxOptionWindow.cpp
+++ b/Engine/Source/Editor/UnrealEd/Private/Fbx/FbxOptionWindow.cpp
@@ -52,13 +52,13 @@ void SFbxOptionWindow::Construct(const FArguments& InArgs)
 			SNew(SVerticalBox)
 			+ SVerticalBox::Slot()
 			.AutoHeight()
-			.Padding(2)
+			.Padding(2.0f)
 			[
 				SAssignNew(ImportTypeDisplay, SBox)
 			]
 			+SVerticalBox::Slot()
 			.AutoHeight()
-			.Padding(2)
+			.Padding(2.0f)
 			[
 				SNew(SBorder)
 				.Padding(FMargin(3))
@@ -84,7 +84,7 @@ void SFbxOptionWindow::Construct(const FArguments& InArgs)
 			]
 			+ SVerticalBox::Slot()
 			.AutoHeight()
-			.Padding(2)
+			.Padding(2.0f)
 			[
 				SAssignNew(InspectorBox, SBox)
 				.MaxDesiredHeight(650.0f)
@@ -92,10 +92,10 @@ void SFbxOptionWindow::Construct(const FArguments& InArgs)
 			]
 			+ SVerticalBox::Slot()
 			.AutoHeight()
-			.Padding(2)
+			.Padding(2.0f)
 			[
 				SNew(SUniformGridPanel)
-				.SlotPadding(2)
+				.SlotPadding(2.0f)
 				+ SUniformGridPanel::Slot(1, 0)
 				[
 					SAssignNew(ImportAllButton, SPrimaryButton)
diff --git a/Engine/Source/Editor/UnrealEd/Private/Fbx/FbxSceneImportData.cpp b/Engine/Source/Editor/UnrealEd/Private/Fbx/FbxSceneImportData.cpp
--- a/Engine/Source/Editor/UnrealEd/Private/Fbx/FbxSceneImportData.cpp
+++ b/Engine/Source/Editor/UnrealEd/Private/Fbx/FbxSceneImportData.cpp
@@ -95,10 +95,10 @@ UnFbx::FBXImportOptions *JSONToFbxOption(TSharedPtr<FJsonValue> OptionJsonValue,
 		(*DataObj)->TryGetNumberField("G", G);
 		(*DataObj)->TryGetNumberField("B", B);
 		(*DataObj)->TryGetNumberField("A", A);
-		Option->VertexOverrideColor.R = (float)R;
-		Option->VertexOverrideColor.G = (float)G;
-		Option->VertexOverrideColor.B = (float)B;
-		Option->VertexOverrideColor.A = (float)A;
+		Option->VertexOverrideColor.R = static_cast<uint8>(R);
+		Option->VertexOverrideColor.G = static_cast<uint8>(G);
+		Option->VertexOverrideColor.B = static_cast<uint8>(B);
+		Option->VertexOverrideColor.A = static_cast<uint8>(A);
 	}
 	OptionObj->TryGetBoolField("bRemoveDegenerates", Option->bRemoveDegenerates);
 	OptionObj->TryGetBoolField("bBuildReversedIndexBuffer", Option->bBuildReversedIndexBuffer);
diff --git a/Engine/Source/Editor/UnrealEd/Private/Fbx/SFbxSceneOptionWindow.cpp b/Engine/Source/Editor/UnrealEd/Private/Fbx/SFbxSceneOptionWindow.cpp
--- a/Engine/Source/Editor/UnrealEd/Private/Fbx/SFbxSceneOptionWindow.cpp
+++ b/Engine/Source/Editor/UnrealEd/Private/Fbx/SFbxSceneOptionWindow.cpp
@@ -791,7 +791,7 @@ TSharedRef<SDockTab> SFbxSceneOptionWindow::SpawnSceneReimportTab(const FSpawnTa
 				.AutoHeight()
 				[
 					SNew(SUniformGridPanel)
-					.SlotPadding(2)
+					.SlotPadding(2.0f)
 					+ SUniformGridPanel::Slot(0, 0)
 					[
 						SNew(SHorizontalBox)
@@ -948,7 +948,7 @@ TSharedRef<SDockTab> SFbxSceneOptionWindow::SpawnStaticMeshReimportTab(const FSp
 					.AutoHeight()
 					[
 						SNew(SUniformGridPanel)
-						.SlotPadding(2)
+						.SlotPadding(2.0f)
 						+ SUniformGridPanel::Slot(0, 0)
 						.VAlign(VAlign_Center)
 						.HAlign(HAlign_Left)
@@ -1254,7 +1254,7 @@ void SFbxSceneOptionWindow::Construct(const FArguments& InArgs)
 			SNew(SVerticalBox)
 			+SVerticalBox::Slot()
 			.AutoHeight()
-			.Padding(2)
+			.Padding(2.0f)
 			[
 				SNew(SBorder)
 				.Padding(FMargin(3))
@@ -1281,17 +1281,17 @@ void SFbxSceneOptionWindow::Construct(const FArguments& InArgs)
 			]
 			+ SVerticalBox::Slot()
 			.FillHeight(1.0f)
-			.Padding(2)
+			.Padding(2.0f)
 			[
 				SpawnDockTab().ToSharedRef()
 			]
 			+ SVerticalBox::Slot()
 			.AutoHeight()
 			.HAlign(HAlign_Right)
-			.Padding(2)
+			.Padding(2.0f)
 			[
 				SNew(SUniformGridPanel)
-				.SlotPadding(2)
+				.SlotPadding(2.0f)
 				+ SUniformGridPanel::Slot(0, 0)
 				[
 					IDocumentation::Get()->CreateAnchor(FString("Engine/Content/FBX/ImportOptions"))
diff --git a/Engine/Source/Editor/UnrealEd/Private/Fbx/SSceneMaterialsListView.cpp b/Engine/Source/Editor/UnrealEd/Private/Fbx/SSceneMaterialsListView.cpp
--- a/Engine/Source/Editor/UnrealEd/Private/Fbx/SSceneMaterialsListView.cpp
+++ b/Engine/Source/Editor/UnrealEd/Private/Fbx/SSceneMaterialsListView.cpp
@@ -348,8 +348,8 @@ TSharedPtr<SWidget> SFbxSceneMaterialsListView::OnOpenContextMenu()
 			//We have to limit the height in case there is a lot of item, the computed height can overpass the 16 bit integer maximum.
 			TSharedRef<SWidget> Widget =
 				SNew(SBox)
-				.WidthOverride(300)
-				.HeightOverride(300)
+				.WidthOverride(300.0f)
+				.HeightOverride(300.0f)
 				[
 					ContentBrowserModule.Get().CreateAssetPicker(Config)
 				];
diff --git a/Engine/Source/Editor/UnrealEd/Private/FeedbackContextEditor.cpp b/Engine/Source/Editor/UnrealEd/Private/FeedbackContextEditor.cpp
--- a/Engine/Source/Editor/UnrealEd/Private/FeedbackContextEditor.cpp
+++ b/Engine/Source/Editor/UnrealEd/Private/FeedbackContextEditor.cpp
@@ -147,7 +147,7 @@ public:
 			.Padding(FMargin(FixedPaddingH))
 			[
 				SNew(SBox).
-				WidthOverride(FixedWidth) 
+				WidthOverride(static_cast<float>(FixedWidth))
 				[
 					VerticalBox 
 				]
@@ -232,7 +232,7 @@ private:
 			.AutoHeight()
 			[
 				SNew(SBox)
-				.HeightOverride(SecondaryBarHeight)
+				.HeightOverride(static_cast<float>(SecondaryBarHeight))
 				[
 					SNew(SProgressBar)
 					.BorderPadding(FVector2D::ZeroVector)
diff --git a/Engine/Source/Editor/UnrealEd/Private/GeomFitUtils.cpp b/Engine/Source/Editor/UnrealEd/Private/GeomFitUtils.cpp
--- a/Engine/Source/Editor/UnrealEd/Private/GeomFitUtils.cpp
+++ b/Engine/Source/Editor/UnrealEd/Private/GeomFitUtils.cpp
@@ -24,28 +24,11 @@
 #include "StaticMeshAttributes.h"
 #include "Settings/EditorExperimentalSettings.h"
 
-#define LOCAL_EPS (0.01f)
-static void AddVertexIfNotPresent(TArray<FVector> &vertices, FVector &newVertex)
-{
-	bool isPresent = 0;
-
-	for(int32 i=0; i<vertices.Num() && !isPresent; i++)
-	{
-		float diffSqr = (newVertex - vertices[i]).SizeSquared();
-		if(diffSqr < LOCAL_EPS * LOCAL_EPS)
-			isPresent = 1;
-	}
-
-	if(!isPresent)
-		vertices.Add(newVertex);
-
-}
-
 static bool PromptToRemoveExistingCollision(UStaticMesh* StaticMesh)
 {
 	check(StaticMesh);
 	UBodySetup* bs = StaticMesh->GetBodySetup();
-	if(bs && (bs->AggGeom.GetElementCount() > 0))
+	if (bs && (bs->AggGeom.GetElementCount() > 0))
 	{
 		// If we already have some simplified collision for this mesh - check before we clobber it.
 		/*const EAppReturnType::Type ret = FMessageDialog::Open(EAppMsgType::YesNoCancel, NSLOCTEXT("UnrealEd", "StaticMeshAlreadyHasGeom", "Static Mesh already has simple collision.\nDo you want to replace it?"));
diff --git a/Engine/Source/Editor/UnrealEd/Private/GraphEditor.cpp b/Engine/Source/Editor/UnrealEd/Private/GraphEditor.cpp
--- a/Engine/Source/Editor/UnrealEd/Private/GraphEditor.cpp
+++ b/Engine/Source/Editor/UnrealEd/Private/GraphEditor.cpp
@@ -152,10 +152,10 @@ void SGraphEditor::FocusCommentNodes(TArray<UEdGraphNode*> &CommentNodes, TArray
 		CommentNode->SetNodeUnrelated(true);
 
 		const FSlateRect CommentRect(
-			CommentNode->NodePosX, 
-			CommentNode->NodePosY, 
-			CommentNode->NodePosX + CommentNode->NodeWidth, 
-			CommentNode->NodePosY + CommentNode->NodeHeight
+			static_cast<float>(CommentNode->NodePosX),
+			static_cast<float>(CommentNode->NodePosY),
+			static_cast<float>(CommentNode->NodePosX + CommentNode->NodeWidth),
+			static_cast<float>(CommentNode->NodePosY + CommentNode->NodeHeight)
 		);
 
 		for (auto& RelatedNode : RelatedNodes)
diff --git a/Engine/Source/Editor/UnrealEd/Private/GroupActor.cpp b/Engine/Source/Editor/UnrealEd/Private/GroupActor.cpp
--- a/Engine/Source/Editor/UnrealEd/Private/GroupActor.cpp
+++ b/Engine/Source/Editor/UnrealEd/Private/GroupActor.cpp
@@ -309,13 +309,13 @@ void GetBoundingVectorsForGroup(AGroupActor* GroupActor, FViewport* Viewport, FV
 				FBox ActorBox = Actor->GetComponentsBoundingBox( true );
 
 				// MinVector
-				OutVectorMin.X = FMath::Min<float>( ActorBox.Min.X, OutVectorMin.X );
-				OutVectorMin.Y = FMath::Min<float>( ActorBox.Min.Y, OutVectorMin.Y );
-				OutVectorMin.Z = FMath::Min<float>( ActorBox.Min.Z, OutVectorMin.Z );
+				OutVectorMin.X = FMath::Min<FVector::FReal>( ActorBox.Min.X, OutVectorMin.X );
+				OutVectorMin.Y = FMath::Min<FVector::FReal>( ActorBox.Min.Y, OutVectorMin.Y );
+				OutVectorMin.Z = FMath::Min<FVector::FReal>( ActorBox.Min.Z, OutVectorMin.Z );
 				// MaxVector
-				OutVectorMax.X = FMath::Max<float>( ActorBox.Max.X, OutVectorMax.X );
-				OutVectorMax.Y = FMath::Max<float>( ActorBox.Max.Y, OutVectorMax.Y );
-				OutVectorMax.Z = FMath::Max<float>( ActorBox.Max.Z, OutVectorMax.Z );
+				OutVectorMax.X = FMath::Max<FVector::FReal>( ActorBox.Max.X, OutVectorMax.X );
+				OutVectorMax.Y = FMath::Max<FVector::FReal>( ActorBox.Max.Y, OutVectorMax.Y );
+				OutVectorMax.Z = FMath::Max<FVector::FReal>( ActorBox.Max.Z, OutVectorMax.Z );
 			}
 		}
 	}	
diff --git a/Engine/Source/Editor/UnrealEd/Private/Kismet2/BlueprintEditorUtils.cpp b/Engine/Source/Editor/UnrealEd/Private/Kismet2/BlueprintEditorUtils.cpp
--- a/Engine/Source/Editor/UnrealEd/Private/Kismet2/BlueprintEditorUtils.cpp
+++ b/Engine/Source/Editor/UnrealEd/Private/Kismet2/BlueprintEditorUtils.cpp
@@ -6368,7 +6368,7 @@ void FBlueprintEditorUtils::MarkBlueprintChildrenAsModified(UBlueprint* InBluepr
 	TArray<FAssetData> Children;
 	if (GetChildrenOfBlueprint(InBlueprint, Children) > 0)
 	{
-		int32 Unloaded = Algo::CountIf(Children,
+		SIZE_T Unloaded = Algo::CountIf(Children,
 			[](const FAssetData& Asset)
 			{
 				return !Asset.IsAssetLoaded();
@@ -6379,7 +6379,7 @@ void FBlueprintEditorUtils::MarkBlueprintChildrenAsModified(UBlueprint* InBluepr
 		EAppReturnType::Type DialogResponse = EAppReturnType::Yes;
 		if (Unloaded > 0)
 		{
-			FText Message = FText::Format(LOCTEXT("LoadChildrenPopupMessage", "Load {0} unloaded child blueprints to fix up phantom references?"), FText::FromString(FString::FromInt(Unloaded)));
+			FText Message = FText::Format(LOCTEXT("LoadChildrenPopupMessage", "Load {0} unloaded child blueprints to fix up phantom references?"), FText::FromString(LexToString(Unloaded)));
 			FText Title = LOCTEXT("LoadChildrenPopupTitle", "Load Unloaded Children?");
 			DialogResponse = FMessageDialog::Open(EAppMsgType::YesNo, Message, &Title);
 		}
@@ -8788,8 +8788,8 @@ void FBlueprintEditorUtils::OpenReparentBlueprintMenu( const TArray< UBlueprint*
 
 	TSharedRef<SBox> ClassPickerBox = 
 		SNew(SBox)
-		.WidthOverride(280)
-		.HeightOverride(400)
+		.WidthOverride(280.0f)
+		.HeightOverride(400.0f)
 		[
 			SNew(SBorder)
 			.BorderImage(FAppStyle::GetBrush("Menu.Background"))
@@ -9153,7 +9153,7 @@ bool FBlueprintEditorUtils::PropertyValueFromString_Direct(const FProperty* Prop
 		}
 		else if (const FByteProperty* ByteProperty = CastField<const FByteProperty>(Property))
 		{
-			int32 IntValue = 0;
+			int64 IntValue = 0;
 			if (const UEnum* Enum = ByteProperty->Enum)
 			{
 				if (StrValue.Len() < NAME_SIZE)
@@ -9174,10 +9174,11 @@ bool FBlueprintEditorUtils::PropertyValueFromString_Direct(const FProperty* Prop
 			}
 			else
 			{
-				bParseSucceeded = FDefaultValueHelper::ParseInt(StrValue, IntValue);
+				bParseSucceeded = FDefaultValueHelper::ParseInt64(StrValue, IntValue);
 			}
+
 			bParseSucceeded = bParseSucceeded && (IntValue <= 255) && (IntValue >= 0);
-			ByteProperty->SetPropertyValue(DirectValue, IntValue);
+			ByteProperty->SetPropertyValue(DirectValue, static_cast<uint8>(IntValue));
 		}
 		else if (const FEnumProperty* EnumProperty = CastField<const FEnumProperty>(Property))
 		{
diff --git a/Engine/Source/Editor/UnrealEd/Private/Kismet2/KismetDebugUtilities.cpp b/Engine/Source/Editor/UnrealEd/Private/Kismet2/KismetDebugUtilities.cpp
--- a/Engine/Source/Editor/UnrealEd/Private/Kismet2/KismetDebugUtilities.cpp
+++ b/Engine/Source/Editor/UnrealEd/Private/Kismet2/KismetDebugUtilities.cpp
@@ -422,8 +422,8 @@ void FKismetDebugUtilities::OnScriptException(const UObject* ActiveObject, const
 					TSharedPtr<SMultiLineEditableText> TextBlock;
 					TSharedRef<SWidget> DisplayWidget =
 						SNew(SBox)
-						.MaxDesiredHeight(512)
-						.MaxDesiredWidth(512)
+						.MaxDesiredHeight(512.0f)
+						.MaxDesiredWidth(512.0f)
 						.Content()
 						[
 							SNew(SBorder)
diff --git a/Engine/Source/Editor/UnrealEd/Private/Lightmass/LightmassLandscapeRender.cpp b/Engine/Source/Editor/UnrealEd/Private/Lightmass/LightmassLandscapeRender.cpp
--- a/Engine/Source/Editor/UnrealEd/Private/Lightmass/LightmassLandscapeRender.cpp
+++ b/Engine/Source/Editor/UnrealEd/Private/Lightmass/LightmassLandscapeRender.cpp
@@ -51,7 +51,7 @@ void RenderLandscapeMaterialForLightmass(const FLandscapeStaticLightingMesh* Lan
 
 	const float fraction = 1.0f / NumSubsections;
 	const FVector2D PositionScale = FVector2D(RenderTarget->GetSizeXY()) * fraction * PatchExpandScale;
-	const float LayerScale = SubsectionSizeQuads;
+	const FVector::FReal LayerScale = SubsectionSizeQuads;
 	const float WeightmapSubsection = LandscapeComponent->WeightmapSubsectionOffset;
 	const FVector2D WeightmapBias = FVector2D(LandscapeComponent->WeightmapScaleBias.Z, LandscapeComponent->WeightmapScaleBias.W);
 	const FVector2D WeightmapScale = FVector2D(LandscapeComponent->WeightmapScaleBias.X, LandscapeComponent->WeightmapScaleBias.Y) * SubsectionSizeQuads;
@@ -220,7 +220,7 @@ void GetLandscapeOpacityData(const FLandscapeStaticLightingMesh* LandscapeMesh,
 	}
 
 	// Scale up the hole map to compensate for lightmass using point-sampling
-	static const float ScaleFactor = 3;
+	static const int32 ScaleFactor = 3;
 	InOutSizeX = (XSize - 1) * ScaleFactor;
 	InOutSizeY = (YSize - 1) * ScaleFactor;
 
diff --git a/Engine/Source/Editor/UnrealEd/Private/Lightmass/LightmassRender.cpp b/Engine/Source/Editor/UnrealEd/Private/Lightmass/LightmassRender.cpp
--- a/Engine/Source/Editor/UnrealEd/Private/Lightmass/LightmassRender.cpp
+++ b/Engine/Source/Editor/UnrealEd/Private/Lightmass/LightmassRender.cpp
@@ -371,7 +371,7 @@ public:
 		if (bStrataEnabled)
 		{
 			EStrataBlendMode StrataBlendMode = MaterialInterface->GetStrataBlendMode();
-			uint8 LegacyBlendMode = MaterialInterface->GetBlendMode();
+			uint8 LegacyBlendMode = static_cast<uint8>(MaterialInterface->GetBlendMode());
 			EStrataMaterialExportContext StrataMaterialExportContext = (StrataBlendMode == SBM_Opaque || StrataBlendMode == SBM_Masked) ? EStrataMaterialExportContext::SMEC_Opaque : EStrataMaterialExportContext::SMEC_Translucent;
 
 			if (Usage == EMaterialShaderMapUsage::LightmassExportDiffuse)
diff --git a/Engine/Source/Editor/UnrealEd/Private/PackageBackup.cpp b/Engine/Source/Editor/UnrealEd/Private/PackageBackup.cpp
--- a/Engine/Source/Editor/UnrealEd/Private/PackageBackup.cpp
+++ b/Engine/Source/Editor/UnrealEd/Private/PackageBackup.cpp
@@ -324,7 +324,7 @@ bool FAutoPackageBackup::PerformBackupSpaceMaintenance( TArray<FBackupFileInfo>&
 		};
 		InBackupFiles.Sort( FCompareFBackupFileInfo() );
 		
-		int32 CurSpaceUsed = InSpaceUsed;
+		int64 CurSpaceUsed = InSpaceUsed;
 		TArray<FBackupFileInfo>::TConstIterator BackupFileIter( InBackupFiles );
 		
 		// Iterate through the backup files until all of the files have been deleted or enough space has been freed
diff --git a/Engine/Source/Editor/UnrealEd/Private/SActorEditorContext.cpp b/Engine/Source/Editor/UnrealEd/Private/SActorEditorContext.cpp
--- a/Engine/Source/Editor/UnrealEd/Private/SActorEditorContext.cpp
+++ b/Engine/Source/Editor/UnrealEd/Private/SActorEditorContext.cpp
@@ -85,7 +85,7 @@ void SActorEditorContext::Rebuild()
 						SNew(SButton)
 						.Cursor(EMouseCursor::Default)
 						.ButtonStyle(FAppStyle::Get(), "NoBorder")
-						.ContentPadding(0)
+						.ContentPadding(0.0f)
 						.OnClicked_Lambda([this]()
 						{
 							bIsContextExpanded = !bIsContextExpanded;
@@ -163,7 +163,7 @@ void SActorEditorContext::Rebuild()
 							SNew(SButton)
 							.Cursor(EMouseCursor::Default)
 							.ButtonStyle(FAppStyle::Get(), "SimpleButton")
-							.ContentPadding(0)
+							.ContentPadding(0.0f)
 							.Visibility_Lambda([Client, this]() { return (Client && Client->CanResetContext(World)) ? EVisibility::Visible : EVisibility::Collapsed; })
 							.ToolTipText(FText::Format(LOCTEXT("ResetActorEditorContextTooltip", "Reset Current {0}"), FText::FromString(Info.Title)))
 							.OnClicked_Lambda([Client, this]()
diff --git a/Engine/Source/Editor/UnrealEd/Private/SCSVImportOptions.cpp b/Engine/Source/Editor/UnrealEd/Private/SCSVImportOptions.cpp
--- a/Engine/Source/Editor/UnrealEd/Private/SCSVImportOptions.cpp
+++ b/Engine/Source/Editor/UnrealEd/Private/SCSVImportOptions.cpp
@@ -167,7 +167,7 @@ void SCSVImportOptions::Construct(const FArguments& InArgs)
 			.Padding(2)
 			[
 				SNew(SBox)
-				.WidthOverride(400)
+				.WidthOverride(400.0f)
 				.Visibility(this, &SCSVImportOptions::GetDetailsPanelVis)
 				[
 					PropertyView.ToSharedRef()
diff --git a/Engine/Source/Editor/UnrealEd/Private/SClassPickerDialog.cpp b/Engine/Source/Editor/UnrealEd/Private/SClassPickerDialog.cpp
--- a/Engine/Source/Editor/UnrealEd/Private/SClassPickerDialog.cpp
+++ b/Engine/Source/Editor/UnrealEd/Private/SClassPickerDialog.cpp
@@ -170,7 +170,7 @@ void SClassPickerDialog::Construct(const FArguments& InArgs)
 				+SVerticalBox::Slot()
 				.HAlign(HAlign_Right)
 				.VAlign(VAlign_Bottom)
-				.Padding(8)
+				.Padding(8.0f)
 				[
 					SNew(SUniformGridPanel)
 					.SlotPadding(FAppStyle::GetMargin("StandardDialog.SlotPadding"))
diff --git a/Engine/Source/Editor/UnrealEd/Private/SComponentClassCombo.cpp b/Engine/Source/Editor/UnrealEd/Private/SComponentClassCombo.cpp
--- a/Engine/Source/Editor/UnrealEd/Private/SComponentClassCombo.cpp
+++ b/Engine/Source/Editor/UnrealEd/Private/SComponentClassCombo.cpp
@@ -94,10 +94,10 @@ void SComponentClassCombo::Construct(const FArguments& InArgs)
 			[
 				SNew(SBorder)
 				.BorderImage(FAppStyle::GetBrush("Menu.Background"))
-				.Padding(2)
+				.Padding(2.0f)
 				[
 					SNew(SBox)
-					.WidthOverride(250)
+					.WidthOverride(250.0f)
 					[				
 						SNew(SVerticalBox)
 						+SVerticalBox::Slot()
@@ -114,7 +114,7 @@ void SComponentClassCombo::Construct(const FArguments& InArgs)
 							.Padding(2.0f, 2.0f)
 							[
 								SNew(SComboButton)
-								.ContentPadding(0)
+								.ContentPadding(0.0f)
 								.ForegroundColor(FSlateColor::UseForeground())
 								.ComboButtonStyle(FAppStyle::Get(), "SimpleComboButton")
 								.HasDownArrow(false)
diff --git a/Engine/Source/Editor/UnrealEd/Private/STransformViewportToolbar.cpp b/Engine/Source/Editor/UnrealEd/Private/STransformViewportToolbar.cpp
--- a/Engine/Source/Editor/UnrealEd/Private/STransformViewportToolbar.cpp
+++ b/Engine/Source/Editor/UnrealEd/Private/STransformViewportToolbar.cpp
@@ -122,7 +122,7 @@ TSharedRef<SWidget> STransformViewportToolBar::GenerateSurfaceSnappingMenu()
 						GetMutableDefault<ULevelEditorViewportSettings>()->SnapToSurface.SnapOffsetExtent = Val;
 					}))
 					.MinValue(0.f)
-					.MaxValue(HALF_WORLD_MAX)
+					.MaxValue(static_cast<float>(HALF_WORLD_MAX))
 					.MaxSliderValue(1000.f) // 'Sensible' range for the slider (10m)
 					.AllowSpin(true)
 				]
@@ -507,10 +507,10 @@ TSharedRef<SWidget> STransformViewportToolBar::FillCameraSpeedMenu()
 			.Padding(FMargin(0.0f, 2.0f))
 			[
 				SAssignNew(CamSpeedScalarBox, SSpinBox<float>)
-				.MinValue(1)
- 			    .MaxValue(TNumericLimits<int32>::Max())
-			    .MinSliderValue(1)
-			    .MaxSliderValue(128)
+				.MinValue(1.0f)
+ 			    .MaxValue(std::numeric_limits<float>::max())
+			    .MinSliderValue(1.0f)
+			    .MaxSliderValue(128.0f)
 				.Value(this, &STransformViewportToolBar::GetCamSpeedScalarBoxValue)
 				.OnValueChanged(this, &STransformViewportToolBar::OnSetCamSpeedScalarBoxValue)
 				.ToolTipText(LOCTEXT("CameraSpeedScalar_ToolTip", "Scalar to increase camera movement range"))
diff --git a/Engine/Source/Editor/UnrealEd/Private/SViewportToolBarIconMenu.cpp b/Engine/Source/Editor/UnrealEd/Private/SViewportToolBarIconMenu.cpp
--- a/Engine/Source/Editor/UnrealEd/Private/SViewportToolBarIconMenu.cpp
+++ b/Engine/Source/Editor/UnrealEd/Private/SViewportToolBarIconMenu.cpp
@@ -48,8 +48,8 @@ void SViewportToolBarIconMenu::Construct( const FArguments& InArgs )
 			.AutoWidth()
 			[
 				SNew(SBox)
-				.WidthOverride( 16 )
-				.HeightOverride( 16 )
+				.WidthOverride(16.0f)
+				.HeightOverride(16.0f)
 				.HAlign(HAlign_Center)
 				.VAlign(VAlign_Center)
 				[
diff --git a/Engine/Source/Editor/UnrealEd/Private/Settings/EditorStyleClasses.cpp b/Engine/Source/Editor/UnrealEd/Private/Settings/EditorStyleClasses.cpp
--- a/Engine/Source/Editor/UnrealEd/Private/Settings/EditorStyleClasses.cpp
+++ b/Engine/Source/Editor/UnrealEd/Private/Settings/EditorStyleClasses.cpp
@@ -39,7 +39,7 @@ UEditorStyleSettings::UEditorStyleSettings( const FObjectInitializer& ObjectInit
 	RuleColor = FLinearColor(0.008f, 0.008f, 0.008f);
 	CenterColor = FLinearColor::Black;
 
-	GridSnapSize = 16.f;
+	GridSnapSize = 16;
 
 	bShowFriendlyNames = true;
 	bShowNativeComponentNames = true;
diff --git a/Engine/Source/Editor/UnrealEd/Private/Text/HLSLSyntaxHighlighterMarshaller.cpp b/Engine/Source/Editor/UnrealEd/Private/Text/HLSLSyntaxHighlighterMarshaller.cpp
--- a/Engine/Source/Editor/UnrealEd/Private/Text/HLSLSyntaxHighlighterMarshaller.cpp
+++ b/Engine/Source/Editor/UnrealEd/Private/Text/HLSLSyntaxHighlighterMarshaller.cpp
@@ -456,8 +456,8 @@ public:
 		const TSharedRef<FSlateFontMeasure> FontMeasure = FSlateApplication::Get().GetRenderer()->GetFontMeasureService();
 		FVector2D Size = FontMeasure->Measure(**Text, StartIndex, EndIndex, Style.Font, true, Scale) + ShadowOffsetToApply;
 
-		Size.X -= TabWidth * static_cast<float>(TabCount) * Scale;
-		Size.X += SpaceWidth * static_cast<float>(TabCount * NumSpacesPerTab) * Scale;
+		Size.X -= TabWidth * TabCount * Scale;
+		Size.X += SpaceWidth * TabCount * NumSpacesPerTab * Scale;
 
 		return Size;
 	}
@@ -481,9 +481,9 @@ protected:
 private:
 	int32 NumSpacesPerTab;
 
-	float TabWidth;
+	double TabWidth;
 
-	float SpaceWidth;
+	double SpaceWidth;
 };
 
 
diff --git a/Engine/Source/Editor/UnrealEd/Private/Texture2DPreview.cpp b/Engine/Source/Editor/UnrealEd/Private/Texture2DPreview.cpp
--- a/Engine/Source/Editor/UnrealEd/Private/Texture2DPreview.cpp
+++ b/Engine/Source/Editor/UnrealEd/Private/Texture2DPreview.cpp
@@ -101,13 +101,13 @@ public:
 		}
 		
 		SetShaderValue(InRHICmdList, ShaderRHI,ColorWeights,InColorWeightsValue);
-		FVector4f PackedParametersValue(InGammaValue, InMipLevel, bInIsNormalMap ? 1.0 : -1.0f, bInIsSingleVTPhysicalSpace ? 0 : InLayerIndex);
+		FVector4f PackedParametersValue(InGammaValue, InMipLevel, bInIsNormalMap ? 1.0f : -1.0f, bInIsSingleVTPhysicalSpace ? 0 : InLayerIndex);
 		SetShaderValue(InRHICmdList, ShaderRHI, PackedParameters, PackedParametersValue);
 
 		// Store slice count and selected slice index for texture array
 		if (bInIsTextureArray)
 		{
-			const float NumSlicesData = InTextureValue ? InTextureValue->GetSizeZ() : 1;
+			const float NumSlicesData = InTextureValue ? static_cast<float>(InTextureValue->GetSizeZ()) : 1.0f;
 			SetShaderValue(InRHICmdList, ShaderRHI, NumSlices, NumSlicesData);
 			SetShaderValue(InRHICmdList, ShaderRHI, SliceIndex, InSliceIndex);
 		}
diff --git a/Engine/Source/Editor/UnrealEd/Private/ThumbnailRendering/FontThumbnailRenderer.cpp b/Engine/Source/Editor/UnrealEd/Private/ThumbnailRendering/FontThumbnailRenderer.cpp
--- a/Engine/Source/Editor/UnrealEd/Private/ThumbnailRendering/FontThumbnailRenderer.cpp
+++ b/Engine/Source/Editor/UnrealEd/Private/ThumbnailRendering/FontThumbnailRenderer.cpp
@@ -36,7 +36,7 @@ void UFontThumbnailRenderer::GetThumbnailSize(UObject* Object, float Zoom, uint3
 		case EFontCacheType::Runtime:
 			if(Font->CompositeFont.DefaultTypeface.Fonts.Num() > 0)
 			{
-				OutWidth = OutHeight = Zoom * 256.0f;
+				OutWidth = OutHeight = static_cast<uint32>(Zoom * 256.0f);
 			}
 			break;
 
diff --git a/Engine/Source/Editor/UnrealEd/Private/Toolkits/SGlobalOpenAssetDialog.cpp b/Engine/Source/Editor/UnrealEd/Private/Toolkits/SGlobalOpenAssetDialog.cpp
--- a/Engine/Source/Editor/UnrealEd/Private/Toolkits/SGlobalOpenAssetDialog.cpp
+++ b/Engine/Source/Editor/UnrealEd/Private/Toolkits/SGlobalOpenAssetDialog.cpp
@@ -37,8 +37,8 @@ void SGlobalOpenAssetDialog::Construct(const FArguments& InArgs, FVector2D InSiz
 	ChildSlot
 	[
 		SNew(SBox)
-		.WidthOverride(InSize.X)
-		.HeightOverride(InSize.Y)
+		.WidthOverride(static_cast<float>(InSize.X))
+		.HeightOverride(static_cast<float>(InSize.Y))
 		[
 			SNew(SVerticalBox)
 			+ SVerticalBox::Slot()
diff --git a/Engine/Source/Editor/UnrealEd/Private/Toolkits/SGlobalTabSwitchingDialog.cpp b/Engine/Source/Editor/UnrealEd/Private/Toolkits/SGlobalTabSwitchingDialog.cpp
--- a/Engine/Source/Editor/UnrealEd/Private/Toolkits/SGlobalTabSwitchingDialog.cpp
+++ b/Engine/Source/Editor/UnrealEd/Private/Toolkits/SGlobalTabSwitchingDialog.cpp
@@ -110,8 +110,8 @@ public:
 			.Padding(4.0f, 4.0f)
 			[
 				SNew(SBox)
-				.WidthOverride(ThumbnailSize)
-				.HeightOverride(ThumbnailSize)
+				.WidthOverride(static_cast<float>(ThumbnailSize))
+				.HeightOverride(static_cast<float>(ThumbnailSize))
 				[
 					Thumbnail->MakeThumbnailWidget(ThumbnailConfig)
 				]
@@ -429,8 +429,8 @@ void SGlobalTabSwitchingDialog::Construct(const FArguments& InArgs, FVector2D In
 		.ForegroundColor(FCoreStyle::Get().GetSlateColor("DefaultForeground"))
 		[
 			SNew(SBox)
-			.WidthOverride(InSize.X)
-			.HeightOverride(InSize.Y)
+			.WidthOverride(static_cast<float>(InSize.X))
+			.HeightOverride(static_cast<float>(InSize.Y))
 			.Padding(FMargin(12.0f, 12.0f, 12.0f, 0.0f))
 			[
 				SNew(SVerticalBox)
diff --git a/Engine/Source/Editor/UnrealEd/Private/VertexSnapping.cpp b/Engine/Source/Editor/UnrealEd/Private/VertexSnapping.cpp
--- a/Engine/Source/Editor/UnrealEd/Private/VertexSnapping.cpp
+++ b/Engine/Source/Editor/UnrealEd/Private/VertexSnapping.cpp
@@ -20,8 +20,8 @@
 
 namespace VertexSnappingConstants
 {
-	const float MaxSnappingDistance = 300;
-	const float MaxSquaredDistanceFromCamera = FMath::Square( 5000 );
+	const float MaxSnappingDistance = 300.0f;
+	const float MaxSquaredDistanceFromCamera = FMath::Square( 5000.0f );
 	const float FadeTime = 0.15f;
 	const FLinearColor VertexHelperColor = FColor(  17, 105, 238, 255 );
 };
@@ -354,7 +354,7 @@ struct FVertexSnappingArgs
 bool FVertexSnappingImpl::GetClosestVertexOnComponent( const FSnapActor& SnapActor, UPrimitiveComponent* Component, const FVertexSnappingArgs& InArgs, FSnappingVertex& OutClosestLocation )
 {
 	// Current closest distance 
-	float ClosestDistance = FLT_MAX;
+	double ClosestDistance = std::numeric_limits<double>::max();
 	bool bHasAnyVerts = false;
 
 	const FPlane& ActorPlane = InArgs.ActorPlane;
@@ -386,8 +386,8 @@ bool FVertexSnappingImpl::GetClosestVertexOnComponent( const FSnapActor& SnapAct
 			// Ignore backface vertices when translating in screen space
 			bool bIsBackface = false;
 			bool bOutside = false;
-			float Distance = 0;
-			float DistanceFromCamera = 0;
+			double Distance = 0.0;
+			double DistanceFromCamera = 0.0;
 			if( CurrentAxis != EAxisList::Screen )
 			{
 				// Compute the distance to the plane the actor is on
@@ -415,7 +415,7 @@ bool FVertexSnappingImpl::GetClosestVertexOnComponent( const FSnapActor& SnapAct
 
 					if( !bOutside )
 					{
-						DistanceFromCamera = View->IsPerspectiveProjection() ? FVector::DistSquared( Position, View->ViewMatrices.GetViewOrigin() ) : 0;
+						DistanceFromCamera = View->IsPerspectiveProjection() ? FVector::DistSquared( Position, View->ViewMatrices.GetViewOrigin() ) : 0.0;
 						Distance = FVector::DistSquared( FVector( MousePosition, 0 ), FVector( PixelPos, 0 ) );
 					}
 				}
@@ -451,7 +451,7 @@ bool FVertexSnappingImpl::GetClosestVertexOnComponent( const FSnapActor& SnapAct
 FSnappingVertex FVertexSnappingImpl::GetClosestVertex( const TArray<FSnapActor>& Actors, const FVertexSnappingArgs& InArgs )
 {
 	// The current closest distance
-	float ClosestDistance = FLT_MAX;
+	double ClosestDistance = std::numeric_limits<double>::max();
 
 	const FPlane& ActorPlane = InArgs.ActorPlane;
 	EAxisList::Type CurrentAxis = InArgs.CurrentAxis;
@@ -484,7 +484,7 @@ FSnappingVertex FVertexSnappingImpl::GetClosestVertex( const TArray<FSnapActor>&
 				ClosestLocationOnComponent.Normal = FVector::ZeroVector;
 			}
 
-			float Distance = 0;
+			double Distance = 0.0;
 			if( CurrentAxis != EAxisList::Screen )
 			{
 				// Compute the distance from the point being snapped.  When not in screen space we snap to the plane created by the current closest vertex
@@ -601,7 +601,7 @@ bool FVertexSnappingImpl::SnapLocationToNearestVertex( FVector& Location, const
 
 	GetPossibleSnapActors( AllowedSnappingBox, MouseLocation.IntPoint(), ViewportClient, View, EAxisList::Screen, ActorsToIgnore, ActorsInBox );
 
-	FViewportCursorLocation Cursor(View, ViewportClient, MouseLocation.X, MouseLocation.Y );
+	FViewportCursorLocation Cursor(View, ViewportClient, static_cast<int32>(MouseLocation.X), static_cast<int32>(MouseLocation.Y));
 
 	FPlane ActorPlane( Location, Cursor.GetDirection() );
 
@@ -752,7 +752,7 @@ void FVertexSnappingImpl::SnapDragDelta( FVertexSnappingArgs& InArgs, const FVec
 		FVector ClosestPoint = GetClosestVertex( PossibleSnapPointActors, InArgs ).Position;
 
 		FVector PrevDragDelta = DragDelta;
-		float Distance = 0;
+		double Distance = 0;
 		if( CurrentAxis != EAxisList::Screen )
 		{
 			// Compute a distance from the stat location to the snap point.  
@@ -778,10 +778,10 @@ void FVertexSnappingImpl::SnapDragDelta( FVertexSnappingArgs& InArgs, const FVec
 		if( ViewportClient->IsPerspective() )
 		{
 			// Distance from start location to the location the actor would be in without snapping
-			float DistFromPreSnapToDesiredUnsnapped = FVector::DistSquared( PreSnapLocation, DesiredUnsnappedLocation );
+			double DistFromPreSnapToDesiredUnsnapped = FVector::DistSquared( PreSnapLocation, DesiredUnsnappedLocation );
 
 			// Distance from the new location of the actor without snapping to the location with snapping
-			float DistFromDesiredUnsnappedToSnapped = FVector::DistSquared( DesiredUnsnappedLocation, SnappedLocation );
+			double DistFromDesiredUnsnappedToSnapped = FVector::DistSquared( DesiredUnsnappedLocation, SnappedLocation );
 
 			// Only snap if the distance to the snapped location is less than the distance to the unsnapped location.  
 			// This allows the user to control the speed of snapping based on how fast they move the mouse and also avoids jerkiness when the mouse is behind the snap location
@@ -802,8 +802,8 @@ void FVertexSnappingImpl::SnapDragDelta( FVertexSnappingArgs& InArgs, const FVec
 			FVector2D PStoML = MousePosition-PreSnapLocationPixel;
 
 			// Only snap if the distance to the snapped location is less than the distance to the unsnapped location
-			float Dist2 = PStoML.SizeSquared();
-			float Dist1 = SLtoML.SizeSquared();
+			double Dist2 = PStoML.SizeSquared();
+			double Dist1 = SLtoML.SizeSquared();
 			if( Dist1 >= Dist2 || ClosestPoint == DesiredUnsnappedLocation )
 			{
 				DragDelta = FVector::ZeroVector;
diff --git a/Engine/Source/Editor/UnrealEd/Public/AssetThumbnail.h b/Engine/Source/Editor/UnrealEd/Public/AssetThumbnail.h
--- a/Engine/Source/Editor/UnrealEd/Public/AssetThumbnail.h
+++ b/Engine/Source/Editor/UnrealEd/Public/AssetThumbnail.h
@@ -284,9 +284,9 @@ private:
 		/** Render target for slate */
 		FSlateTextureRenderTarget2DResource* ThumbnailRenderTarget;
 		/** The time since last access */
-		float LastAccessTime;
+		double LastAccessTime;
 		/** The time since last update */
-		float LastUpdateTime;
+		double LastUpdateTime;
 		/** Width of the thumbnail */
 		uint32 Width;
 		/** Height of the thumbnail */
diff --git a/Engine/Source/Editor/UnrealEd/Public/GameModeInfoCustomizer.h b/Engine/Source/Editor/UnrealEd/Public/GameModeInfoCustomizer.h
--- a/Engine/Source/Editor/UnrealEd/Public/GameModeInfoCustomizer.h
+++ b/Engine/Source/Editor/UnrealEd/Public/GameModeInfoCustomizer.h
@@ -87,7 +87,7 @@ public:
 			.AutoWidth()
 			[
 				SNew(SBox)
-				.WidthOverride(125)
+				.WidthOverride(125.0f)
 				[
 					SNew(SClassPropertyEntryBox)
 					.AllowNone(bAllowNone)
@@ -144,7 +144,7 @@ public:
 			.AutoWidth()
 			[
 				SNew(SBox)
-				.WidthOverride(125)
+				.WidthOverride(125.0f)
 				[
 					SNew(SClassPropertyEntryBox)
 					.AllowNone(bAllowNone)
diff --git a/Engine/Source/Editor/UnrealEd/Public/Kismet2/CompilerResultsLog.h b/Engine/Source/Editor/UnrealEd/Public/Kismet2/CompilerResultsLog.h
--- a/Engine/Source/Editor/UnrealEd/Public/Kismet2/CompilerResultsLog.h
+++ b/Engine/Source/Editor/UnrealEd/Public/Kismet2/CompilerResultsLog.h
@@ -356,7 +356,8 @@ protected:
 		// read to next "@@":
 		if (const TCHAR* DelimiterStr = FCString::Strstr(Format, TEXT("@@")))
 		{
-			OutMessage.AddToken(FTextToken::Create(FText::FromString(FString(DelimiterStr - Format, Format))));
+			int32 TokenLength = UE_PTRDIFF_TO_INT32(DelimiterStr - Format);
+			OutMessage.AddToken(FTextToken::Create(FText::FromString(FString(TokenLength, Format))));
 			FEdGraphToken::Create(First, this, OutMessage, OutSourceNode);
 
 			const TCHAR* NextChunk = DelimiterStr + FCString::Strlen(TEXT("@@"));
