UE-160860 - Fix UnrealEd Module truncation warnings
Part 3 of N

diff --git a/Engine/Source/Editor/UnrealEd/Private/AutoReimport/AutoReimportManager.cpp b/Engine/Source/Editor/UnrealEd/Private/AutoReimport/AutoReimportManager.cpp
--- a/Engine/Source/Editor/UnrealEd/Private/AutoReimport/AutoReimportManager.cpp
+++ b/Engine/Source/Editor/UnrealEd/Private/AutoReimport/AutoReimportManager.cpp
@@ -758,7 +758,7 @@ TOptional<ECurrentState> FAutoReimportManager::ProcessDeletions()
 		Monitor->ExtractAssetsToDelete(AssetsToDelete);
 	}
 
-	FeedbackContextOverride->MainTask->EnterProgressFrame(TotalWork);
+	FeedbackContextOverride->MainTask->EnterProgressFrame(static_cast<float>(TotalWork));
 
 	if (AssetsToDelete.Num() > 0)
 	{
diff --git a/Engine/Source/Editor/UnrealEd/Private/Commandlets/AssetRegistryGenerator.cpp b/Engine/Source/Editor/UnrealEd/Private/Commandlets/AssetRegistryGenerator.cpp
--- a/Engine/Source/Editor/UnrealEd/Private/Commandlets/AssetRegistryGenerator.cpp
+++ b/Engine/Source/Editor/UnrealEd/Private/Commandlets/AssetRegistryGenerator.cpp
@@ -316,7 +316,7 @@ public:
 		if (ExtIndex >= UE_ARRAY_COUNT(Extensions))
 			return true;
 
-		int32 LengthWithoutExtension = Extension - FilenameOrDirectory;
+		int32 LengthWithoutExtension = UE_PTRDIFF_TO_INT32(Extension - FilenameOrDirectory);
 		FString FilenameWithoutExtension(LengthWithoutExtension, FilenameOrDirectory);
 
 		if (int64* CurrentPackageSize = PackageFileSizes.Find(FilenameWithoutExtension))
@@ -2670,7 +2670,7 @@ bool FAssetRegistryPackageMessage::TryRead(FCbObject&& Object, FPackageData& Pac
 	{
 		return false;
 	}
-	AssetDatas.Reset(AssetDatasArray.Num());
+	AssetDatas.Reset(IntCastChecked<int32>(AssetDatasArray.Num()));
 	for (FCbFieldView ElementField : AssetDatasArray)
 	{
 		FAssetData& AssetData = AssetDatas.Emplace_GetRef();
diff --git a/Engine/Source/Editor/UnrealEd/Private/Commandlets/ContentCommandlets.cpp b/Engine/Source/Editor/UnrealEd/Private/Commandlets/ContentCommandlets.cpp
--- a/Engine/Source/Editor/UnrealEd/Private/Commandlets/ContentCommandlets.cpp
+++ b/Engine/Source/Editor/UnrealEd/Private/Commandlets/ContentCommandlets.cpp
@@ -2013,7 +2013,7 @@ void UResavePackagesCommandlet::PerformAdditionalOperations(class UWorld* World,
 
 				while (Processor->IsProxyGenerationRunning())
 				{
-					FTSTicker::GetCoreTicker().Tick(FApp::GetDeltaTime());
+					FTSTicker::GetCoreTicker().Tick(static_cast<float>(FApp::GetDeltaTime()));
 					FThreadManager::Get().Tick();
 					FTaskGraphInterface::Get().ProcessThreadUntilIdle(ENamedThreads::GameThread);
 					FPlatformProcess::Sleep(0.1f);
@@ -2294,7 +2294,7 @@ struct FUnreferencedObject
 	/** Full name of object */
 	FString ObjectName;
 	/** Size on disk as recorded in FObjectExport */
-	int32 SerialSize;
+	int64 SerialSize;
 
 	/**
 	 * Constructor for easy creation in a TArray
diff --git a/Engine/Source/Editor/UnrealEd/Private/Commandlets/DerivedDataCacheCommandlet.cpp b/Engine/Source/Editor/UnrealEd/Private/Commandlets/DerivedDataCacheCommandlet.cpp
--- a/Engine/Source/Editor/UnrealEd/Private/Commandlets/DerivedDataCacheCommandlet.cpp
+++ b/Engine/Source/Editor/UnrealEd/Private/Commandlets/DerivedDataCacheCommandlet.cpp
@@ -151,7 +151,7 @@ static void TickCookObjects(bool bCookComplete = false)
 	if (bCookComplete || LastTickTime + TickPeriod <= CurrentTime)
 	{
 		TRACE_CPUPROFILER_EVENT_SCOPE(TickCookObjects);
-		FTickableCookObject::TickObjects(CurrentTime - LastTickTime, bCookComplete);
+		FTickableCookObject::TickObjects(static_cast<float>(CurrentTime - LastTickTime), bCookComplete);
 		LastTickTime = CurrentTime;
 	}
 }
diff --git a/Engine/Source/Editor/UnrealEd/Private/Commandlets/DiffPackagesCommandlet.cpp b/Engine/Source/Editor/UnrealEd/Private/Commandlets/DiffPackagesCommandlet.cpp
--- a/Engine/Source/Editor/UnrealEd/Private/Commandlets/DiffPackagesCommandlet.cpp
+++ b/Engine/Source/Editor/UnrealEd/Private/Commandlets/DiffPackagesCommandlet.cpp
@@ -1356,7 +1356,7 @@ void LoadNativePropertyData( UObject* Object, TArray<uint8>& out_NativePropertyD
 		((FArchive*)ObjectLinker)->Precache(NativeStartPos, NativePropertySerialSize);
 
 		// allocate enough space to contain the data we're about to read from disk
-		out_NativePropertyData.AddZeroed(NativePropertySerialSize);
+		out_NativePropertyData.AddZeroed(IntCastChecked<int32, int64>(NativePropertySerialSize));
 		((FArchive*)ObjectLinker)->Serialize(out_NativePropertyData.GetData(), NativePropertySerialSize);
 
 		// return the linker to its previous position
diff --git a/Engine/Source/Editor/UnrealEd/Private/CookOnTheFlyServer.cpp b/Engine/Source/Editor/UnrealEd/Private/CookOnTheFlyServer.cpp
--- a/Engine/Source/Editor/UnrealEd/Private/CookOnTheFlyServer.cpp
+++ b/Engine/Source/Editor/UnrealEd/Private/CookOnTheFlyServer.cpp
@@ -1387,7 +1387,7 @@ void UCookOnTheFlyServer::TickCookStatus(UE::Cook::FTickStackData& StackData)
 	if (LastCookableObjectTickTime + TickCookableObjectsFrameTime <= CurrentTime)
 	{
 		UE_SCOPED_COOKTIMER(TickCookableObjects);
-		FTickableCookObject::TickObjects(CurrentTime - LastCookableObjectTickTime, false /* bTickComplete */);
+		FTickableCookObject::TickObjects(static_cast<float>(CurrentTime - LastCookableObjectTickTime), false /* bTickComplete */);
 		LastCookableObjectTickTime = CurrentTime;
 	}
 
@@ -1608,7 +1608,7 @@ void UCookOnTheFlyServer::UpdateDisplay(UE::Cook::FTickStackData& StackData, boo
 	using namespace UE::Cook;
 
 	const double CurrentTime = StackData.LoopStartTime;
-	const float DeltaProgressDisplayTime = CurrentTime - LastProgressDisplayTime;
+	const double DeltaProgressDisplayTime = CurrentTime - LastProgressDisplayTime;
 	if (!bForceDisplay && DeltaProgressDisplayTime < DisplayUpdatePeriodSeconds)
 	{
 		return;
@@ -1632,7 +1632,7 @@ void UCookOnTheFlyServer::UpdateDisplay(UE::Cook::FTickStackData& StackData, boo
 		LastCookPendingCount = CookPendingCount;
 		LastProgressDisplayTime = CurrentTime;
 	}
-	const float DeltaDiagnosticsDisplayTime = CurrentTime - LastDiagnosticsDisplayTime;
+	const double DeltaDiagnosticsDisplayTime = CurrentTime - LastDiagnosticsDisplayTime;
 	if (bForceDisplay || DeltaDiagnosticsDisplayTime > GCookProgressDiagnosticTime)
 	{
 		uint32 OpenFileHandles = 0;
@@ -1748,7 +1748,7 @@ void UCookOnTheFlyServer::FPollable::TriggerInternal(UCookOnTheFlyServer& COTFS)
 		LocalQueueKey.NextTimeSeconds = FMath::Min(CurrentTime, TimeAfterHeapTop);
 		this->NextTimeIdleSeconds = LocalQueueKey.NextTimeSeconds;
 
-		int32 Index = KeyInQueue - COTFS.Pollables.GetData();
+		int32 Index = UE_PTRDIFF_TO_INT32(KeyInQueue - COTFS.Pollables.GetData());
 		COTFS.Pollables.HeapRemoveAt(Index, false /* bAllowShrinking */);
 		COTFS.Pollables.HeapPush(MoveTemp(LocalQueueKey));
 		COTFS.PollNextTimeSeconds = 0;
@@ -1786,7 +1786,7 @@ void UCookOnTheFlyServer::FPollable::RunNowInternal(UCookOnTheFlyServer & COTFS,
 		LocalQueueKey.NextTimeSeconds = TimeLastRun + this->PeriodSeconds;
 		this->NextTimeIdleSeconds = TimeLastRun + this->PeriodIdleSeconds;
 
-		int32 Index = KeyInQueue - COTFS.Pollables.GetData();
+		int32 Index = UE_PTRDIFF_TO_INT32(KeyInQueue - COTFS.Pollables.GetData());
 		COTFS.PollNextTimeSeconds = FMath::Min(LocalQueueKey.NextTimeSeconds, COTFS.PollNextTimeSeconds);
 		COTFS.PollNextTimeIdleSeconds = FMath::Min(this->NextTimeIdleSeconds, COTFS.PollNextTimeIdleSeconds);
 		COTFS.Pollables.HeapRemoveAt(Index, false /* bAllowShrinking */);
@@ -2101,13 +2101,13 @@ void UCookOnTheFlyServer::WaitForAsync(UE::Cook::FTickStackData& StackData)
 	// Sleep until the next time that DecideNextCookAction will find work to do, up to a maximum of WaitForAsyncSleepSeconds
 	UE_SCOPED_COOKTIMER(WaitForAsync);
 	double CurrentTime = FPlatformTime::Seconds();
-	float SleepDuration = WaitForAsyncSleepSeconds;
+	double SleepDuration = WaitForAsyncSleepSeconds;
 	SleepDuration = FMath::Min(SleepDuration, StackData.Timer.GetEndTimeSeconds() - CurrentTime);
 	SleepDuration = FMath::Min(SleepDuration, PollNextTimeIdleSeconds - CurrentTime);
 	SleepDuration = FMath::Min(SleepDuration, SaveBusyRetryTimeSeconds - CurrentTime);
 	SleepDuration = FMath::Min(SleepDuration, LoadBusyRetryTimeSeconds - CurrentTime);
 	SleepDuration = FMath::Max(SleepDuration, 0);
-	FPlatformProcess::Sleep(SleepDuration);
+	FPlatformProcess::Sleep(static_cast<float>(SleepDuration));
 }
 
 UCookOnTheFlyServer::ECookAction UCookOnTheFlyServer::DecideNextCookAction(UE::Cook::FTickStackData& StackData)
@@ -4526,7 +4526,7 @@ void UCookOnTheFlyServer::EvaluateGarbageCollectionResults(int32 NumObjectsBefor
 	int64 VirtualMemAfterGC = MemStatsAfterGC.UsedVirtual;
 	int64 VirtualMemFreed = MemStatsBeforeGC.UsedVirtual - MemStatsAfterGC.UsedVirtual;
 
-	int64 ExpectedObjectsFreed = MemoryExpectedFreedToSpreadRatio * NumObjectsSpread;
+	int64 ExpectedObjectsFreed = static_cast<int64>(MemoryExpectedFreedToSpreadRatio * static_cast<float>(NumObjectsSpread));
 	double ExpectedMemFreed = MemoryExpectedFreedToSpreadRatio * VirtualMemSpread;
 	if ((NumObjectsFreed >= ExpectedObjectsFreed || NumObjectsBeforeGC - NumObjectsMin < ExpectedObjectsFreed) &&
 		(VirtualMemFreed >= ExpectedMemFreed || VirtualMemBeforeGC - VirtualMemMin <= ExpectedMemFreed))
@@ -8246,7 +8246,7 @@ void UCookOnTheFlyServer::CookByTheBookFinished()
 	{
 		UE_SCOPED_COOKTIMER(TickCookableObjects);
 		const double CurrentTime = FPlatformTime::Seconds();
-		FTickableCookObject::TickObjects(CurrentTime - LastCookableObjectTickTime, true /* bTickComplete */);
+		FTickableCookObject::TickObjects(static_cast<float>(CurrentTime - LastCookableObjectTickTime), true /* bTickComplete */);
 		LastCookableObjectTickTime = CurrentTime;
 	}
 
@@ -10492,7 +10492,7 @@ uint32 UCookOnTheFlyServer::CookFullLoadAndSave()
 			if (LastCookableObjectTickTime + TickCookableObjectsFrameTime <= CurrentTime)
 			{
 				UE_SCOPED_COOKTIMER(TickCookableObjects);
-				FTickableCookObject::TickObjects(CurrentTime - LastCookableObjectTickTime, false /* bTickComplete */);
+				FTickableCookObject::TickObjects(static_cast<float>(CurrentTime - LastCookableObjectTickTime), false /* bTickComplete */);
 				LastCookableObjectTickTime = CurrentTime;
 			}
 
diff --git a/Engine/Source/Editor/UnrealEd/Private/EdGraphNode_Comment.cpp b/Engine/Source/Editor/UnrealEd/Private/EdGraphNode_Comment.cpp
--- a/Engine/Source/Editor/UnrealEd/Private/EdGraphNode_Comment.cpp
+++ b/Engine/Source/Editor/UnrealEd/Private/EdGraphNode_Comment.cpp
@@ -44,8 +44,8 @@ namespace FEdGraphNode_Comment_Utils
 UEdGraphNode_Comment::UEdGraphNode_Comment(const FObjectInitializer& ObjectInitializer)
 	: Super(ObjectInitializer)
 {
-	NodeWidth = 400.0f;
-	NodeHeight = 100.0f;
+	NodeWidth = 400;
+	NodeHeight = 100;
 	FontSize = 18;
 	CommentColor = FLinearColor::White;
 	bColorCommentBubble = false;
@@ -153,8 +153,8 @@ void UEdGraphNode_Comment::ResizeNode(const FVector2D& NewSize)
 {
 	if (bCanResizeNode) 
 	{
-		NodeHeight = NewSize.Y;
-		NodeWidth = NewSize.X;
+		NodeHeight = UE::LWC::FloatToIntCastChecked<int32>(NewSize.Y);
+		NodeWidth = UE::LWC::FloatToIntCastChecked<int32>(NewSize.X);
 	}
 }
 
@@ -174,12 +174,12 @@ void UEdGraphNode_Comment::ClearNodesUnderComment()
 
 void UEdGraphNode_Comment::SetBounds(const class FSlateRect& Rect)
 {
-	NodePosX = Rect.Left;
-	NodePosY = Rect.Top;
+	NodePosX = UE::LWC::FloatToIntCastChecked<int32>(Rect.Left);
+	NodePosY = UE::LWC::FloatToIntCastChecked<int32>(Rect.Top);
 
 	FVector2D Size = Rect.GetSize();
-	NodeWidth = Size.X;
-	NodeHeight = Size.Y;
+	NodeWidth = UE::LWC::FloatToIntCastChecked<int32>(Size.X);
+	NodeHeight = UE::LWC::FloatToIntCastChecked<int32>(Size.Y);
 }
 
 const FCommentNodeSet& UEdGraphNode_Comment::GetNodesUnderComment() const
diff --git a/Engine/Source/Editor/UnrealEd/Private/EditorActor.cpp b/Engine/Source/Editor/UnrealEd/Private/EditorActor.cpp
--- a/Engine/Source/Editor/UnrealEd/Private/EditorActor.cpp
+++ b/Engine/Source/Editor/UnrealEd/Private/EditorActor.cpp
@@ -917,7 +917,7 @@ bool UUnrealEdEngine::DeleteActors(const TArray<AActor*>& InActorsToDelete, UWor
 
 			if (bWarnAboutSoftReferences)
 			{
-				FScopedSlowTask SlowTask(ActorsToDeletePaths.Num(), LOCTEXT("ComputeActorSoftReferences", "Computing References"));
+				FScopedSlowTask SlowTask(static_cast<float>(ActorsToDeletePaths.Num()), LOCTEXT("ComputeActorSoftReferences", "Computing References"));
 				SlowTask.MakeDialogDelayed(1.0f);
 
 				FAssetToolsModule& AssetToolsModule = FModuleManager::GetModuleChecked<FAssetToolsModule>(TEXT("AssetTools"));
diff --git a/Engine/Source/Editor/UnrealEd/Private/EditorComponents.cpp b/Engine/Source/Editor/UnrealEd/Private/EditorComponents.cpp
--- a/Engine/Source/Editor/UnrealEd/Private/EditorComponents.cpp
+++ b/Engine/Source/Editor/UnrealEd/Private/EditorComponents.cpp
@@ -582,8 +582,8 @@ void FEditorCommonDrawHelper::DrawGridSection(double ViewportGridY, double MaxSi
 
 	const double InvTransitionRegion = 1.0 / TransitionRegion;
 	const double Fract = IncValue - FMath::FloorToDouble(IncValue);
-	float AlphaA = FMath::Clamp<float>(Fract * InvTransitionRegion, 0.0f, 1.0f);
-	float AlphaB = FMath::Clamp<float>(InvTransitionRegion - Fract * InvTransitionRegion, 0.0f, 1.0f);
+	float AlphaA = FMath::Clamp<float>(static_cast<float>(Fract * InvTransitionRegion), 0.0f, 1.0f);
+	float AlphaB = FMath::Clamp<float>(static_cast<float>(InvTransitionRegion - Fract * InvTransitionRegion), 0.0f, 1.0f);
 
 	if(IncValue < -0.5)
 	{
@@ -600,8 +600,8 @@ void FEditorCommonDrawHelper::DrawGridSection(double ViewportGridY, double MaxSi
 	FLinearColor MinorColor = FMath::Lerp(Background, FLinearColor::White, 0.02f);
 
 	const FMatrix InvViewProjMatrix = View->ViewMatrices.GetInvViewProjectionMatrix();
-	int32 FirstLine = FMath::TruncToInt(InvViewProjMatrix.TransformPosition(FVector(-1, -1, 0.5f)).Component(Axis) / ViewportGridY);
-	int32 LastLine = FMath::TruncToInt(InvViewProjMatrix.TransformPosition(FVector(+1, +1, 0.5f)).Component(Axis) / ViewportGridY);
+	int64 FirstLine = FMath::TruncToInt(InvViewProjMatrix.TransformPosition(FVector(-1, -1, 0.5f)).Component(Axis) / ViewportGridY);
+	int64 LastLine = FMath::TruncToInt(InvViewProjMatrix.TransformPosition(FVector(+1, +1, 0.5f)).Component(Axis) / ViewportGridY);
 	if (FirstLine > LastLine)
 	{
 		Exchange(FirstLine, LastLine);
@@ -657,8 +657,8 @@ void FEditorCommonDrawHelper::DrawPivot(const FSceneView* View,FPrimitiveDrawInt
 
 	const FVector PivLoc = GLevelEditorModeTools().SnappedLocation;
 
-	const float ZoomFactor = FMath::Min<float>(View->ViewMatrices.GetProjectionMatrix().M[0][0], View->ViewMatrices.GetProjectionMatrix().M[1][1]);
-	const float WidgetRadius = View->ViewMatrices.GetViewProjectionMatrix().TransformPosition(PivLoc).W * (PivotSize / ZoomFactor);
+	const FVector::FReal ZoomFactor = FMath::Min<FVector::FReal>(View->ViewMatrices.GetProjectionMatrix().M[0][0], View->ViewMatrices.GetProjectionMatrix().M[1][1]);
+	const FVector::FReal WidgetRadius = View->ViewMatrices.GetViewProjectionMatrix().TransformPosition(PivLoc).W * (PivotSize / ZoomFactor);
 
 	const FVector CamX = CameraToWorld.TransformVector( FVector(1,0,0) );
 	const FVector CamY = CameraToWorld.TransformVector( FVector(0,1,0) );
diff --git a/Engine/Source/Editor/UnrealEd/Private/EditorDomain/EditorDomainSave.cpp b/Engine/Source/Editor/UnrealEd/Private/EditorDomain/EditorDomainSave.cpp
--- a/Engine/Source/Editor/UnrealEd/Private/EditorDomain/EditorDomainSave.cpp
+++ b/Engine/Source/Editor/UnrealEd/Private/EditorDomain/EditorDomainSave.cpp
@@ -1321,7 +1321,7 @@ bool FEditorDomainSaveClient::TrySendRequests(TArray<FPackagePath>& LocalRequest
 			FString PackageName = PackagePath.GetPackageName();
 			Message << PackageName;
 		}
-		MessageHeader.Size = Message.Tell();
+		MessageHeader.Size = IntCastChecked<uint32>(Message.Tell());
 		Message.Seek(0);
 		Message << MessageHeader;
 		Message.Seek(MessageHeader.Size);
diff --git a/Engine/Source/Editor/UnrealEd/Private/EditorLevelUtils.cpp b/Engine/Source/Editor/UnrealEd/Private/EditorLevelUtils.cpp
--- a/Engine/Source/Editor/UnrealEd/Private/EditorLevelUtils.cpp
+++ b/Engine/Source/Editor/UnrealEd/Private/EditorLevelUtils.cpp
@@ -314,7 +314,7 @@ ULevel* UEditorLevelUtils::AddLevelsToWorld(UWorld* InWorld, TArray<FString> Pac
 		return nullptr;
 	}
 
-	FScopedSlowTask SlowTask(PackageNames.Num(), LOCTEXT("AddLevelsToWorldTask", "Adding Levels to World"));
+	FScopedSlowTask SlowTask(static_cast<float>(PackageNames.Num()), LOCTEXT("AddLevelsToWorldTask", "Adding Levels to World"));
 	SlowTask.MakeDialog();
 
 	// Sort the level packages alphabetically by name.
diff --git a/Engine/Source/Editor/UnrealEd/Private/Factories/VectorFieldFactory.cpp b/Engine/Source/Editor/UnrealEd/Private/Factories/VectorFieldFactory.cpp
--- a/Engine/Source/Editor/UnrealEd/Private/Factories/VectorFieldFactory.cpp
+++ b/Engine/Source/Editor/UnrealEd/Private/Factories/VectorFieldFactory.cpp
@@ -224,7 +224,7 @@ UObject* UVectorFieldStaticFactory::FactoryCreateBinary(
 	if ( IsFluidGridAscii( Type ) )
 	{
 		FString String;
-		const int32 BufferSize = BufferEnd - Buffer;
+		const int32 BufferSize = UE_PTRDIFF_TO_INT32(BufferEnd - Buffer);
 		FFileHelper::BufferToString( String, Buffer, BufferSize );
 		TCHAR* Stream = (TCHAR*)*String;
 		
diff --git a/Engine/Source/Editor/UnrealEd/Private/Fbx/FbxAnimationExport.cpp b/Engine/Source/Editor/UnrealEd/Private/Fbx/FbxAnimationExport.cpp
--- a/Engine/Source/Editor/UnrealEd/Private/Fbx/FbxAnimationExport.cpp
+++ b/Engine/Source/Editor/UnrealEd/Private/Fbx/FbxAnimationExport.cpp
@@ -31,7 +31,7 @@ namespace UnFbx
 			return false;
 		}
 
-		const float FrameRate = AnimSeq->GetDataModel()->GetFrameRate().AsDecimal();
+		const double FrameRate = AnimSeq->GetDataModel()->GetFrameRate().AsDecimal();
 		//Configure the scene time line
 		{
 			FbxGlobalSettings& SceneGlobalSettings = Scene->GetGlobalSettings();
diff --git a/Engine/Source/Editor/UnrealEd/Private/HierarchicalLOD.cpp b/Engine/Source/Editor/UnrealEd/Private/HierarchicalLOD.cpp
--- a/Engine/Source/Editor/UnrealEd/Private/HierarchicalLOD.cpp
+++ b/Engine/Source/Editor/UnrealEd/Private/HierarchicalLOD.cpp
@@ -1187,7 +1187,7 @@ void FHierarchicalLODBuilder::MergeClustersAndBuildActors(ULevel* InLevel, const
 			{
 				const int32 NumClusters = Clusters.Num();
 
-				FScopedSlowTask SlowTask(NumClusters, FText::Format(LOCTEXT("HierarchicalLOD_BuildClusters", "Building Clusters for LOD {LODIndex} of {LevelName}..."), Arguments));
+				FScopedSlowTask SlowTask(static_cast<float>(NumClusters), FText::Format(LOCTEXT("HierarchicalLOD_BuildClusters", "Building Clusters for LOD {LODIndex} of {LevelName}..."), Arguments));
 				SlowTask.MakeDialog();
 
 				TArray<FLODCluster> ValidMergedClusters;
@@ -1270,7 +1270,7 @@ void FHierarchicalLODBuilder::MergeClustersAndBuildActors(ULevel* InLevel, const
 				}
 			}
 
-			FScopedSlowTask SlowTask(TotalValidCluster, FText::Format(LOCTEXT("HierarchicalLOD_MergeActors", "Merging Actors for LOD {LODIndex} of {LevelName}..."), Arguments));
+			FScopedSlowTask SlowTask(static_cast<float>(TotalValidCluster), FText::Format(LOCTEXT("HierarchicalLOD_MergeActors", "Merging Actors for LOD {LODIndex} of {LevelName}..."), Arguments));
 			SlowTask.MakeDialog();
 
 			for (FLODCluster& Cluster : Clusters)
diff --git a/Engine/Source/Editor/UnrealEd/Private/ImportUtils/SkeletalMeshImportUtils.cpp b/Engine/Source/Editor/UnrealEd/Private/ImportUtils/SkeletalMeshImportUtils.cpp
--- a/Engine/Source/Editor/UnrealEd/Private/ImportUtils/SkeletalMeshImportUtils.cpp
+++ b/Engine/Source/Editor/UnrealEd/Private/ImportUtils/SkeletalMeshImportUtils.cpp
@@ -195,9 +195,9 @@ bool SkeletalMeshImportUtils::ProcessImportMeshSkeleton(const USkeleton* Skeleto
 	for (int32 b = 0; b < RefSkeleton.GetRawBoneNum(); b++)
 	{
 		int32 Parent = RefSkeleton.GetRawParentIndex(b);
-		int32 Depth = 1.0f;
+		int32 Depth = 1;
 
-		SkeletalDepths[b] = 1.0f;
+		SkeletalDepths[b] = 1;
 		if (Parent != INDEX_NONE)
 		{
 			Depth += SkeletalDepths[Parent];
@@ -554,7 +554,7 @@ void SkeletalMeshImportUtils::ApplySkinning(USkeletalMesh* SkeletalMesh, FSkelet
 						DestTangentZ.Normalize();
 						FVector SrcTangentZ = FVector4(SrcVertex.TangentZ);
 						SrcTangentZ.Normalize();
-						float AngleDiff = FMath::Abs(FMath::Acos(FVector::DotProduct(DestTangentZ, SrcTangentZ)));
+						double AngleDiff = FMath::Abs(FMath::Acos(FVector::DotProduct(DestTangentZ, SrcTangentZ)));
 						if (AngleDiff < MinNormalAngle)
 						{
 							MinNormalAngle = AngleDiff;
@@ -584,11 +584,11 @@ void SkeletalMeshImportUtils::ApplySkinning(USkeletalMesh* SkeletalMesh, FSkelet
 					int32 OverrideIndex;
 					if (Section.BoneMap.Find(OriginalBoneIndex, OverrideIndex))
 					{
-						DestVertex.InfluenceBones[InfluenceIndex] = OverrideIndex;
+						DestVertex.InfluenceBones[InfluenceIndex] = IntCastChecked<FBoneIndexType>(OverrideIndex);
 					}
 					else
 					{
-						DestVertex.InfluenceBones[InfluenceIndex] = Section.BoneMap.Add(OriginalBoneIndex);
+						DestVertex.InfluenceBones[InfluenceIndex] = IntCastChecked<FBoneIndexType>(Section.BoneMap.Add(OriginalBoneIndex));
 						DestLODModel.ActiveBoneIndices.AddUnique(OriginalBoneIndex);
 					}
 					bUseBone = true;
@@ -762,7 +762,7 @@ void SkeletalMeshImportUtils::RestoreExistingSkelMeshData(const TSharedPtr<const
 							}
 							else
 							{
-								LODModelCopy->ActiveBoneIndices[j] = NewBoneIndex;
+								LODModelCopy->ActiveBoneIndices[j] = IntCastChecked<FBoneIndexType>(NewBoneIndex);
 							}
 						}
 						else
@@ -788,7 +788,7 @@ void SkeletalMeshImportUtils::RestoreExistingSkelMeshData(const TSharedPtr<const
 							}
 							else
 							{
-								LODModelCopy->RequiredBones[j] = NewBoneIndex;
+								LODModelCopy->RequiredBones[j] = IntCastChecked<FBoneIndexType>(NewBoneIndex);
 							}
 						}
 						else
@@ -818,7 +818,7 @@ void SkeletalMeshImportUtils::RestoreExistingSkelMeshData(const TSharedPtr<const
 							}
 							else
 							{
-								Section.BoneMap[BoneIndex] = NewBoneIndex;
+								Section.BoneMap[BoneIndex] = IntCastChecked<FBoneIndexType>(NewBoneIndex);
 							}
 						}
 						if (bMissingBone)
diff --git a/Engine/Source/Editor/UnrealEd/Private/ImportUtils/StaticMeshImportUtils.cpp b/Engine/Source/Editor/UnrealEd/Private/ImportUtils/StaticMeshImportUtils.cpp
--- a/Engine/Source/Editor/UnrealEd/Private/ImportUtils/StaticMeshImportUtils.cpp
+++ b/Engine/Source/Editor/UnrealEd/Private/ImportUtils/StaticMeshImportUtils.cpp
@@ -465,11 +465,11 @@ bool StaticMeshImportUtils::AddSphereGeomFromVerts( const TArray<FVector3f>& Ver
 
 	FVector Center, Extents;
 	Box.GetCenterAndExtents(Center, Extents);
-	float Longest = 2.f * Extents.GetMax();
-	float Shortest = 2.f * Extents.GetMin();
+	double Longest = 2.f * Extents.GetMax();
+	double Shortest = 2.f * Extents.GetMin();
 
 	// check that the AABB is roughly a square (5% tolerance)
-	if((Longest - Shortest)/Longest > 0.05f)
+	if((Longest - Shortest) / Longest > 0.05f)
 	{
 		UE_LOG(LogStaticMeshImportUtils, Log, TEXT("AddSphereGeomFromVerts (%s): Sphere bounding box not square."), ObjName);
 		return false;
diff --git a/Engine/Source/Editor/UnrealEd/Private/Kismet2/EnumEditorUtils.cpp b/Engine/Source/Editor/UnrealEd/Private/Kismet2/EnumEditorUtils.cpp
--- a/Engine/Source/Editor/UnrealEd/Private/Kismet2/EnumEditorUtils.cpp
+++ b/Engine/Source/Editor/UnrealEd/Private/Kismet2/EnumEditorUtils.cpp
@@ -481,7 +481,7 @@ int64 FEnumEditorUtils::ResolveEnumerator(const UEnum* Enum, FArchive& Ar, int64
 			if (OldName.Value == EnumeratorValue)
 			{
 				const FName EnumeratorName = OldName.Key;
-				const int32 NewEnumValue = Enum->GetValueByName(EnumeratorName);
+				const int64 NewEnumValue = Enum->GetValueByName(EnumeratorName);
 				if(INDEX_NONE != NewEnumValue)
 				{
 					return NewEnumValue;
diff --git a/Engine/Source/Editor/UnrealEd/Private/Kismet2/KismetReinstanceUtilities.cpp b/Engine/Source/Editor/UnrealEd/Private/Kismet2/KismetReinstanceUtilities.cpp
--- a/Engine/Source/Editor/UnrealEd/Private/Kismet2/KismetReinstanceUtilities.cpp
+++ b/Engine/Source/Editor/UnrealEd/Private/Kismet2/KismetReinstanceUtilities.cpp
@@ -129,7 +129,7 @@ struct FReplaceReferenceHelper
 		{
 			BP_SCOPED_COMPILER_EVENT_STAT(EKismetReinstancerStats_ReplaceReferences);
 
-			FScopedSlowTask SlowTask(Targets.Num(), NSLOCTEXT("Kismet", "PerformingReplaceReferences", "Performing replace references..."));
+			FScopedSlowTask SlowTask(static_cast<float>(Targets.Num()), NSLOCTEXT("Kismet", "PerformingReplaceReferences", "Performing replace references..."));
 			SlowTask.MakeDialogDelayed(1.0f);
 
 			for (UObject* Obj : Targets)
diff --git a/Engine/Source/Editor/UnrealEd/Private/LevelEditorViewport.cpp b/Engine/Source/Editor/UnrealEd/Private/LevelEditorViewport.cpp
--- a/Engine/Source/Editor/UnrealEd/Private/LevelEditorViewport.cpp
+++ b/Engine/Source/Editor/UnrealEd/Private/LevelEditorViewport.cpp
@@ -2655,7 +2655,7 @@ void FLevelEditorViewportClient::ProjectActorsIntoWorld(const TArray<AActor*>& A
 			// We only snap things that are on screen
 			FVector2D ScreenPos;
 			FIntPoint ViewportSize = InViewport->GetSizeXY();
-			if (SceneView->WorldToPixel(NewActorPosition, ScreenPos) && FMath::IsWithin<float>(ScreenPos.X, 0, ViewportSize.X) && FMath::IsWithin<float>(ScreenPos.Y, 0, ViewportSize.Y))
+			if (SceneView->WorldToPixel(NewActorPosition, ScreenPos) && FMath::IsWithin(ScreenPos.X, 0, ViewportSize.X) && FMath::IsWithin(ScreenPos.Y, 0, ViewportSize.Y))
 			{
 				bIsOnScreen = true;
 				Cursor = FViewportCursorLocation(SceneView, this, static_cast<int32>(ScreenPos.X), static_cast<int32>(ScreenPos.Y));
diff --git a/Engine/Source/Editor/UnrealEd/Private/PackageTools.cpp b/Engine/Source/Editor/UnrealEd/Private/PackageTools.cpp
--- a/Engine/Source/Editor/UnrealEd/Private/PackageTools.cpp
+++ b/Engine/Source/Editor/UnrealEd/Private/PackageTools.cpp
@@ -1169,7 +1169,7 @@ UPackageTools::UPackageTools(const FObjectInitializer& ObjectInitializer)
 			// Recompile any BPs that had their references updated
 			if (BlueprintsToRecompileThisBatch.Num() > 0)
 			{
-				FScopedSlowTask CompilingBlueprintsSlowTask(BlueprintsToRecompileThisBatch.Num(), NSLOCTEXT("UnrealEd", "CompilingBlueprints", "Compiling Blueprints"));
+				FScopedSlowTask CompilingBlueprintsSlowTask(static_cast<float>(BlueprintsToRecompileThisBatch.Num()), NSLOCTEXT("UnrealEd", "CompilingBlueprints", "Compiling Blueprints"));
 
 				// Gather up all loaded BP assets.
 				TArray<UObject*> BPs;
diff --git a/Engine/Source/Editor/UnrealEd/Private/SCommonEditorViewportToolbarBase.cpp b/Engine/Source/Editor/UnrealEd/Private/SCommonEditorViewportToolbarBase.cpp
--- a/Engine/Source/Editor/UnrealEd/Private/SCommonEditorViewportToolbarBase.cpp
+++ b/Engine/Source/Editor/UnrealEd/Private/SCommonEditorViewportToolbarBase.cpp
@@ -494,8 +494,8 @@ void SCommonEditorViewportToolbarBase::OnFOVValueChanged(float NewValue) const
 
 TSharedRef<SWidget> SCommonEditorViewportToolbarBase::GenerateScreenPercentageMenu() const
 {
-	const int32 PreviewScreenPercentageMin = ISceneViewFamilyScreenPercentage::kMinTSRResolutionFraction * 100.0f;
-	const int32 PreviewScreenPercentageMax = ISceneViewFamilyScreenPercentage::kMaxTSRResolutionFraction * 100.0f;
+	const int32 PreviewScreenPercentageMin = static_cast<int32>(ISceneViewFamilyScreenPercentage::kMinTSRResolutionFraction * 100.0f);
+	const int32 PreviewScreenPercentageMax = static_cast<int32>(ISceneViewFamilyScreenPercentage::kMaxTSRResolutionFraction * 100.0f);
 
 	return
 		SNew(SBox)
diff --git a/Engine/Source/Editor/UnrealEd/Private/SKeySelector.cpp b/Engine/Source/Editor/UnrealEd/Private/SKeySelector.cpp
--- a/Engine/Source/Editor/UnrealEd/Private/SKeySelector.cpp
+++ b/Engine/Source/Editor/UnrealEd/Private/SKeySelector.cpp
@@ -170,7 +170,7 @@ void SKeySelector::Construct(const FArguments& InArgs)
 		[
 			SAssignNew(KeyComboButton, SComboButton)
 			.OnGetMenuContent(this, &SKeySelector::GetMenuContent)
-			.ContentPadding(0)
+			.ContentPadding(0.0f)
 			.ToolTipText(this, &SKeySelector::GetKeyDescriptionToolTip)	// Longer key descriptions can overrun the visible space in the combo button if the parent width is constrained, so we reflect them in the tooltip too.
 			.ButtonContent()
 			[
diff --git a/Engine/Source/Editor/UnrealEd/Private/SScalabilitySettings.cpp b/Engine/Source/Editor/UnrealEd/Private/SScalabilitySettings.cpp
--- a/Engine/Source/Editor/UnrealEd/Private/SScalabilitySettings.cpp
+++ b/Engine/Source/Editor/UnrealEd/Private/SScalabilitySettings.cpp
@@ -19,7 +19,7 @@ ECheckBoxState SScalabilitySettings::IsGroupQualityLevelSelected(const TCHAR* In
 {
 	int32 QualityLevel = -1;
 
-	if (FCString::Strcmp(InGroupName, TEXT("ResolutionQuality")) == 0) { QualityLevel = CachedQualityLevels.ResolutionQuality; }
+	if (FCString::Strcmp(InGroupName, TEXT("ResolutionQuality")) == 0) { QualityLevel = static_cast<int32>(CachedQualityLevels.ResolutionQuality); }
 	else if (FCString::Strcmp(InGroupName, TEXT("ViewDistanceQuality")) == 0) QualityLevel = CachedQualityLevels.ViewDistanceQuality;
 	else if (FCString::Strcmp(InGroupName, TEXT("AntiAliasingQuality")) == 0) QualityLevel = CachedQualityLevels.AntiAliasingQuality;
 	else if (FCString::Strcmp(InGroupName, TEXT("PostProcessQuality")) == 0) QualityLevel = CachedQualityLevels.PostProcessQuality;
@@ -36,7 +36,7 @@ ECheckBoxState SScalabilitySettings::IsGroupQualityLevelSelected(const TCHAR* In
 
 void SScalabilitySettings::OnGroupQualityLevelChanged(ECheckBoxState NewState, const TCHAR* InGroupName, int32 InQualityLevel)
 {
-	if (FCString::Strcmp(InGroupName, TEXT("ResolutionQuality")) == 0) CachedQualityLevels.ResolutionQuality = InQualityLevel;
+	if (FCString::Strcmp(InGroupName, TEXT("ResolutionQuality")) == 0) CachedQualityLevels.ResolutionQuality = static_cast<float>(InQualityLevel);
 	else if (FCString::Strcmp(InGroupName, TEXT("ViewDistanceQuality")) == 0) CachedQualityLevels.ViewDistanceQuality = InQualityLevel;
 	else if (FCString::Strcmp(InGroupName, TEXT("AntiAliasingQuality")) == 0) CachedQualityLevels.AntiAliasingQuality = InQualityLevel;
 	else if (FCString::Strcmp(InGroupName, TEXT("PostProcessQuality")) == 0) CachedQualityLevels.PostProcessQuality = InQualityLevel;
diff --git a/Engine/Source/Editor/UnrealEd/Private/SkeletalMeshEdit.cpp b/Engine/Source/Editor/UnrealEd/Private/SkeletalMeshEdit.cpp
--- a/Engine/Source/Editor/UnrealEd/Private/SkeletalMeshEdit.cpp
+++ b/Engine/Source/Editor/UnrealEd/Private/SkeletalMeshEdit.cpp
@@ -262,7 +262,7 @@ bool UEditorEngine::ReimportFbxAnimation( USkeleton* Skeleton, UAnimSequence* An
 				if (CurAnimStack)
 				{
 					// set current anim stack
-					int32 ResampleRate = DEFAULT_SAMPLERATE;
+					int32 ResampleRate = static_cast<int32>(DEFAULT_SAMPLERATE);
 					if (FbxImporter->ImportOptions->bResample)
 					{
 						if(FbxImporter->ImportOptions->ResampleRate > 0)
diff --git a/Engine/Source/Editor/UnrealEd/Private/SnappingUtils.cpp b/Engine/Source/Editor/UnrealEd/Private/SnappingUtils.cpp
--- a/Engine/Source/Editor/UnrealEd/Private/SnappingUtils.cpp
+++ b/Engine/Source/Editor/UnrealEd/Private/SnappingUtils.cpp
@@ -197,7 +197,7 @@ bool FEditorViewportSnapping::SnapActorsToNearestActor( FVector& Drag, FLevelEdi
 			// Nearest results
 			const AActor* BestActor = NULL;
 			FVector BestPoint = FVector::ZeroVector;
-			float BestSqrdDist = 0.0f;
+			double BestSqrdDist = 0.0f;
 
 			// Find the nearest actor to the pivot point that isn't part of the selection
 			const FVector PivotLocation = Tools.PivotLocation;
@@ -226,8 +226,8 @@ bool FEditorViewportSnapping::SnapActorsToNearestActor( FVector& Drag, FLevelEdi
 
 					// Is this the nearest actor to the pivot?
 					const FVector Point = Actor->GetActorLocation();
-					const float SqrdDist = FVector::DistSquared( PivotLocation, Point );
-					if ( BestActor == NULL || SqrdDist < BestSqrdDist )
+					const double SqrdDist = FVector::DistSquared( PivotLocation, Point );
+					if ( BestActor == nullptr || SqrdDist < BestSqrdDist )
 					{
 						BestActor = Actor;
 						BestPoint = Point;
@@ -307,12 +307,12 @@ void FEditorViewportSnapping::SnapScale(FVector& Point, const FVector& GridBase)
 			if (GetDefault<ULevelEditorViewportSettings>()->PreserveNonUniformScale)
 			{
 				// when using 'auto-precision', we take the max component & snap its scale, then proportionally scale the other components
-				float MaxComponent = Point.GetAbsMax();
+				double MaxComponent = Point.GetAbsMax();
 				if(MaxComponent == 0.0f)
 				{
 					MaxComponent = 1.0f;
 				}
-				const float SnappedMaxComponent = FMath::GridSnap(MaxComponent, GEditor->GetScaleGridSize());
+				const double SnappedMaxComponent = FMath::GridSnap(MaxComponent, GEditor->GetScaleGridSize());
 				Point = Point * (SnappedMaxComponent / MaxComponent);
 			}
 			else
diff --git a/Engine/Source/Editor/UnrealEd/Private/Tests/AutomationEditorCommon.cpp b/Engine/Source/Editor/UnrealEd/Private/Tests/AutomationEditorCommon.cpp
--- a/Engine/Source/Editor/UnrealEd/Private/Tests/AutomationEditorCommon.cpp
+++ b/Engine/Source/Editor/UnrealEd/Private/Tests/AutomationEditorCommon.cpp
@@ -315,10 +315,10 @@ void FAutomationEditorCommonUtils::ApplyCustomFactorySettings(UFactory* InFactor
 * @param InTestName is the folder that has the same name as the test. (For Example: "Performance").
 * @param InItemBeingTested is the name for the thing that is being tested. (For Example: "MapName").
 * @param InFileName is the name of the file with an extension
-* @param InNumberToBeWritten is the float number that is expected to be written to the file.
+* @param InEntry is the double-precision number that is expected to be written to the file.
 * @param Delimiter is the delimiter to be used. TEXT(",")
 */
-void FAutomationEditorCommonUtils::WriteToTextFile(const FString& InTestName, const FString& InTestItem, const FString& InFileName, const float& InEntry, const FString& Delimiter)
+void FAutomationEditorCommonUtils::WriteToTextFile(const FString& InTestName, const FString& InTestItem, const FString& InFileName, const double& InEntry, const FString& Delimiter)
 {
 	//Performance file locations and setups.
 	FString FileSaveLocation = FPaths::Combine(*FPaths::AutomationLogDir(), *InTestName, *InTestItem, *InFileName);
diff --git a/Engine/Source/Editor/UnrealEd/Private/TexAlignTools.cpp b/Engine/Source/Editor/UnrealEd/Private/TexAlignTools.cpp
--- a/Engine/Source/Editor/UnrealEd/Private/TexAlignTools.cpp
+++ b/Engine/Source/Editor/UnrealEd/Private/TexAlignTools.cpp
@@ -316,10 +316,10 @@ void UTexAlignerFit::AlignSurf( ETexAlign InTexAlignType, UModel* InModel, FBspS
 
 	// Find a corner of the polygon that's closest to a 90 degree angle.  When there are multiple corners with
 	// similar angles, we'll use the one closest to the local space bottom-left along the polygon's plane
-	const float DesiredAbsDotProduct = 0.0f;
+	const double DesiredAbsDotProduct = 0.0f;
 	int32 BestVertexIndex = INDEX_NONE;
-	float BestDotProductDiff = 10000.0f;
-	float BestPositivity = 10000.0f;
+	double BestDotProductDiff = 10000.0f;
+	double BestPositivity = 10000.0f;
 	for( int32 VertexIndex = 0; VertexIndex < WorldSpacePolyVertices.Num(); ++VertexIndex )
 	{
 		// Compute the previous and next vertex in the winding
@@ -331,14 +331,14 @@ void UTexAlignerFit::AlignSurf( ETexAlign InTexAlignType, UModel* InModel, FBspS
 		const FVector& NextVertex = WorldSpacePolyVertices[ NextWindingVertexIndex ];
 
 		// Compute the corner angle
-		float AbsDotProduct = FMath::Abs( ( PrevVertex - CurVertex ).GetSafeNormal() | ( NextVertex - CurVertex ).GetSafeNormal() );
+		double AbsDotProduct = FMath::Abs( ( PrevVertex - CurVertex ).GetSafeNormal() | ( NextVertex - CurVertex ).GetSafeNormal() );
 
 		// Compute how 'positive' this vertex is relative to the bottom left position in the polygon's plane
 		FVector PolySpaceVertex = WorldToPolyRotationMatrix.InverseTransformVector( CurVertex - FirstPolyVertex );
-		const float Positivity = PolySpaceVertex.X + PolySpaceVertex.Y;
+		const double Positivity = PolySpaceVertex.X + PolySpaceVertex.Y;
 
 		// Is the corner angle closer to 90 degrees than our current best?
-		const float DotProductDiff = FMath::Abs( AbsDotProduct - DesiredAbsDotProduct );
+		const double DotProductDiff = FMath::Abs( AbsDotProduct - DesiredAbsDotProduct );
 		if( FMath::IsNearlyEqual( DotProductDiff, BestDotProductDiff, 0.1f ) )
 		{
 			// This angle is just as good as the current best, so check to see which is closer to the local space
@@ -379,12 +379,13 @@ void UTexAlignerFit::AlignSurf( ETexAlign InTexAlignType, UModel* InModel, FBspS
 	FVector PolyNormal = (FVector)InPoly->Normal;
 	WorldToTextureRotationMatrix.SetAxes( &TextureRightVec, &TextureUpVec, &PolyNormal );
 
+	
 
 	// Compute bounds of polygon along plane
-	float MinX = FLT_MAX;
-	float MaxX = -FLT_MAX;
-	float MinY = FLT_MAX;
-	float MaxY = -FLT_MAX;
+	double MinX = std::numeric_limits<double>::max();
+	double MaxX = std::numeric_limits<double>::min();
+	double MinY = std::numeric_limits<double>::max();
+	double MaxY = std::numeric_limits<double>::min();
 	for( int32 VertexIndex = 0; VertexIndex < WorldSpacePolyVertices.Num(); ++VertexIndex )
 	{
 		const FVector& CurVertex = WorldSpacePolyVertices[ VertexIndex ];
@@ -420,8 +421,8 @@ void UTexAlignerFit::AlignSurf( ETexAlign InTexAlignType, UModel* InModel, FBspS
 
 	// Apply scale to UV vectors.  We incorporate the parameterized tiling rations and scale by our texture size
 	const float WorldTexelScale = UModel::GetGlobalBSPTexelScale();
-	const float TextureSizeU = FMath::Abs( MaxX - MinX );
-	const float TextureSizeV = FMath::Abs( MaxY - MinY );
+	const double TextureSizeU = FMath::Abs( MaxX - MinX );
+	const double TextureSizeV = FMath::Abs( MaxY - MinY );
 	FVector TextureUVector = UTile * TextureRightVec * WorldTexelScale / TextureSizeU;
 	FVector TextureVVector = VTile * TextureUpVec * WorldTexelScale / TextureSizeV;
 
diff --git a/Engine/Source/Editor/UnrealEd/Private/ThumbnailExternalCache.cpp b/Engine/Source/Editor/UnrealEd/Private/ThumbnailExternalCache.cpp
--- a/Engine/Source/Editor/UnrealEd/Private/ThumbnailExternalCache.cpp
+++ b/Engine/Source/Editor/UnrealEd/Private/ThumbnailExternalCache.cpp
@@ -142,7 +142,7 @@ void FSaveThumbnailCache::Save(FArchive& Ar, const TArrayView<FAssetData> InAsse
 	const int32 NumAssetDatas = InAssetDatas.Num();
 
 	FText StatusText = LOCTEXT("SaveStatus", "Saving Thumbnails: {0}");
-	FScopedSlowTask SlowTask( (double)NumAssetDatas / (double)TaskBatchSize, FText::Format(StatusText, FText::AsNumber(NumAssetDatas)));
+	FScopedSlowTask SlowTask( (float)NumAssetDatas / (float)TaskBatchSize, FText::Format(StatusText, FText::AsNumber(NumAssetDatas)));
 	SlowTask.MakeDialog(/*bShowCancelButton*/ false);
 
 	TArray<FPackageThumbnailRecord> PackageThumbnailRecords;
@@ -236,7 +236,7 @@ void FSaveThumbnailCache::Save(FArchive& Ar, const TArrayView<FAssetData> InAsse
 		SaveTime += FPlatformTime::Seconds() - SaveTimeStart;
 		LoadTime += SaveTimeStart - LoadTimeStart;
 
-		SlowTask.EnterProgressFrame((double)AssetsForSegment.Num() / (double)TaskBatchSize, FText::Format(StatusText, FText::AsNumber(NumAssetDatas - EndIndex)));
+		SlowTask.EnterProgressFrame((float)AssetsForSegment.Num() / (float)TaskBatchSize, FText::Format(StatusText, FText::AsNumber(NumAssetDatas - EndIndex)));
 		if (SlowTask.ShouldCancel())
 		{
 			PackageThumbnailRecords.Reset();
@@ -553,7 +553,7 @@ bool FThumbnailExternalCache::LoadCacheFileIndex(FArchive& Ar, const TSharedPtr<
 	int64 NumPackages = 0;
 	Ar << NumPackages;
 
-	CacheFile->NameToEntry.Reserve(NumPackages);
+	CacheFile->NameToEntry.Reserve(IntCastChecked<int32>(NumPackages));
 
 	FString PackageNameString;
 	for (int64 i=0; i < NumPackages; ++i)
diff --git a/Engine/Source/Editor/UnrealEd/Private/Utils.cpp b/Engine/Source/Editor/UnrealEd/Private/Utils.cpp
--- a/Engine/Source/Editor/UnrealEd/Private/Utils.cpp
+++ b/Engine/Source/Editor/UnrealEd/Private/Utils.cpp
@@ -89,7 +89,7 @@ void FUnrealEdUtils::DrawWidget(const FSceneView* View,FPrimitiveDrawInterface*
 	const FVector WidgetOrigin = WidgetMatrix.GetOrigin();
 
 	// Calculate size to draw widget so it takes up the same screen space.
-	const double ZoomFactor = FMath::Min<float>(View->ViewMatrices.GetProjectionMatrix().M[0][0], View->ViewMatrices.GetProjectionMatrix().M[1][1]);
+	const double ZoomFactor = FMath::Min<double>(View->ViewMatrices.GetProjectionMatrix().M[0][0], View->ViewMatrices.GetProjectionMatrix().M[1][1]);
 	const double WidgetRadius = View->Project(WidgetOrigin).W * (UnrealEd_WidgetSize / ZoomFactor);
 
 	// Choose its color. Highlight manipulated axis in yellow.
diff --git a/Engine/Source/Editor/UnrealEd/Public/Tests/AutomationEditorCommon.h b/Engine/Source/Editor/UnrealEd/Public/Tests/AutomationEditorCommon.h
--- a/Engine/Source/Editor/UnrealEd/Public/Tests/AutomationEditorCommon.h
+++ b/Engine/Source/Editor/UnrealEd/Public/Tests/AutomationEditorCommon.h
@@ -72,10 +72,10 @@ public:
 	* @param InTestName - is the folder that has the same name as the test. (For Example: "Performance").
 	* @param InItemBeingTested - is the name for the thing that is being tested. (For Example: "MapName").
 	* @param InFileName - is the name of the file with an extension
-	* @param InNumberToBeWritten - is the float number that is expected to be written to the file.
+	* @param InEntry - is the double-precision number that is expected to be written to the file.
 	* @param Delimiter - is the delimiter to be used. TEXT(",")
 	*/
-	static void WriteToTextFile(const FString& InTestName, const FString& InTestItem, const FString& InFileName, const float& InEntry, const FString& Delimiter);
+	static void WriteToTextFile(const FString& InTestName, const FString& InTestItem, const FString& InFileName, const double& InEntry, const FString& Delimiter);
 
 	/**
 	* Returns the sum of the numbers available in an array of float.
