Fixed bug with the way used to determine whether a source control provider uses a custom project base directory

diff --git a/Engine/Source/Developer/SourceControl/Private/SourceControlModule.cpp b/Engine/Source/Developer/SourceControl/Private/SourceControlModule.cpp
--- a/Engine/Source/Developer/SourceControl/Private/SourceControlModule.cpp
+++ b/Engine/Source/Developer/SourceControl/Private/SourceControlModule.cpp
@@ -551,6 +551,19 @@ FString FSourceControlModule::GetSourceControlProjectDir() const
 	return FPaths::ProjectDir();
 }
 
+bool FSourceControlModule::UsesCustomProjectDir() const
+{
+	if (SourceControlProjectDirDelegate.IsBound())
+	{
+		FString ProjectDir = SourceControlProjectDirDelegate.Execute();
+		if (!ProjectDir.IsEmpty())
+		{
+			return true;
+		}
+	}
+	return false;
+}
+
 FSourceControlFileStatusMonitor& FSourceControlModule::GetSourceControlFileStatusMonitor()
 {
 	return *SourceControlFileStatusMonitor;
diff --git a/Engine/Source/Developer/SourceControl/Private/SourceControlModule.h b/Engine/Source/Developer/SourceControl/Private/SourceControlModule.h
--- a/Engine/Source/Developer/SourceControl/Private/SourceControlModule.h
+++ b/Engine/Source/Developer/SourceControl/Private/SourceControlModule.h
@@ -58,6 +58,7 @@ public:
 	virtual void RegisterSourceControlProjectDirDelegate(const FSourceControlProjectDirDelegate& SourceControlProjectDirDelegate) override;
 	virtual void UnregisterSourceControlProjectDirDelegate() override;
 	virtual FString GetSourceControlProjectDir() const override;
+	virtual bool UsesCustomProjectDir() const override;
 	virtual FSourceControlFileStatusMonitor& GetSourceControlFileStatusMonitor() override;
 
 	/** Save the settings to the ini file */
diff --git a/Engine/Source/Developer/SourceControl/Public/ISourceControlModule.h b/Engine/Source/Developer/SourceControl/Public/ISourceControlModule.h
--- a/Engine/Source/Developer/SourceControl/Public/ISourceControlModule.h
+++ b/Engine/Source/Developer/SourceControl/Public/ISourceControlModule.h
@@ -250,6 +250,11 @@ public:
 	 */
 	virtual FString GetSourceControlProjectDir() const = 0;
 
+	/**
+	 * Returns whether a delegate has been registered to specify the project base directory to be used by the source control
+	 */
+	virtual bool UsesCustomProjectDir() const = 0;
+
 	/**
 	 * Returns the object used to monitor the source control status of a collection of files.
 	 */
diff --git a/Engine/Source/Editor/ContentBrowser/Private/FrontendFilters.cpp b/Engine/Source/Editor/ContentBrowser/Private/FrontendFilters.cpp
--- a/Engine/Source/Editor/ContentBrowser/Private/FrontendFilters.cpp
+++ b/Engine/Source/Editor/ContentBrowser/Private/FrontendFilters.cpp
@@ -850,14 +850,14 @@ void FFrontendFilter_NotSourceControlled::RequestStatus()
 		TSharedRef<FUpdateStatus, ESPMode::ThreadSafe> UpdateStatusOperation = ISourceControlOperation::Create<FUpdateStatus>();
 
 		TArray<FString> Filenames;
-		FString SourceControlProjectDir = ISourceControlModule::Get().GetSourceControlProjectDir();
-		if (SourceControlProjectDir.IsEmpty())
+		if (ISourceControlModule::Get().UsesCustomProjectDir())
 		{
-			Filenames = FSourceControlWindows::GetSourceControlLocations(/*bContentOnly*/true);
+			FString SourceControlProjectDir = ISourceControlModule::Get().GetSourceControlProjectDir();
+			Filenames.Add(SourceControlProjectDir);
 		}
 		else
 		{
-			Filenames.Add(SourceControlProjectDir);
+			Filenames = FSourceControlWindows::GetSourceControlLocations(/*bContentOnly*/true);
 		}
 		UpdateStatusOperation->SetCheckingAllFiles(false);
 		SourceControlProvider.Execute(UpdateStatusOperation, Filenames, EConcurrency::Asynchronous, FSourceControlOperationComplete::CreateSP(this, &FFrontendFilter_NotSourceControlled::SourceControlOperationComplete));
diff --git a/Engine/Source/Editor/SourceControlWindows/Private/SourceControlWindows.cpp b/Engine/Source/Editor/SourceControlWindows/Private/SourceControlWindows.cpp
--- a/Engine/Source/Editor/SourceControlWindows/Private/SourceControlWindows.cpp
+++ b/Engine/Source/Editor/SourceControlWindows/Private/SourceControlWindows.cpp
@@ -89,14 +89,14 @@ bool FSourceControlWindows::ChoosePackagesToCheckIn(const FSourceControlWindowsO
 
 	// make sure we update the SCC status of all packages (this could take a long time, so we will run it as a background task)
 	TArray<FString> Filenames;
-	FString SourceControlProjectDir = ISourceControlModule::Get().GetSourceControlProjectDir();
-	if (SourceControlProjectDir.IsEmpty())
+	if (ISourceControlModule::Get().UsesCustomProjectDir())
 	{
-		Filenames = GetSourceControlLocations();
+		FString SourceControlProjectDir = ISourceControlModule::Get().GetSourceControlProjectDir();
+		Filenames.Add(SourceControlProjectDir);
 	}
 	else
 	{
-		Filenames.Add(SourceControlProjectDir);
+		Filenames = GetSourceControlLocations();
 	}
 	
 	// make sure the SourceControlProvider state cache is populated as well
@@ -267,9 +267,7 @@ bool FSourceControlWindows::PromptForCheckin(FCheckinResultInfo& OutResultInfo,
 	{
 		SourceControlHelpers::RevertUnchangedFiles(SourceControlProvider, Description.FilesForSubmit);
 
-		const FString SCCProjectDir = ISourceControlModule::Get().GetSourceControlProjectDir();
-
-		if (SCCProjectDir.IsEmpty())
+		if (!ISourceControlModule::Get().UsesCustomProjectDir())
 		{
 			// Make sure all files are still checked out
 			for (int32 VerifyIndex = Description.FilesForSubmit.Num() - 1; VerifyIndex >= 0; --VerifyIndex)
@@ -445,15 +443,15 @@ void FSourceControlWindows::ChoosePackagesToCheckInCompleted(const TArray<UPacka
 	TArray<FString> PendingDeletePaths;
 
 	bool bUseSourceControlStateCache = true;
-	FString SourceControlProjectDir = ISourceControlModule::Get().GetSourceControlProjectDir();
-	if (SourceControlProjectDir.IsEmpty())
+	if (ISourceControlModule::Get().UsesCustomProjectDir())
 	{
-		PendingDeletePaths = GetSourceControlLocations();
+		FString SourceControlProjectDir = ISourceControlModule::Get().GetSourceControlProjectDir();
+		PendingDeletePaths.Add(SourceControlProjectDir);
+		bUseSourceControlStateCache = false;
 	}
 	else
 	{
-		PendingDeletePaths.Add(SourceControlProjectDir);
-		bUseSourceControlStateCache = false;
+		PendingDeletePaths = GetSourceControlLocations();
 	}
 
 	PromptForCheckin(OutResultInfo, PackageNames, PendingDeletePaths, ConfigFiles, bUseSourceControlStateCache);
diff --git a/Engine/Source/Editor/UnrealEd/Private/FileHelpers.cpp b/Engine/Source/Editor/UnrealEd/Private/FileHelpers.cpp
--- a/Engine/Source/Editor/UnrealEd/Private/FileHelpers.cpp
+++ b/Engine/Source/Editor/UnrealEd/Private/FileHelpers.cpp
@@ -4524,7 +4524,7 @@ void FEditorFileUtils::LoadDefaultMapAtStartup()
 void FEditorFileUtils::FindAllPackageFiles(TArray<FString>& OutPackages)
 {
 	FString SourceControlProjectDir = ISourceControlModule::Get().GetSourceControlProjectDir();
-	if (!SourceControlProjectDir.IsEmpty())
+	if (ISourceControlModule::Get().UsesCustomProjectDir())
 	{
 		FPackageName::FindPackagesInDirectory(OutPackages, SourceControlProjectDir);
 		return;
@@ -4554,8 +4554,6 @@ void FEditorFileUtils::FindAllSubmittablePackageFiles(TMap<FString, FSourceContr
 {
 	ISourceControlProvider& SourceControlProvider = ISourceControlModule::Get().GetProvider();
 
-	const FString SCCProjectDir = ISourceControlModule::Get().GetSourceControlProjectDir();
-
 	TArray<FString> Packages;
 	FEditorFileUtils::FindAllPackageFiles(Packages);
 
@@ -4577,9 +4575,9 @@ void FEditorFileUtils::FindAllSubmittablePackageFiles(TMap<FString, FSourceContr
 		FSourceControlStatePtr SourceControlState = SourceControlProvider.GetState(FPaths::ConvertRelativePathToFull(Filename), EStateCacheUsage::Use);
 
 		// Only include non-map packages that are currently checked out or packages not under source control
-		if (SCCProjectDir.IsEmpty())
+		if (ISourceControlModule::Get().UsesCustomProjectDir())
 		{
-			if (SourceControlState.IsValid() && SourceControlState->IsCurrent() &&
+			if (SourceControlState.IsValid() &&
 				(SourceControlState->CanCheckIn() || (!SourceControlState->IsSourceControlled() && SourceControlState->CanAdd())) &&
 				(bIncludeMaps || !IsMapPackageAsset(*Filename)))
 			{
@@ -4588,7 +4586,7 @@ void FEditorFileUtils::FindAllSubmittablePackageFiles(TMap<FString, FSourceContr
 		}
 		else
 		{
-			if (SourceControlState.IsValid() && 
+			if (SourceControlState.IsValid() && SourceControlState->IsCurrent() &&
 				(SourceControlState->CanCheckIn() || (!SourceControlState->IsSourceControlled() && SourceControlState->CanAdd())) &&
 				(bIncludeMaps || !IsMapPackageAsset(*Filename)))
 			{
@@ -4604,20 +4602,10 @@ void FEditorFileUtils::FindAllSubmittableProjectFiles(TMap<FString, FSourceContr
 {
 	ISourceControlProvider& SourceControlProvider = ISourceControlModule::Get().GetProvider();
 
-	const FString SCCProjectDir = ISourceControlModule::Get().GetSourceControlProjectDir();
-	if (SCCProjectDir.IsEmpty())
+	if (ISourceControlModule::Get().UsesCustomProjectDir())
 	{
-		// Handle just the project file
-		FSourceControlStatePtr SourceControlState = SourceControlProvider.GetState(FPaths::ConvertRelativePathToFull(FPaths::GetProjectFilePath()), EStateCacheUsage::Use);
+		const FString SCCProjectDir = ISourceControlModule::Get().GetSourceControlProjectDir();
 
-		if (SourceControlState.IsValid() && SourceControlState->IsCurrent() &&
-			(SourceControlState->CanCheckIn() || (!SourceControlState->IsSourceControlled() && SourceControlState->CanAdd())))
-		{
-			OutProjectFiles.Add(FPaths::GetProjectFilePath(), MoveTemp(SourceControlState));
-		}
-	}
-	else
-	{
 		// Handle non-package files in the project directory
 		TArray<FSourceControlStateRef> SourceControlStates = SourceControlProvider.GetCachedStateByPredicate(
 			[SCCProjectDir](const FSourceControlStateRef& SourceControlState)
@@ -4641,6 +4629,17 @@ void FEditorFileUtils::FindAllSubmittableProjectFiles(TMap<FString, FSourceContr
 			}
 		}
 	}
+	else
+	{
+		// Handle just the project file
+		FSourceControlStatePtr SourceControlState = SourceControlProvider.GetState(FPaths::ConvertRelativePathToFull(FPaths::GetProjectFilePath()), EStateCacheUsage::Use);
+
+		if (SourceControlState.IsValid() && SourceControlState->IsCurrent() &&
+			(SourceControlState->CanCheckIn() || (!SourceControlState->IsSourceControlled() && SourceControlState->CanAdd())))
+		{
+			OutProjectFiles.Add(FPaths::GetProjectFilePath(), MoveTemp(SourceControlState));
+		}
+	}
 }
 
 static void FindAllConfigFilesRecursive(TArray<FString>& OutConfigFiles, const FString& ParentDirectory)
