CookMultiprocess: Temporarily change WaitForConnect timeout to be 10 minutes instead of one minute to avoid timing out while director is busy during start of the cook.
#rnx

diff --git a/Engine/Source/Editor/UnrealEd/Private/Cooker/CookWorkerClient.cpp b/Engine/Source/Editor/UnrealEd/Private/Cooker/CookWorkerClient.cpp
--- a/Engine/Source/Editor/UnrealEd/Private/Cooker/CookWorkerClient.cpp
+++ b/Engine/Source/Editor/UnrealEd/Private/Cooker/CookWorkerClient.cpp
@@ -21,7 +21,7 @@ namespace UE::Cook
 
 namespace CookWorkerClient
 {
-constexpr float WaitForConnectReplyTimeout = 60.f;
+constexpr float WaitForConnectReplyTimeout = 600.f;
 }
 
 FCookWorkerClient::FCookWorkerClient(UCookOnTheFlyServer& InCOTFS)
@@ -263,7 +263,7 @@ void FCookWorkerClient::CreateServerSocket(const FDirectorConnectionInfo& Connec
 		return;
 	}
 
-	constexpr float WaitForConnectTimeout = 60.f * 10;
+	constexpr float WaitForConnectTimeout = 600.f * 10;
 	float ConditionalTimeoutSeconds = IsCookIgnoreTimeouts() ? MAX_flt : WaitForConnectTimeout;
 	bool bServerSocketReady = ServerSocket->Wait(ESocketWaitConditions::WaitForWrite, FTimespan::FromSeconds(ConditionalTimeoutSeconds));
 	if (!bServerSocketReady)
