Uniquely instance subobjects when importing default value terms to non-native class variable properties on the new CDO during a Blueprint compile.

Change summary:
- Reverted previous change (21708020) as that would have failed with a NULL OwnerObject. Pasting/duplicating a variable just imports the clipboard text to a new descriptor value; we don't need to add the port flag until compile time.
- Modified FBlueprintEditorUtils::PropertyValueToString()/PropertyValueFromString() APIs to include an additional parameter to support callers that need to add one or more extra port flag(s) to the underlying ImportText()/ExportText() call.
- Modified FKismetCompilerContext::CopyTermDefaultsToDefaultObject() to check for properties that either contain or represent an instanced object reference and append the PPF_InstanceSubobjects port flag in that case.

diff --git a/Engine/Source/Editor/Kismet/Private/SMyBlueprint.cpp b/Engine/Source/Editor/Kismet/Private/SMyBlueprint.cpp
--- a/Engine/Source/Editor/Kismet/Private/SMyBlueprint.cpp
+++ b/Engine/Source/Editor/Kismet/Private/SMyBlueprint.cpp
@@ -3768,7 +3768,7 @@ void SMyBlueprint::OnPasteVariable()
 	FBPVariableDescription Description;
 	FStringOutputDevice Errors;
 	const TCHAR* Import = ClipboardText.GetCharArray().GetData() + FCString::Strlen(VAR_PREFIX);
-	FBPVariableDescription::StaticStruct()->ImportText(Import, &Description, nullptr, PPF_InstanceSubobjects, &Errors, FBPVariableDescription::StaticStruct()->GetName());
+	FBPVariableDescription::StaticStruct()->ImportText(Import, &Description, nullptr, PPF_None, &Errors, FBPVariableDescription::StaticStruct()->GetName());
 	if (Errors.IsEmpty())
 	{
 		FBPVariableDescription NewVar = FBlueprintEditorUtils::DuplicateVariableDescription(Blueprint, Description);
diff --git a/Engine/Source/Editor/KismetCompiler/Private/KismetCompiler.cpp b/Engine/Source/Editor/KismetCompiler/Private/KismetCompiler.cpp
--- a/Engine/Source/Editor/KismetCompiler/Private/KismetCompiler.cpp
+++ b/Engine/Source/Editor/KismetCompiler/Private/KismetCompiler.cpp
@@ -1721,26 +1721,33 @@ void FKismetCompilerContext::CopyTermDefaultsToDefaultObject(UObject* DefaultObj
 				++PropertiesAssigned;
 
 				const FString& Value = *ValuePtr;
-				if(FObjectProperty* AsObjectProperty = CastField<FObjectProperty>(Property))
+				if (FObjectProperty* AsObjectProperty = CastField<FObjectProperty>(Property))
 				{
 					// Value is the fully qualified name, so just search for it:
 					UObject* Result = StaticFindObjectSafe(UObject::StaticClass(), nullptr, *Value);
-					if(Result)
+					if (Result)
 					{
 						// Object may be of a type that is also being compiled and therefore REINST_, so get real class:
 						UClass* RealClass = Result->GetClass()->GetAuthoritativeClass();
 
 						// If object is compatible, write it into cdo:
-						if( RealClass && RealClass->IsChildOf(AsObjectProperty->PropertyClass) )
+						if (RealClass && RealClass->IsChildOf(AsObjectProperty->PropertyClass))
 						{
-							AsObjectProperty->SetObjectPropertyValue( AsObjectProperty->ContainerPtrToValuePtr<uint8>(DefaultObject), Result );
+							AsObjectProperty->SetObjectPropertyValue(AsObjectProperty->ContainerPtrToValuePtr<uint8>(DefaultObject), Result);
 							continue;
 						}
 					}
 				}
 
-				const bool bParseSuccedded = FBlueprintEditorUtils::PropertyValueFromString(Property, Value, reinterpret_cast<uint8*>(DefaultObject));
-				if(!bParseSuccedded)
+				// If this property contains an instanced reference, set the flag to make sure we uniquely instance those objects when we import the default value.
+				int32 PortFlags = PPF_None;
+				if (Property->HasAnyPropertyFlags(CPF_ContainsInstancedReference | CPF_InstancedReference))
+				{
+					PortFlags |= PPF_InstanceSubobjects;
+				}
+
+				const bool bParseSucceeded = FBlueprintEditorUtils::PropertyValueFromString(Property, Value, reinterpret_cast<uint8*>(DefaultObject), DefaultObject, PortFlags);
+				if(!bParseSucceeded)
 				{
 					const FString ErrorMessage = *FText::Format(
 						LOCTEXT("ParseDefaultValueErrorFmt", "Can't parse default value '{0}' for @@. Property: {1}."),
diff --git a/Engine/Source/Editor/UnrealEd/Private/Kismet2/BlueprintEditorUtils.cpp b/Engine/Source/Editor/UnrealEd/Private/Kismet2/BlueprintEditorUtils.cpp
--- a/Engine/Source/Editor/UnrealEd/Private/Kismet2/BlueprintEditorUtils.cpp
+++ b/Engine/Source/Editor/UnrealEd/Private/Kismet2/BlueprintEditorUtils.cpp
@@ -9117,12 +9117,12 @@ bool FBlueprintEditorUtils::IsObjectADebugCandidate( AActor* InActorObject, UBlu
 	return bPassesFlags && bCanDebugThisObject;
 }
 
-bool FBlueprintEditorUtils::PropertyValueFromString(const FProperty* Property, const FString& StrValue, uint8* Container, UObject* OwningObject)
+bool FBlueprintEditorUtils::PropertyValueFromString(const FProperty* Property, const FString& StrValue, uint8* Container, UObject* OwningObject, int32 PortFlags)
 {
-	return PropertyValueFromString_Direct(Property, StrValue, Property->ContainerPtrToValuePtr<uint8>(Container), OwningObject);
+	return PropertyValueFromString_Direct(Property, StrValue, Property->ContainerPtrToValuePtr<uint8>(Container), OwningObject, PortFlags);
 }
 
-bool FBlueprintEditorUtils::PropertyValueFromString_Direct(const FProperty* Property, const FString& StrValue, uint8* DirectValue, UObject* OwningObject)
+bool FBlueprintEditorUtils::PropertyValueFromString_Direct(const FProperty* Property, const FString& StrValue, uint8* DirectValue, UObject* OwningObject, int32 PortFlags)
 {
 	bool bParseSucceeded = true;
 	if (!Property->IsA(FStructProperty::StaticClass()))
@@ -9228,7 +9228,7 @@ bool FBlueprintEditorUtils::PropertyValueFromString_Direct(const FProperty* Prop
 		else if (Property->IsA(FTextProperty::StaticClass()))
 		{
 			FStringOutputDevice ImportError;
-			const TCHAR* EndOfParsedBuff = Property->ImportText_Direct(*StrValue, DirectValue, OwningObject, PPF_SerializedAsImportText, &ImportError);
+			const TCHAR* EndOfParsedBuff = Property->ImportText_Direct(*StrValue, DirectValue, OwningObject, PPF_SerializedAsImportText | PortFlags, &ImportError);
 			bParseSucceeded = EndOfParsedBuff && ImportError.IsEmpty();
 		}
 		else
@@ -9239,7 +9239,7 @@ bool FBlueprintEditorUtils::PropertyValueFromString_Direct(const FProperty* Prop
 				: *StrValue;
 
 			FStringOutputDevice ImportError;
-			const TCHAR* EndOfParsedBuff = Property->ImportText_Direct(*StrValue, DirectValue, OwningObject, PPF_SerializedAsImportText, &ImportError);
+			const TCHAR* EndOfParsedBuff = Property->ImportText_Direct(*StrValue, DirectValue, OwningObject, PPF_SerializedAsImportText | PortFlags, &ImportError);
 			bParseSucceeded = EndOfParsedBuff && ImportError.IsEmpty();
 		}
 	}
@@ -9286,7 +9286,7 @@ bool FBlueprintEditorUtils::PropertyValueFromString_Direct(const FProperty* Prop
 			ensure(1 == StructProperty->ArrayDim);
 
 			FStringOutputDevice ImportError;
-			const TCHAR* EndOfParsedBuff = StructProperty->ImportText_Direct(StrValue.IsEmpty() ? TEXT("()") : *StrValue, DirectValue, OwningObject, PPF_SerializedAsImportText, &ImportError);
+			const TCHAR* EndOfParsedBuff = StructProperty->ImportText_Direct(StrValue.IsEmpty() ? TEXT("()") : *StrValue, DirectValue, OwningObject, PPF_SerializedAsImportText | PortFlags, &ImportError);
 			bParseSucceeded &= EndOfParsedBuff && ImportError.IsEmpty();
 		}
 	}
@@ -9294,12 +9294,12 @@ bool FBlueprintEditorUtils::PropertyValueFromString_Direct(const FProperty* Prop
 	return bParseSucceeded;
 }
 
-bool FBlueprintEditorUtils::PropertyValueToString(const FProperty* Property, const uint8* Container, FString& OutForm, UObject* OwningObject)
+bool FBlueprintEditorUtils::PropertyValueToString(const FProperty* Property, const uint8* Container, FString& OutForm, UObject* OwningObject, int32 PortFlags)
 {
-	return PropertyValueToString_Direct(Property, Property->ContainerPtrToValuePtr<const uint8>(Container), OutForm, OwningObject);
+	return PropertyValueToString_Direct(Property, Property->ContainerPtrToValuePtr<const uint8>(Container), OutForm, OwningObject, PortFlags);
 }
 
-bool FBlueprintEditorUtils::PropertyValueToString_Direct(const FProperty* Property, const uint8* DirectValue, FString& OutForm, UObject* OwningObject)
+bool FBlueprintEditorUtils::PropertyValueToString_Direct(const FProperty* Property, const uint8* DirectValue, FString& OutForm, UObject* OwningObject, int32 PortFlags)
 {
 	check(Property && DirectValue);
 	OutForm.Reset();
@@ -9343,7 +9343,7 @@ bool FBlueprintEditorUtils::PropertyValueToString_Direct(const FProperty* Proper
 	if (OutForm.IsEmpty())
 	{
 		const uint8* DefaultValue = DirectValue;	
-		bSucceeded = Property->ExportText_Direct(OutForm, DirectValue, DefaultValue, OwningObject, PPF_SerializedAsImportText);
+		bSucceeded = Property->ExportText_Direct(OutForm, DirectValue, DefaultValue, OwningObject, PPF_SerializedAsImportText | PortFlags);
 	}
 	return bSucceeded;
 }
diff --git a/Engine/Source/Editor/UnrealEd/Public/Kismet2/BlueprintEditorUtils.h b/Engine/Source/Editor/UnrealEd/Public/Kismet2/BlueprintEditorUtils.h
--- a/Engine/Source/Editor/UnrealEd/Public/Kismet2/BlueprintEditorUtils.h
+++ b/Engine/Source/Editor/UnrealEd/Public/Kismet2/BlueprintEditorUtils.h
@@ -1329,25 +1329,25 @@ public:
 	 * Copies the value from the passed in string into a property. ContainerMem points to the Struct or Class containing Property 
 	 * NOTE: This function does not work correctly with static arrays.
 	 */
-	static bool PropertyValueFromString(const FProperty* Property, const FString& StrValue, uint8* Container, UObject* OwningObject = nullptr);
+	static bool PropertyValueFromString(const FProperty* Property, const FString& StrValue, uint8* Container, UObject* OwningObject = nullptr, int32 PortFlags = PPF_None);
 
 	/** 
 	 * Copies the value from the passed in string into a property. DirectValue is the raw memory address of the property value 
 	 * NOTE: This function does not work correctly with static arrays.
 	 */
-	static bool PropertyValueFromString_Direct(const FProperty* Property, const FString& StrValue, uint8* DirectValue, UObject* OwningObject = nullptr);
+	static bool PropertyValueFromString_Direct(const FProperty* Property, const FString& StrValue, uint8* DirectValue, UObject* OwningObject = nullptr, int32 PortFlags = PPF_None);
 
 	/** 
 	 * Copies the value from a property into the string OutForm. ContainerMem points to the Struct or Class containing Property 
 	 * NOTE: This function does not work correctly with static arrays.
 	 */
-	static bool PropertyValueToString(const FProperty* Property, const uint8* Container, FString& OutForm, UObject* OwningObject = nullptr);
+	static bool PropertyValueToString(const FProperty* Property, const uint8* Container, FString& OutForm, UObject* OwningObject = nullptr, int32 PortFlags = PPF_None);
 
 	/** 
 	 * Copies the value from a property into the string OutForm. DirectValue is the raw memory address of the property value 
 	 * NOTE: This function does not work correctly with static arrays.
 	 */
-	static bool PropertyValueToString_Direct(const FProperty* Property, const uint8* DirectValue, FString& OutForm, UObject* OwningObject = nullptr);
+	static bool PropertyValueToString_Direct(const FProperty* Property, const uint8* DirectValue, FString& OutForm, UObject* OwningObject = nullptr, int32 PortFlags = PPF_None);
 
 	/** Call PostEditChange() on all Actors based on the given Blueprint */
 	static void PostEditChangeBlueprintActors(UBlueprint* Blueprint, bool bComponentEditChange = false);
