Move IAssetReferenceFilter out of EditorEngine.h into its own header file
#rnx

diff --git a/Engine/Source/Editor/UnrealEd/Classes/Editor/AssetReferenceFilter.h b/Engine/Source/Editor/UnrealEd/Classes/Editor/AssetReferenceFilter.h
new file mode 100644
--- /dev/null
+++ b/Engine/Source/Editor/UnrealEd/Classes/Editor/AssetReferenceFilter.h
@@ -0,0 +1,16 @@
+// Copyright Epic Games, Inc. All Rights Reserved.
+
+#pragma once
+
+#include "CoreTypes.h"
+
+struct FAssetData;
+
+/** Used in filtering allowed references between assets. Implement a subclass of this and return it in OnMakeAssetReferenceFilter */
+class IAssetReferenceFilter
+{
+public:
+	virtual ~IAssetReferenceFilter() { }
+	/** Filter function to pass/fail an asset. Called in some situations that are performance-sensitive so is expected to be fast. */
+	virtual bool PassesFilter(const FAssetData& AssetData, FText* OutOptionalFailureReason = nullptr) const = 0;
+};
\ No newline at end of file
diff --git a/Engine/Source/Editor/UnrealEd/Classes/Editor/EditorEngine.h b/Engine/Source/Editor/UnrealEd/Classes/Editor/EditorEngine.h
--- a/Engine/Source/Editor/UnrealEd/Classes/Editor/EditorEngine.h
+++ b/Engine/Source/Editor/UnrealEd/Classes/Editor/EditorEngine.h
@@ -31,6 +31,7 @@
 #include "UnrealEngine.h"
 #include "Templates/PimplPtr.h"
 #include "Templates/UniqueObj.h"
+#include "Editor/AssetReferenceFilter.h"
 
 #include "EditorEngine.generated.h"
 
@@ -274,15 +275,6 @@ struct FAssetReferenceFilterContext
 	TArray<FAssetData> ReferencingAssets;
 };
 
-/** Used in filtering allowed references between assets. Implement a subclass of this and return it in OnMakeAssetReferenceFilter */
-class IAssetReferenceFilter
-{
-public:
-	virtual ~IAssetReferenceFilter() { }
-	/** Filter function to pass/fail an asset. Called in some situations that are performance-sensitive so is expected to be fast. */
-	virtual bool PassesFilter(const FAssetData& AssetData, FText* OutOptionalFailureReason = nullptr) const = 0;
-};
-
 /**
  * Engine that drives the Editor.
  * Separate from UGameEngine because it may have much different functionality than desired for an instance of a game itself.
