Choose optimal source pixel format for UDIM when importing textures of different pixels formats.

diff --git a/Engine/Plugins/Interchange/Runtime/Source/Import/Private/Texture/InterchangeTextureFactory.cpp b/Engine/Plugins/Interchange/Runtime/Source/Import/Private/Texture/InterchangeTextureFactory.cpp
--- a/Engine/Plugins/Interchange/Runtime/Source/Import/Private/Texture/InterchangeTextureFactory.cpp
+++ b/Engine/Plugins/Interchange/Runtime/Source/Import/Private/Texture/InterchangeTextureFactory.cpp
@@ -521,41 +521,40 @@ namespace UE::Interchange::Private::InterchangeTextureFactory
 			BlockedImage.BlocksData.Reserve(Images.Num());
 
 			bool bMismatchedFormats = false;
+			bool bMismatchedGammaSpace = false;
 			for (int32 Index = 0; Index < Images.Num(); ++Index)
 			{
+				if (BlockedImage.bSRGB != Images[Index].bSRGB)
+				{
+					bMismatchedGammaSpace = true;
+					BlockedImage.bSRGB = false;
+					BlockedImage.Format = TSF_RGBA32F;
+				}
+
 				if (BlockedImage.Format != Images[Index].Format)
 				{
 					bMismatchedFormats = true;
-					break;
+					BlockedImage.Format = FImageCoreUtils::GetCommonSourceFormat(BlockedImage.Format, Images[Index].Format);
 				}
 			}
 
-			if (bMismatchedFormats)
+			if (bMismatchedGammaSpace || bMismatchedFormats)
 			{
-				UE_LOG(LogInterchangeImport, Warning, TEXT("Mismatched UDIM image formats, converting all to BGRA8 or RGBA16F ..."));
-
-				if (FTextureSource::IsHDR(BlockedImage.Format))
-				{
-					BlockedImage.Format = TSF_RGBA16F;
-					BlockedImage.bSRGB = false;
-				}
-				else
-				{
-					BlockedImage.Format = TSF_BGRA8;
-					BlockedImage.bSRGB = true;
-				}
+				UE_LOG(LogInterchangeImport, Warning, TEXT("Mismatched UDIM image %s, converting all to %s/%s ..."), bMismatchedGammaSpace ? TEXT("gamma spaces") : TEXT("pixel formats"),
+					ERawImageFormat::GetName(FImageCoreUtils::ConvertToRawImageFormat(BlockedImage.Format)), BlockedImage.bSRGB ? TEXT("sRGB") : TEXT("Linear"));
 
 				for (UE::Interchange::FImportImage& Image : Images)
 				{
-					if (Image.Format != BlockedImage.Format)
+					if (Image.bSRGB != BlockedImage.bSRGB || Image.Format != BlockedImage.Format)
 					{
 						ERawImageFormat::Type ImageRawFormat = FImageCoreUtils::ConvertToRawImageFormat(Image.Format);
-						FImageView SourceImage(Image.RawData.GetData(), Image.SizeX, Image.SizeY, ImageRawFormat);
-						FImage DestImage(Image.SizeX, Image.SizeY, FImageCoreUtils::ConvertToRawImageFormat(BlockedImage.Format));
+						FImageView SourceImage(Image.RawData.GetData(), Image.SizeX, Image.SizeY, 1, ImageRawFormat, Image.bSRGB ? EGammaSpace::sRGB : EGammaSpace::Linear);
+						FImage DestImage(Image.SizeX, Image.SizeY, FImageCoreUtils::ConvertToRawImageFormat(BlockedImage.Format), BlockedImage.bSRGB ? EGammaSpace::sRGB : EGammaSpace::Linear);
 						FImageCore::CopyImage(SourceImage, DestImage);
 
 						Image.RawData = MakeUniqueBufferFromArray(MoveTemp(DestImage.RawData));
 						Image.Format = BlockedImage.Format;
+						Image.bSRGB = BlockedImage.bSRGB;
 					}
 				}
 			}
diff --git a/Engine/Source/Editor/UnrealEd/Private/Factories/EditorFactories.cpp b/Engine/Source/Editor/UnrealEd/Private/Factories/EditorFactories.cpp
--- a/Engine/Source/Editor/UnrealEd/Private/Factories/EditorFactories.cpp
+++ b/Engine/Source/Editor/UnrealEd/Private/Factories/EditorFactories.cpp
@@ -3207,6 +3207,7 @@ UTexture* UTextureFactory::ImportTextureUDIM(UClass* Class, UObject* InParent, F
 	const EImageImportFlags ImportFlags = EImageImportFlags::None;
 
 	bool bMismatchedFormats = false;
+	bool bMismatchedGammaSpace = false;
 
 	for (const auto& It : UDIMIndexToFile)
 	{
@@ -3228,9 +3229,17 @@ UTexture* UTextureFactory::ImportTextureUDIM(UClass* Class, UObject* InParent, F
 					TCSettings = Image.CompressionSettings;
 				}
 
-				if ( Format != Image.Format )
+				if (bSRGB != Image.SRGB)
+				{
+					bMismatchedGammaSpace = true;
+					bSRGB = false;
+					Format = TSF_RGBA32F;
+				}
+
+				if (Format != Image.Format)
 				{
 					bMismatchedFormats = true;
+					Format = FImageCoreUtils::GetCommonSourceFormat(Format, Image.Format);
 				}
 
 				const int32 UDIMIndex = It.Key;
@@ -3252,32 +3261,23 @@ UTexture* UTextureFactory::ImportTextureUDIM(UClass* Class, UObject* InParent, F
 		return nullptr;
 	}
 
-	if ( bMismatchedFormats )
+	if (bMismatchedGammaSpace || bMismatchedFormats)
 	{
-		Warn->Logf(ELogVerbosity::Warning, TEXT("Mismatched UDIM image formats, converting all to BGRA8 or RGBA16F ..."));
-		
-		if ( FTextureSource::IsHDR(Format) )
-		{
-			Format = TSF_RGBA16F;
-			bSRGB = false;
-		}
-		else
-		{
-			Format = TSF_BGRA8;
-			bSRGB = true;
-		}
+		Warn->Logf(ELogVerbosity::Warning, TEXT("Mismatched UDIM image %s, converting all to %s/%s ..."), bMismatchedGammaSpace ? TEXT("gamma spaces") : TEXT("pixel formats"),
+			ERawImageFormat::GetName(FImageCoreUtils::ConvertToRawImageFormat(Format)), bSRGB ? TEXT("sRGB") : TEXT("Linear"));
 
 		for( FImportImage & Image : SourceImages )
 		{
-			if ( Image.Format != Format )
+			if (Image.SRGB != bSRGB || Image.Format != Format)
 			{
-				ERawImageFormat::Type ImageRawFormat = FImageCoreUtils::ConvertToRawImageFormat( Image.Format );
-				FImageView SourceImage( Image.RawData.GetData(), Image.SizeX, Image.SizeY, ImageRawFormat );
-				FImage DestImage( Image.SizeX, Image.SizeY, FImageCoreUtils::ConvertToRawImageFormat(Format) );
-				FImageCore::CopyImage(SourceImage,DestImage);
+				ERawImageFormat::Type ImageRawFormat = FImageCoreUtils::ConvertToRawImageFormat(Image.Format);
+				FImageView SourceImage(Image.RawData.GetData(), Image.SizeX, Image.SizeY, 1, ImageRawFormat, Image.SRGB ? EGammaSpace::sRGB : EGammaSpace::Linear);
+				FImage DestImage(Image.SizeX, Image.SizeY, FImageCoreUtils::ConvertToRawImageFormat(Format), bSRGB ? EGammaSpace::sRGB : EGammaSpace::Linear);
+				FImageCore::CopyImage(SourceImage, DestImage);
 
 				Image.RawData = MoveTemp(DestImage.RawData);
 				Image.Format = Format;
+				Image.SRGB = bSRGB;
 			}				
 		}
 	}
