Updated UPanelExtensionSubsystem to support multiple FPanelExtensionFactory elements that each create a widget

#rnx

diff --git a/Engine/Source/Editor/UnrealEd/Private/Subsystems/PanelExtensionSubsystem.cpp b/Engine/Source/Editor/UnrealEd/Private/Subsystems/PanelExtensionSubsystem.cpp
--- a/Engine/Source/Editor/UnrealEd/Private/Subsystems/PanelExtensionSubsystem.cpp
+++ b/Engine/Source/Editor/UnrealEd/Private/Subsystems/PanelExtensionSubsystem.cpp
@@ -76,6 +76,10 @@ void UPanelExtensionSubsystem::RegisterPanelFactory(FName ExtensionPanelID, cons
 {
 	TArray<FPanelExtensionFactory>& Factories = ExtensionPointMap.FindOrAdd(ExtensionPanelID);
 	Factories.Add(InPanelExtensionFactory);
+	Factories.StableSort([](const FPanelExtensionFactory& One, const FPanelExtensionFactory& Two)
+	{
+		return One.SortWeight > Two.SortWeight;
+	});
 
 	UPanelExtensionSubsystem::FPanelFactoryRegistryChanged& PanelFactoryRegsitered = OnPanelFactoryRegistryChanged(ExtensionPanelID);
 	PanelFactoryRegsitered.Broadcast();
@@ -117,18 +121,34 @@ TSharedRef<SWidget> UPanelExtensionSubsystem::CreateWidget(FName ExtensionPanelI
 	const TArray<FPanelExtensionFactory>* ExtensionArray = ExtensionPointMap.Find(ExtensionPanelID);
 	if (ExtensionArray && ExtensionArray->Num() > 0)
 	{
-		const FPanelExtensionFactory& Extension = (*ExtensionArray)[0]; // TODO: We need to add support for multiple widgets
-		if (Extension.CreateExtensionWidget.IsBound())
-		{
-			return Extension.CreateExtensionWidget.Execute(ExtensionContext);
-		}
-		PRAGMA_DISABLE_DEPRECATION_WARNINGS
-		else if (Extension.CreateWidget.IsBound())
+		TSharedRef<SHorizontalBox> ExtensionWidgets = SNew(SHorizontalBox);
+
+		for (const FPanelExtensionFactory& Extension : *ExtensionArray)
 		{
-			TArray<UObject*> Temp;
-			return Extension.CreateWidget.Execute(Temp);
+			TSharedPtr<SWidget> ExtensionWidget;
+			if (Extension.CreateExtensionWidget.IsBound())
+			{
+				ExtensionWidget = Extension.CreateExtensionWidget.Execute(ExtensionContext);
+			}
+			PRAGMA_DISABLE_DEPRECATION_WARNINGS
+			else if (Extension.CreateWidget.IsBound())
+			{
+				TArray<UObject*> Temp;
+				ExtensionWidget = Extension.CreateWidget.Execute(Temp);
+			}
+			PRAGMA_ENABLE_DEPRECATION_WARNINGS
+
+			if (ExtensionWidget && ExtensionWidget != SNullWidget::NullWidget)
+			{
+				ExtensionWidgets->AddSlot()
+				.AutoWidth()
+				[
+					ExtensionWidget.ToSharedRef()
+				];
+			}
 		}
-		PRAGMA_ENABLE_DEPRECATION_WARNINGS
+
+		return ExtensionWidgets;
 	}
 	return SNullWidget::NullWidget;
 }
diff --git a/Engine/Source/Editor/UnrealEd/Public/Subsystems/PanelExtensionSubsystem.h b/Engine/Source/Editor/UnrealEd/Public/Subsystems/PanelExtensionSubsystem.h
--- a/Engine/Source/Editor/UnrealEd/Public/Subsystems/PanelExtensionSubsystem.h
+++ b/Engine/Source/Editor/UnrealEd/Public/Subsystems/PanelExtensionSubsystem.h
@@ -21,6 +21,9 @@ public:
 	/** An identifier to allow removal later on. */
 	FName Identifier;
 
+	/** Weight used when sorting against other factories for this panel (higher weights sort first) */
+	int32 SortWeight = 0;
+
 	/**
 	 * Delegate that generates the SExtensionPanel content widget. 
 	 * The FWeakObjectPtr param is an opaque context specific to the panel that is being extended
