CookMultiprocess: Move CookDirector's communication with CookWorkers to another thread so the communication does not time out when the director process is busy saving and loading packages.
#rnx

diff --git a/Engine/Source/Editor/UnrealEd/Private/Cooker/CompactBinaryTCP.cpp b/Engine/Source/Editor/UnrealEd/Private/Cooker/CompactBinaryTCP.cpp
--- a/Engine/Source/Editor/UnrealEd/Private/Cooker/CompactBinaryTCP.cpp
+++ b/Engine/Source/Editor/UnrealEd/Private/Cooker/CompactBinaryTCP.cpp
@@ -58,6 +58,7 @@ public:
 	/** Note that the elements in AppendMessages are mutated */
 	static EConnectionStatus TryWritePacket(FSocket* Socket, FSendBuffer& Buffer,
 		TArrayView<FMarshalledMessage>&& AppendMessages, uint64 MaxPacketSize);
+	static void QueueMessage(FSendBuffer& Buffer, FMarshalledMessage&& Message);
 };
 
 EConnectionStatus FCompactBinaryTCPImpl::TryReadPacket(FSocket* Socket, FReceiveBuffer& Buffer,
@@ -308,6 +309,11 @@ EConnectionStatus FCompactBinaryTCPImpl::TryWritePacket(FSocket* Socket, FSendBu
 	}
 }
 
+void FCompactBinaryTCPImpl::QueueMessage(FSendBuffer& Buffer, FMarshalledMessage&& Message)
+{
+	Buffer.PendingMessages.Add(MoveTemp(Message));
+}
+
 EConnectionStatus TryReadPacket(FSocket* Socket, FReceiveBuffer& Buffer,
 	TArray<FMarshalledMessage>& Messages, uint64 MaxPacketSize)
 {
@@ -341,6 +347,18 @@ EConnectionStatus TryWritePacket(FSocket* Socket, FSendBuffer& Buffer,
 		TArrayView<FMarshalledMessage>(&Marshalled, 1), MaxPacketSize);
 }
 
+void QueueMessage(FSendBuffer& Buffer, const IMessage& AppendMessage)
+{
+	FMarshalledMessage Marshalled;
+	Marshalled.MessageType = AppendMessage.GetMessageType();
+	FCbWriter Writer;
+	Writer.BeginObject();
+	AppendMessage.Write(Writer);
+	Writer.EndObject();
+	Marshalled.Object = Writer.Save().AsObject();
+	FCompactBinaryTCPImpl::QueueMessage(Buffer, MoveTemp(Marshalled));
+}
+
 EConnectionStatus TryFlushBuffer(FSocket* Socket, FSendBuffer& Buffer, uint64 MaxPacketSize)
 {
 	return FCompactBinaryTCPImpl::TryWritePacket(Socket, Buffer, TArrayView<FMarshalledMessage>(), MaxPacketSize);
diff --git a/Engine/Source/Editor/UnrealEd/Private/Cooker/CompactBinaryTCP.h b/Engine/Source/Editor/UnrealEd/Private/Cooker/CompactBinaryTCP.h
--- a/Engine/Source/Editor/UnrealEd/Private/Cooker/CompactBinaryTCP.h
+++ b/Engine/Source/Editor/UnrealEd/Private/Cooker/CompactBinaryTCP.h
@@ -88,6 +88,8 @@ EConnectionStatus TryWritePacket(FSocket* Socket, FSendBuffer& Buffer,
 	FMarshalledMessage&& AppendMessage, uint64 MaxPacketSize = 0);
 EConnectionStatus TryWritePacket(FSocket* Socket, FSendBuffer& Buffer,
 	const IMessage& AppendMessage, uint64 MaxPacketSize = 0);
+void QueueMessage(FSendBuffer& Buffer, const IMessage& AppendMessage);
+
 
 /** Attempt to finish sending packets that were previously queued by TryWritePacket */
 EConnectionStatus TryFlushBuffer(FSocket* Socket, FSendBuffer& Buffer, uint64 MaxPacketSize = 0);
diff --git a/Engine/Source/Editor/UnrealEd/Private/Cooker/CookDirector.cpp b/Engine/Source/Editor/UnrealEd/Private/Cooker/CookDirector.cpp
--- a/Engine/Source/Editor/UnrealEd/Private/Cooker/CookDirector.cpp
+++ b/Engine/Source/Editor/UnrealEd/Private/Cooker/CookDirector.cpp
@@ -6,11 +6,13 @@
 #include "CompactBinaryTCP.h"
 #include "CookMPCollector.h"
 #include "CookPackageData.h"
+#include "CookPlatformManager.h"
 #include "CookWorkerServer.h"
 #include "CookOnTheSide/CookOnTheFlyServer.h"
 #include "CoreGlobals.h"
 #include "LoadBalanceCookBurden.h"
 #include "HAL/PlatformMisc.h"
+#include "HAL/RunnableThread.h"
 #include "Math/NumericLimits.h"
 #include "Misc/CommandLine.h"
 #include "Misc/ConfigCacheIni.h"
@@ -20,6 +22,7 @@
 #include "Sockets.h"
 #include "SocketSubsystem.h"
 #include "String/ParseTokens.h"
+#include "UnrealEdMisc.h"
 
 extern CORE_API int32 GNumForegroundWorkers; // TaskGraph.cpp
 
@@ -27,14 +30,23 @@ namespace UE::Cook
 {
 
 FCookDirector::FCookDirector(UCookOnTheFlyServer& InCOTFS)
-	: COTFS(InCOTFS)
+	: RunnableShunt(*this) 
+	, COTFS(InCOTFS)
 {
 	WorkersStalledStartTimeSeconds = MAX_flt;
 	WorkersStalledWarnTimeSeconds = MAX_flt;
+	ShutdownEvent->Reset();
 
 	ParseConfig();
-	ResetCookWorkerCount();
-	UE_LOG(LogCook, Display, TEXT("CookMultiprocess is enabled with %d CookWorker processes."), CookWorkerCount);
+	ISocketSubsystem* SocketSubsystem = ISocketSubsystem::Get();
+	if (!SocketSubsystem)
+	{
+		UE_LOG(LogCook, Error, TEXT("CookDirector initialization failure: platform does not support network sockets. CookWorkers will be disabled."));
+	}
+	else
+	{
+		UE_LOG(LogCook, Display, TEXT("CookMultiprocess is enabled with %d CookWorker processes."), RequestedCookWorkerCount);
+	}
 }
 
 void FCookDirector::ParseConfig()
@@ -84,32 +96,14 @@ void FCookDirector::ParseConfig()
 	}
 }
 
-void FCookDirector::ResetCookWorkerCount()
-{
-	CookWorkerCount = RequestedCookWorkerCount;
-	if (CookWorkerCount > 0)
-	{
-		ISocketSubsystem* SocketSubsystem = ISocketSubsystem::Get();
-		if (!SocketSubsystem)
-		{
-			UE_LOG(LogCook, Error, TEXT("CookDirector initialization failure: platform does not support network sockets. CookWorkers will be disabled."));
-			CookWorkerCount = 0;
-		}
-	}
-
-	if (CookWorkerCount > 0)
-	{
-		TryCreateWorkerConnectSocket();
-	}
-
-}
-
 FCookDirector::~FCookDirector()
 {
+	StopCommunicationThread();
+
 	TSet<FPackageData*> AbortedAssignments;
-	for (TPair<int32, TUniquePtr<FCookWorkerServer>>& Pair : RemoteWorkers)
+	for (TPair<int32, TRefCountPtr<FCookWorkerServer>>& Pair : RemoteWorkers)
 	{
-		Pair.Value->AbortWorker(AbortedAssignments);
+		Pair.Value->AbortWorker(AbortedAssignments, ECookDirectorThread::SchedulerThread);
 	}
 	for (FPackageData* PackageData : AbortedAssignments)
 	{
@@ -122,6 +116,56 @@ FCookDirector::~FCookDirector()
 	Sockets::CloseSocket(WorkerConnectSocket);
 }
 
+void FCookDirector::LaunchCommunicationThread()
+{
+	if (!CommunicationThread && FPlatformProcess::SupportsMultithreading())
+	{
+		CommunicationThread = FRunnableThread::Create(&RunnableShunt, TEXT("FCookDirector"), 0, TPri_Normal);
+	}
+}
+
+void FCookDirector::StopCommunicationThread()
+{
+	ShutdownEvent->Trigger();
+	if (CommunicationThread)
+	{
+		CommunicationThread->WaitForCompletion();
+		delete CommunicationThread;
+		CommunicationThread = nullptr;
+	}
+	ShutdownEvent->Reset();
+}
+
+uint32 FCookDirector::RunCommunicationThread()
+{
+	constexpr float TickPeriod = 1.f;
+	constexpr float MinSleepTime = 0.001f;
+	for (;;)
+	{
+		double StartTime = FPlatformTime::Seconds();
+		TickCommunication(ECookDirectorThread::CommunicateThread);
+
+		double CurrentTime = FPlatformTime::Seconds();
+		float RemainingDuration = StartTime + TickPeriod - CurrentTime;
+		uint32 WaitTimeMilliseconds = static_cast<uint32>(RemainingDuration * .001f);
+		if (ShutdownEvent->Wait(WaitTimeMilliseconds))
+		{
+			break;
+		}
+	}
+	return 0;
+}
+
+uint32 FCookDirector::FRunnableShunt::Run()
+{
+	return Director.RunCommunicationThread();
+}
+
+void FCookDirector::FRunnableShunt::Stop()
+{
+	Director.ShutdownEvent->Trigger();
+}
+
 void FCookDirector::StartCook(const FBeginCookContext& InBeginContext)
 {
 	BeginCookContext.Set(InBeginContext);
@@ -131,10 +175,22 @@ void FCookDirector::AssignRequests(TArrayView<UE::Cook::FPackageData*> Requests,
 	TMap<FPackageData*, TArray<FPackageData*>>&& RequestGraph)
 {
 	ActivateMachineResourceReduction();
-	InitializeWorkers();
 
-	int32 NumRemoteWorkers = RemoteWorkers.Num();
-	if (NumRemoteWorkers == 0)
+	TArray<TRefCountPtr<FCookWorkerServer>> SortedWorkers;
+	int32 MaxRemoteIndex = -1;
+	FScopeLock CommunicationScopeLock(&CommunicationLock);
+	{
+		InitializeWorkers();
+
+		// Convert the Map of RemoteWorkers into an array sorted by WorkerIndex
+		SortedWorkers.Reserve(RemoteWorkers.Num());
+		for (TPair<int32, TRefCountPtr<FCookWorkerServer>>& Pair : RemoteWorkers)
+		{
+			SortedWorkers.Add(Pair.Value);
+			MaxRemoteIndex = FMath::Max(Pair.Key, MaxRemoteIndex);
+		}
+	}
+	if (SortedWorkers.IsEmpty())
 	{
 		OutAssignments.SetNum(Requests.Num());
 		for (FWorkerId& Assignment : OutAssignments)
@@ -143,18 +199,9 @@ void FCookDirector::AssignRequests(TArrayView<UE::Cook::FPackageData*> Requests,
 		}
 		return;
 	}
-
-	// Convert the Map of RemoteWorkers into an array sorted by WorkerIndex
-	TArray<FCookWorkerServer*> SortedWorkers;
-	SortedWorkers.Reserve(NumRemoteWorkers);
-	int32 MaxRemoteIndex = -1;
-	for (TPair<int32, TUniquePtr<FCookWorkerServer>>& Pair : RemoteWorkers)
-	{
-		SortedWorkers.Add(Pair.Value.Get());
-		MaxRemoteIndex = FMath::Max(Pair.Key, MaxRemoteIndex);
-	}
 	check(MaxRemoteIndex >= 0);
-	SortedWorkers.Sort([](const FCookWorkerServer& A, const FCookWorkerServer& B) { return A.GetWorkerId() < B.GetWorkerId(); });
+	SortedWorkers.Sort([](const TRefCountPtr<FCookWorkerServer>& A, const TRefCountPtr<FCookWorkerServer>& B)
+		{ return A->GetWorkerId() < B->GetWorkerId(); });
 
 	// Call the LoadBalancing algorithm to split the requests among the LocalWorker and RemoteWorkers
 	LoadBalance(SortedWorkers, Requests, MoveTemp(RequestGraph), OutAssignments);
@@ -196,7 +243,7 @@ void FCookDirector::AssignRequests(TArrayView<UE::Cook::FPackageData*> Requests,
 			TArray<FPackageData*>& RemoteBatch = RemoteBatches[RemoteIndex];
 			if (RemoteBatch.Num() == 0)
 			{
-				RemoteBatch.Reserve(2 * Requests.Num() / (NumRemoteWorkers + 1));
+				RemoteBatch.Reserve(2 * Requests.Num() / (SortedWorkers.Num() + 1));
 			}
 			RemoteBatch.Add(Requests[RequestIndex]);
 		}
@@ -206,100 +253,199 @@ void FCookDirector::AssignRequests(TArrayView<UE::Cook::FPackageData*> Requests,
 
 	// Assign each batch to the FCookWorkerServer in RemoteWorkers;
 	// the CookWorkerServer's tick will handle sending the message to the remote process
-	for (TPair<int32, TUniquePtr<FCookWorkerServer>>& Pair : RemoteWorkers)
+	for (const TRefCountPtr<FCookWorkerServer>& RemoteWorker : SortedWorkers)
 	{
-		Pair.Value->AppendAssignments(RemoteBatches[Pair.Key]);
+		RemoteWorker->AppendAssignments(RemoteBatches[RemoteWorker->GetWorkerId().GetRemoteIndex()], ECookDirectorThread::SchedulerThread);
 	}
-	TickWorkerConnects();
 
 	bIsFirstAssignment = false;
 }
 
 void FCookDirector::RemoveFromWorker(FPackageData& PackageData)
 {
-	for (TPair<int32, TUniquePtr<FCookWorkerServer>>& Pair : RemoteWorkers)
+	TArray<FCookWorkerServer*> Workers;
+	{
+		FScopeLock CommunicationScopeLock(&CommunicationLock);
+		for (TPair<int32, TRefCountPtr<FCookWorkerServer>>& Pair : RemoteWorkers)
+		{
+			Workers.Add(Pair.Value);
+		}
+	}
+	for (FCookWorkerServer* Worker : Workers)
 	{
-		Pair.Value->AbortAssignment(PackageData);
+		Worker->AbortAssignment(PackageData, ECookDirectorThread::SchedulerThread);
 	}
 }
 
-
 void FCookDirector::TickFromSchedulerThread()
 {
-	TickWorkerConnects();
-	for (TPair<int32, TUniquePtr<FCookWorkerServer>>& Pair : RemoteWorkers)
+	if (!CommunicationThread)
 	{
-		FCookWorkerServer& RemoteWorker = *Pair.Value;
-		RemoteWorker.TickFromSchedulerThread();
-		if (RemoteWorker.IsShuttingDown())
-		{
-			TUniquePtr<FCookWorkerServer>& Existing = ShuttingDownWorkers.FindOrAdd(&RemoteWorker);
-			check(!Existing); // We should not be able to send the same pointer into ShuttingDown twice
-		}
+		TickCommunication(ECookDirectorThread::SchedulerThread);
 	}
 
-	TickWorkerShutdowns();
+	TArray<TRefCountPtr<FCookWorkerServer>, TInlineAllocator<16>> WorkersWithMessage;
+	{
+		FScopeLock CommunicationScopeLock(&CommunicationLock);
+		for (TPair<int32, TRefCountPtr<FCookWorkerServer>>& Pair : RemoteWorkers)
+		{
+			if (Pair.Value->HasMessages())
+			{
+				WorkersWithMessage.Add(Pair.Value);
+			}
+		}
+		for (TPair<FCookWorkerServer*, TRefCountPtr<FCookWorkerServer>>& Pair : ShuttingDownWorkers)
+		{
+			if (Pair.Value && Pair.Value->HasMessages())
+			{
+				WorkersWithMessage.Add(Pair.Value);
+			}
+		}
 
+	}
 	bool bIsStalled = COTFS.IsMultiprocessLocalWorkerIdle() && !COTFS.PackageDatas->GetAssignedToWorkerSet().IsEmpty();
+	for (TRefCountPtr<FCookWorkerServer>& Worker : WorkersWithMessage)
+	{
+		Worker->HandleReceiveMessages(ECookDirectorThread::SchedulerThread);
+		bIsStalled = false;
+	}
+	WorkersWithMessage.Empty();
+
 	SetWorkersStalled(bIsStalled);
 }
 
-void FCookDirector::PumpCookComplete(bool& bCompleted)
+void FCookDirector::TickCommunication(ECookDirectorThread TickThread)
 {
-	TickWorkerConnects();
+	bool bHasShutdownWorkers = false;
+	TickWorkerConnects(TickThread);
+	TArray<TRefCountPtr<FCookWorkerServer>, TInlineAllocator<16>> LocalRemoteWorkers;
+	{
+		FScopeLock CommunicationScopeLock(&CommunicationLock);
+		for (const TPair<int32, TRefCountPtr<FCookWorkerServer>>& Pair: RemoteWorkers)
+		{
+			LocalRemoteWorkers.Add(Pair.Value);
+		}
+		if (!RemoteWorkers.IsEmpty())
+		{
+			bWorkersActive = true;
+		}
+		else
+		{
+			bWorkersActive = false;
+			for (const TPair <FCookWorkerServer*, TRefCountPtr<FCookWorkerServer>>& Pair : ShuttingDownWorkers)
+			{
+				FCookWorkerServer* RemoteWorker = Pair.Key;
+				check(RemoteWorker->IsShuttingDown());
+				bWorkersActive = bWorkersActive || RemoteWorker->IsFlushingBeforeShutdown();
+			}
+		}
+		bHasShutdownWorkers = !ShuttingDownWorkers.IsEmpty();
+	}
+
+	for (TRefCountPtr<FCookWorkerServer>& RemoteWorker: LocalRemoteWorkers)
+	{
+		RemoteWorker->TickCommunication(TickThread);
+		if (RemoteWorker->IsShuttingDown())
+		{
+			FScopeLock CommunicationScopeLock(&CommunicationLock);
+			TRefCountPtr<FCookWorkerServer>& Existing = ShuttingDownWorkers.FindOrAdd(RemoteWorker.GetReference());
+			check(!Existing); // We should not be able to send the same pointer into ShuttingDown twice
+			bHasShutdownWorkers = true;
+		}
+	}
 
-	// MPCOOKTODO: Messages sent from the CookWorkers about discovered packages might come in after
-	// the last message sent about completed saves. These discovered packages can cause the cook to fall
-	// back to incomplete. Don't send CookComplete messages to the CookWorkers until all CookWorkers have
-	// reported they are idle.
-	for (TPair<int32, TUniquePtr<FCookWorkerServer>>& Pair : RemoteWorkers)
+	if (bHasShutdownWorkers)
 	{
-		FCookWorkerServer& RemoteWorker = *Pair.Value;
-		RemoteWorker.SignalCookComplete();
-		check(RemoteWorker.IsShuttingDown());
-		TUniquePtr<FCookWorkerServer>& Existing = ShuttingDownWorkers.FindOrAdd(&RemoteWorker);
-		check(!Existing); // We should not be able to send the same pointer into ShuttingDown twice
+		TickWorkerShutdowns(TickThread);
 	}
-	TickWorkerShutdowns();
-	check(RemoteWorkers.Num() == 0);
+}
 
-	bCompleted = true;
-	for (const TPair <FCookWorkerServer*, TUniquePtr<FCookWorkerServer>>& Pair : ShuttingDownWorkers)
+void FCookDirector::PumpCookComplete(bool& bCompleted)
+{
 	{
-		FCookWorkerServer* RemoteWorker = Pair.Key;
-		check(RemoteWorker->IsShuttingDown());
-		bCompleted = bCompleted && !RemoteWorker->IsFlushingBeforeShutdown();
+		FScopeLock CommunicationScopeLock(&CommunicationLock);
+		if (!bCookCompleteSent)
+		{
+			bool bAllIdle = true;
+			for (TPair<int32, TRefCountPtr<FCookWorkerServer>>& Pair : RemoteWorkers)
+			{
+				FCookWorkerServer& RemoteWorker = *Pair.Value;
+				// MPCOOKTODO: Messages sent from the CookWorkers about discovered packages might come in after
+				// the last message sent about completed saves. These discovered packages can cause the cook to fall
+				// back to incomplete. Don't send CookComplete messages to the CookWorkers until all CookWorkers have
+				// reported they are idle and have no further messages to send.
+				if (RemoteWorker.HasAssignments())
+				{
+					bAllIdle = false;
+					break;
+				}
+			}
+			if (bAllIdle)
+			{
+				for (TPair<int32, TRefCountPtr<FCookWorkerServer>>& Pair : RemoteWorkers)
+				{
+					FCookWorkerServer& RemoteWorker = *Pair.Value;
+					RemoteWorker.SignalCookComplete(ECookDirectorThread::SchedulerThread);
+					check(RemoteWorker.IsShuttingDown());
+				}
+				bCookCompleteSent = true;
+			}
+		}
+		bCompleted = !bWorkersActive;
 	}
-	SetWorkersStalled(!bCompleted);
+	TickFromSchedulerThread();
 }
 
 void FCookDirector::ShutdownCookSession()
 {
+	StopCommunicationThread();
+
 	// Cancel any inprogress workers and move them to the Shutdown list
-	while (!RemoteWorkers.IsEmpty())
+	for (;;)
 	{
-		AbortWorker(TMap<int32, TUniquePtr<FCookWorkerServer>>::TIterator(RemoteWorkers).Value()->GetWorkerId());
+		TRefCountPtr<FCookWorkerServer> RemoteWorker;
+		{
+			FScopeLock CommunicationScopeLock(&CommunicationLock);
+			if (RemoteWorkers.IsEmpty())
+			{
+				break;
+			}
+			RemoteWorker = TMap<int32, TRefCountPtr<FCookWorkerServer>>::TIterator(RemoteWorkers).Value();
+		}
+		AbortWorker(RemoteWorker->GetWorkerId(), ECookDirectorThread::SchedulerThread);
 	}
 
 	// Immediately shutdown any gracefully shutting down workers
-	for (TPair<FCookWorkerServer*, TUniquePtr<FCookWorkerServer>>& Pair : ShuttingDownWorkers)
+	TArray<TRefCountPtr<FCookWorkerServer>, TInlineAllocator<16>> WorkersNeedingAbort;
 	{
-		if (Pair.Key->IsFlushingBeforeShutdown())
+		FScopeLock CommunicationScopeLock(&CommunicationLock);
+		for (TPair<FCookWorkerServer*, TRefCountPtr<FCookWorkerServer>>& Pair : ShuttingDownWorkers)
 		{
-			TSet<FPackageData*> UnusedPendingPackages;
-			Pair.Key->AbortWorker(UnusedPendingPackages);
+			check(Pair.Value); // The Value was set by AbortWorker for any new entries, and all old entries guarantee the value is set
+			if (Pair.Key->IsFlushingBeforeShutdown())
+			{
+				WorkersNeedingAbort.Add(Pair.Value);
+			}
 		}
 	}
+	for (TRefCountPtr<FCookWorkerServer>& RemoteWorker : WorkersNeedingAbort)
+	{
+		TSet<FPackageData*> UnusedPendingPackages;
+		RemoteWorker->AbortWorker(UnusedPendingPackages, ECookDirectorThread::SchedulerThread);
+	}
 
 	// Wait for all the shutdowns to complete
 	for (;;)
 	{
-		TickWorkerShutdowns();
-		if (ShuttingDownWorkers.IsEmpty())
+		TickWorkerShutdowns(ECookDirectorThread::SchedulerThread);
 		{
-			break;
+			FScopeLock CommunicationScopeLock(&CommunicationLock);
+			if (ShuttingDownWorkers.IsEmpty())
+			{
+				break;
+			}
 		}
-		constexpr double SleepSeconds = 0.010;
+		constexpr float SleepSeconds = 0.010f;
 		FPlatformProcess::Sleep(SleepSeconds);
 	}
 
@@ -307,8 +453,10 @@ void FCookDirector::ShutdownCookSession()
 	PendingConnections.Reset();
 
 	// Restore the FCookDirector to its original state so that it is ready for a new session
-	ResetCookWorkerCount();
+	bWorkersInitialized = false;
 	bIsFirstAssignment = true;
+	bCookCompleteSent = false;
+	bWorkersActive = false;
 }
 
 void FCookDirector::Register(IMPCollector* Collector)
@@ -401,6 +549,13 @@ bool FCookDirector::TryCreateWorkerConnectSocket()
 	{
 		return true;
 	}
+	ISocketSubsystem* SocketSubsystem = ISocketSubsystem::Get();
+	if (!SocketSubsystem)
+	{
+		// Error was already logged in the constructor
+		return false;
+	}
+
 	FString ErrorReason;
 	TSharedPtr<FInternetAddr> ListenAddr;
 	WorkerConnectSocket = Sockets::CreateListenSocket(WorkerConnectPort, ListenAddr, WorkerConnectAuthority,
@@ -409,7 +564,6 @@ bool FCookDirector::TryCreateWorkerConnectSocket()
 	{
 		UE_LOG(LogCook, Error, TEXT("CookDirector could not create listen socket, CookWorkers will be disabled. Reason: %s."),
 			*ErrorReason);
-		CookWorkerCount = 0;
 		return false;
 	}
 	return true;
@@ -418,16 +572,56 @@ bool FCookDirector::TryCreateWorkerConnectSocket()
 
 void FCookDirector::InitializeWorkers()
 {
-	if (RemoteWorkers.Num() >= CookWorkerCount)
+	if (bWorkersInitialized)
 	{
 		return;
 	}
+	bWorkersInitialized = true;
+
+	check(!CommunicationThread);
+	check(RemoteWorkers.IsEmpty());
+	bool bSucceeded = false;
+	ON_SCOPE_EXIT
+	{
+		if (!bSucceeded)
+		{
+			bWorkersActive = false;
+		}
+	};
+
+	if (!TryCreateWorkerConnectSocket())
+	{
+		return;
+	}
+
+	RemoteWorkers.Reserve(RequestedCookWorkerCount);
+	for (int32 RemoteIndex = 0; RemoteIndex < RequestedCookWorkerCount; ++RemoteIndex)
+	{
+		RemoteWorkers.Add(RemoteIndex, new FCookWorkerServer(*this, FWorkerId::FromRemoteIndex(RemoteIndex)));
+	}
+	bWorkersActive = true;
+
+	ConstructReadonlyThreadVariables();
+	ShutdownEvent->Reset();
+	LaunchCommunicationThread();
+	bSucceeded = true;
+}
+
+void FCookDirector::RecreateWorkers()
+{
+	// TODO: Finish implementing the recreation of workers that have crashed
 
 	// Find any unused RemoteIndex less than the maximum used RemoteIndex
+	FScopeLock CommunicationScopeLock(&CommunicationLock);
+	if (RemoteWorkers.Num() >= RequestedCookWorkerCount || !WorkerConnectSocket)
+	{
+		return;
+	}
+
 	TArray<uint8> UnusedRemoteIndexes;
 	RemoteWorkers.KeySort(TLess<>());
 	uint8 NextPossiblyOpenIndex = 0;
-	for (TPair<int32, TUniquePtr<FCookWorkerServer>>& Pair : RemoteWorkers)
+	for (TPair<int32, TRefCountPtr<FCookWorkerServer>>& Pair : RemoteWorkers)
 	{
 		check(NextPossiblyOpenIndex <= Pair.Key);
 		while (NextPossiblyOpenIndex != Pair.Key)
@@ -435,9 +629,10 @@ void FCookDirector::InitializeWorkers()
 			UnusedRemoteIndexes.Add(NextPossiblyOpenIndex++);
 		}
 	}
+
 	// Add RemoteWorkers, pulling the RemoteIndex id from the UnusedRemoteIndexes if any exist
 	// otherwise use the next integer because all indexes up to RemoteWorkers.Num() are in use.
-	while (RemoteWorkers.Num() < CookWorkerCount)
+	while (RemoteWorkers.Num() < RequestedCookWorkerCount)
 	{
 		uint8 RemoteIndex;
 		if (UnusedRemoteIndexes.Num())
@@ -449,10 +644,11 @@ void FCookDirector::InitializeWorkers()
 		{
 			RemoteIndex = RemoteWorkers.Num();
 		}
-		RemoteWorkers.Add(RemoteIndex, MakeUnique<FCookWorkerServer>(*this, FWorkerId::FromRemoteIndex(RemoteIndex)));
+		RemoteWorkers.Add(RemoteIndex, new FCookWorkerServer(*this, FWorkerId::FromRemoteIndex(RemoteIndex)));
+		bWorkersActive = true;
 	}
-}
 
+}
 void FCookDirector::ActivateMachineResourceReduction()
 {
 	if (bHasReducedMachineResources)
@@ -506,7 +702,7 @@ void FCookDirector::ActivateMachineResourceReduction()
 		HyperThreadCount, FPlatformMisc::NumberOfCoresIncludingHyperthreads());
 }
 
-void FCookDirector::TickWorkerConnects()
+void FCookDirector::TickWorkerConnects(ECookDirectorThread TickThread)
 {
 	using namespace UE::CompactBinaryTCP;
 
@@ -560,27 +756,33 @@ void FCookDirector::TickWorkerConnects()
 			UE_LOG(LogCook, Warning, TEXT("Pending connection sent an invalid Connection Message. Connection will be ignored."));
 			continue;
 		}
-		TUniquePtr<FCookWorkerServer>* RemoteWorkerPtr = RemoteWorkers.Find(Message.RemoteIndex);
-		if (!RemoteWorkerPtr)
+		TRefCountPtr<FCookWorkerServer> RemoteWorker;
 		{
-			TStringBuilder<256> ValidIndexes;
-			if (RemoteWorkers.Num())
+			FScopeLock CommunicationScopeLock(&CommunicationLock);
+			TRefCountPtr<FCookWorkerServer>* RemoteWorkerPtr;
+			RemoteWorkerPtr = RemoteWorkers.Find(Message.RemoteIndex);
+			if (!RemoteWorkerPtr)
 			{
-				RemoteWorkers.KeySort(TLess<>());
-				for (TPair<int32, TUniquePtr<FCookWorkerServer>>& Pair : RemoteWorkers)
+				TStringBuilder<256> ValidIndexes;
+				if (RemoteWorkers.Num())
 				{
-					ValidIndexes.Appendf(TEXT("%d,"), Pair.Key);
+					RemoteWorkers.KeySort(TLess<>());
+					for (TPair<int32, TRefCountPtr<FCookWorkerServer>>& Pair : RemoteWorkers)
+					{
+						ValidIndexes.Appendf(TEXT("%d,"), Pair.Key);
+					}
+					ValidIndexes.RemoveSuffix(1); // Remove the terminating comma
 				}
-				ValidIndexes.RemoveSuffix(1); // Remove the terminating comma
+				UE_LOG(LogCook, Warning, TEXT("Pending connection sent a Connection Message with invalid RemoteIndex %d. ValidIndexes = {%s}. Connection will be ignored."),
+					Message.RemoteIndex, *ValidIndexes);
+				continue;
 			}
-			UE_LOG(LogCook, Warning, TEXT("Pending connection sent a Connection Message with invalid RemoteIndex %d. ValidIndexes = {%s}. Connection will be ignored."),
-				Message.RemoteIndex, *ValidIndexes);
-			continue;
+			RemoteWorker = *RemoteWorkerPtr;
 		}
-		FCookWorkerServer& RemoteWorker = **RemoteWorkerPtr;
+
 		FSocket* LocalSocket = LocalConn.DetachSocket();
 		Messages.RemoveAt(0);
-		if (!RemoteWorker.TryHandleConnectMessage(Message, LocalSocket, MoveTemp(Messages)))
+		if (!RemoteWorker->TryHandleConnectMessage(Message, LocalSocket, MoveTemp(Messages), TickThread))
 		{
 			UE_LOG(LogCook, Warning, TEXT("Pending connection sent a Connection Message with an already in-use RemoteIndex. Connection will be ignored."));
 			Sockets::CloseSocket(LocalSocket);
@@ -589,37 +791,58 @@ void FCookDirector::TickWorkerConnects()
 	}
 }
 
-void FCookDirector::TickWorkerShutdowns()
+void FCookDirector::TickWorkerShutdowns(ECookDirectorThread TickThread)
 {
-	if (ShuttingDownWorkers.IsEmpty())
+	// Move any newly shutting down workers from RemoteWorkers
+	TArray<TRefCountPtr<FCookWorkerServer>> NewShutdowns;
+	TArray<TRefCountPtr<FCookWorkerServer>, TInlineAllocator<16>> LocalRemoteWorkers;
 	{
-		return;
+		FScopeLock RemoteWorkersScopeLock(&CommunicationLock);
+		for (TPair<FCookWorkerServer*, TRefCountPtr<FCookWorkerServer>>& Pair : ShuttingDownWorkers)
+		{
+			if (!Pair.Value)
+			{
+				NewShutdowns.Emplace(Pair.Key);
+			}
+			else
+			{
+				LocalRemoteWorkers.Add(Pair.Value);
+			}
+		}
 	}
-
-	// Move any newly shutting down workers from RemoteWorkers
-	TArray<FCookWorkerServer*> NewShutdowns;
-	for (TPair<FCookWorkerServer*, TUniquePtr<FCookWorkerServer>>& Pair : ShuttingDownWorkers)
+	if (!NewShutdowns.IsEmpty())
 	{
-		if (!Pair.Value)
+		for (FCookWorkerServer* NewShutdown : NewShutdowns)
 		{
-			NewShutdowns.Add(Pair.Key);
+			AbortWorker(NewShutdown->GetWorkerId(), TickThread);
+			{
+				FScopeLock CommunicationScopeLock(&CommunicationLock);
+				check(ShuttingDownWorkers.FindOrAdd(NewShutdown).IsValid()); // Abort worker should have set the value
+			}
+			LocalRemoteWorkers.Emplace(NewShutdown);
 		}
 	}
-	for (FCookWorkerServer* NewShutdown : NewShutdowns)
+
+	TArray<TRefCountPtr<FCookWorkerServer>, TInlineAllocator<16>> CompletedWorkers;
+	for (TRefCountPtr<FCookWorkerServer>& RemoteWorker : LocalRemoteWorkers)
 	{
-		AbortWorker(NewShutdown->GetWorkerId());
+		RemoteWorker->TickCommunication(TickThread);
+		if (RemoteWorker->IsShutdownComplete())
+		{
+			CompletedWorkers.Add(RemoteWorker);
+		}
 	}
-	for (TMap<FCookWorkerServer*, TUniquePtr<FCookWorkerServer>>::TIterator Iter(ShuttingDownWorkers); Iter; ++Iter)
+	LocalRemoteWorkers.Empty();
+
+	if (!CompletedWorkers.IsEmpty())
 	{
-		FCookWorkerServer& Worker = *Iter.Key();
-		check(Iter.Value()); // All non-yet-moved workers should have been moved by the AbortWorker calls above
-		Worker.TickFromSchedulerThread();
-		if (Worker.IsShutdownComplete())
+		FScopeLock CommunicationScopeLock(&CommunicationLock);
+		for (TRefCountPtr<FCookWorkerServer>& CompletedWorker : CompletedWorkers)
 		{
-			// Worker is now deleted, do not access
-			Iter.RemoveCurrent();
+			ShuttingDownWorkers.Remove(CompletedWorker.GetReference());
 		}
 	}
+	CompletedWorkers.Empty();
 }
 
 FString FCookDirector::GetWorkerCommandLine(FWorkerId WorkerId)
@@ -683,7 +906,7 @@ bool FDirectorConnectionInfo::TryParseCommandLine()
 	return true;
 }
 
-void FCookDirector::LoadBalance(TConstArrayView<FCookWorkerServer*> SortedWorkers, TArrayView<FPackageData*> Requests,
+void FCookDirector::LoadBalance(TConstArrayView<TRefCountPtr<FCookWorkerServer>> SortedWorkers, TArrayView<FPackageData*> Requests,
 	TMap<FPackageData*, TArray<FPackageData*>>&& RequestGraph, TArray<FWorkerId>& OutAssignments)
 {
 	TArray<FWorkerId> AllWorkers;
@@ -708,22 +931,24 @@ void FCookDirector::LoadBalance(TConstArrayView<FCookWorkerServer*> SortedWorker
 	return LoadBalanceCookBurden(AllWorkers, Requests, MoveTemp(RequestGraph), OutAssignments, bLogResults);
 }
 
-void FCookDirector::AbortWorker(FWorkerId WorkerId)
+void FCookDirector::AbortWorker(FWorkerId WorkerId, ECookDirectorThread TickThread)
 {
 	check(!WorkerId.IsLocal());
 	int32 Index = WorkerId.GetRemoteIndex();
-	TUniquePtr<FCookWorkerServer> RemoteWorker;
-	RemoteWorkers.RemoveAndCopyValue(Index, RemoteWorker);
-	if (!RemoteWorker)
+	TRefCountPtr<FCookWorkerServer> RemoteWorker;
 	{
-		return;
+		FScopeLock RemoteWorkersScopeLock(&CommunicationLock);
+		RemoteWorkers.RemoveAndCopyValue(Index, RemoteWorker);
+		if (!RemoteWorker)
+		{
+			return;
+		}
 	}
-	--CookWorkerCount;
 	TSet<FPackageData*> PackagesToReassign;
-	RemoteWorker->AbortAssignments(PackagesToReassign);
+	RemoteWorker->AbortAssignments(PackagesToReassign, TickThread);
 	if (!RemoteWorker->IsShuttingDown())
 	{
-		RemoteWorker->AbortWorker(PackagesToReassign);
+		RemoteWorker->AbortWorker(PackagesToReassign, TickThread);
 	}
 	for (FPackageData* PackageData : PackagesToReassign)
 	{
@@ -731,9 +956,44 @@ void FCookDirector::AbortWorker(FWorkerId WorkerId)
 		PackageData->SetWorkerAssignment(FWorkerId::Invalid());
 		PackageData->SendToState(UE::Cook::EPackageState::Request, ESendFlags::QueueAddAndRemove);
 	}
-	TUniquePtr<FCookWorkerServer>& Existing = ShuttingDownWorkers.FindOrAdd(RemoteWorker.Get());
-	check(!Existing); // We should not be able to abort a worker twice because we removed it from RemoteWorkers above
-	Existing = MoveTemp(RemoteWorker);
+	{
+		FScopeLock RemoteWorkersScopeLock(&CommunicationLock);
+		TRefCountPtr<FCookWorkerServer>& Existing = ShuttingDownWorkers.FindOrAdd(RemoteWorker.GetReference());
+		check(!Existing); // We should not be able to abort a worker twice because we removed it from RemoteWorkers above
+		Existing = MoveTemp(RemoteWorker);
+	}
+}
+
+void FCookDirector::ConstructReadonlyThreadVariables()
+{
+	IsCookIgnoreTimeouts(); // The global variables are read-only; call them now to initialize them
+	CommandletExecutablePath = FUnrealEdMisc::Get().GetProjectEditorBinaryPath();
+
+	InitialConfigMessage = MakeUnique<FInitialConfigMessage>();
+	const TArray<const ITargetPlatform*>& SessionPlatforms = COTFS.PlatformManager->GetSessionPlatforms();
+	TArray<ITargetPlatform*> OrderedSessionPlatforms;
+	OrderedSessionPlatforms.Reset(SessionPlatforms.Num());
+	for (const ITargetPlatform* TargetPlatform : SessionPlatforms)
+	{
+		OrderedSessionPlatforms.Add(const_cast<ITargetPlatform*>(TargetPlatform));
+	}
+	InitialConfigMessage->ReadFromLocal(COTFS, OrderedSessionPlatforms,
+		*COTFS.CookByTheBookOptions, *COTFS.CookOnTheFlyOptions, BeginCookContext);
 }
 
+const FInitialConfigMessage& FCookDirector::GetInitialConfigMessage()
+{
+	return *InitialConfigMessage;
+}
+
+FCookDirector::FLaunchInfo FCookDirector::GetLaunchInfo(FWorkerId WorkerId)
+{
+	FLaunchInfo Info;
+	Info.ShowWorkerOption = GetShowWorkerOption();
+	Info.CommandletExecutable = CommandletExecutablePath;
+	Info.WorkerCommandLine = GetWorkerCommandLine(WorkerId);
+	return Info;
+}
+
+
 }
diff --git a/Engine/Source/Editor/UnrealEd/Private/Cooker/CookDirector.h b/Engine/Source/Editor/UnrealEd/Private/Cooker/CookDirector.h
--- a/Engine/Source/Editor/UnrealEd/Private/Cooker/CookDirector.h
+++ b/Engine/Source/Editor/UnrealEd/Private/Cooker/CookDirector.h
@@ -9,21 +9,41 @@
 #include "Containers/UnrealString.h"
 #include "CookSockets.h"
 #include "CookTypes.h"
+#include "HAL/CriticalSection.h"
+#include "HAL/Event.h"
+#include "HAL/Runnable.h"
 #include "Memory/SharedBuffer.h"
 #include "Misc/Guid.h"
+#include "Misc/ScopeLock.h"
+#include "Templates/RefCounting.h"
 #include "Templates/UniquePtr.h"
 
+#include <atomic>
+
 class FCbObject;
 class FCbWriter;
+class FRunnableThread;
 class UCookOnTheFlyServer;
 namespace UE::Cook { class FCookWorkerServer; }
 namespace UE::Cook { class IMPCollector; }
+namespace UE::Cook { struct FInitialConfigMessage; }
 namespace UE::Cook { struct FPackageData; }
 namespace UE::Cook { struct FWorkerId; }
 
 namespace UE::Cook
 {
 
+/**
+ * The categories of thread that can pump the communication with CookWorkers. It can be pumped either from
+ * the cooker's scheduler thread (aka Unreal's game thread or main thread) or from a worker thread.
+ */
+enum class ECookDirectorThread : uint8
+{
+	SchedulerThread,
+	CommunicateThread,
+	Invalid,
+};
+
 /**
  * Helper for CookOnTheFlyServer that sends requests to CookWorker processes for load/save and merges
  * their replies into the local process's cook results.
@@ -31,6 +51,7 @@ namespace UE::Cook
 class FCookDirector
 {
 public:
+
 	FCookDirector(UCookOnTheFlyServer& InCOTFS);
 	~FCookDirector();
 
@@ -65,7 +86,24 @@ public:
 	/** Unegister a Collector that was registered. */
 	void Unregister(IMPCollector* Collector);
 
+	/** Data used by a CookWorkerServer to launch the remote process. */
+	struct FLaunchInfo
+	{
+		EShowWorker ShowWorkerOption;
+		FString CommandletExecutable;
+		FString WorkerCommandLine;
+	};
+	FLaunchInfo GetLaunchInfo(FWorkerId WorkerId);
+
+	/** The message CookWorkerServer sends to the remote process once it is ready to connect. */
+	const FInitialConfigMessage& GetInitialConfigMessage();
+
 private:
+	enum class ELoadBalanceAlgorithm
+	{
+		Striped,
+		CookBurden,
+	};
 	/** CookWorker connections that have not yet identified which CookWorker they are. */
 	struct FPendingConnection
 	{
@@ -82,28 +120,54 @@ private:
 		FSocket* Socket = nullptr;
 		UE::CompactBinaryTCP::FReceiveBuffer Buffer;
 	};
+	/** Struct that implements the FRunnable interface and forwards it to to named functions on this FCookDirector. */
+	struct FRunnableShunt : public FRunnable
+	{
+		FRunnableShunt(FCookDirector& InDirector) : Director(InDirector) {}
+		virtual uint32 Run() override;
+		virtual void Stop() override;
+		FCookDirector& Director;
+	};
+
+private:
 	/** Helper for constructor parsing. */
 	void ParseConfig();
-	/** Sets CookWorkerCount for a new session */
-	void ResetCookWorkerCount();
 	/** Initialization helper: create the listen socket. */
 	bool TryCreateWorkerConnectSocket();
-	/** Initialization helper: add the local Server for a remote worker, worker process not yet created. */
+	/**
+	 * Construct CookWorkerServers and communication thread if not yet constructed.
+	 * The CookWorkerServers are constructed to Uninitialized; the worker process is created later.
+	 */
 	void InitializeWorkers();
+	/** Copy to snapshot variables the data required on the communication thread that can only be read from the scheduler thread. */
+	void ConstructReadonlyThreadVariables();
+	/** Construct CookWorkerServers if necessary to replace workers that have crashed. */
+	void RecreateWorkers();
 	/** Reduce memory settings, cpusettings, and anything else that needs to be shared with CookWorkers. */
 	void ActivateMachineResourceReduction();
+	/** Start the communication thread if not already started. */
+	void LaunchCommunicationThread();
+	/** Signal the communication thread to stop, wait for it to finish, and deallocate it. */
+	void StopCommunicationThread();
+	/** Entry point for the communication thread. */
+	uint32 RunCommunicationThread();
+	/**
+	 * Execute a single frame of communication with CookWorkers: send/receive to all CookWorkers,
+	 * including connecting, ongoing communication, and shutting down.
+	 */
+	void TickCommunication(ECookDirectorThread TickThread);
 	/** Tick helper: tick any workers that have not yet finished initialization. */
-	void TickWorkerConnects();
+	void TickWorkerConnects(ECookDirectorThread TickThread);
 	/** Tick helper: tick any workers that are shutting down. */
-	void TickWorkerShutdowns();
+	void TickWorkerShutdowns(ECookDirectorThread TickThread);
 	/** Get the commandline to launch a worker process with. */
 	FString GetWorkerCommandLine(FWorkerId WorkerId);
 	/** Calls the configured LoadBalanceAlgorithm. Input Requests have been sorted by leaf to root load order. */
-	void LoadBalance(TConstArrayView<FCookWorkerServer*> Workers, TArrayView<FPackageData*> Requests,
+	void LoadBalance(TConstArrayView<TRefCountPtr<FCookWorkerServer>> Workers, TArrayView<FPackageData*> Requests,
 		TMap<FPackageData*, TArray<FPackageData*>>&& RequestGraph, TArray<FWorkerId>& OutAssignments);
 
 	/** Move the given worker from active workers to the list of workers shutting down. */
-	void AbortWorker(FWorkerId WorkerId);
+	void AbortWorker(FWorkerId WorkerId, ECookDirectorThread TickThread);
 	/**
 	 * Periodically update whether (1) local server is done and (2) no results from cookworkers have come in.
 	 * Send warning when it goes on too long.
@@ -111,31 +175,43 @@ private:
 	void SetWorkersStalled(bool bInWorkersStalled);
 
 private:
-	enum class ELoadBalanceAlgorithm
-	{
-		Striped,
-		CookBurden,
-	};
-	FBeginCookContextForWorker BeginCookContext;
-	TMap<int32, TUniquePtr<FCookWorkerServer>> RemoteWorkers;
-	TMap<FCookWorkerServer*, TUniquePtr<FCookWorkerServer>> ShuttingDownWorkers;
-	TMap<FGuid, TRefCountPtr<IMPCollector>> MessageHandlers;
+	// Synchronization primitives that can be used from any thread
+	FCriticalSection CommunicationLock;
+	FEventRef ShutdownEvent {EEventMode::ManualReset};
+
+	// Data only accessible from the SchedulerThread
+	FRunnableShunt RunnableShunt;
+	FRunnableThread* CommunicationThread = nullptr;
 	TArray<FPendingConnection> PendingConnections;
-	FString WorkerConnectAuthority;
 	UCookOnTheFlyServer& COTFS;
-	FSocket* WorkerConnectSocket = nullptr;
 	double WorkersStalledStartTimeSeconds = 0.;
 	double WorkersStalledWarnTimeSeconds = 0.;
+	bool bWorkersInitialized = false;
+	bool bHasReducedMachineResources = false;
+	bool bIsFirstAssignment = true;
+	bool bCookCompleteSent = false;
+	bool bWorkersStalled = false;
+
+	// Data that is read-only while the CommunicationThread is active and is readable from any thread
+	FBeginCookContextForWorker BeginCookContext;
+	TMap<FGuid, TRefCountPtr<IMPCollector>> MessageHandlers;
+	TUniquePtr<FInitialConfigMessage> InitialConfigMessage;
+	FString WorkerConnectAuthority;
+	FString CommandletExecutablePath;
 	int32 RequestedCookWorkerCount = 0;
-	int32 CookWorkerCount = 0;
 	int32 WorkerConnectPort = 0;
 	int32 CoreLimit = 0;
-	bool bWorkersStalled = false;
-	bool bIsFirstAssignment = true;
-	bool bHasReducedMachineResources = false;
 	EShowWorker ShowWorkerOption = EShowWorker::CombinedLogs;
 	ELoadBalanceAlgorithm LoadBalanceAlgorithm = ELoadBalanceAlgorithm::CookBurden;
 
+	// Data only accessible from the CommunicationThread (or if the CommunicationThread is inactive)
+	FSocket* WorkerConnectSocket = nullptr;
+
+	// Data shared between SchedulerThread and CommunicationThread that can only be accessed inside CommunicationLock
+	TMap<int32, TRefCountPtr<FCookWorkerServer>> RemoteWorkers;
+	TMap<FCookWorkerServer*, TRefCountPtr<FCookWorkerServer>> ShuttingDownWorkers;
+	bool bWorkersActive = false;
+
 	friend class UE::Cook::FCookWorkerServer;
 };
 
diff --git a/Engine/Source/Editor/UnrealEd/Private/Cooker/CookWorkerServer.cpp b/Engine/Source/Editor/UnrealEd/Private/Cooker/CookWorkerServer.cpp
--- a/Engine/Source/Editor/UnrealEd/Private/Cooker/CookWorkerServer.cpp
+++ b/Engine/Source/Editor/UnrealEd/Private/Cooker/CookWorkerServer.cpp
@@ -13,6 +13,7 @@
 #include "Interfaces/ITargetPlatformManagerModule.h"
 #include "Math/NumericLimits.h"
 #include "Misc/AssertionMacros.h"
+#include "Misc/ScopeLock.h"
 #include "PackageResultsMessage.h"
 #include "PackageTracker.h"
 #include "UnrealEdMisc.h"
@@ -29,14 +30,10 @@ FCookWorkerServer::FCookWorkerServer(FCookDirector& InDirector, FWorkerId InWork
 
 FCookWorkerServer::~FCookWorkerServer()
 {
-	PendingPackages.Append(PackagesToAssign);
-	PackagesToAssign.Empty();
-	for (FPackageData* PackageData : PendingPackages)
-	{
-		check(PackageData->IsInProgress()); // Packages that were assigned to a worker should be in the AssignedToWorker state
-		PackageData->SetWorkerAssignment(FWorkerId::Invalid());
-		PackageData->SendToState(UE::Cook::EPackageState::Request, ESendFlags::QueueAddAndRemove);
-	}
+	FCommunicationScopeLock ScopeLock(this, ECookDirectorThread::CommunicateThread, ETickAction::Queue);
+
+	checkf(PendingPackages.IsEmpty() && PackagesToAssign.IsEmpty(),
+		TEXT("CookWorkerServer still has assigned packages when it is being destroyed; we will leak them and block the cook."));
 
 	if (ConnectStatus == EConnectStatus::Connected || ConnectStatus == EConnectStatus::PumpingCookComplete || ConnectStatus == EConnectStatus::WaitForDisconnect)
 	{
@@ -66,12 +63,19 @@ void FCookWorkerServer::ShutdownRemoteProcess()
 	DetachFromRemoteProcess();
 }
 
-void FCookWorkerServer::AppendAssignments(TArrayView<FPackageData*> Assignments)
+void FCookWorkerServer::AppendAssignments(TArrayView<FPackageData*> Assignments, ECookDirectorThread TickThread)
 {
+	FCommunicationScopeLock ScopeLock(this, TickThread, ETickAction::Queue);
 	PackagesToAssign.Append(Assignments);
 }
 
-void FCookWorkerServer::AbortAssignments(TSet<FPackageData*>& OutPendingPackages)
+void FCookWorkerServer::AbortAssignments(TSet<FPackageData*>& OutPendingPackages, ECookDirectorThread TickThread)
+{
+	FCommunicationScopeLock ScopeLock(this, TickThread, ETickAction::Queue);
+	AbortAssignmentsInLock(OutPendingPackages);
+}
+
+void FCookWorkerServer::AbortAssignmentsInLock(TSet<FPackageData*>& OutPendingPackages)
 {
 	if (PendingPackages.Num())
 	{
@@ -92,8 +96,10 @@ void FCookWorkerServer::AbortAssignments(TSet<FPackageData*>& OutPendingPackages
 	PackagesToAssign.Empty();
 }
 
-void FCookWorkerServer::AbortAssignment(FPackageData& PackageData)
+void FCookWorkerServer::AbortAssignment(FPackageData& PackageData, ECookDirectorThread TickThread)
 {
+	FCommunicationScopeLock ScopeLock(this, TickThread, ETickAction::Queue);
+
 	if (PendingPackages.Remove(&PackageData))
 	{
 		if (ConnectStatus == EConnectStatus::Connected)
@@ -107,9 +113,11 @@ void FCookWorkerServer::AbortAssignment(FPackageData& PackageData)
 	PackagesToAssign.Remove(&PackageData);
 }
 
-void FCookWorkerServer::AbortWorker(TSet<FPackageData*>& OutPendingPackages)
+void FCookWorkerServer::AbortWorker(TSet<FPackageData*>& OutPendingPackages, ECookDirectorThread TickThread)
 {
-	AbortAssignments(OutPendingPackages);
+	FCommunicationScopeLock ScopeLock(this, TickThread, ETickAction::Tick);
+
+	AbortAssignmentsInLock(OutPendingPackages);
 	switch (ConnectStatus)
 	{
 	case EConnectStatus::Uninitialized: // Fall through
@@ -155,21 +163,38 @@ void FCookWorkerServer::SendToState(EConnectStatus TargetStatus)
 
 bool FCookWorkerServer::IsShuttingDown() const
 {
+	FScopeLock CommunicationScopeLock(&CommunicationLock);
 	return ConnectStatus == EConnectStatus::PumpingCookComplete || ConnectStatus == EConnectStatus::WaitForDisconnect || ConnectStatus == EConnectStatus::LostConnection;
 }
 
 bool FCookWorkerServer::IsFlushingBeforeShutdown() const
 {
+	FScopeLock CommunicationScopeLock(&CommunicationLock);
 	return ConnectStatus == EConnectStatus::PumpingCookComplete;
 }
 
 bool FCookWorkerServer::IsShutdownComplete() const
 {
+	FScopeLock CommunicationScopeLock(&CommunicationLock);
 	return ConnectStatus == EConnectStatus::LostConnection;
 }
 
-bool FCookWorkerServer::TryHandleConnectMessage(FWorkerConnectMessage& Message, FSocket* InSocket, TArray<UE::CompactBinaryTCP::FMarshalledMessage>&& OtherPacketMessages)
+bool FCookWorkerServer::HasAssignments() const
 {
+	FScopeLock CommunicationScopeLock(&CommunicationLock);
+	return !PackagesToAssign.IsEmpty() || !PendingPackages.IsEmpty();
+}
+
+bool FCookWorkerServer::HasMessages() const
+{
+	FScopeLock CommunicationScopeLock(&CommunicationLock);
+	return !ReceiveMessages.IsEmpty();
+}
+
+bool FCookWorkerServer::TryHandleConnectMessage(FWorkerConnectMessage& Message, FSocket* InSocket, TArray<UE::CompactBinaryTCP::FMarshalledMessage>&& OtherPacketMessages, ECookDirectorThread TickThread)
+{
+	FCommunicationScopeLock ScopeLock(this, TickThread, ETickAction::Tick);
+
 	if (ConnectStatus != EConnectStatus::WaitForConnect)
 	{
 		return false;
@@ -178,22 +203,21 @@ bool FCookWorkerServer::TryHandleConnectMessage(FWorkerConnectMessage& Message,
 	Socket = InSocket;
 
 	SendToState(EConnectStatus::Connected);
-	HandleReceiveMessages(MoveTemp(OtherPacketMessages));
-	FInitialConfigMessage ConfigMessage;
-	const TArray<const ITargetPlatform*>& SessionPlatforms = COTFS.PlatformManager->GetSessionPlatforms();
-	OrderedSessionPlatforms.Reset(SessionPlatforms.Num());
-	for (const ITargetPlatform* TargetPlatform : SessionPlatforms)
+	for (UE::CompactBinaryTCP::FMarshalledMessage& OtherMessage : OtherPacketMessages)
 	{
-		OrderedSessionPlatforms.Add(const_cast<ITargetPlatform*>(TargetPlatform));
+		ReceiveMessages.Add(MoveTemp(OtherMessage));
 	}
-	ConfigMessage.ReadFromLocal(COTFS, OrderedSessionPlatforms,
-		*COTFS.CookByTheBookOptions, *COTFS.CookOnTheFlyOptions, Director.BeginCookContext);
-	SendMessage(ConfigMessage);
+	HandleReceiveMessagesInternal();
+	const FInitialConfigMessage& InitialConfigMessage = Director.GetInitialConfigMessage();
+	OrderedSessionPlatforms = InitialConfigMessage.GetOrderedSessionPlatforms();
+	SendMessage(InitialConfigMessage);
 	return true;
 }
 
-void FCookWorkerServer::TickFromSchedulerThread()
+void FCookWorkerServer::TickCommunication(ECookDirectorThread TickThread)
 {
+	FCommunicationScopeLock ScopeLock(this, TickThread, ETickAction::Tick);
+
 	for (;;)
 	{
 		switch (ConnectStatus)
@@ -251,8 +275,10 @@ void FCookWorkerServer::TickFromSchedulerThread()
 	}
 }
 
-void FCookWorkerServer::SignalCookComplete()
+void FCookWorkerServer::SignalCookComplete(ECookDirectorThread TickThread)
 {
+	FCommunicationScopeLock ScopeLock(this, TickThread, ETickAction::Tick);
+
 	switch (ConnectStatus)
 	{
 	case EConnectStatus::Uninitialized: // Fall through
@@ -270,18 +296,17 @@ void FCookWorkerServer::SignalCookComplete()
 
 void FCookWorkerServer::LaunchProcess()
 {
-	bool bShowCookWorkers = Director.GetShowWorkerOption() == FCookDirector::EShowWorker::SeparateWindows;
+	FCookDirector::FLaunchInfo LaunchInfo = Director.GetLaunchInfo(WorkerId);
+	bool bShowCookWorkers = LaunchInfo.ShowWorkerOption == FCookDirector::EShowWorker::SeparateWindows;
 
-	FString CommandletExecutable = FUnrealEdMisc::Get().GetProjectEditorBinaryPath();
-	FString CommandLine = Director.GetWorkerCommandLine(WorkerId);
-	CookWorkerHandle = FPlatformProcess::CreateProc(*CommandletExecutable, *CommandLine,
+	CookWorkerHandle = FPlatformProcess::CreateProc(*LaunchInfo.CommandletExecutable, *LaunchInfo.WorkerCommandLine,
 		true /* bLaunchDetached */, !bShowCookWorkers /* bLaunchHidden */, !bShowCookWorkers /* bLaunchReallyHidden */,
-		&CookWorkerProcessId, 0 /* PriorityModifier */, *FPaths::GetPath(CommandletExecutable),
+		&CookWorkerProcessId, 0 /* PriorityModifier */, *FPaths::GetPath(LaunchInfo.CommandletExecutable),
 		nullptr /* PipeWriteChild */);
 	if (CookWorkerHandle.IsValid())
 	{
 		UE_LOG(LogCook, Display, TEXT("CookWorkerServer %d launched CookWorker as PID %u with commandline \"%s\"."),
-			WorkerId.GetRemoteIndex(), CookWorkerProcessId, *CommandLine);
+			WorkerId.GetRemoteIndex(), CookWorkerProcessId, *LaunchInfo.WorkerCommandLine);
 		SendToState(EConnectStatus::WaitForConnect);
 	}
 	else
@@ -399,24 +424,45 @@ void FCookWorkerServer::PumpReceiveMessages()
 		bTerminateImmediately = true;
 		return;
 	}
-	HandleReceiveMessages(MoveTemp(Messages));
+	for (FMarshalledMessage& Message : Messages)
+	{
+		ReceiveMessages.Add(MoveTemp(Message));
+	}
+	HandleReceiveMessagesInternal();
 }
 
-void FCookWorkerServer::HandleReceiveMessages(TArray<UE::CompactBinaryTCP::FMarshalledMessage>&& Messages)
+void FCookWorkerServer::HandleReceiveMessages(ECookDirectorThread TickThread)
 {
-	for (UE::CompactBinaryTCP::FMarshalledMessage& Message : Messages)
+	FCommunicationScopeLock ScopeLock(this, TickThread, ETickAction::Queue);
+	HandleReceiveMessagesInternal();
+}
+
+void FCookWorkerServer::HandleReceiveMessagesInternal()
+{
+	while (!ReceiveMessages.IsEmpty())
 	{
-		if (Message.MessageType == FAbortWorkerMessage::MessageType)
+		UE::CompactBinaryTCP::FMarshalledMessage& PeekMessage = ReceiveMessages[0];
+
+		if (PeekMessage.MessageType == FAbortWorkerMessage::MessageType)
 		{
+			UE::CompactBinaryTCP::FMarshalledMessage Message = ReceiveMessages.PopFrontValue();
 			UE_CLOG(ConnectStatus != EConnectStatus::PumpingCookComplete && ConnectStatus != EConnectStatus::WaitForDisconnect,
 				LogCook, Error, TEXT("CookWorkerServer %d remote process shut down unexpectedly. Assigned packages will be returned to the director."),
 				WorkerId.GetRemoteIndex());
 
 			SendMessage(FAbortWorkerMessage(FAbortWorkerMessage::AbortAcknowledge));
 			SendToState(EConnectStatus::WaitForDisconnect);
+			ReceiveMessages.Reset();
+			break;
+		}
+
+		if (TickState.TickThread != ECookDirectorThread::SchedulerThread)
+		{
 			break;
 		}
-		else if (Message.MessageType == FPackageResultsMessage::MessageType)
+
+		UE::CompactBinaryTCP::FMarshalledMessage Message = ReceiveMessages.PopFrontValue();
+		if (Message.MessageType == FPackageResultsMessage::MessageType)
 		{
 			FPackageResultsMessage ResultsMessage;
 			if (!ResultsMessage.TryRead(MoveTemp(Message.Object)))
@@ -464,6 +510,8 @@ void FCookWorkerServer::HandleReceiveMessages(TArray<UE::CompactBinaryTCP::FMars
 
 void FCookWorkerServer::HandleReceivedPackagePlatformMessages(FPackageData& PackageData, const ITargetPlatform* TargetPlatform, TArray<UE::CompactBinaryTCP::FMarshalledMessage>&& Messages)
 {
+	check(TickState.TickThread == ECookDirectorThread::SchedulerThread);
+
 	for (UE::CompactBinaryTCP::FMarshalledMessage& Message : Messages)
 	{
 		if (Message.MessageType == FAssetRegistryPackageMessage::MessageType)
@@ -485,11 +533,21 @@ void FCookWorkerServer::HandleReceivedPackagePlatformMessages(FPackageData& Pack
 
 void FCookWorkerServer::SendMessage(const UE::CompactBinaryTCP::IMessage& Message)
 {
-	UE::CompactBinaryTCP::TryWritePacket(Socket, SendBuffer, Message);
+	if (TickState.TickAction == ETickAction::Tick)
+	{
+		UE::CompactBinaryTCP::TryWritePacket(Socket, SendBuffer, Message);
+	}
+	else
+	{
+		check(TickState.TickAction == ETickAction::Queue);
+		UE::CompactBinaryTCP::QueueMessage(SendBuffer, Message);
+	}
 }
 
 void FCookWorkerServer::RecordResults(FPackageResultsMessage& Message)
 {
+	check(TickState.TickThread == ECookDirectorThread::SchedulerThread);
+
 	for (FPackageRemoteResult& Result : Message.Results)
 	{
 		FPackageData* PackageData = COTFS.PackageDatas->FindPackageDataByPackageName(Result.PackageName);
@@ -547,6 +605,8 @@ void FCookWorkerServer::LogInvalidMessage(const TCHAR* MessageTypeName)
 
 void FCookWorkerServer::AddDiscoveredPackage(FDiscoveredPackage&& DiscoveredPackage)
 {
+	check(TickState.TickThread == ECookDirectorThread::SchedulerThread);
+
 	FPackageData& PackageData = COTFS.PackageDatas->FindOrAddPackageData(DiscoveredPackage.PackageName,
 		DiscoveredPackage.NormalizedFileName);
 	if (PackageData.IsInProgress() || PackageData.HasAnyCookedPlatform())
@@ -563,6 +623,30 @@ void FCookWorkerServer::AddDiscoveredPackage(FDiscoveredPackage&& DiscoveredPack
 	COTFS.QueueDiscoveredPackageData(PackageData, MoveTemp(DiscoveredPackage.Instigator));
 }
 
+FCookWorkerServer::FTickState::FTickState()
+{
+	TickThread = ECookDirectorThread::Invalid;
+	TickAction = ETickAction::Invalid;
+}
+
+FCookWorkerServer::FCommunicationScopeLock::FCommunicationScopeLock(FCookWorkerServer* InServer, ECookDirectorThread TickThread, ETickAction TickAction)
+	: ScopeLock(&InServer->CommunicationLock)
+	, Server(*InServer)
+{
+	check(TickThread != ECookDirectorThread::Invalid);
+	check(TickAction != ETickAction::Invalid);
+	check(Server.TickState.TickThread == ECookDirectorThread::Invalid);
+	Server.TickState.TickThread = TickThread;
+	Server.TickState.TickAction = TickAction;
+}
+
+FCookWorkerServer::FCommunicationScopeLock::~FCommunicationScopeLock()
+{
+	check(Server.TickState.TickThread != ECookDirectorThread::Invalid);
+	Server.TickState.TickThread = ECookDirectorThread::Invalid;
+	Server.TickState.TickAction = ETickAction::Invalid;
+}
+
 FAssignPackagesMessage::FAssignPackagesMessage(TArray<FAssignPackageData>&& InPackageDatas)
 	: PackageDatas(MoveTemp(InPackageDatas))
 {
diff --git a/Engine/Source/Editor/UnrealEd/Private/Cooker/CookWorkerServer.h b/Engine/Source/Editor/UnrealEd/Private/Cooker/CookWorkerServer.h
--- a/Engine/Source/Editor/UnrealEd/Private/Cooker/CookWorkerServer.h
+++ b/Engine/Source/Editor/UnrealEd/Private/Cooker/CookWorkerServer.h
@@ -5,12 +5,15 @@
 #include "CompactBinaryTCP.h"
 #include "Containers/Array.h"
 #include "Containers/ArrayView.h"
+#include "Containers/RingBuffer.h"
 #include "Containers/Set.h"
 #include "CookOnTheSide/CookOnTheFlyServer.h"
 #include "CookPackageData.h"
 #include "CookSockets.h"
 #include "CookTypes.h"
+#include "HAL/CriticalSection.h"
 #include "Misc/Guid.h"
+#include "Templates/RefCounting.h"
 
 class FSocket;
 class ITargetPlatform;
@@ -18,6 +21,7 @@ struct FProcHandle;
 namespace UE::CompactBinaryTCP { struct FMarshalledMessage; }
 namespace UE::Cook { class FCookDirector; }
 namespace UE::Cook { struct FDiscoveredPackage; }
+namespace UE::Cook { enum class ECookDirectorThread : uint8; };
 namespace UE::Cook { struct FPackageData; }
 namespace UE::Cook { struct FPackageResultsMessage; }
 namespace UE::Cook { struct FWorkerConnectMessage; }
@@ -26,7 +30,7 @@ namespace UE::Cook
 {
 
 /** Class in a Director process that communicates over a Socket with FCookWorkerClient in a CookWorker process. */
-class FCookWorkerServer
+class FCookWorkerServer : public FThreadSafeRefCountedObject
 {
 public:
 	FCookWorkerServer(FCookDirector& InDirector, FWorkerId InWorkerId);
@@ -35,22 +39,28 @@ public:
 	FWorkerId GetWorkerId() const { return WorkerId; }
 
 	/** Add the given assignments for the CookWorker. They will be sent during Tick */
-	void AppendAssignments(TArrayView<FPackageData*> Assignments);
+	void AppendAssignments(TArrayView<FPackageData*> Assignments, ECookDirectorThread TickThread);
 	/** Remove assignment of the package from local state and from the connected Client. */
-	void AbortAssignment(FPackageData& PackageData);
+	void AbortAssignment(FPackageData& PackageData, ECookDirectorThread TickThread);
 	/**
 	 * Remove assignment of the all assigned packages from local state and from the connected Client.
 	 * Report all packages that were unassigned.
 	 */
-	void AbortAssignments(TSet<FPackageData*>& OutPendingPackages);
+	void AbortAssignments(TSet<FPackageData*>& OutPendingPackages, ECookDirectorThread TickThread);
 	/** AbortAssignments and tell the connected Client to gracefully terminate. Report all packages that were unassigned. */
-	void AbortWorker(TSet<FPackageData*>& OutPendingPackages);
+	void AbortWorker(TSet<FPackageData*>& OutPendingPackages, ECookDirectorThread TickThread);
 	/** Take over the Socket for a CookWorker that has just connected. */
-	bool TryHandleConnectMessage(FWorkerConnectMessage& Message, FSocket* InSocket, TArray<UE::CompactBinaryTCP::FMarshalledMessage>&& OtherPacketMessages);
+	bool TryHandleConnectMessage(FWorkerConnectMessage& Message, FSocket* InSocket, TArray<UE::CompactBinaryTCP::FMarshalledMessage>&& OtherPacketMessages, ECookDirectorThread TickThread);
+
 	/** Periodic Tick function to send and receive messages to the Client. */
-	void TickFromSchedulerThread();
+	void TickCommunication(ECookDirectorThread TickThread);
 	/** Called when the COTFS Server has detected all packages are complete. Tell the CookWorker to flush messages and exit. */
-	void SignalCookComplete();
+	void SignalCookComplete(ECookDirectorThread TickThread);
+	/**
+	 * Execute the respond for all messages that have been received and that can be executed from the given thread. 
+	 * This is the entry point for the scheduler thread to pick up messages that the CommunicationThread has queued.
+	 */
+	void HandleReceiveMessages(ECookDirectorThread TickThread);
 
 	/** Is this either shutting down or completed shutdown of its remote Client? */
 	bool IsShuttingDown() const;
@@ -58,6 +68,10 @@ public:
 	bool IsFlushingBeforeShutdown() const;
 	/** Is this not yet or no longer connected to a remote Client? */
 	bool IsShutdownComplete() const;
+	/** Does this Server have any package assignments the remmote CookWorker is supposed to save but hasn't yet? */
+	bool HasAssignments() const;
+	/** Does this Server have any ReceivedMessages that need to be processed by the Scheduler thread? */
+	bool HasMessages() const;
 
 private:
 	enum class EConnectStatus
@@ -69,6 +83,31 @@ private:
 		WaitForDisconnect,
 		LostConnection,
 	};
+	enum class ETickAction
+	{
+		Tick,
+		Queue,
+		Invalid,
+	};
+	/**
+	 * Stores from which thread the public function on *this was called, and whether that public function is a pumping
+	 * function that should send/receive network messages or merely an accessor function that should send-to-queue or
+	 * read-from-queued network messages.
+	 */
+	struct FTickState
+	{
+		FTickState();
+		ECookDirectorThread TickThread;
+		ETickAction TickAction;
+	};
+	/** An RAII structure that enters the lock and sets the TickState information required by many functions. */
+	struct FCommunicationScopeLock
+	{
+		FScopeLock ScopeLock;
+		FCookWorkerServer& Server;
+		FCommunicationScopeLock(FCookWorkerServer* InServer, ECookDirectorThread TickThread, ETickAction TickAction);
+		~FCommunicationScopeLock();
+	};
 
 private:
 	/** Helper for PumpConnect, launch the remote Client process. */
@@ -83,6 +122,8 @@ private:
 	void SendPendingPackages();
 	/** Helper for Tick, pump receive messages from a connected Client. */
 	void PumpReceiveMessages();
+	/** The main implementation of AbortAssignments, only callable from inside the lock. */
+	void AbortAssignmentsInLock(TSet<FPackageData*>& OutPendingPackages);
 	/** Send the message immediately to the Socket. If cannot complete immediately, it will be finished during Tick. */
 	void SendMessage(const UE::CompactBinaryTCP::IMessage& Message);
 	/** Send this into the given state. Update any state-dependent variables. */
@@ -91,23 +132,30 @@ private:
 	void DetachFromRemoteProcess();
 	/** Kill the Client process (non-graceful termination), and close the connection resources. */
 	void ShutdownRemoteProcess();
+	/** The main implementation of HandleReceiveMessages, only callable from inside the lock. */
+	void HandleReceiveMessagesInternal();
 	/** Helper for PumpReceiveMessages: dispatch the messages received from the socket. */
-	void HandleReceiveMessages(TArray<UE::CompactBinaryTCP::FMarshalledMessage>&& Messages);
 	void HandleReceivedPackagePlatformMessages(FPackageData& PackageData, const ITargetPlatform* TargetPlatform, TArray<UE::CompactBinaryTCP::FMarshalledMessage>&& Messages);
 	/** Add results from the client to the local CookOnTheFlyServer. */
 	void RecordResults(FPackageResultsMessage& Message);
 	void LogInvalidMessage(const TCHAR* MessageTypeName);
 	void AddDiscoveredPackage(FDiscoveredPackage&& DiscoveredPackage);
 
+	// Lock guarding access to all data on *this
+	mutable FCriticalSection CommunicationLock;
+
+	// All data can only be read or written while the CommunicationLock is entered.
 	TArray<FPackageData*> PackagesToAssign;
 	TSet<FPackageData*> PendingPackages;
 	TArray<ITargetPlatform*> OrderedSessionPlatforms;
 	UE::CompactBinaryTCP::FSendBuffer SendBuffer;
 	UE::CompactBinaryTCP::FReceiveBuffer ReceiveBuffer;
+	TRingBuffer<UE::CompactBinaryTCP::FMarshalledMessage> ReceiveMessages;
 	FCookDirector& Director;
 	UCookOnTheFlyServer& COTFS;
 	FSocket* Socket = nullptr;
 	FProcHandle CookWorkerHandle;
+	FTickState TickState;
 	uint32 CookWorkerProcessId = 0;
 	double ConnectStartTimeSeconds = 0.;
 	double ConnectTestStartTimeSeconds = 0.;
@@ -201,8 +249,8 @@ public:
 	FBeginCookConfigSettings&& ConsumeBeginCookConfigSettings() { return MoveTemp(BeginCookSettings); }
 	FCookByTheBookOptions&& ConsumeCookByTheBookOptions() { return MoveTemp(CookByTheBookOptions); }
 	FCookOnTheFlyOptions&& ConsumeCookOnTheFlyOptions() { return MoveTemp(CookOnTheFlyOptions); }
-	const FBeginCookContextForWorker& GetBeginCookContext() { return BeginCookContext; }
-	const TArray<ITargetPlatform*>& GetOrderedSessionPlatforms() { return OrderedSessionPlatforms; }
+	const FBeginCookContextForWorker& GetBeginCookContext() const { return BeginCookContext; }
+	const TArray<ITargetPlatform*>& GetOrderedSessionPlatforms() const { return OrderedSessionPlatforms; }
 	bool IsZenStore() const { return bZenStore; }
 
 public:
