Fix for dragging a texture on a cube when a material already exists, and some protection against redirects in similar cases.

diff --git a/Engine/Source/Editor/UnrealEd/Private/LevelEditorViewport.cpp b/Engine/Source/Editor/UnrealEd/Private/LevelEditorViewport.cpp
--- a/Engine/Source/Editor/UnrealEd/Private/LevelEditorViewport.cpp
+++ b/Engine/Source/Editor/UnrealEd/Private/LevelEditorViewport.cpp
@@ -571,8 +571,7 @@ UObject* FLevelEditorViewportClient::GetOrCreateMaterialFromTexture(UTexture* Un
 
 	const FString MaterialFullName = bCreateMaterialInstance ? (TEXT("MI_") + TextureShortName) : (TextureShortName + TEXT("_Mat"));
 	FString NewPackageName = FPaths::Combine(NewPackageFolder, MaterialFullName);
-	NewPackageName = UPackageTools::SanitizePackageName(NewPackageName);
-	UPackage* Package = CreatePackage(*NewPackageName);
+	NewPackageName = UPackageTools::SanitizePackageName(NewPackageName);	
 
 	// See if the material asset already exists with the expected name, if it does, just return
 	// an instance of it.
@@ -580,7 +579,8 @@ UObject* FLevelEditorViewportClient::GetOrCreateMaterialFromTexture(UTexture* Un
 	if (AssetRegistry.GetAssetsByPackageName(*NewPackageName, OutAssetData) && (OutAssetData.Num() > 0))
 	{
 		UObject* FoundAsset = OutAssetData[0].GetAsset();
-		if (FoundAsset->IsA(UMaterialInterface::StaticClass()))
+		if (FoundAsset != nullptr && // Due to redirects we might have an asset by this name that actually doesn't exist under that name anymore.
+			FoundAsset->IsA(UMaterialInterface::StaticClass()))
 		{
 			return FoundAsset;
 		}
@@ -588,6 +588,8 @@ UObject* FLevelEditorViewportClient::GetOrCreateMaterialFromTexture(UTexture* Un
 		return nullptr;
 	}
 
+	UPackage* Package = CreatePackage(*NewPackageName);
+
 	// Variations for Base Maps.
 	TArray<FString> BaseSuffixes;
 	BaseSuffixes.Add("_D");
diff --git a/Engine/Source/Runtime/CoreUObject/Public/AssetRegistry/AssetData.h b/Engine/Source/Runtime/CoreUObject/Public/AssetRegistry/AssetData.h
--- a/Engine/Source/Runtime/CoreUObject/Public/AssetRegistry/AssetData.h
+++ b/Engine/Source/Runtime/CoreUObject/Public/AssetRegistry/AssetData.h
@@ -500,6 +500,8 @@ PRAGMA_ENABLE_DEPRECATION_WARNINGS
 	/** 
 	 * Returns the asset UObject if it is loaded or loads the asset if it is unloaded then returns the result
 	 * 
+	 * Note this can return nullptr even if it looks like the asset exists because it doesn't handle redirects.
+	 * 
 	 * @param bLoad (optional) loads the asset if it is unloaded.
 	 * @param LoadTags (optional) allows passing specific tags to the linker when loading the asset (@see ULevel::LoadAllExternalObjectsTag for an example usage)
 	 */
@@ -542,6 +544,8 @@ PRAGMA_ENABLE_DEPRECATION_WARNINGS
 	/** 
 	 * Returns the asset UObject if it is loaded or loads the asset if it is unloaded then returns the result 
 	 * 
+	 * Note this can return nullptr even if it looks like the asset exists because it doesn't handle redirects.
+	 * 
 	 * @param LoadTags (optional) allows passing specific tags to the linker when loading the asset (@see ULevel::LoadAllExternalObjectsTag for an example usage)
 	 */
 	UObject* GetAsset(TSet<FName> LoadTags = {}) const
