CookMultiprocess: Add replication of log messages from the CookWorkers.
#rnx

diff --git a/Engine/Source/Editor/UnrealEd/Private/CookOnTheFlyServer.cpp b/Engine/Source/Editor/UnrealEd/Private/CookOnTheFlyServer.cpp
--- a/Engine/Source/Editor/UnrealEd/Private/CookOnTheFlyServer.cpp
+++ b/Engine/Source/Editor/UnrealEd/Private/CookOnTheFlyServer.cpp
@@ -7840,8 +7840,8 @@ namespace UE::Cook
 class FShaderLibraryCollector : public IMPCollector
 {
 public:
-	virtual FGuid GetMessageType() override { return MessageType; }
-	virtual const TCHAR* GetDebugName() override { return TEXT("FShaderLibraryCollector"); }
+	virtual FGuid GetMessageType() const override { return MessageType; }
+	virtual const TCHAR* GetDebugName() const override { return TEXT("FShaderLibraryCollector"); }
 
 	virtual void ClientTick(FClientContext& Context) override;
 	virtual void ReceiveMessage(FServerContext& Context, FCbObjectView Message) override;
@@ -9767,8 +9767,8 @@ namespace UE::Cook
 class FEDLMPCollector : public IMPCollector
 {
 public:
-	virtual FGuid GetMessageType() override { return MessageType; }
-	virtual const TCHAR* GetDebugName() override { return TEXT("FEDLMPCollector"); }
+	virtual FGuid GetMessageType() const override { return MessageType; }
+	virtual const TCHAR* GetDebugName() const override { return TEXT("FEDLMPCollector"); }
 
 	virtual void ClientTick(FClientContext& Context) override;
 	virtual void ReceiveMessage(FServerContext& Context, FCbObjectView Message) override;
diff --git a/Engine/Source/Editor/UnrealEd/Private/Cooker/CookDirector.cpp b/Engine/Source/Editor/UnrealEd/Private/Cooker/CookDirector.cpp
--- a/Engine/Source/Editor/UnrealEd/Private/Cooker/CookDirector.cpp
+++ b/Engine/Source/Editor/UnrealEd/Private/Cooker/CookDirector.cpp
@@ -47,6 +47,8 @@ FCookDirector::FCookDirector(UCookOnTheFlyServer& InCOTFS)
 	{
 		UE_LOG(LogCook, Display, TEXT("CookMultiprocess is enabled with %d CookWorker processes."), RequestedCookWorkerCount);
 	}
+
+	Register(new FLogMessagesMessageHandler());
 }
 
 void FCookDirector::ParseConfig()
diff --git a/Engine/Source/Editor/UnrealEd/Private/Cooker/CookMPCollector.h b/Engine/Source/Editor/UnrealEd/Private/Cooker/CookMPCollector.h
--- a/Engine/Source/Editor/UnrealEd/Private/Cooker/CookMPCollector.h
+++ b/Engine/Source/Editor/UnrealEd/Private/Cooker/CookMPCollector.h
@@ -4,6 +4,7 @@
 
 #include "Containers/Array.h"
 #include "Containers/ArrayView.h"
+#include "Cooker/CookTypes.h"
 #include "HAL/Platform.h"
 #include "Misc/Guid.h"
 #include "Templates/RefCounting.h"
@@ -51,17 +52,19 @@ public:
 
 		uint8 PlatformToInt(const ITargetPlatform* Platform) const;
 		const ITargetPlatform* IntToPlatform(uint8 PlatformAsInt) const;
+		FWorkerId GetWorkerId() const { return WorkerId; }
 
 	private:
 		TConstArrayView<const ITargetPlatform*> Platforms;
+		FWorkerId WorkerId;
 
 		friend class FCookWorkerServer;
 	};
 
 	virtual ~IMPCollector() {}
 
-	virtual FGuid GetMessageType() = 0;
-	virtual const TCHAR* GetDebugName() = 0;
+	virtual FGuid GetMessageType() const = 0;
+	virtual const TCHAR* GetDebugName() const = 0;
 
 	virtual void ClientTick(FClientContext& Context) {}
 	virtual void ReceiveMessage(FServerContext& Context, FCbObjectView Message);
diff --git a/Engine/Source/Editor/UnrealEd/Private/Cooker/CookWorkerClient.cpp b/Engine/Source/Editor/UnrealEd/Private/Cooker/CookWorkerClient.cpp
--- a/Engine/Source/Editor/UnrealEd/Private/Cooker/CookWorkerClient.cpp
+++ b/Engine/Source/Editor/UnrealEd/Private/Cooker/CookWorkerClient.cpp
@@ -27,6 +27,9 @@ constexpr float WaitForConnectReplyTimeout = 60.f;
 FCookWorkerClient::FCookWorkerClient(UCookOnTheFlyServer& InCOTFS)
 	: COTFS(InCOTFS)
 {
+	FLogMessagesMessageHandler* Handler = new FLogMessagesMessageHandler();
+	Handler->InitializeClient();
+	Register(Handler);
 }
 
 FCookWorkerClient::~FCookWorkerClient()
diff --git a/Engine/Source/Editor/UnrealEd/Private/Cooker/CookWorkerServer.cpp b/Engine/Source/Editor/UnrealEd/Private/Cooker/CookWorkerServer.cpp
--- a/Engine/Source/Editor/UnrealEd/Private/Cooker/CookWorkerServer.cpp
+++ b/Engine/Source/Editor/UnrealEd/Private/Cooker/CookWorkerServer.cpp
@@ -497,6 +497,7 @@ void FCookWorkerServer::HandleReceiveMessagesInternal()
 				check(*Collector);
 				IMPCollector::FServerContext Context;
 				Context.Platforms = OrderedSessionPlatforms;
+				Context.WorkerId = WorkerId;
 				(*Collector)->ReceiveMessage(Context, Message.Object);
 			}
 			else
@@ -875,4 +876,97 @@ bool FDiscoveredPackagesMessage::TryRead(FCbObject&& Object)
 
 FGuid FDiscoveredPackagesMessage::MessageType(TEXT("C9F5BC5C11484B06B346B411F1ED3090"));
 
+FCbWriter& operator<<(FCbWriter& Writer, const FReplicatedLogData& Package)
+{
+	Writer.BeginArray();
+	Writer << Package.Category;
+	uint8 Verbosity = static_cast<uint8>(Package.Verbosity);
+	check(Verbosity < ELogVerbosity::NumVerbosity);
+	Writer << Verbosity;
+	Writer << Package.Message;
+	Writer.EndArray();
+	return Writer;
+}
+
+bool LoadFromCompactBinary(FCbFieldView Field, FReplicatedLogData& OutPackage)
+{
+	bool bOk = true;
+	FCbFieldViewIterator It = Field.CreateViewIterator();
+	bOk = LoadFromCompactBinary(*It++, OutPackage.Category) & bOk;
+	uint8 Verbosity;
+	if (LoadFromCompactBinary(*It++, Verbosity) && Verbosity <= ELogVerbosity::NumVerbosity)
+	{
+		OutPackage.Verbosity = static_cast<ELogVerbosity::Type>(Verbosity);
+	}
+	else
+	{
+		bOk = false;
+		OutPackage.Verbosity = static_cast<ELogVerbosity::Type>(0);
+	}
+	bOk = LoadFromCompactBinary(*It++, OutPackage.Message) & bOk;
+	return bOk;
+}
+
+FGuid FLogMessagesMessageHandler::MessageType(TEXT("DB024D28203D4FBAAAF6AAD7080CF277"));
+
+FLogMessagesMessageHandler::~FLogMessagesMessageHandler()
+{
+	if (bRegistered && GLog)
+	{
+		GLog->RemoveOutputDevice(this);
+	}
+}
+void FLogMessagesMessageHandler::InitializeClient()
+{
+	check(!bRegistered);
+	GLog->AddOutputDevice(this);
+	bRegistered = true;
+}
+
+void FLogMessagesMessageHandler::ClientTick(FClientContext& Context)
+{
+	{
+		FScopeLock QueueScopeLock(&QueueLock);
+		Swap(QueuedLogs, QueuedLogsBackBuffer);
+	}
+	if (!QueuedLogsBackBuffer.IsEmpty())
+	{
+		FCbWriter Writer;
+		Writer.BeginObject();
+		Writer << "Messages" << QueuedLogsBackBuffer;
+		Writer.EndObject();
+		Context.AddMessage(Writer.Save().AsObject());
+		QueuedLogsBackBuffer.Reset();
+	}
+}
+
+void FLogMessagesMessageHandler::ReceiveMessage(FServerContext& Context, FCbObjectView InMessage)
+{
+	TArray<FReplicatedLogData> Messages;
+	if (!LoadFromCompactBinary(InMessage["Messages"], Messages))
+	{
+		UE_LOG(LogCook, Error, TEXT("FLogMessagesMessageHandler received corrupted message from CookWorker"));
+		return;
+	}
+
+	for (FReplicatedLogData& LogData : Messages)
+	{
+		GLog->CategorizedLogf(LogData.Category, LogData.Verbosity, TEXT("[CookWorker %d]: %s"),
+			Context.GetWorkerId().GetRemoteIndex(), *LogData.Message);
+	}
+}
+
+void FLogMessagesMessageHandler::Serialize(const TCHAR* V, ELogVerbosity::Type Verbosity,
+	const FName& Category)
+{
+	FScopeLock QueueScopeLock(&QueueLock);
+	QueuedLogs.Add(FReplicatedLogData{ FString(V), Category, Verbosity });
+}
+
+void FLogMessagesMessageHandler::Serialize(const TCHAR* V, ELogVerbosity::Type Verbosity,
+	const FName& Category, const double Time)
+{
+	Serialize(V, Verbosity, Category);
+}
+
 }
\ No newline at end of file
diff --git a/Engine/Source/Editor/UnrealEd/Private/Cooker/CookWorkerServer.h b/Engine/Source/Editor/UnrealEd/Private/Cooker/CookWorkerServer.h
--- a/Engine/Source/Editor/UnrealEd/Private/Cooker/CookWorkerServer.h
+++ b/Engine/Source/Editor/UnrealEd/Private/Cooker/CookWorkerServer.h
@@ -8,11 +8,14 @@
 #include "Containers/RingBuffer.h"
 #include "Containers/Set.h"
 #include "CookOnTheSide/CookOnTheFlyServer.h"
+#include "CookMPCollector.h"
 #include "CookPackageData.h"
 #include "CookSockets.h"
 #include "CookTypes.h"
 #include "HAL/CriticalSection.h"
+#include "Logging/LogVerbosity.h"
 #include "Misc/Guid.h"
+#include "Misc/OutputDevice.h"
 #include "Templates/RefCounting.h"
 
 class FSocket;
@@ -293,4 +296,45 @@ public:
 	static FGuid MessageType;
 };
 
+/** Stores the data passed into FOutputDevice::Serialize, for replication to the CookDirector. */
+struct FReplicatedLogData
+{
+	FString Message;
+	FName Category;
+	ELogVerbosity::Type Verbosity;
+};
+FCbWriter& operator<<(FCbWriter& Writer, const FReplicatedLogData& Package);
+bool LoadFromCompactBinary(FCbFieldView Field, FReplicatedLogData& OutPackage);
+
+/**
+ * Send log messages from CookWorkers to the CookDirector, which marks them up with the CookWorkerId and
+ * then prints them to its own log.
+ */
+class FLogMessagesMessageHandler : public IMPCollector, public FOutputDevice
+{
+public:
+	~FLogMessagesMessageHandler();
+	void InitializeClient();
+
+	// IMPCollector
+	virtual FGuid GetMessageType() const override { return MessageType; }
+	virtual const TCHAR* GetDebugName() const override { return TEXT("FLogMessagesMessageHandler"); }
+	virtual void ClientTick(FClientContext& Context) override;
+	virtual void ReceiveMessage(FServerContext& Context, FCbObjectView Message) override;
+
+	// FOutputDevice
+	virtual void Serialize(const TCHAR* V, ELogVerbosity::Type Verbosity, const FName& Category) override;
+	virtual void Serialize(const TCHAR* V, ELogVerbosity::Type Verbosity, const FName& Category, const double Time) override;
+	virtual bool CanBeUsedOnAnyThread() const override { return true; }
+	virtual bool CanBeUsedOnMultipleThreads() const override { return true; }
+
+private:
+	FCriticalSection QueueLock;
+	TArray<FReplicatedLogData> QueuedLogs;
+	TArray<FReplicatedLogData> QueuedLogsBackBuffer;
+	bool bRegistered = false;
+
+	static FGuid MessageType;
+};
+
 }
\ No newline at end of file
