UE-167071 - Batch 5 of UnrealEd truncation warnings

diff --git a/Engine/Source/Editor/UnrealEd/Private/HierarchicalLOD.cpp b/Engine/Source/Editor/UnrealEd/Private/HierarchicalLOD.cpp
--- a/Engine/Source/Editor/UnrealEd/Private/HierarchicalLOD.cpp
+++ b/Engine/Source/Editor/UnrealEd/Private/HierarchicalLOD.cpp
@@ -588,7 +588,7 @@ void FHierarchicalLODBuilder::InitializeClusters(ULevel* InLevel, const int32 LO
 						AActor* Actor2 = InActors[SubActorId];
 
 						FLODCluster NewClusterCandidate = FLODCluster(Actor1, Actor2);
-						float NewClusterCost = NewClusterCandidate.GetCost();
+						double NewClusterCost = NewClusterCandidate.GetCost();
 
 						if (NewClusterCost <= CullCost)
 						{
@@ -898,7 +898,7 @@ void FHierarchicalLODBuilder::BuildMeshesForLODActors(bool bForceAll)
 				FHierarchicalLODProxyProcessor* Processor = Module.GetProxyProcessor();
 				while (Processor->IsProxyGenerationRunning())
 				{
-					FTSTicker::GetCoreTicker().Tick(FApp::GetDeltaTime());
+					FTSTicker::GetCoreTicker().Tick(static_cast<float>(FApp::GetDeltaTime()));
 					FThreadManager::Get().Tick();
 					FTaskGraphInterface::Get().ProcessThreadUntilIdle(ENamedThreads::GameThread);
 					FPlatformProcess::Sleep(0.1f);
@@ -1213,7 +1213,7 @@ void FHierarchicalLODBuilder::MergeClustersAndBuildActors(ULevel* InLevel, const
 							// if valid, see if it contains any of this actors
 							if (MergedCluster.Contains(Cluster))
 							{
-								float MergeCost = MergedCluster.GetMergedCost(Cluster);
+								double MergeCost = MergedCluster.GetMergedCost(Cluster);
 
 								// merge two clusters
 								if (MergeCost <= HighestCost)
@@ -1361,7 +1361,7 @@ ALODActor* FHierarchicalLODBuilder::CreateLODActor(const FLODCluster& InCluster,
 			NewActor = LevelWorld->SpawnActor<ALODActor>(ALODActor::StaticClass(), Transform, ActorSpawnParams);
 			NewLODActors.Add(NewActor);
 			NewActor->LODLevel = LODIdx + 1;
-			NewActor->CachedNumHLODLevels = InLevel->GetWorldSettings()->GetNumHierarchicalLODLevels();
+			NewActor->CachedNumHLODLevels = IntCastChecked<uint8>(InLevel->GetWorldSettings()->GetNumHierarchicalLODLevels());
 			NewActor->SetDrawDistance(0.0f);
 
 			// now set as parent
diff --git a/Engine/Source/Editor/UnrealEd/Private/LODCluster.cpp b/Engine/Source/Editor/UnrealEd/Private/LODCluster.cpp
--- a/Engine/Source/Editor/UnrealEd/Private/LODCluster.cpp
+++ b/Engine/Source/Editor/UnrealEd/Private/LODCluster.cpp
@@ -23,7 +23,7 @@
 
 
 /** Utility function to calculate overlap of two spheres */
-const float CalculateOverlap(const FSphere& ASphere, const float AFillingFactor, const FSphere& BSphere, const float BFillingFactor)
+double CalculateOverlap(const FSphere& ASphere, const double AFillingFactor, const FSphere& BSphere, const double BFillingFactor)
 {
 	// if it doesn't intersect, return zero 
 	if (!ASphere.Intersects(BSphere))
@@ -46,39 +46,39 @@ const float CalculateOverlap(const FSphere& ASphere, const float AFillingFactor,
 		return ASphere.GetVolume();
 	}
 
-	float Distance = (ASphere.Center-BSphere.Center).Size();
+	double Distance = (ASphere.Center - BSphere.Center).Size();
 	check(!FMath::IsNearlyZero(Distance));
 
-	float ARadius = ASphere.W;
-	float BRadius = BSphere.W;
+	double ARadius = ASphere.W;
+	double BRadius = BSphere.W;
 
-	float ACapHeight = (BRadius*BRadius - (ARadius - Distance)*(ARadius - Distance)) / (2*Distance);
-	float BCapHeight = (ARadius*ARadius - (BRadius - Distance)*(BRadius - Distance)) / (2*Distance);
+	double ACapHeight = (BRadius * BRadius - (ARadius - Distance) * (ARadius - Distance)) / (2 * Distance);
+	double BCapHeight = (ARadius * ARadius - (BRadius - Distance) * (BRadius - Distance)) / (2 * Distance);
 
-	if (ACapHeight<=0.f || BCapHeight<=0.f)
+	if ((ACapHeight <= 0.f) || (BCapHeight <= 0.f))
 	{
 		// it's possible to get cap height to be less than 0 
 		// since when we do check intersect, we do have regular tolerance
 		return 0.f;		
 	}
 
-	float OverlapRadius1 = ((ARadius + BRadius)*(ARadius + BRadius) - Distance*Distance) * (Distance*Distance - (ARadius - BRadius)*(ARadius - BRadius));
-	float OverlapRadius2 = 2 * Distance;
-	float OverlapRadius = FMath::Sqrt(OverlapRadius1) / OverlapRadius2;
-	float OverlapRadiusSq = FMath::Square(OverlapRadius);
+	double OverlapRadius1 = ((ARadius + BRadius) * (ARadius + BRadius) - Distance*Distance) * (Distance * Distance - (ARadius - BRadius) * (ARadius - BRadius));
+	double OverlapRadius2 = 2 * Distance;
+	double OverlapRadius = FMath::Sqrt(OverlapRadius1) / OverlapRadius2;
+	double OverlapRadiusSq = FMath::Square(OverlapRadius);
 
-	float ConstPI = PI/6.0f;
-	float AVolume = ConstPI*(3*OverlapRadiusSq + ACapHeight*ACapHeight) * ACapHeight;
-	float BVolume = ConstPI*(3*OverlapRadiusSq + BCapHeight*BCapHeight) * BCapHeight;
+	double ConstPI = UE_PI / 6.0f;
+	double AVolume = ConstPI * (3 * OverlapRadiusSq + ACapHeight * ACapHeight) * ACapHeight;
+	double BVolume = ConstPI * (3 * OverlapRadiusSq + BCapHeight * BCapHeight) * BCapHeight;
 
-	float TotalVolume = AFillingFactor*AVolume + BFillingFactor*BVolume;
+	double TotalVolume = AFillingFactor * AVolume + BFillingFactor * BVolume;
 	return TotalVolume;
 }
 
 /** Utility function that calculates filling factor */
-const float CalculateFillingFactor(const FSphere& ASphere, const float AFillingFactor, const FSphere& BSphere, const float BFillingFactor)
+double CalculateFillingFactor(const FSphere& ASphere, const double AFillingFactor, const FSphere& BSphere, const double BFillingFactor)
 {
-	const float OverlapVolume = CalculateOverlap( ASphere, AFillingFactor, BSphere, BFillingFactor);
+	const double OverlapVolume = CalculateOverlap( ASphere, AFillingFactor, BSphere, BFillingFactor);
 	FSphere UnionSphere = ASphere + BSphere;
 	// it shouldn't be zero or it should be checked outside
 	ensure(UnionSphere.W != 0.f);
@@ -208,12 +208,12 @@ bool FLODCluster::operator==(const FLODCluster& Other) const
 	return Actors.Num() == Other.Actors.Num() && Actors.Includes(Other.Actors);
 }
 
-float FLODCluster::GetMergedCost(const FLODCluster& Other) const
+double FLODCluster::GetMergedCost(const FLODCluster& Other) const
 {
-	float MergedFillingFactor = CalculateFillingFactor(Bound, FillingFactor, Other.Bound, Other.FillingFactor);
+	double MergedFillingFactor = CalculateFillingFactor(Bound, FillingFactor, Other.Bound, Other.FillingFactor);
 	FSphere MergedBound = Bound + Other.Bound;
 
-	float MergedClusterCost = (MergedBound.W * MergedBound.W * MergedBound.W) / MergedFillingFactor;
+	double MergedClusterCost = (MergedBound.W * MergedBound.W * MergedBound.W) / MergedFillingFactor;
 	return MergedClusterCost;
 }
 
diff --git a/Engine/Source/Editor/UnrealEd/Public/LODCluster.h b/Engine/Source/Editor/UnrealEd/Public/LODCluster.h
--- a/Engine/Source/Editor/UnrealEd/Public/LODCluster.h
+++ b/Engine/Source/Editor/UnrealEd/Public/LODCluster.h
@@ -47,13 +47,13 @@ struct FLODCluster
 	inline bool IsValid() const { return bValid; }
 
 	/** Return cost of the cluster, lower is better */
-	inline const float GetCost() const
+	inline double GetCost() const
 	{
 		return ClusterCost;
 	}
 
 	/** Return cost of the union of this cluster & the other cluster, lower is better */
-	float GetMergedCost(const FLODCluster& Other) const;
+	double GetMergedCost(const FLODCluster& Other) const;
 
 	/** Compare clusters and returns true when this contains any of Other's actors */
 	bool Contains(FLODCluster& Other) const;
@@ -67,9 +67,9 @@ struct FLODCluster
 	/** Cluster bounds */
 	FSphere	Bound;
 	/** Filling factor for this cluster, determines how much of the cluster's bounds/area is occupied by the contained actors*/
-	float FillingFactor;
+	double FillingFactor;
 	/** Cached cluster cost, FMath::Pow(Bound.W, 3) / FillingFactor */
-	float ClusterCost;
+	double ClusterCost;
 
 private:
 	/**
