Editor Replace Actor: added an option to avoid copying related properties between old and new actor. The original behavior is kept, which was to copy related properties.

#rnx

diff --git a/Engine/Source/Editor/LevelEditor/Private/LevelEditorActions.cpp b/Engine/Source/Editor/LevelEditor/Private/LevelEditorActions.cpp
--- a/Engine/Source/Editor/LevelEditor/Private/LevelEditorActions.cpp
+++ b/Engine/Source/Editor/LevelEditor/Private/LevelEditorActions.cpp
@@ -1525,12 +1525,13 @@ AActor* FLevelEditorActionCallbacks::AddActorFromClass( UClass* ActorClass )
 	return NewActor;
 }
 
+extern bool GReplaceSelectedActorsWithSelectedClassCopyProperties;
 void FLevelEditorActionCallbacks::ReplaceActors_Clicked( UActorFactory* ActorFactory, FAssetData AssetData )
 {
-	FLevelEditorActionCallbacks::ReplaceActors(ActorFactory, AssetData);
+	FLevelEditorActionCallbacks::ReplaceActors(ActorFactory, AssetData, GReplaceSelectedActorsWithSelectedClassCopyProperties);
 }
 
-AActor* FLevelEditorActionCallbacks::ReplaceActors( UActorFactory* ActorFactory, const FAssetData& AssetData )
+AActor* FLevelEditorActionCallbacks::ReplaceActors( UActorFactory* ActorFactory, const FAssetData& AssetData, bool bCopySourceProperties )
 {
 	AActor* NewActor = NULL;
 
@@ -1539,7 +1540,7 @@ AActor* FLevelEditorActionCallbacks::ReplaceActors( UActorFactory* ActorFactory,
 	if( ActorFactory->CanCreateActorFrom( AssetData, ErrorMessage ) )
 	{
 		// Replace all selected actors with actors created from the specified factory
-		GEditor->ReplaceSelectedActors( ActorFactory, AssetData );
+		GEditor->ReplaceSelectedActors( ActorFactory, AssetData, bCopySourceProperties );
 
 		if ( IPlacementModeModule::IsAvailable() )
 		{
diff --git a/Engine/Source/Editor/LevelEditor/Private/LevelEditorCreateActorMenu.cpp b/Engine/Source/Editor/LevelEditor/Private/LevelEditorCreateActorMenu.cpp
--- a/Engine/Source/Editor/LevelEditor/Private/LevelEditorCreateActorMenu.cpp
+++ b/Engine/Source/Editor/LevelEditor/Private/LevelEditorCreateActorMenu.cpp
@@ -447,8 +447,32 @@ void LevelEditorCreateActorMenu::FillAddReplaceContextMenuSections(FToolMenuSect
 	}
 }
 
+bool GReplaceSelectedActorsWithSelectedClassCopyProperties = true;
 void LevelEditorCreateActorMenu::FillAddReplaceActorMenu(UToolMenu* Menu, EActorCreateMode::Type CreateMode)
 {
+	if ( CreateMode == EActorCreateMode::Replace )
+	{
+		FToolMenuSection& Section = Menu->AddSection("Options", NSLOCTEXT("LevelViewportContextMenu", "Options", "Options"));
+
+		GReplaceSelectedActorsWithSelectedClassCopyProperties = true;
+
+		FToolMenuEntry ToolMenuEntry = FToolMenuEntry::InitMenuEntry(
+			"CopyProperties",
+			NSLOCTEXT("LevelViewportContextMenu", "CopyProperties", "Copy Properties"),
+			FText(),
+			FSlateIcon(),
+			FUIAction(
+				FExecuteAction::CreateLambda([]() { GReplaceSelectedActorsWithSelectedClassCopyProperties = !GReplaceSelectedActorsWithSelectedClassCopyProperties; }),
+				FCanExecuteAction(),
+				FGetActionCheckState::CreateLambda([] { return GReplaceSelectedActorsWithSelectedClassCopyProperties ? ECheckBoxState::Checked : ECheckBoxState::Unchecked; } )
+			),
+			EUserInterfaceActionType::ToggleButton
+		);
+		ToolMenuEntry.bShouldCloseWindowAfterMenuSelection = false;
+
+		Section.AddEntry(ToolMenuEntry);
+	}
+
 	{
 		FToolMenuSection& Section = Menu->AddSection("ContentBrowserActor", NSLOCTEXT("LevelViewportContextMenu", "AssetSelectionSection", "Selection"));
 		FAssetData TargetAssetData;
diff --git a/Engine/Source/Editor/LevelEditor/Public/LevelEditorActions.h b/Engine/Source/Editor/LevelEditor/Public/LevelEditorActions.h
--- a/Engine/Source/Editor/LevelEditor/Public/LevelEditorActions.h
+++ b/Engine/Source/Editor/LevelEditor/Public/LevelEditorActions.h
@@ -964,7 +964,7 @@ public:
 	 * @param ActorFactory	The actor factory to use in replacement
 	 */
 	static void ReplaceActors_Clicked( UActorFactory* ActorFactory, FAssetData AssetData );
-	static AActor* ReplaceActors( UActorFactory* ActorFactory, const FAssetData& AssetData );
+	static AActor* ReplaceActors( UActorFactory* ActorFactory, const FAssetData& AssetData, bool bCopySourceProperties = true );
 
 	/**
 	 * Called when the ReplaceActor command is executed and a class is selected in the actor browser
diff --git a/Engine/Source/Editor/UnrealEd/Classes/Editor/EditorEngine.h b/Engine/Source/Editor/UnrealEd/Classes/Editor/EditorEngine.h
--- a/Engine/Source/Editor/UnrealEd/Classes/Editor/EditorEngine.h
+++ b/Engine/Source/Editor/UnrealEd/Classes/Editor/EditorEngine.h
@@ -1257,7 +1257,7 @@ public:
 	 * 
 	 * @param Factory - the Factory to use to create Actors
 	 */
-	void ReplaceSelectedActors(UActorFactory* Factory, const FAssetData& AssetData);
+	void ReplaceSelectedActors(UActorFactory* Factory, const FAssetData& AssetData, bool bCopySourceProperties = true);
 
 	/**
 	 * Replaces specified Actors with the same number of a different kind of Actor using the specified factory to spawn the new Actors
@@ -1268,7 +1268,7 @@ public:
 	 * @param ActorsToReplace - Actors to replace
 	 * @param OutNewActors - Actors that were created
 	 */
-	void ReplaceActors(UActorFactory* Factory, const FAssetData& AssetData, const TArray<AActor*>& ActorsToReplace, TArray<AActor*>* OutNewActors = nullptr);
+	void ReplaceActors(UActorFactory* Factory, const FAssetData& AssetData, const TArray<AActor*>& ActorsToReplace, TArray<AActor*>* OutNewActors = nullptr, bool bCopySourceProperties = true);
 
 	/**
 	 * Converts passed in brushes into a single static mesh actor. 
diff --git a/Engine/Source/Editor/UnrealEd/Private/EditorEngine.cpp b/Engine/Source/Editor/UnrealEd/Private/EditorEngine.cpp
--- a/Engine/Source/Editor/UnrealEd/Private/EditorEngine.cpp
+++ b/Engine/Source/Editor/UnrealEd/Private/EditorEngine.cpp
@@ -5148,7 +5148,7 @@ namespace ReattachActorsHelper
 	}
 }
 
-void UEditorEngine::ReplaceSelectedActors(UActorFactory* Factory, const FAssetData& AssetData)
+void UEditorEngine::ReplaceSelectedActors(UActorFactory* Factory, const FAssetData& AssetData, bool bCopySourceProperties)
 {
 	UObject* ObjectForFactory = NULL;
 
@@ -5185,10 +5185,10 @@ void UEditorEngine::ReplaceSelectedActors(UActorFactory* Factory, const FAssetDa
 		}
 	}
 
-	ReplaceActors(Factory, AssetData, ActorsToReplace);
+	ReplaceActors(Factory, AssetData, ActorsToReplace, nullptr, bCopySourceProperties);
 }
 
-void UEditorEngine::ReplaceActors(UActorFactory* Factory, const FAssetData& AssetData, const TArray<AActor*>& ActorsToReplace, TArray<AActor*>* OutNewActors)
+void UEditorEngine::ReplaceActors(UActorFactory* Factory, const FAssetData& AssetData, const TArray<AActor*>& ActorsToReplace, TArray<AActor*>* OutNewActors, bool bCopySourceProperties)
 {
 	// Cache for attachment info of all actors being converted.
 	TArray<ReattachActorsHelper::FActorAttachmentCache> AttachmentInfo;
@@ -5235,7 +5235,7 @@ void UEditorEngine::ReplaceActors(UActorFactory* Factory, const FAssetData& Asse
 		// create the actor
 		NewActor = Factory->CreateActor(Asset, Level, OldTransform, SpawnParams);
 		// For blueprints, try to copy over properties
-		if (Factory->IsA(UActorFactoryBlueprint::StaticClass()))
+		if (bCopySourceProperties && Factory->IsA(UActorFactoryBlueprint::StaticClass()))
 		{
 			UBlueprint* Blueprint = CastChecked<UBlueprint>(Asset);
 			// Only try to copy properties if this blueprint is based on the actor
