Made AssetEditorSubsystem resilient to dangling object pointers in its asset editor bookkeeping.
This is to tentatively fix a crash where GetAllEditedAssets returns invalid object pointers.
Also made sure that reload package wouldn't end up opening new asset editors.

diff --git a/Engine/Source/Editor/UnrealEd/Private/Subsystems/AssetEditorSubsystem.cpp b/Engine/Source/Editor/UnrealEd/Private/Subsystems/AssetEditorSubsystem.cpp
--- a/Engine/Source/Editor/UnrealEd/Private/Subsystems/AssetEditorSubsystem.cpp
+++ b/Engine/Source/Editor/UnrealEd/Private/Subsystems/AssetEditorSubsystem.cpp
@@ -122,9 +122,9 @@ TArray<IAssetEditorInstance*> UAssetEditorSubsystem::FindEditorsForAssetAndSubOb
 {
 	TArray<IAssetEditorInstance*> EditorInstances;
 
-	for (const TPair<UObject*, IAssetEditorInstance*>& Pair : OpenedAssets)
+	for (const TPair<FAssetEntry, IAssetEditorInstance*>& Pair : OpenedAssets)
 	{
-		if (Pair.Key == Asset || Pair.Key->IsIn(Asset))
+		if (Pair.Key.RawPtr == Asset || (Pair.Key.ObjectPtr.IsValid() && Pair.Key.ObjectPtr.Get()->IsIn(Asset)))
 		{		
 			EditorInstances.Add(Pair.Value);
 		}
@@ -170,10 +170,10 @@ void UAssetEditorSubsystem::RemoveAssetFromAllEditors(UObject* Asset)
 void UAssetEditorSubsystem::CloseOtherEditors(UObject* Asset, IAssetEditorInstance* OnlyEditor)
 {
 	TArray<UObject*> AllAssets;
-	for (TMultiMap<UObject*, IAssetEditorInstance*>::TIterator It(OpenedAssets); It; ++It)
+	for (TMultiMap<FAssetEntry, IAssetEditorInstance*>::TIterator It(OpenedAssets); It; ++It)
 	{
 		IAssetEditorInstance* Editor = It.Value();
-		if (Asset == It.Key() && Editor != OnlyEditor)
+		if (Asset == It.Key().RawPtr && Editor != OnlyEditor)
 		{
 			Editor->CloseWindow();
 		}
@@ -186,10 +186,10 @@ void UAssetEditorSubsystem::CloseOtherEditors(UObject* Asset, IAssetEditorInstan
 TArray<UObject*> UAssetEditorSubsystem::GetAllEditedAssets()
 {
 	TArray<UObject*> AllAssets;
-	for (TMultiMap<UObject*, IAssetEditorInstance*>::TIterator It(OpenedAssets); It; ++It)
+	for (TMultiMap<FAssetEntry, IAssetEditorInstance*>::TIterator It(OpenedAssets); It; ++It)
 	{
-		UObject* Asset = It.Key();
-		if (Asset != NULL)
+		UObject* Asset = It.Key().ObjectPtr.Get();
+		if (Asset != nullptr)
 		{
 			AllAssets.AddUnique(Asset);
 		}
@@ -245,7 +245,7 @@ void UAssetEditorSubsystem::NotifyAssetClosed(UObject* Asset, IAssetEditorInstan
 void UAssetEditorSubsystem::NotifyEditorClosed(IAssetEditorInstance* InInstance)
 {
 	// Remove all assets associated with the editor
-	TArray<UObject*> Assets;
+	TArray<FAssetEntry> Assets;
 	OpenedEditors.MultiFind(InInstance, /*out*/ Assets);
 	for (int32 AssetIndex = 0; AssetIndex < Assets.Num(); ++AssetIndex)
 	{
@@ -268,10 +268,10 @@ void UAssetEditorSubsystem::NotifyEditorClosed(IAssetEditorInstance* InInstance)
 bool UAssetEditorSubsystem::CloseAllAssetEditors()
 {
 	bool bAllEditorsClosed = true;
-	for (TMultiMap<IAssetEditorInstance*, UObject*>::TIterator It(OpenedEditors); It; ++It)
+	for (TMultiMap<IAssetEditorInstance*, FAssetEntry>::TIterator It(OpenedEditors); It; ++It)
 	{
 		IAssetEditorInstance* Editor = It.Key();
-		if (Editor != NULL)
+		if (Editor != nullptr)
 		{
 			if (!Editor->CloseWindow())
 			{
@@ -916,17 +916,17 @@ void UAssetEditorSubsystem::SaveOpenAssetEditors(const bool bOnShutdown)
 		TArray<FString> OpenAssets;
 		if (bOnShutdown || !FPlatformMisc::IsDebuggerPresent())
 		{
-			for (const TPair<IAssetEditorInstance*, UObject*>& EditorPair : OpenedEditors)
+			for (const TPair<IAssetEditorInstance*, FAssetEntry>& EditorPair : OpenedEditors)
 			{
 				IAssetEditorInstance* Editor = EditorPair.Key;
-				if (Editor != NULL)
+				if (Editor != nullptr)
 				{
-					UObject* EditedObject = EditorPair.Value;
-					if (EditedObject != NULL)
+					UObject* EditedObject = EditorPair.Value.ObjectPtr.Get();
+					if (EditedObject != nullptr)
 					{
 						// only record assets that have a valid saved package
 						UPackage* Package = EditedObject->GetOutermost();
-						if (Package != NULL && Package->GetFileSize() != 0)
+						if (Package != nullptr && Package->GetFileSize() != 0)
 						{
 							OpenAssets.Add(EditedObject->GetPathName());
 						}
@@ -964,27 +964,34 @@ void UAssetEditorSubsystem::HandlePackageReloaded(const EPackageReloadPhase InPa
 		}
 
 		/** Look for replacement for assets that are open now so we can reopen */
-		for (TPair<UObject*, IAssetEditorInstance*>& AssetEditorPair : OpenedAssets)
+		for (TPair<FAssetEntry, IAssetEditorInstance*>& AssetEditorPair : OpenedAssets)
 		{
 			UObject* NewAsset = nullptr;
-			if (InPackageReloadedEvent->GetRepointedObject(AssetEditorPair.Key, NewAsset))
+			if (InPackageReloadedEvent->GetRepointedObject(AssetEditorPair.Key.ObjectPtr.Get(), NewAsset))
 			{
 				if (NewAsset)
 				{
 					PendingAssetsToOpen.AddUnique(NewAsset);
 				}
 
-				UObject* OldAsset = AssetEditorPair.Key;
+				UObject* OldAsset = AssetEditorPair.Key.RawPtr; // Not validating the asset here since we'd want to close editors for garbage collected assets
 				ObjectsToClose.AddUnique(OldAsset);
 
-				// The asset being reloaded might have other assets that depend on it. Find the list of
-				// external referencers to this asset and mark them to be closed and reopened as well.
-				TArray<FReferencerInformation> AssetInternalReferencers, AssetExternalReferencers;
-				OldAsset->RetrieveReferencers(&AssetInternalReferencers, &AssetExternalReferencers);
-				for (const FReferencerInformation& Ref : AssetExternalReferencers)
+				if (AssetEditorPair.Key.ObjectPtr.IsValid())
 				{
-					ObjectsToClose.AddUnique(Ref.Referencer);
-					PendingAssetsToOpen.AddUnique(Ref.Referencer);
+					// The asset being reloaded might have other assets that depend on it. Find the list of
+					// external referencers to this asset and mark them to be closed and reopened as well.
+					TArray<FReferencerInformation> AssetInternalReferencers, AssetExternalReferencers;
+					AssetEditorPair.Key.ObjectPtr.Get()->RetrieveReferencers(&AssetInternalReferencers, &AssetExternalReferencers);
+					for (const FReferencerInformation& Ref : AssetExternalReferencers)
+					{
+						ObjectsToClose.AddUnique(Ref.Referencer);
+
+						if (!FindEditorsForAssetAndSubObjects(Ref.Referencer).IsEmpty())
+						{
+							PendingAssetsToOpen.AddUnique(Ref.Referencer);
+						}
+					}
 				}
 			}
 		}
diff --git a/Engine/Source/Editor/UnrealEd/Public/Subsystems/AssetEditorSubsystem.h b/Engine/Source/Editor/UnrealEd/Public/Subsystems/AssetEditorSubsystem.h
--- a/Engine/Source/Editor/UnrealEd/Public/Subsystems/AssetEditorSubsystem.h
+++ b/Engine/Source/Editor/UnrealEd/Public/Subsystems/AssetEditorSubsystem.h
@@ -288,11 +288,36 @@ private:
 		}
 	};
 
+	/**
+	 * Utility struct for assets book keeping.
+	 * We'll index assets using the raw pointer, as it should be possible to clear an entry after an asset has been garbage collected
+	 * but any access to the UObject should go through the weak object pointer.
+	 */
+	struct FAssetEntry
+	{
+		// Implicit constructor
+		FAssetEntry(UObject* InRawPtr)
+			: RawPtr(InRawPtr)
+			, ObjectPtr(InRawPtr)
+		{
+		}
+
+		UObject* RawPtr;
+		FWeakObjectPtr ObjectPtr;
+
+		bool operator==(const FAssetEntry& Other) const { return RawPtr == Other.RawPtr; }
+
+		inline friend uint32 GetTypeHash(const FAssetEntry& AssetEntry)
+		{
+			return GetTypeHash(AssetEntry.RawPtr);
+		}
+	};
+
 	/** Holds the opened assets. */
-	TMultiMap<UObject*, IAssetEditorInstance*> OpenedAssets;
+	TMultiMap<FAssetEntry, IAssetEditorInstance*> OpenedAssets;
 
 	/** Holds the opened editors. */
-	TMultiMap<IAssetEditorInstance*, UObject*> OpenedEditors;
+	TMultiMap<IAssetEditorInstance*, FAssetEntry> OpenedEditors;
 
 	/** Holds the times that editors were opened. */
 	TMap<IAssetEditorInstance*, FOpenedEditorTime> OpenedEditorTimes;
