Improve logging

diff --git a/Engine/Source/Editor/UnrealEd/Private/Kismet2/KismetReinstanceUtilities.cpp b/Engine/Source/Editor/UnrealEd/Private/Kismet2/KismetReinstanceUtilities.cpp
--- a/Engine/Source/Editor/UnrealEd/Private/Kismet2/KismetReinstanceUtilities.cpp
+++ b/Engine/Source/Editor/UnrealEd/Private/Kismet2/KismetReinstanceUtilities.cpp
@@ -57,7 +57,7 @@ static FAutoConsoleVariableRef CVarUseLegacyAnimInstanceReinstancingBehavior(
 	TEXT("Use the legacy re-instancing behavior for anim instances where the instance is destroyed and re-created.")
 );
 
-extern ENGINE_API bool GLogReinstancerReferenceReplacement;
+extern ENGINE_API UObject* GLogReinstancerReferenceReplacementObj;
 
 struct FReplaceReferenceHelper
 {
@@ -96,14 +96,17 @@ struct FReplaceReferenceHelper
 			return;
 		}
 
-		if (GLogReinstancerReferenceReplacement)
+		if (GLogReinstancerReferenceReplacementObj)
 		{
 			FString Classes;
 			for (UObject* Obj : SourceObjects)
 			{
 				Classes.Append(Obj->GetClass()->GetName() + " ");
 			}
-			UE_LOG(LogBlueprint, Warning, TEXT("Replacing References to: %s"), *Classes);
+			UE_LOG(LogBlueprint, Warning, TEXT("%s %s Replacing References to: %s"), 
+				*(GLogReinstancerReferenceReplacementObj->GetName()), 
+				*(GLogReinstancerReferenceReplacementObj->GetClass()->GetName()), 
+				*Classes);
 		}
 
 		// Remember what values were in UActorChannel::Actor so we can restore them later (this should only affect reinstancing during PIE)
diff --git a/Engine/Source/Runtime/Engine/Private/SkeletalMeshComponentPhysics.cpp b/Engine/Source/Runtime/Engine/Private/SkeletalMeshComponentPhysics.cpp
--- a/Engine/Source/Runtime/Engine/Private/SkeletalMeshComponentPhysics.cpp
+++ b/Engine/Source/Runtime/Engine/Private/SkeletalMeshComponentPhysics.cpp
@@ -2679,7 +2679,7 @@ void USkeletalMeshComponent::AddClothingBounds(FBoxSphereBounds& InOutBounds, co
 }
 
 #if WITH_EDITOR
-ENGINE_API bool GLogReinstancerReferenceReplacement = false;
+ENGINE_API UObject* GLogReinstancerReferenceReplacementObj = nullptr;
 #endif
 
 void USkeletalMeshComponent::RecreateClothingActors()
@@ -2695,7 +2695,7 @@ void USkeletalMeshComponent::RecreateClothingActors()
 	}
 
 #if WITH_EDITOR
-	TGuardValue<bool> EnableReinstancingLog(GLogReinstancerReferenceReplacement, true);
+	TGuardValue<UObject*> EnableReinstancingLog(GLogReinstancerReferenceReplacementObj, this);
 #endif
 
 	if(CVarEnableClothPhysics.GetValueOnGameThread() && (SkelMesh->GetMeshClothingAssets().Num() > 0))
