Interchange: fix editor reimport callback not called.

#rnx

diff --git a/Engine/Source/Editor/UnrealEd/Private/Editor.cpp b/Engine/Source/Editor/UnrealEd/Private/Editor.cpp
--- a/Engine/Source/Editor/UnrealEd/Private/Editor.cpp
+++ b/Engine/Source/Editor/UnrealEd/Private/Editor.cpp
@@ -410,37 +410,14 @@ UE::Interchange::FAssetImportResultRef FReimportManager::ReimportAsync(UObject*
 					int32 RealSourceFileIndex = SourceFileIndex == INDEX_NONE ? 0 : SourceFileIndex;
 					int32 RealValidSourceFileIndex = SourceFilenames.IsValidIndex(RealSourceFileIndex) ? RealSourceFileIndex : 0;
 					UE::Interchange::FScopedSourceData ScopedSourceData(SourceFilenames[RealValidSourceFileIndex]);
-
 					CanReimportHandler->SetReimportSourceIndex(Obj, SourceFileIndex);
-
 					if (InterchangeManager.CanTranslateSourceData(ScopedSourceData.GetSourceData()))
 					{
-						auto PostImportedLambda = [](UObject* ImportedObject)
-						{
-							if (ImportedObject)
-							{
-								TArray<UObject*> ObjectArray;
-								ObjectArray.Add(ImportedObject);
-								//UAssetToolsImpl::Get().SyncBrowserToAssets(ObjectArray);
-								GEditor->BroadcastObjectReimported(ImportedObject);
-								if (FEngineAnalytics::IsAvailable())
-								{
-									TArray<FAnalyticsEventAttribute> Attributes;
-									Attributes.Add(FAnalyticsEventAttribute(TEXT("ObjectType"), ImportedObject->GetClass()->GetName()));
-									FEngineAnalytics::GetProvider().RecordEvent(TEXT("Editor.Usage.AssetReimported"), Attributes);
-								}
-								//PostReimport.Broadcast(ImportedObject, true);
-								GEditor->RedrawAllViewports();
-							}
-						};
-						FDelegateHandle PostImportHandle = InterchangeManager.OnAssetPostImport.AddLambda(PostImportedLambda);
-
 						FImportAssetParameters ImportAssetParameters;
 						ImportAssetParameters.bIsAutomated = GIsAutomationTesting || FApp::IsUnattended() || IsRunningCommandlet() || GIsRunningUnattendedScript;
 						ImportAssetParameters.ReimportAsset = Obj;
 						ImportAssetParameters.ReimportSourceIndex = SourceFileIndex;
 						UE::Interchange::FAssetImportResultRef ImportResult = InterchangeManager.ImportAssetAsync(FString(), ScopedSourceData.GetSourceData(), ImportAssetParameters);
-						InterchangeManager.OnAssetPostImport.Remove(PostImportHandle);
 						return ImportResult;
 					}
 				}
@@ -691,6 +668,25 @@ void FReimportManager::SortHandlersIfNeeded()
 	}
 }
 
+void FReimportManager::OnInterchangePostReimported(UObject* ReimportAsset) const
+{
+	if (ReimportAsset)
+	{
+		TArray<UObject*> ObjectArray;
+		ObjectArray.Add(ReimportAsset);
+		//UAssetToolsImpl::Get().SyncBrowserToAssets(ObjectArray);
+		GEditor->BroadcastObjectReimported(ReimportAsset);
+		if (FEngineAnalytics::IsAvailable())
+		{
+			TArray<FAnalyticsEventAttribute> Attributes;
+			Attributes.Add(FAnalyticsEventAttribute(TEXT("ObjectType"), ReimportAsset->GetClass()->GetName()));
+			FEngineAnalytics::GetProvider().RecordEvent(TEXT("Editor.Usage.AssetReimported"), Attributes);
+		}
+		PostReimport.Broadcast(ReimportAsset, true);
+		GEditor->RedrawAllViewports();
+	}
+}
+
 bool FReimportManager::ReimportMultiple(TArrayView<UObject*> Objects, bool bAskForNewFileIfMissing /*= false*/, bool bShowNotification /*= true*/, FString PreferredReimportFile /*= TEXT("")*/, FReimportHandler* SpecifiedReimportHandler /*= nullptr */, int32 SourceFileIndex /*= INDEX_NONE*/, bool bForceNewFile /*= false*/, bool bAutomated /*= false*/)
 {
 	bool bBulkSuccess = true;
@@ -914,11 +910,26 @@ FReimportManager::FReimportManager()
 
 	// Create reimport handler for PhysicalMaterialMasks
 	UPhysicalMaterialMaskFactory::StaticClass();
+
+	InterchangePostReimportedDelegateHandle = UInterchangeManager::GetInterchangeManager().OnAssetPostReimport.AddRaw(this, &FReimportManager::OnInterchangePostReimported);
+
+	UInterchangeManager& InterchangeManager = UInterchangeManager::GetInterchangeManager();
+	InterchangeManager.OnPreDestroyInterchangeManager.AddLambda([InterchangePostReimportedDelegateHandleClosure = InterchangePostReimportedDelegateHandle]()
+		{
+			if (InterchangePostReimportedDelegateHandleClosure.IsValid())
+			{
+				UInterchangeManager::GetInterchangeManager().OnAssetPostImport.Remove(InterchangePostReimportedDelegateHandleClosure);
+			}
+		});
 }
 
 FReimportManager::~FReimportManager()
 {
 	Handlers.Empty();
+	if (InterchangePostReimportedDelegateHandle.IsValid())
+	{
+		UInterchangeManager::GetInterchangeManager().OnAssetPostImport.Remove(InterchangePostReimportedDelegateHandle);
+	}
 }
 
 int32 FReimportHandler::GetPriority() const
diff --git a/Engine/Source/Editor/UnrealEd/Public/EditorReimportHandler.h b/Engine/Source/Editor/UnrealEd/Public/EditorReimportHandler.h
--- a/Engine/Source/Editor/UnrealEd/Public/EditorReimportHandler.h
+++ b/Engine/Source/Editor/UnrealEd/Public/EditorReimportHandler.h
@@ -158,6 +158,8 @@ private:
 	/** Sort Reimport handlers by priority if they are unsorted */
 	void SortHandlersIfNeeded();
 
+	void OnInterchangePostReimported(UObject* ReimportAsset) const;
+
 	/** Reimport handlers registered with this manager */
 	TArray<FReimportHandler*> Handlers;
 
@@ -170,6 +172,8 @@ private:
 	/** Delegate to call after the asset is reimported */
 	FPostReimportNotification PostReimport;
 
+	FDelegateHandle InterchangePostReimportedDelegateHandle;
+
 	/** Constructor */
 	FReimportManager();
 
diff --git a/Engine/Source/Runtime/Interchange/Engine/Private/Tasks/InterchangeTaskCompletion.cpp b/Engine/Source/Runtime/Interchange/Engine/Private/Tasks/InterchangeTaskCompletion.cpp
--- a/Engine/Source/Runtime/Interchange/Engine/Private/Tasks/InterchangeTaskCompletion.cpp
+++ b/Engine/Source/Runtime/Interchange/Engine/Private/Tasks/InterchangeTaskCompletion.cpp
@@ -198,10 +198,8 @@ void UE::Interchange::FTaskCompletion::DoTask(ENamedThreads::Type CurrentThread,
 			{
 				InterchangeManager->OnAssetPostReimport.Broadcast(Asset);
 			}
-			else
-			{
-				InterchangeManager->OnAssetPostImport.Broadcast(Asset);
-			}
+			//We broadcast this event for both import and reimport.
+			InterchangeManager->OnAssetPostImport.Broadcast(Asset);
 		}
 	}
 
