Fix crash when exiting the editor

#rnx

diff --git a/Engine/Source/Editor/UnrealEd/Private/Editor.cpp b/Engine/Source/Editor/UnrealEd/Private/Editor.cpp
--- a/Engine/Source/Editor/UnrealEd/Private/Editor.cpp
+++ b/Engine/Source/Editor/UnrealEd/Private/Editor.cpp
@@ -911,14 +911,14 @@ FReimportManager::FReimportManager()
 	// Create reimport handler for PhysicalMaterialMasks
 	UPhysicalMaterialMaskFactory::StaticClass();
 
-	InterchangePostReimportedDelegateHandle = UInterchangeManager::GetInterchangeManager().OnAssetPostReimport.AddRaw(this, &FReimportManager::OnInterchangePostReimported);
-
 	UInterchangeManager& InterchangeManager = UInterchangeManager::GetInterchangeManager();
+	InterchangePostReimportedDelegateHandle = InterchangeManager.OnAssetPostReimport.AddRaw(this, &FReimportManager::OnInterchangePostReimported);
+
 	InterchangeManager.OnPreDestroyInterchangeManager.AddLambda([InterchangePostReimportedDelegateHandleClosure = InterchangePostReimportedDelegateHandle]()
 		{
 			if (InterchangePostReimportedDelegateHandleClosure.IsValid())
 			{
-				UInterchangeManager::GetInterchangeManager().OnAssetPostImport.Remove(InterchangePostReimportedDelegateHandleClosure);
+				UInterchangeManager::GetInterchangeManager().OnAssetPostReimport.Remove(InterchangePostReimportedDelegateHandleClosure);
 			}
 		});
 }
@@ -927,13 +927,7 @@ FReimportManager::~FReimportManager()
 {
 	Handlers.Empty();
 
-	/*
-	// you can't do this because ~FReimportManager is called from FReimportManager::Instance at cexit shutdown time
-	if (InterchangePostReimportedDelegateHandle.IsValid())
-	{
-		UInterchangeManager::GetInterchangeManager().OnAssetPostImport.Remove(InterchangePostReimportedDelegateHandle);
-	}
-	*/
+	// you can't do much here because ~FReimportManager is called from FReimportManager::Instance at cexit shutdown time
 }
 
 int32 FReimportHandler::GetPriority() const
