Fixed SummarizeTrace CIS error re "LoadModule" scopes.

diff --git a/Engine/Source/Editor/UnrealEd/Private/Commandlets/SummarizeTraceCommandlet.cpp b/Engine/Source/Editor/UnrealEd/Private/Commandlets/SummarizeTraceCommandlet.cpp
--- a/Engine/Source/Editor/UnrealEd/Private/Commandlets/SummarizeTraceCommandlet.cpp
+++ b/Engine/Source/Editor/UnrealEd/Private/Commandlets/SummarizeTraceCommandlet.cpp
@@ -409,6 +409,11 @@ void FSummarizeCpuProfilerProvider::OnCpuScopeTree(uint32 ThreadId, const FScope
 
 const FString* FSummarizeCpuProfilerProvider::LookupScopeName(uint32 ScopeId)
 {
+	if (int32(ScopeId) < 0)
+	{
+		ScopeId = Metadatas[~ScopeId].TimerId;
+	}
+
 	if (ScopeId < static_cast<uint32>(ScopeNames.Num()) && ScopeNames[ScopeId])
 	{
 		return &ScopeNames[ScopeId].GetValue();
@@ -810,15 +815,15 @@ void FSummarizeCpuScopeDurationAnalyzer::OnCpuScopeAnalysisEnd()
  *
  * For the LoadModule example described above, if the analyzer gets this scope tree as input:
  *
- * |-LoadModule_Module1----------------------------------------------------------|
+ * |-LoadModule(Module1)---------------------------------------------------------|
  *    |- StartupModule -------------------------------------------------------|
- *      |-LoadModule_Module1Dep1------------|  |-LoadModule_Module1Dep2------|
+ *      |-LoadModule(Module1Dep1)-----------|  |-LoadModule(Module1Dep2)-----|
  *        |-StartupModule-----------------|      |-StartupModule----------|
  *
  * It would turn it into the one below if the REGEX to match was "LoadModule_.*"
  * 
- * |-LoadModule_Module1----------------------------------------------------------|
- *      |-LoadModule_Module1Dep1------------|  |-LoadModule_Module1Dep2------|
+ * |-LoadModule(Module1)---------------------------------------------------------|
+ *      |-LoadModule(Module1Dep1)-----------|  |-LoadModule(Module1Dep2)-----|
  *
  * And it would compute the exclusive time required to load Module1 by substracting the time consumed to
  * load Module1Dep1 and Module1Dep2.
@@ -1903,12 +1908,12 @@ int32 USummarizeTraceCommandlet::Main(const FString& CmdLineParams)
 			CollectedScopeSummaries.Append(ScopeSummaries);
 		});
 
-	// Analyze 'LoadModule*' scope timer hierarchically to account individual load time only (substracting time consumed to load dependent module(s)).
+	// Analyze 'LoadModule' scope timer hierarchically to account individual load time only (substracting time consumed to load dependent module(s)).
 	TSharedPtr<FSummarizeCpuScopeHierarchyAnalyzer> HierarchicalScopeAnalyzer = MakeShared<FSummarizeCpuScopeHierarchyAnalyzer>(
 		TEXT("LoadModule"), // Analyzer Name.
 		[](const FString& ScopeName)
 		{
-			return ScopeName.StartsWith("LoadModule_"); // When analyzing a tree of scopes, only keeps scope with name starting with 'LoadModule'.
+			return ScopeName.Equals("LoadModule"); // When analyzing a tree of scopes, only keeps scope with name 'LoadModule'.
 		},
 		[&CollectedScopeSummaries](const FSummarizeScope& AllModulesStats, TArray<FSummarizeScope>&& ModuleStats)
 		{
