Further improve logging

diff --git a/Engine/Source/Editor/UnrealEd/Private/Kismet2/KismetReinstanceUtilities.cpp b/Engine/Source/Editor/UnrealEd/Private/Kismet2/KismetReinstanceUtilities.cpp
--- a/Engine/Source/Editor/UnrealEd/Private/Kismet2/KismetReinstanceUtilities.cpp
+++ b/Engine/Source/Editor/UnrealEd/Private/Kismet2/KismetReinstanceUtilities.cpp
@@ -103,8 +103,10 @@ struct FReplaceReferenceHelper
 			{
 				Classes.Append(Obj->GetClass()->GetName() + " ");
 			}
-			UE_LOG(LogBlueprint, Warning, TEXT("%s %s Replacing References to: %s"), 
-				*(GLogReinstancerReferenceReplacementObj->GetName()), 
+			UE_LOG(LogBlueprint, Warning, TEXT("%s %s %s %s Replacing References to: %s"), 
+				*(GLogReinstancerReferenceReplacementObj->GetOutermost()->GetName()),
+				GLogReinstancerReferenceReplacementObj->GetOuter() ? *(GLogReinstancerReferenceReplacementObj->GetOuter()->GetName()) : TEXT("nullouter"),
+				*(GLogReinstancerReferenceReplacementObj->GetName()),
 				*(GLogReinstancerReferenceReplacementObj->GetClass()->GetName()), 
 				*Classes);
 		}
