Added support for world partitions to base iterate packages commandlet. This much improves the performance of the Moderation Artifact generation.

diff --git a/Engine/Source/Editor/UnrealEd/Classes/Commandlets/BaseIteratePackagesCommandlet.h b/Engine/Source/Editor/UnrealEd/Classes/Commandlets/BaseIteratePackagesCommandlet.h
--- a/Engine/Source/Editor/UnrealEd/Classes/Commandlets/BaseIteratePackagesCommandlet.h
+++ b/Engine/Source/Editor/UnrealEd/Classes/Commandlets/BaseIteratePackagesCommandlet.h
@@ -7,12 +7,21 @@
 #include "UObject/ObjectMacros.h"
 #include "Commandlets/Commandlet.h"
 #include "Engine/EngineTypes.h"
+// world partition builder includes
+#include "WorldPartition/WorldPartitionBuilder.h"
+
 #include "BaseIteratePackagesCommandlet.generated.h"
 
 
+
+
 // Log category should be accessible by derived classes
 UNREALED_API DECLARE_LOG_CATEGORY_EXTERN(LogIteratePackagesCommandlet, Log, All);
 
+DECLARE_DELEGATE_TwoParams(FBaseIteratePackagesCommandletActorCallback, AActor* /*Actor*/, bool& /*bSavePackage*/);
+DECLARE_DELEGATE_TwoParams(FBaseIteratePackagesCommandletObjectCallback, UObject* /*Actor*/, bool& /*bSavePackage*/);
+
+
 UCLASS()
 // Added UNREALED_API to expose this to the save packages test
 // Base commandlet used to iterate packages provided on a commandline
@@ -83,6 +92,8 @@ protected:
 	/** if we should auto checkin packages that were checked out**/
 	bool bAutoCheckIn;
 
+	bool bUseWorldPartitionBuilder;
+
 	/** Should we build lighting for the packages we are saving? **/
 	//bool bShouldBuildLighting;
 
@@ -116,6 +127,8 @@ protected:
 	/** Ignore package version changelist **/
 	bool bIgnoreChangelist;
 
+	bool bKeepPackageGUIDOnSave;
+
 	/** Filter packages based on a collection **/
 	TSet<FName> CollectionFilter;
 
@@ -224,8 +237,12 @@ protected:
 	*/
 	virtual void PerformAdditionalOperations(class UWorld* World, bool& bSavePackage) { }
 
+	virtual void PerformWorldBuilderAdditionalOperations(class AActor* Actor, bool& bSavePackage) { }
+	virtual void PerformWorldBuilderAdditionalOperations(class UObject* Object, bool& bSavePackage) { }
+
 	virtual void PostPerformAdditionalOperations(class UPackage* Package) { }
 
+
 	/**
 	 * Allows the commandlet to do something after all the packages have been processed
 	 */
@@ -277,6 +294,7 @@ protected:
 
 private:
 
+	void SavePackages(const TArray<UPackage*>& PackagesToSave);
 
 public:		
 	//~ Begin UCommandlet Interface
diff --git a/Engine/Source/Editor/UnrealEd/Private/Commandlets/BaseIteratePackagesCommandlet.cpp b/Engine/Source/Editor/UnrealEd/Private/Commandlets/BaseIteratePackagesCommandlet.cpp
--- a/Engine/Source/Editor/UnrealEd/Private/Commandlets/BaseIteratePackagesCommandlet.cpp
+++ b/Engine/Source/Editor/UnrealEd/Private/Commandlets/BaseIteratePackagesCommandlet.cpp
@@ -30,13 +30,9 @@
 #include "SourceControlHelpers.h"
 #include "ISourceControlModule.h"
 #include "Engine/MapBuildDataRegistry.h"
-//#include "Particles/ParticleSystem.h"
 #include "EngineGlobals.h"
-//#include "Particles/ParticleEmitter.h"
 #include "GameFramework/WorldSettings.h"
-//#include "Engine/StaticMesh.h"
 #include "AssetRegistry/AssetData.h"
-//#include "Engine/Brush.h"
 #include "Editor.h"
 #include "EditorWorldUtils.h"
 #include "FileHelpers.h"
@@ -44,9 +40,6 @@
 #include "CollectionManagerModule.h"
 #include "ICollectionManager.h"
 #include "CommandletSourceControlUtils.h"
-//#include "WorldPartition/WorldPartition.h"
-//#include "WorldPartition/WorldPartitionHelpers.h"
-//#include "LevelInstance/LevelInstanceInterface.h"
 #include "AssetCompilingManager.h"
 #include "PackageHelperFunctions.h"
 #include "PackageTools.h"
@@ -60,39 +53,19 @@ DEFINE_LOG_CATEGORY(LogIteratePackagesCommandlet);
 #include "AssetRegistry/AssetRegistryModule.h"
 #include "IDirectoryWatcher.h"
 #include "DirectoryWatcherModule.h"
-//#include "Particles/Material/ParticleModuleMeshMaterial.h"
-//#include "Particles/ParticleLODLevel.h"
-//#include "Particles/ParticleModuleRequired.h"
-//#include "Particles/TypeData/ParticleModuleTypeDataMesh.h"
 #include "Engine/LevelStreaming.h"
 #include "EditorBuildUtils.h"
 
-// for UBaseIteratePackagesCommandlet::PerformAdditionalOperations building lighting code
-//#include "LightingBuildOptions.h"
-
-// for UBaseIteratePackagesCommandlet::PerformAdditionalOperations building navigation data
-//#include "EngineUtils.h"
-//#include "NavigationData.h"
-//#include "AI/NavigationSystemBase.h"
-
 // For preloading FFindInBlueprintSearchManager
 #include "FindInBlueprintManager.h"
 
 #include "ShaderCompiler.h"
-//#include "IHierarchicalLODUtilities.h"
-//#include "HierarchicalLODUtilitiesModule.h"
-//#include "HierarchicalLOD.h"
-//#include "HierarchicalLODProxyProcessor.h"
-//#include "HLOD/HLODEngineSubsystem.h"
-//#include "GenericPlatform/GenericPlatformProcess.h"
-//#include "HAL/ThreadManager.h"
-
-//#include "ICollectionManager.h"
-//#include "CollectionManagerModule.h"
-//#include "UObject/UObjectThreadContext.h"
-//#include "Engine/LODActor.h"
-//#include "PerQualityLevelProperties.h"
-//#include "Misc/RedirectCollector.h"
+
+// world partition includes
+#include "WorldPartition/WorldPartitionHelpers.h"
+#include "WorldPartition/WorldPartitionActorDesc.h"
+#include "WorldPartition/LoaderAdapter/LoaderAdapterShape.h"
+#include "UObject/GCObjectScopeGuard.h"
 
 /**-----------------------------------------------------------------------------
  *	UBaseIteratePackagesCommandlet commandlet.
@@ -461,6 +434,28 @@ void UBaseIteratePackagesCommandlet::OnAddResaveOnDemandPackage(FName SystemName
 	}
 }
 
+void UBaseIteratePackagesCommandlet::SavePackages(const TArray<UPackage*>& PackagesToSave)
+{
+	for (UPackage* PackageToSave : PackagesToSave)
+	{
+		const FString LongPackageName = PackageToSave->GetName();
+		FString Extension = FPackageName::GetAssetPackageExtension();
+		if (PackageToSave->ContainsMap())
+		{
+			Extension = FPackageName::GetMapPackageExtension();
+		}
+		FString FileName;
+		if (FPackageName::TryConvertLongPackageNameToFilename(LongPackageName, FileName, Extension))
+		{
+			ESaveFlags SaveFlags = bKeepPackageGUIDOnSave ? SAVE_KeepGUID : SAVE_None;
+			SavePackageHelper(PackageToSave, FileName, RF_Standalone, GWarn, SaveFlags);
+		}
+		else
+		{
+			UE_LOG(LogIteratePackagesCommandlet, Error, TEXT("Unable to save package %s, couldn't convert package name to filename"), *LongPackageName);
+		}
+	}
+}
 
 TSet<FName> UBaseIteratePackagesCommandlet::ParseResaveOnDemandSystems()
 {
@@ -611,7 +606,10 @@ void UBaseIteratePackagesCommandlet::LoadAndSaveOnePackage(const FString& Filena
 			UPackage* Package = nullptr;
 			{
 				TRACE_CPUPROFILER_EVENT_SCOPE(UBaseIteratePackagesCommandlet::LoadAndSaveOnePackage::LoadPackage);
-				Package = LoadPackage(NULL, *Filename, 0);
+
+				Package = LoadWorldPackageForEditor(PackagePath.GetPackageName());
+
+				// Package = LoadPackage(NULL, *Filename, 0);
 			}
 
 			if (Package == NULL)
@@ -646,25 +644,109 @@ void UBaseIteratePackagesCommandlet::LoadAndSaveOnePackage(const FString& Filena
 				}
 			}
 
+			UWorld* World = UWorld::FindWorldInPackage(Package);
+			if (World)
 			{
-				UWorld* World = UWorld::FindWorldInPackage(Package);
-				if (World)
-				{
-					PerformAdditionalOperations(World, bSavePackage);
-				}
+				PerformAdditionalOperations(World, bSavePackage);
 			}
-
+			
 			// hook to allow performing additional checks without lumping everything into this one function
 			PerformAdditionalOperations(Package,bSavePackage);
 
-			VerboseMessage(TEXT("Post PerformAdditionalOperations"));
+			if (bUseWorldPartitionBuilder && World && UWorld::IsPartitionedWorld(World))
+			{
+				// Load configuration file
+				FString WorldConfigFilename = FPackageName::LongPackageNameToFilename(World->GetPackage()->GetName(), TEXT(".ini"));
+				if (FPlatformFileManager::Get().GetPlatformFile().FileExists(*WorldConfigFilename))
+				{
+					LoadConfig(GetClass(), *WorldConfigFilename);
+				}
+
+
+				// process all the objects loaded with the base package.
+				ForEachObjectWithOuter(Package, [this, &bSavePackage](UObject* Object)
+					{
+						PerformAdditionalOperations(Object, bSavePackage);
+					});
+
+				UWorld::InitializationValues IVS;
+				IVS.RequiresHitProxies(false);
+				IVS.ShouldSimulatePhysics(false);
+				IVS.EnableTraceCollision(false);
+				IVS.CreateNavigation(false);
+				IVS.CreateAISystem(false);
+				IVS.AllowAudioPlayback(false);
+				IVS.CreatePhysicsScene(true);
+
+				FScopedEditorWorld ScopeEditorWorld(World, IVS);
+
+				/*TUniquePtr<FLoaderAdapterShape> LoaderAdapterShape;
+				FBox Bounds = FBox(FVector(-HALF_WORLD_MAX, -HALF_WORLD_MAX, -HALF_WORLD_MAX), FVector(HALF_WORLD_MAX, HALF_WORLD_MAX, HALF_WORLD_MAX));
+				LoaderAdapterShape = MakeUnique<FLoaderAdapterShape>(World, Bounds, TEXT("Loaded Region"));
+				LoaderAdapterShape->Load();*/
+
+				TArray<UPackage*> PackagesToSave;
 
-			// Check for any special per object operations
-			ForEachObjectWithOuter(Package, [this, &bSavePackage](UObject* Object)
+				FWorldPartitionHelpers::FForEachActorWithLoadingParams ForEachActorWithLoadingParams;
+
+				ForEachActorWithLoadingParams.ActorClasses = { AActor::StaticClass()};
+
+				ForEachActorWithLoadingParams.OnPreGarbageCollect = [&PackagesToSave, this]()
+					{
+						SavePackages(PackagesToSave);
+						//UWorldPartitionBuilder::SavePackages(PackagesToSave, PackageHelper, true);
+						PackagesToSave.Empty();
+					};
+
+				UWorldPartition* WorldPartition = World->GetWorldPartition();
+
+				FWorldPartitionHelpers::ForEachActorWithLoading(WorldPartition, [&PackagesToSave, this](const FWorldPartitionActorDesc* ActorDesc)
+					{
+						AActor* Actor = ActorDesc->GetActor();
+
+						if (!Actor)
+						{
+							return true;
+						}
+
+						bool bSavePackage = false;
+						PerformWorldBuilderAdditionalOperations(Actor, bSavePackage);
+
+						UPackage* Package = Actor->GetExternalPackage();
+						check(Package);
+
+						TArray<UObject*> DependantObjects;
+						ForEachObjectWithPackage(Package, [this, &bSavePackage](UObject* Object)
+							{
+								if (!Cast<UMetaData>(Object))
+								{
+									PerformWorldBuilderAdditionalOperations(Object, bSavePackage);
+								}
+								return true;
+							}, true);
+
+						if (bSavePackage)
+						{
+							PackagesToSave.Add(Package);
+							return true;
+						}
+						return true;
+					}, ForEachActorWithLoadingParams);
+
+				SavePackages(PackagesToSave);
+				PackagesToSave.Empty();
+			}
+			else
 			{
-				PerformAdditionalOperations(Object, bSavePackage);
-			});
-			
+				VerboseMessage(TEXT("Post PerformAdditionalOperations"));
+
+				// Check for any special per object operations
+				ForEachObjectWithOuter(Package, [this, &bSavePackage](UObject* Object)
+				{
+					PerformAdditionalOperations(Object, bSavePackage);
+				});
+			}
+
 			PostPerformAdditionalOperations(Package);
 
 			VerboseMessage(TEXT("Post PerformAdditionalOperations Loop"));
@@ -789,7 +871,6 @@ void UBaseIteratePackagesCommandlet::LoadAndSaveOnePackage(const FString& Filena
 							*Filename,Linker->Summary.GetFileVersionUE().ToValue(), Linker->Summary.GetFileVersionLicenseeUE(), GPackageFileUEVersion.ToValue(), VER_LATEST_ENGINE_LICENSEEUE4 );
 					}
 
-					const static bool bKeepPackageGUIDOnSave = FParse::Param(FCommandLine::Get(), TEXT("KeepPackageGUIDOnSave"));
 					ESaveFlags SaveFlags = bKeepPackageGUIDOnSave ? SAVE_KeepGUID : SAVE_None;
 					
 					if (bIsReadOnly == false || SourceControlQueue == nullptr)
@@ -819,7 +900,6 @@ void UBaseIteratePackagesCommandlet::LoadAndSaveOnePackage(const FString& Filena
 				}
 			}
 		}
-
 		static int32 Counter = 0;
 
 		if (!GarbageCollectionFrequency || Counter++ % GarbageCollectionFrequency == 0)
@@ -842,6 +922,7 @@ void UBaseIteratePackagesCommandlet::LoadAndSaveOnePackage(const FString& Filena
 	}
 }
 
+
 void UBaseIteratePackagesCommandlet::DeleteOnePackage(const FString& Filename)
 {
 	TRACE_CPUPROFILER_EVENT_SCOPE(UBaseIteratePackagesCommandlet::DeleteOnePackage);
@@ -977,6 +1058,10 @@ int32 UBaseIteratePackagesCommandlet::Main( const FString& Params )
 	bOnlyPayloadTrailers = Switches.Contains(TEXT("OnlyPayloadTrailers"));
 	/** only process packages containing materials */
 	bOnlyMaterials = Switches.Contains(TEXT("onlymaterials"));
+
+	bUseWorldPartitionBuilder = false;
+
+	bKeepPackageGUIDOnSave = Switches.Contains(TEXT("KeepPackageGUIDOnSave"));
 	
 	/** check for filtering packages by collection. **/
 	FString FilterByCollection;
