PR #9702: [Source Control][Subversion][Plastic SCM] Fix regression: missing "Make writable" button in PromptToCheckoutPackages when the source control provider doesn't uses readonly flags! (Contributed by [REDACTED_PERSON])

diff --git a/Engine/Source/Editor/UnrealEd/Private/FileHelpers.cpp b/Engine/Source/Editor/UnrealEd/Private/FileHelpers.cpp
--- a/Engine/Source/Editor/UnrealEd/Private/FileHelpers.cpp
+++ b/Engine/Source/Editor/UnrealEd/Private/FileHelpers.cpp
@@ -1847,12 +1847,10 @@ bool FEditorFileUtils::PromptToCheckoutPackagesInternal(bool bCheckDirty, const
 		// The checkout button should be disabled if no packages can be checked out.
 		CheckoutPackagesDialogModule.AddButton(DRT_CheckOut, NSLOCTEXT("PackagesDialogModule", "Dlg_CheckOutButtonp", "Check Out Selected"), NSLOCTEXT("PackagesDialogModule", "Dlg_CheckOutTooltip", "Attempt to Check Out Checked Assets"), CheckOutSelectedDisabledAttrib );
 		
-		// Make writable button to make checked files writable (if supported by source control provider)
-		if (SourceControlProvider.UsesLocalReadOnlyState())
-		{
-			CheckoutPackagesDialogModule.AddButton(DRT_MakeWritable, NSLOCTEXT("PackagesDialogModule", "Dlg_MakeWritableButton", "Make Writable"), NSLOCTEXT("PackagesDialogModule", "Dlg_MakeWritableTooltip", "Makes selected files writable on disk"));
-		}
-
+		// Make writable button to make checked files writable (to be able to save them)
+		// Note: this is needed when unable to checkout (not only for Perforce, but also for Plastic SCM and Subversion, even though they don't use read-only flags)
+		CheckoutPackagesDialogModule.AddButton(DRT_MakeWritable, NSLOCTEXT("PackagesDialogModule", "Dlg_MakeWritableButton", "Make Writable"), NSLOCTEXT("PackagesDialogModule", "Dlg_MakeWritableTooltip", "Makes selected files writable on disk"));
+	
 		if (bAllowSkip)
 		{
 			// Skip button to skip checkout step
