Speculative fixes for nullptr this check in UStruct::IsChildOf

#rnx

diff --git a/Engine/Source/Editor/LevelEditor/Private/SLevelEditor.cpp b/Engine/Source/Editor/LevelEditor/Private/SLevelEditor.cpp
--- a/Engine/Source/Editor/LevelEditor/Private/SLevelEditor.cpp
+++ b/Engine/Source/Editor/LevelEditor/Private/SLevelEditor.cpp
@@ -1917,7 +1917,7 @@ void SLevelEditor::HandleAssetsDeleted(const TArray<UClass*>& DeletedClasses)
 	bool bDeletedMaterials = false;
 	for (UClass* AssetClass : DeletedClasses)
 	{
-		if (AssetClass->IsChildOf<UMaterialInterface>())
+		if (AssetClass == nullptr || AssetClass->IsChildOf<UMaterialInterface>())
 		{
 			bDeletedMaterials = true;
 			break;
diff --git a/Engine/Source/Editor/UMGEditor/Private/Palette/SPaletteViewModel.cpp b/Engine/Source/Editor/UMGEditor/Private/Palette/SPaletteViewModel.cpp
--- a/Engine/Source/Editor/UMGEditor/Private/Palette/SPaletteViewModel.cpp
+++ b/Engine/Source/Editor/UMGEditor/Private/Palette/SPaletteViewModel.cpp
@@ -667,9 +667,9 @@ void FWidgetCatalogViewModel::OnReloadComplete(EReloadCompleteReason Reason)
 
 void FWidgetCatalogViewModel::HandleOnAssetsDeleted(const TArray<UClass*>& DeletedAssetClasses)
 {
-	for (auto DeletedAssetClass : DeletedAssetClasses)
+	for (const UClass* DeletedAssetClass : DeletedAssetClasses)
 	{
-		if (DeletedAssetClass->IsChildOf(UWidgetBlueprint::StaticClass()))
+		if ((DeletedAssetClass == nullptr) || DeletedAssetClass->IsChildOf(UWidgetBlueprint::StaticClass()))
 		{
 			bRebuildRequested = true;
 		}
diff --git a/Engine/Source/Editor/UnrealEd/Private/EditorEngine.cpp b/Engine/Source/Editor/UnrealEd/Private/EditorEngine.cpp
--- a/Engine/Source/Editor/UnrealEd/Private/EditorEngine.cpp
+++ b/Engine/Source/Editor/UnrealEd/Private/EditorEngine.cpp
@@ -510,26 +510,29 @@ void UEditorEngine::GetContentBrowserSelections(TArray<FAssetData>& Selection) c
 
 USelection* UEditorEngine::GetSelectedSet( const UClass* Class ) const
 {
-	USelection* SelectedSet = GetSelectedActors();
-	if ( Class->IsChildOf( AActor::StaticClass() ) )
+	if (Class != nullptr)
 	{
-		return SelectedSet;
-	}
-	else
-	{
-		//make sure this actor isn't derived off of an interface class
-		for ( FSelectionIterator It( GetSelectedActorIterator() ) ; It ; ++It )
+		USelection* SelectedSet = GetSelectedActors();
+		if (Class->IsChildOf(AActor::StaticClass()))
+		{
+			return SelectedSet;
+		}
+		else
 		{
-			AActor* TestActor = static_cast<AActor*>( *It );
-			if (TestActor->GetClass()->ImplementsInterface(Class))
+			//make sure this actor isn't derived off of an interface class
+			for (FSelectionIterator It(GetSelectedActorIterator()); It; ++It)
 			{
-				return SelectedSet;
+				AActor* TestActor = static_cast<AActor*>(*It);
+				if (TestActor->GetClass()->ImplementsInterface(Class))
+				{
+					return SelectedSet;
+				}
 			}
 		}
-
-		//no actor matched the interface class
-		return GetSelectedObjects();
 	}
+
+	//no actor matched the interface class
+	return GetSelectedObjects();
 }
 
 const UClass* UEditorEngine::GetFirstSelectedClass( const UClass* const RequiredParentClass ) const
@@ -741,6 +744,11 @@ void UEditorEngine::InitEditor(IEngineLoop* InEngineLoop)
 		for (TObjectIterator<UClass> ObjectIt; ObjectIt; ++ObjectIt)
 		{
 			UClass* TestClass = *ObjectIt;
+			if (TestClass == nullptr)
+			{
+				continue;
+			}
+
 			if (TestClass->IsChildOf(UActorFactory::StaticClass()))
 			{
 				if (!TestClass->HasAnyClassFlags(CLASS_Abstract))
@@ -772,7 +780,7 @@ void UEditorEngine::InitEditor(IEngineLoop* InEngineLoop)
 			const UClass* DefaultActorClass = VolumeFactoryClass->GetDefaultObject<UActorFactory>()->NewActorClass;
 			for (UClass* VolumeClass : VolumeClasses)
 			{
-				if (DefaultActorClass && !VolumeClass->IsChildOf(DefaultActorClass))
+				if (DefaultActorClass && (VolumeClass && !VolumeClass->IsChildOf(DefaultActorClass)))
 				{
 					continue;
 				}
@@ -1174,7 +1182,12 @@ void UEditorEngine::CreateVolumeFactoriesForNewClasses(const TArray<UClass*>& Ne
 		for (TObjectIterator<UClass> ObjectIt; ObjectIt; ++ObjectIt)
 		{
 			UClass* TestClass = *ObjectIt;
-			if (!TestClass->HasAnyClassFlags(CLASS_Abstract) && TestClass->IsChildOf(UActorFactoryVolume::StaticClass()))
+			if (TestClass == nullptr)
+			{
+				continue;
+			}
+		
+		if (!TestClass->HasAnyClassFlags(CLASS_Abstract) && TestClass->IsChildOf(UActorFactoryVolume::StaticClass()))
 			{
 				ActorFactories.Reserve(ActorFactories.Num() + NewVolumeClasses.Num());
 				for (UClass* NewVolumeClass : NewVolumeClasses)
