[Editor]
* Removed includes from EditorEngine.h

diff --git a/Engine/Source/Editor/UnrealEd/Classes/Editor/EditorEngine.h b/Engine/Source/Editor/UnrealEd/Classes/Editor/EditorEngine.h
--- a/Engine/Source/Editor/UnrealEd/Classes/Editor/EditorEngine.h
+++ b/Engine/Source/Editor/UnrealEd/Classes/Editor/EditorEngine.h
@@ -3,6 +3,7 @@
 #pragma once
 
 #include "CoreMinimal.h"
+#include "Elements/Framework/TypedElementHandle.h"
 #include "SlateFwd.h"
 #include "UObject/ObjectMacros.h"
 #include "UObject/UObjectGlobals.h"
@@ -17,7 +18,6 @@
 #include "TimerManager.h"
 #include "UObject/UObjectAnnotation.h"
 #include "Engine/Brush.h"
-#include "Model.h"
 #include "Engine/Engine.h"
 #if UE_ENABLE_INCLUDE_ORDER_DEPRECATED_IN_5_2
 #include "Settings/LevelEditorPlaySettings.h"
@@ -29,12 +29,16 @@
 #include "PlayInEditorDataTypes.h"
 #include "EditorSubsystem.h"
 #include "Subsystems/SubsystemCollection.h"
-#include "RHI.h"
 #include "UnrealEngine.h"
 #include "Templates/PimplPtr.h"
 #include "Templates/UniqueObj.h"
 #include "Editor/AssetReferenceFilter.h"
 
+#if UE_ENABLE_INCLUDE_ORDER_DEPRECATED_IN_5_2
+#include "Model.h"
+#include "RHI.h"
+#endif
+
 #include "EditorEngine.generated.h"
 
 class APlayerStart;
@@ -267,11 +271,7 @@ struct FPreviewPlatformInfo
 	}
 
 	/** returns the preview feature level if active, or GMaxRHIFeatureLevel otherwise */
-	ERHIFeatureLevel::Type GetEffectivePreviewFeatureLevel() const
-	{
-		return bPreviewFeatureLevelActive ? PreviewFeatureLevel : GMaxRHIFeatureLevel;
-	}
-
+	UNREALED_API ERHIFeatureLevel::Type GetEffectivePreviewFeatureLevel() const;
 };
 
 /** Struct used in filtering allowed references between assets. Passes context about the referencers to game-level filters */
diff --git a/Engine/Source/Editor/UnrealEd/Private/EditorEngine.cpp b/Engine/Source/Editor/UnrealEd/Private/EditorEngine.cpp
--- a/Engine/Source/Editor/UnrealEd/Private/EditorEngine.cpp
+++ b/Engine/Source/Editor/UnrealEd/Private/EditorEngine.cpp
@@ -335,6 +335,12 @@ void DestroySelectionSets()
 */
 static TMap<UPackage*, bool> StartupPackageToWarnState;
 
+
+ERHIFeatureLevel::Type FPreviewPlatformInfo::GetEffectivePreviewFeatureLevel() const
+{
+	return bPreviewFeatureLevelActive ? PreviewFeatureLevel : GMaxRHIFeatureLevel;
+}
+
 //////////////////////////////////////////////////////////////////////////
 // UEditorEngine
 
