[Mode Manager] Fix a regression where ActorsDuplicatedNotify was not called in some cases

robomerge[FN_EngineMerge] -UE5-RES

diff --git a/Engine/Source/Editor/UnrealEd/Private/LevelEditorViewport.cpp b/Engine/Source/Editor/UnrealEd/Private/LevelEditorViewport.cpp
--- a/Engine/Source/Editor/UnrealEd/Private/LevelEditorViewport.cpp
+++ b/Engine/Source/Editor/UnrealEd/Private/LevelEditorViewport.cpp
@@ -2797,6 +2797,13 @@ bool FLevelEditorViewportClient::InputWidgetDelta(FViewport* InViewport, EAxisLi
 
 									SelectionSet->SetSelection(DuplicatedElements, SelectionOptions);
 
+									// Selected actors are the same as the selection set at this point, send the legacy mode tools update
+									{
+										TArray<AActor*> SelectedActors = SelectionSet->GetSelectedObjects<AActor>();
+										constexpr bool bDidOffsetDuplicate = false;
+										ModeTools->ActorsDuplicatedNotify(SelectedActors, SelectedActors, bDidOffsetDuplicate);
+									}
+
 									// Force the cached manipulation list to update, as we don't notify about the change above
 									ResetElementsToManipulate();
 
diff --git a/Engine/Source/Editor/UnrealEd/Private/UnrealEdSrv.cpp b/Engine/Source/Editor/UnrealEd/Private/UnrealEdSrv.cpp
--- a/Engine/Source/Editor/UnrealEd/Private/UnrealEdSrv.cpp
+++ b/Engine/Source/Editor/UnrealEd/Private/UnrealEdSrv.cpp
@@ -2779,11 +2779,16 @@ bool UUnrealEdEngine::Exec_Actor( UWorld* InWorld, const TCHAR* Str, FOutputDevi
 					}
 					else
 					{
-						const TArray<FTypedElementHandle> DuplicatedElements = CommonActions->DuplicateSelectedElements(SelectionSet, InWorld, GEditor->GetGridLocationOffset(/*bUniformOffset*/false));
+						const FVector DuplicateOffset = GEditor->GetGridLocationOffset(/*bUniformOffset*/false);
+						const TArray<FTypedElementHandle> DuplicatedElements = CommonActions->DuplicateSelectedElements(SelectionSet, InWorld, DuplicateOffset);
 						if (DuplicatedElements.Num() > 0)
 						{
 							SelectionSet->SetSelection(DuplicatedElements, FTypedElementSelectionOptions());
 							SelectionSet->NotifyPendingChanges();
+
+							// notify the global mode tools, the selection set should be identical to the new actors at this point
+							TArray<AActor*> SelectedActors = SelectionSet->GetSelectedObjects<AActor>();
+							GLevelEditorModeTools().ActorsDuplicatedNotify(SelectedActors, SelectedActors, DuplicateOffset != FVector::ZeroVector);
 						}
 					}
 				}
