Wait on the asset registry to be fully constructed when starting PIE to avoid gameplay errors relying on asset search

#rnx

diff --git a/Engine/Source/Editor/UnrealEd/Private/PlayLevel.cpp b/Engine/Source/Editor/UnrealEd/Private/PlayLevel.cpp
--- a/Engine/Source/Editor/UnrealEd/Private/PlayLevel.cpp
+++ b/Engine/Source/Editor/UnrealEd/Private/PlayLevel.cpp
@@ -2,6 +2,8 @@
 
 #include "PlayLevel.h"
 #include "CoreMinimal.h"
+#include "AssetRegistry/AssetRegistryModule.h"
+#include "AssetRegistry/IAssetRegistry.h"
 #include "Misc/MessageDialog.h"
 #include "Misc/CommandLine.h"
 #include "Misc/Paths.h"
@@ -2622,6 +2624,9 @@ void UEditorEngine::StartPlayInEditorSession(FRequestPlaySessionParams& InReques
 	// Make sure there's no outstanding load requests
 	FlushAsyncLoading();
 
+	// Gameplay relies on asset registry to be fully constructed, wait for completion before starting PIE
+	IAssetRegistry::GetChecked().WaitForCompletion();
+
 	// Update the Blueprint Debugger 
 	FBlueprintEditorUtils::FindAndSetDebuggableBlueprintInstances();
 
