Addition of 'view' method to AssetTypeActions

Iteration on Jason Stasiks CL 17314610. This is a minimal version of CL 17314610 which adds a 'view' method to AssetTypeActions which is now attempted whenever an activated item's 'edit' attempt fails.

The matching CL for the restricted FAssetTypeActions_VerseClass changes is 22472079.

diff --git a/Engine/Plugins/Editor/ContentBrowser/ContentBrowserAliasDataSoure/Source/ContentBrowserAliasDataSource/Private/ContentBrowserAliasDataSource.cpp b/Engine/Plugins/Editor/ContentBrowser/ContentBrowserAliasDataSoure/Source/ContentBrowserAliasDataSource/Private/ContentBrowserAliasDataSource.cpp
--- a/Engine/Plugins/Editor/ContentBrowser/ContentBrowserAliasDataSoure/Source/ContentBrowserAliasDataSource/Private/ContentBrowserAliasDataSource.cpp
+++ b/Engine/Plugins/Editor/ContentBrowser/ContentBrowserAliasDataSoure/Source/ContentBrowserAliasDataSource/Private/ContentBrowserAliasDataSource.cpp
@@ -692,6 +692,21 @@ bool UContentBrowserAliasDataSource::BulkEditItems(TArrayView<const FContentBrow
 	return ContentBrowserAssetData::EditItems(AssetTools, this, InItems);
 }
 
+bool UContentBrowserAliasDataSource::CanViewItem(const FContentBrowserItemData& InItem, FText* OutErrorMsg)
+{
+	return ContentBrowserAssetData::CanViewItem(AssetTools, this, InItem, OutErrorMsg);
+}
+
+bool UContentBrowserAliasDataSource::ViewItem(const FContentBrowserItemData& InItem)
+{
+	return ContentBrowserAssetData::ViewItems(AssetTools, this, MakeArrayView(&InItem, 1));
+}
+
+bool UContentBrowserAliasDataSource::BulkViewItems(TArrayView<const FContentBrowserItemData> InItems)
+{
+	return ContentBrowserAssetData::ViewItems(AssetTools, this, InItems);
+}
+
 bool UContentBrowserAliasDataSource::CanPreviewItem(const FContentBrowserItemData& InItem, FText* OutErrorMsg)
 {
 	return ContentBrowserAssetData::CanPreviewItem(AssetTools, this, InItem, OutErrorMsg);
diff --git a/Engine/Plugins/Editor/ContentBrowser/ContentBrowserAliasDataSoure/Source/ContentBrowserAliasDataSource/Public/ContentBrowserAliasDataSource.h b/Engine/Plugins/Editor/ContentBrowser/ContentBrowserAliasDataSoure/Source/ContentBrowserAliasDataSource/Public/ContentBrowserAliasDataSource.h
--- a/Engine/Plugins/Editor/ContentBrowser/ContentBrowserAliasDataSoure/Source/ContentBrowserAliasDataSource/Public/ContentBrowserAliasDataSource.h
+++ b/Engine/Plugins/Editor/ContentBrowser/ContentBrowserAliasDataSoure/Source/ContentBrowserAliasDataSource/Public/ContentBrowserAliasDataSource.h
@@ -70,6 +70,10 @@ public:
 	virtual bool CanEditItem(const FContentBrowserItemData& InItem, FText* OutErrorMsg) override;
 	virtual bool EditItem(const FContentBrowserItemData& InItem) override;
 	virtual bool BulkEditItems(TArrayView<const FContentBrowserItemData> InItems) override;
+	
+	virtual bool CanViewItem(const FContentBrowserItemData& InItem, FText* OutErrorMsg) override;
+	virtual bool ViewItem(const FContentBrowserItemData& InItem) override;
+	virtual bool BulkViewItems(TArrayView<const FContentBrowserItemData> InItems) override;
 
 	virtual bool CanPreviewItem(const FContentBrowserItemData& InItem, FText* OutErrorMsg) override;
 	virtual bool PreviewItem(const FContentBrowserItemData& InItem) override;
diff --git a/Engine/Plugins/Editor/ContentBrowser/ContentBrowserAssetDataSource/Source/ContentBrowserAssetDataSource/Private/ContentBrowserAssetDataCore.cpp b/Engine/Plugins/Editor/ContentBrowser/ContentBrowserAssetDataSource/Source/ContentBrowserAssetDataSource/Private/ContentBrowserAssetDataCore.cpp
--- a/Engine/Plugins/Editor/ContentBrowser/ContentBrowserAssetDataSource/Source/ContentBrowserAssetDataSource/Private/ContentBrowserAssetDataCore.cpp
+++ b/Engine/Plugins/Editor/ContentBrowser/ContentBrowserAssetDataSource/Source/ContentBrowserAssetDataSource/Private/ContentBrowserAssetDataCore.cpp
@@ -5,6 +5,7 @@
 #include "IAssetTools.h"
 #include "AssetRegistry/IAssetRegistry.h"
 #include "AssetViewUtils.h"
+#include "AssetTypeActivationOpenedMethod.h"
 #include "AssetPropertyTagCache.h"
 #include "ObjectTools.h"
 #include "Misc/NamePermissionList.h"
@@ -205,6 +206,26 @@ bool CanEditAssetFileItem(IAssetTools* InAssetTools, const FContentBrowserAssetF
 	return true;
 }
 
+bool CanViewItem(IAssetTools* InAssetTools, const UContentBrowserDataSource* InOwnerDataSource, const FContentBrowserItemData& InItem, FText* OutErrorMsg)
+{
+	if (TSharedPtr<const FContentBrowserAssetFileItemDataPayload> AssetPayload = GetAssetFileItemPayload(InOwnerDataSource, InItem))
+	{
+		return CanViewAssetFileItem(InAssetTools, *AssetPayload, OutErrorMsg);
+	}
+
+	return false;
+}
+
+bool CanViewAssetFileItem(IAssetTools* InAssetTools, const FContentBrowserAssetFileItemDataPayload& InAssetPayload, FText* OutErrorMsg)
+{
+	const TWeakPtr<IAssetTypeActions> AssetTypeActions = InAssetTools->GetAssetTypeActionsForClass(InAssetPayload.GetAssetData().GetClass());
+	if (AssetTypeActions.IsValid())
+	{
+		return AssetTypeActions.Pin()->SupportsOpenedMethod(EAssetTypeActivationOpenedMethod::View);
+	}
+	return false;
+}
+
 bool CanPreviewItem(IAssetTools* InAssetTools, const UContentBrowserDataSource* InOwnerDataSource, const FContentBrowserItemData& InItem, FText* OutErrorMsg)
 {
 	if (TSharedPtr<const FContentBrowserAssetFileItemDataPayload> AssetPayload = GetAssetFileItemPayload(InOwnerDataSource, InItem))
@@ -220,14 +241,15 @@ bool CanPreviewAssetFileItem(IAssetTools* InAssetTools, const FContentBrowserAss
 	return true;
 }
 
-bool EditOrPreviewAssetFileItems(TArrayView<const TSharedRef<const FContentBrowserAssetFileItemDataPayload>> InAssetPayloads, const bool bIsPreview)
+bool EditOrPreviewAssetFileItems(TArrayView<const TSharedRef<const FContentBrowserAssetFileItemDataPayload>> InAssetPayloads, EAssetTypeActivationMethod::Type ActivationMethod, const EAssetTypeActivationOpenedMethod OpenedMethod)
 {
 	if (InAssetPayloads.Num() == 0)
 	{
 		return false;
 	}
 
-	const EAssetTypeActivationMethod::Type ActivationMethod = bIsPreview ? EAssetTypeActivationMethod::Previewed : EAssetTypeActivationMethod::Opened;
+	ensure(ActivationMethod != EAssetTypeActivationMethod::Type::DoubleClicked);
+
 	UAssetEditorSubsystem* AssetEditorSubsystem = GEditor->GetEditorSubsystem<UAssetEditorSubsystem>();
 
 	TMap<TSharedPtr<IAssetTypeActions>, TArray<FAssetData>> TypeActionsToAssetData;
@@ -259,7 +281,7 @@ bool EditOrPreviewAssetFileItems(TArrayView<const TSharedRef<const FContentBrows
 		TArray<FAssetData>& AssetsToLoad = TypeActionToObjectsPair.Value;
 		if (TypeActions.IsValid())
 		{
-			AssetsToLoad = TypeActions->GetValidAssetsForPreviewOrEdit(AssetsToLoad, bIsPreview);
+			AssetsToLoad = TypeActions->GetValidAssetsForPreviewOrEdit(AssetsToLoad, ActivationMethod == EAssetTypeActivationMethod::Type::Previewed);
 		}
 
 		TArray<UObject*> ObjList;
@@ -285,47 +307,84 @@ bool EditOrPreviewAssetFileItems(TArrayView<const TSharedRef<const FContentBrows
 
 		if (bOpenEditorForAssets)
 		{
-			bSuccessfulEditorOpen &= AssetEditorSubsystem->OpenEditorForAssets(ObjList);
+			bSuccessfulEditorOpen &= AssetEditorSubsystem->OpenEditorForAssets(ObjList, OpenedMethod);
 		}
 	}
 
 	return bSuccessfulEditorOpen;
 }
 
-bool EditOrPreviewItems(IAssetTools* InAssetTools, const UContentBrowserDataSource* InOwnerDataSource, TArrayView<const FContentBrowserItemData> InItems, const bool bIsPreview)
+bool EditOrPreviewItems(IAssetTools* InAssetTools, const UContentBrowserDataSource* InOwnerDataSource, TArrayView<const FContentBrowserItemData> InItems, EAssetTypeActivationMethod::Type ActivationMethod, EAssetTypeActivationOpenedMethod OpenedMethod)
 {
-	TArray<TSharedRef<const FContentBrowserAssetFileItemDataPayload>, TInlineAllocator<16>> AssetPayloads;
+	TArray<TSharedRef<const FContentBrowserAssetFileItemDataPayload>, TInlineAllocator<16>> EditOrPreviewPayloads;
+	TArray<TSharedRef<const FContentBrowserAssetFileItemDataPayload>, TInlineAllocator<16>> ViewPayloads;
 
-	EnumerateAssetFileItemPayloads(InOwnerDataSource, InItems, [InAssetTools, bIsPreview , &AssetPayloads](const TSharedRef<const FContentBrowserAssetFileItemDataPayload>& InAssetPayload)
+	ensure(ActivationMethod != EAssetTypeActivationMethod::Type::DoubleClicked);
+
+	const bool bIsPreview = ActivationMethod == EAssetTypeActivationMethod::Type::Previewed;
+	EnumerateAssetFileItemPayloads(InOwnerDataSource, InItems, [InAssetTools, bIsPreview, OpenedMethod, &EditOrPreviewPayloads, &ViewPayloads](const TSharedRef<const FContentBrowserAssetFileItemDataPayload>& InAssetPayload)
 	{
-		if ((bIsPreview ? CanPreviewAssetFileItem(InAssetTools, *InAssetPayload, nullptr) : CanEditAssetFileItem(InAssetTools, *InAssetPayload, nullptr)))
+		if (bIsPreview)
 		{
-			AssetPayloads.Add(InAssetPayload);
+			if (CanPreviewAssetFileItem(InAssetTools, *InAssetPayload, nullptr))
+			{
+				EditOrPreviewPayloads.Add(InAssetPayload);
+			}
+		}
+		else if (OpenedMethod == EAssetTypeActivationOpenedMethod::Edit)
+		{
+			if (CanEditAssetFileItem(InAssetTools, *InAssetPayload, nullptr))
+			{
+				EditOrPreviewPayloads.Add(InAssetPayload);
+			}
+			else if (CanViewAssetFileItem(InAssetTools, *InAssetPayload, nullptr))
+			{
+				ViewPayloads.Add(InAssetPayload);
+			}
+		}
+		else if (OpenedMethod == EAssetTypeActivationOpenedMethod::View)
+		{
+			if (CanViewAssetFileItem(InAssetTools, *InAssetPayload, nullptr))
+			{
+				ViewPayloads.Add(InAssetPayload);
+			}
 		}
 		return true;
 	});
 
-	return EditOrPreviewAssetFileItems(AssetPayloads, bIsPreview);
+	const bool bEditItems = EditOrPreviewAssetFileItems(EditOrPreviewPayloads, ActivationMethod, EAssetTypeActivationOpenedMethod::Edit);
+	const bool bViewItems = EditOrPreviewAssetFileItems(ViewPayloads, ActivationMethod, EAssetTypeActivationOpenedMethod::View);
+	return bEditItems || bViewItems;
 }
 
 bool EditItems(IAssetTools* InAssetTools, const UContentBrowserDataSource* InOwnerDataSource, TArrayView<const FContentBrowserItemData> InItems)
 {
-	return EditOrPreviewItems(InAssetTools, InOwnerDataSource, InItems, /*bIsPreview*/false);
+	return EditOrPreviewItems(InAssetTools, InOwnerDataSource, InItems, EAssetTypeActivationMethod::Type::Opened, EAssetTypeActivationOpenedMethod::Edit);
 }
 
 bool EditAssetFileItems(TArrayView<const TSharedRef<const FContentBrowserAssetFileItemDataPayload>> InAssetPayloads)
 {
-	return EditOrPreviewAssetFileItems(InAssetPayloads, /*bIsPreview*/false);
+	return EditOrPreviewAssetFileItems(InAssetPayloads, EAssetTypeActivationMethod::Type::Opened, EAssetTypeActivationOpenedMethod::Edit);
+}
+
+bool ViewItems(IAssetTools* InAssetTools, const UContentBrowserDataSource* InOwnerDataSource, TArrayView<const FContentBrowserItemData> InItems)
+{
+	return EditOrPreviewItems(InAssetTools, InOwnerDataSource, InItems, EAssetTypeActivationMethod::Type::Opened, EAssetTypeActivationOpenedMethod::View);
+}
+
+bool ViewAssetFileItems(TArrayView<const TSharedRef<const FContentBrowserAssetFileItemDataPayload>> InAssetPayloads)
+{
+	return EditOrPreviewAssetFileItems(InAssetPayloads, EAssetTypeActivationMethod::Type::Opened, EAssetTypeActivationOpenedMethod::View);
 }
 
 bool PreviewItems(IAssetTools* InAssetTools, const UContentBrowserDataSource* InOwnerDataSource, TArrayView<const FContentBrowserItemData> InItems)
 {
-	return EditOrPreviewItems(InAssetTools, InOwnerDataSource, InItems, /*bIsPreview*/true);
+	return EditOrPreviewItems(InAssetTools, InOwnerDataSource, InItems, EAssetTypeActivationMethod::Type::Previewed, EAssetTypeActivationOpenedMethod::Edit);
 }
 
 bool PreviewAssetFileItems(TArrayView<const TSharedRef<const FContentBrowserAssetFileItemDataPayload>> InAssetPayloads)
 {
-	return EditOrPreviewAssetFileItems(InAssetPayloads, /*bIsPreview*/true);
+	return EditOrPreviewAssetFileItems(InAssetPayloads, EAssetTypeActivationMethod::Type::Previewed, EAssetTypeActivationOpenedMethod::Edit);
 }
 
 bool CanDuplicateItem(IAssetTools* InAssetTools, const UContentBrowserDataSource* InOwnerDataSource, const FContentBrowserItemData& InItem, FText* OutErrorMsg)
diff --git a/Engine/Plugins/Editor/ContentBrowser/ContentBrowserAssetDataSource/Source/ContentBrowserAssetDataSource/Private/ContentBrowserAssetDataSource.cpp b/Engine/Plugins/Editor/ContentBrowser/ContentBrowserAssetDataSource/Source/ContentBrowserAssetDataSource/Private/ContentBrowserAssetDataSource.cpp
--- a/Engine/Plugins/Editor/ContentBrowser/ContentBrowserAssetDataSource/Source/ContentBrowserAssetDataSource/Private/ContentBrowserAssetDataSource.cpp
+++ b/Engine/Plugins/Editor/ContentBrowser/ContentBrowserAssetDataSource/Source/ContentBrowserAssetDataSource/Private/ContentBrowserAssetDataSource.cpp
@@ -1237,6 +1237,21 @@ bool UContentBrowserAssetDataSource::BulkEditItems(TArrayView<const FContentBrow
 	return ContentBrowserAssetData::EditItems(AssetTools, this, InItems);
 }
 
+bool UContentBrowserAssetDataSource::CanViewItem(const FContentBrowserItemData& InItem, FText* OutErrorMsg)
+{
+	return ContentBrowserAssetData::CanViewItem(AssetTools, this, InItem, OutErrorMsg);
+}
+
+bool UContentBrowserAssetDataSource::ViewItem(const FContentBrowserItemData& InItem)
+{
+	return ContentBrowserAssetData::ViewItems(AssetTools, this, MakeArrayView(&InItem, 1));
+}
+
+bool UContentBrowserAssetDataSource::BulkViewItems(TArrayView<const FContentBrowserItemData> InItems)
+{
+	return ContentBrowserAssetData::ViewItems(AssetTools, this, InItems);
+}
+
 bool UContentBrowserAssetDataSource::CanPreviewItem(const FContentBrowserItemData& InItem, FText* OutErrorMsg)
 {
 	return ContentBrowserAssetData::CanPreviewItem(AssetTools, this, InItem, OutErrorMsg);
diff --git a/Engine/Plugins/Editor/ContentBrowser/ContentBrowserAssetDataSource/Source/ContentBrowserAssetDataSource/Public/ContentBrowserAssetDataCore.h b/Engine/Plugins/Editor/ContentBrowser/ContentBrowserAssetDataSource/Source/ContentBrowserAssetDataSource/Public/ContentBrowserAssetDataCore.h
--- a/Engine/Plugins/Editor/ContentBrowser/ContentBrowserAssetDataSource/Source/ContentBrowserAssetDataSource/Public/ContentBrowserAssetDataCore.h
+++ b/Engine/Plugins/Editor/ContentBrowser/ContentBrowserAssetDataSource/Source/ContentBrowserAssetDataSource/Public/ContentBrowserAssetDataCore.h
@@ -53,6 +53,14 @@ namespace ContentBrowserAssetData
 
 	CONTENTBROWSERASSETDATASOURCE_API bool EditAssetFileItems(TArrayView<const TSharedRef<const FContentBrowserAssetFileItemDataPayload>> InAssetPayloads);
 
+	CONTENTBROWSERASSETDATASOURCE_API bool CanViewItem(IAssetTools* InAssetTools, const UContentBrowserDataSource* InOwnerDataSource, const FContentBrowserItemData& InItem, FText* OutErrorMsg);
+
+	CONTENTBROWSERASSETDATASOURCE_API bool CanViewAssetFileItem(IAssetTools* InAssetTools, const FContentBrowserAssetFileItemDataPayload& InAssetPayload, FText* OutErrorMsg);
+
+	CONTENTBROWSERASSETDATASOURCE_API bool ViewItems(IAssetTools* InAssetTools, const UContentBrowserDataSource* InOwnerDataSource, TArrayView<const FContentBrowserItemData> InItems);
+
+	CONTENTBROWSERASSETDATASOURCE_API bool ViewAssetFileItems(TArrayView<const TSharedRef<const FContentBrowserAssetFileItemDataPayload>> InAssetPayloads);
+
 	CONTENTBROWSERASSETDATASOURCE_API bool CanPreviewItem(IAssetTools* InAssetTools, const UContentBrowserDataSource* InOwnerDataSource, const FContentBrowserItemData& InItem, FText* OutErrorMsg);
 
 	CONTENTBROWSERASSETDATASOURCE_API bool CanPreviewAssetFileItem(IAssetTools* InAssetTools, const FContentBrowserAssetFileItemDataPayload& InAssetPayload, FText* OutErrorMsg);
diff --git a/Engine/Plugins/Editor/ContentBrowser/ContentBrowserAssetDataSource/Source/ContentBrowserAssetDataSource/Public/ContentBrowserAssetDataSource.h b/Engine/Plugins/Editor/ContentBrowser/ContentBrowserAssetDataSource/Source/ContentBrowserAssetDataSource/Public/ContentBrowserAssetDataSource.h
--- a/Engine/Plugins/Editor/ContentBrowser/ContentBrowserAssetDataSource/Source/ContentBrowserAssetDataSource/Public/ContentBrowserAssetDataSource.h
+++ b/Engine/Plugins/Editor/ContentBrowser/ContentBrowserAssetDataSource/Source/ContentBrowserAssetDataSource/Public/ContentBrowserAssetDataSource.h
@@ -187,6 +187,12 @@ public:
 	virtual bool EditItem(const FContentBrowserItemData& InItem) override;
 
 	virtual bool BulkEditItems(TArrayView<const FContentBrowserItemData> InItems) override;
+	
+	virtual bool CanViewItem(const FContentBrowserItemData& InItem, FText* OutErrorMsg) override;
+
+	virtual bool ViewItem(const FContentBrowserItemData& InItem) override;
+
+	virtual bool BulkViewItems(TArrayView<const FContentBrowserItemData> InItems) override;
 
 	virtual bool CanPreviewItem(const FContentBrowserItemData& InItem, FText* OutErrorMsg) override;
 
diff --git a/Engine/Source/Developer/AssetTools/Public/AssetTypeActions_Base.h b/Engine/Source/Developer/AssetTools/Public/AssetTypeActions_Base.h
--- a/Engine/Source/Developer/AssetTools/Public/AssetTypeActions_Base.h
+++ b/Engine/Source/Developer/AssetTools/Public/AssetTypeActions_Base.h
@@ -11,6 +11,7 @@
 #include "IAssetTools.h"
 #include "IAssetTypeActions.h"
 #include "AssetToolsModule.h"
+#include "AssetTypeActivationOpenedMethod.h"
 #include "Toolkits/SimpleAssetEditor.h"
 #include "AssetRegistry/ARFilter.h"
 
@@ -35,6 +36,14 @@ public:
 		FSimpleAssetEditor::CreateEditor(EToolkitMode::Standalone, EditWithinLevelEditor, InObjects);
 	}
 
+	virtual void OpenAssetEditor(const TArray<UObject*>& InObjects, const EAssetTypeActivationOpenedMethod OpenedMethod, TSharedPtr<IToolkitHost> EditWithinLevelEditor = TSharedPtr<IToolkitHost>()) override
+	{
+		if (OpenedMethod == EAssetTypeActivationOpenedMethod::Edit)
+		{
+			OpenAssetEditor(InObjects, EditWithinLevelEditor);
+		}
+	}
+
 	virtual bool AssetsActivatedOverride(const TArray<UObject*>& InObjects, EAssetTypeActivationMethod::Type ActivationType) override
 	{
 		return false;
@@ -162,6 +171,11 @@ public:
 		return bIsSupported;
 	}
 
+	virtual bool SupportsOpenedMethod(const EAssetTypeActivationOpenedMethod OpenedMethod) const override
+	{
+		return OpenedMethod == EAssetTypeActivationOpenedMethod::Edit;
+	}
+
 	virtual FName GetFilterName() const override
 	{
 		return FName(*GetSupportedClass()->GetPathName());
diff --git a/Engine/Source/Developer/AssetTools/Public/AssetTypeActivationOpenedMethod.h b/Engine/Source/Developer/AssetTools/Public/AssetTypeActivationOpenedMethod.h
new file mode 100644
--- /dev/null
+++ b/Engine/Source/Developer/AssetTools/Public/AssetTypeActivationOpenedMethod.h
@@ -0,0 +1,19 @@
+// Copyright Epic Games, Inc. All Rights Reserved.
+
+#pragma once
+
+#include "CoreMinimal.h"
+
+/**
+ *	(jcotton) This enum has been extracted into a separate header as it would ideally live in IAssetTypeActions.h or AssetTypeActions_Base.h, 
+ *  however these files are included without module linking (IncludePathModuleNames) in several other modules which makes adding a UENUM() not 
+ *  possible without refactoring.
+ */ 
+
+// Types of permissions allowed when attempting to open an asset in editor via activation (EAssetTypeActivationMethod)
+UENUM()
+enum class EAssetTypeActivationOpenedMethod : int32
+{
+	Edit,
+	View
+};
\ No newline at end of file
diff --git a/Engine/Source/Developer/AssetTools/Public/IAssetTypeActions.h b/Engine/Source/Developer/AssetTools/Public/IAssetTypeActions.h
--- a/Engine/Source/Developer/AssetTools/Public/IAssetTypeActions.h
+++ b/Engine/Source/Developer/AssetTools/Public/IAssetTypeActions.h
@@ -10,6 +10,7 @@
 #include "IAssetTypeActions.generated.h"
 
 class IToolkitHost;
+enum class EAssetTypeActivationOpenedMethod : int32;
 
 namespace EAssetTypeActivationMethod
 {
@@ -21,8 +22,6 @@ namespace EAssetTypeActivationMethod
 	};
 }
 
-class IToolkitHost;
-
 /* Revision information for a single revision of a file in source control */
 USTRUCT(BlueprintType)
 struct FRevisionInfo
@@ -96,6 +95,9 @@ public:
 	/** Opens the asset editor for the specified objects. If EditWithinLevelEditor is valid, the world-centric editor will be used. */
 	virtual void OpenAssetEditor( const TArray<UObject*>& InObjects, TSharedPtr<IToolkitHost> EditWithinLevelEditor = TSharedPtr<IToolkitHost>() ) = 0;
 
+	/** Opens the asset editor for the specified objects. If EditWithinLevelEditor is valid, the world-centric editor will be used. */
+	virtual void OpenAssetEditor( const TArray<UObject*>& InObjects, const EAssetTypeActivationOpenedMethod OpenedMethod, TSharedPtr<IToolkitHost> EditWithinLevelEditor = TSharedPtr<IToolkitHost>() ) = 0;
+
 	/** Allows overriding asset activation to perform asset type specific activation for the supplied assets. This happens when the user double clicks, presses enter, or presses space. Return true if you have overridden the behavior. */
 	virtual bool AssetsActivatedOverride(const TArray<UObject*>& InObjects, EAssetTypeActivationMethod::Type ActivationType) = 0;
 
@@ -177,6 +179,9 @@ public:
 	/** Returns class path name as a package + class FName pair */
 	virtual FTopLevelAssetPath GetClassPathName() const = 0;
 
+	/** Does this asset support edit or view methods? */
+	virtual bool SupportsOpenedMethod(const EAssetTypeActivationOpenedMethod OpenedMethod) const = 0;
+
 	/** Returns thumbnail brush unique for the given asset data.  Returning null falls back to class thumbnail brush. */
 	virtual const FSlateBrush* GetThumbnailBrush(const FAssetData& InAssetData, const FName InClassName) const = 0;
 
diff --git a/Engine/Source/Editor/ContentBrowser/Private/SContentBrowser.cpp b/Engine/Source/Editor/ContentBrowser/Private/SContentBrowser.cpp
--- a/Engine/Source/Editor/ContentBrowser/Private/SContentBrowser.cpp
+++ b/Engine/Source/Editor/ContentBrowser/Private/SContentBrowser.cpp
@@ -3011,7 +3011,7 @@ void SContentBrowser::OnItemsActivated(TArrayView<const FContentBrowserItem> Act
 			for (const FContentBrowserItemData& ItemToEdit : SourceAndItemsPair.Value)
 			{
 				FText EditErrorMsg;
-				if (!SourceAndItemsPair.Key->CanEditItem(ItemToEdit, &EditErrorMsg))
+				if (!SourceAndItemsPair.Key->CanEditItem(ItemToEdit, &EditErrorMsg) && !SourceAndItemsPair.Key->CanViewItem(ItemToEdit, &EditErrorMsg))
 				{
 					AssetViewUtils::ShowErrorNotifcation(EditErrorMsg);
 				}
diff --git a/Engine/Source/Editor/ContentBrowserData/Private/ContentBrowserDataSource.cpp b/Engine/Source/Editor/ContentBrowserData/Private/ContentBrowserDataSource.cpp
--- a/Engine/Source/Editor/ContentBrowserData/Private/ContentBrowserDataSource.cpp
+++ b/Engine/Source/Editor/ContentBrowserData/Private/ContentBrowserDataSource.cpp
@@ -305,6 +305,26 @@ bool UContentBrowserDataSource::BulkEditItems(TArrayView<const FContentBrowserIt
 	return bSuccess;
 }
 
+bool UContentBrowserDataSource::CanViewItem(const FContentBrowserItemData& InItem, FText* OutErrorMsg)
+{
+	return false;
+}
+
+bool UContentBrowserDataSource::ViewItem(const FContentBrowserItemData& InItem)
+{
+	return false;
+}
+
+bool UContentBrowserDataSource::BulkViewItems(TArrayView<const FContentBrowserItemData> InItems)
+{
+	bool bSuccess = false;
+	for (const FContentBrowserItemData& Item : InItems)
+	{
+		bSuccess |= ViewItem(Item);
+	}
+	return bSuccess;
+}
+
 bool UContentBrowserDataSource::CanPreviewItem(const FContentBrowserItemData& InItem, FText* OutErrorMsg)
 {
 	return false;
diff --git a/Engine/Source/Editor/ContentBrowserData/Public/ContentBrowserDataSource.h b/Engine/Source/Editor/ContentBrowserData/Public/ContentBrowserDataSource.h
--- a/Engine/Source/Editor/ContentBrowserData/Public/ContentBrowserDataSource.h
+++ b/Engine/Source/Editor/ContentBrowserData/Public/ContentBrowserDataSource.h
@@ -356,6 +356,35 @@ public:
 	 */
 	virtual bool BulkEditItems(TArrayView<const FContentBrowserItemData> InItems);
 
+	/**
+     * Query whether the given item is can be viewed (a read-only asset editor), optionally providing error information if it cannot.
+     *
+     * @param InItem The item to query.
+     * @param OutErrorMessage Optional error message to fill on failure.
+     *
+     * @return True if the item can be viewed in a read-only editor, false otherwise.
+     */
+    virtual bool CanViewItem(const FContentBrowserItemData& InItem, FText* OutErrorMsg);
+
+    /**
+     * Attempt to open the given item for read-only viewing.
+     *
+     * @param InItem The item to view.
+     *
+     * @return True if the item was opened for read-only viewing, false otherwise.
+     */
+    virtual bool ViewItem(const FContentBrowserItemData& InItem);
+
+    /**
+     * Attempt to open the given items for read-only viewing.
+     * @note The default implementation of this will call ViewItem for each item. Override if you can provide a more efficient implementation.
+     *
+     * @param InItems The items to view.
+     *
+     * @return True if any items were opened for read-only viewing, false otherwise.
+     */
+    virtual bool BulkViewItems(TArrayView<const FContentBrowserItemData> InItems);
+
 	/**
 	 * Query whether the given item is can be previewed, optionally providing error information if it cannot.
 	 *
diff --git a/Engine/Source/Editor/UnrealEd/Private/Subsystems/AssetEditorSubsystem.cpp b/Engine/Source/Editor/UnrealEd/Private/Subsystems/AssetEditorSubsystem.cpp
--- a/Engine/Source/Editor/UnrealEd/Private/Subsystems/AssetEditorSubsystem.cpp
+++ b/Engine/Source/Editor/UnrealEd/Private/Subsystems/AssetEditorSubsystem.cpp
@@ -151,7 +151,6 @@ int32 UAssetEditorSubsystem::CloseAllEditorsForAsset(UObject* Asset)
 	return EditorInstances.Num();
 }
 
-
 void UAssetEditorSubsystem::RemoveAssetFromAllEditors(UObject* Asset)
 {
 	TArray<IAssetEditorInstance*> EditorInstances = FindEditorsForAsset(Asset);
@@ -311,7 +310,7 @@ bool UAssetEditorSubsystem::IsAssetEditable(const UObject* Asset)
 	return true;
 }
 
-bool UAssetEditorSubsystem::OpenEditorForAsset(UObject* Asset, const EToolkitMode::Type ToolkitMode, TSharedPtr< IToolkitHost > OpenedFromLevelEditor, const bool bShowProgressWindow)
+bool UAssetEditorSubsystem::OpenEditorForAsset(UObject* Asset, const EToolkitMode::Type ToolkitMode, TSharedPtr< IToolkitHost > OpenedFromLevelEditor, const bool bShowProgressWindow, const EAssetTypeActivationOpenedMethod OpenedMethod)
 {
 	SCOPE_STALL_REPORTER(UAssetEditorSubsystem::OpenEditorForAsset, 2.0);
 
@@ -331,18 +330,21 @@ bool UAssetEditorSubsystem::OpenEditorForAsset(UObject* Asset, const EToolkitMod
 
 	const bool bBringToFrontIfOpen = true;
 
-	if (UPackage* Package = Asset->GetOutermost())
+	if (OpenedMethod == EAssetTypeActivationOpenedMethod::Edit)
 	{
-		// Don't open asset editors for cooked packages
-		if (Package->bIsCookedForEditor)
+		if (UPackage* Package = Asset->GetOutermost())
 		{
-			return false;
-		}
+			// Don't open asset editors for cooked packages
+			if (Package->bIsCookedForEditor)
+			{
+				return false;
+			}
 
-		if (!AssetToolsModule.Get().GetWritableFolderPermissionList()->PassesStartsWithFilter(Package->GetName()))
-		{
-			AssetToolsModule.Get().NotifyBlockedByWritableFolderFilter();
-			return false;
+			if (!AssetToolsModule.Get().GetWritableFolderPermissionList()->PassesStartsWithFilter(Package->GetName()))
+			{
+				AssetToolsModule.Get().NotifyBlockedByWritableFolderFilter();
+				return false;
+			}
 		}
 	}
 
@@ -391,7 +393,7 @@ bool UAssetEditorSubsystem::OpenEditorForAsset(UObject* Asset, const EToolkitMod
 	// Disallow opening an asset editor for classes
 	bool bCanSummonSimpleAssetEditor = !Asset->IsA<UClass>();
 
-	if (AssetTypeActions.IsValid())
+	if (AssetTypeActions.IsValid() && AssetTypeActions.Pin()->SupportsOpenedMethod(OpenedMethod))
 	{
 		TArray<UObject*> AssetsToEdit;
 		AssetsToEdit.Add(Asset);
@@ -400,7 +402,7 @@ bool UAssetEditorSubsystem::OpenEditorForAsset(UObject* Asset, const EToolkitMod
 		TWeakObjectPtr<UObject> WeakAsset = Asset;
 		const FString AssetPath = Asset->GetPathName();
 
-		AssetTypeActions.Pin()->OpenAssetEditor(AssetsToEdit, ActualToolkitMode == EToolkitMode::WorldCentric ? OpenedFromLevelEditor : TSharedPtr<IToolkitHost>());
+		AssetTypeActions.Pin()->OpenAssetEditor(AssetsToEdit, OpenedMethod, ActualToolkitMode == EToolkitMode::WorldCentric ? OpenedFromLevelEditor : TSharedPtr<IToolkitHost>());
 
 		// If the Asset was destroyed, attempt to find it if it was recreated
 		if (!WeakAsset.IsValid() && !AssetPath.IsEmpty())
@@ -439,14 +441,13 @@ bool UAssetEditorSubsystem::OpenEditorForAsset(UObject* Asset, const EToolkitMod
 		FStudioAnalytics::FireEvent_Loading(TEXT("OpenAssetEditor"), OpenTime, {
 			FAnalyticsEventAttribute(TEXT("AssetPath"), Asset->GetFullName()),
 			FAnalyticsEventAttribute(TEXT("AssetType"), Asset->GetClass()->GetName())
-		});
+			});
 	}
-
+	
 	return true;
 }
 
-
-bool UAssetEditorSubsystem::OpenEditorForAssets_Advanced(const TArray <UObject* >& InAssets, const EToolkitMode::Type ToolkitMode, TSharedPtr< IToolkitHost > OpenedFromLevelEditor)
+bool UAssetEditorSubsystem::OpenEditorForAssets_Advanced(const TArray <UObject*>& InAssets, const EToolkitMode::Type ToolkitMode, TSharedPtr< IToolkitHost > OpenedFromLevelEditor, const EAssetTypeActivationOpenedMethod OpenedMethod)
 {
 	TArray<UObject*> Assets;
 	Assets.Reserve(InAssets.Num());
@@ -474,7 +475,7 @@ bool UAssetEditorSubsystem::OpenEditorForAssets_Advanced(const TArray <UObject*
 
 	if (Assets.Num() == 1)
 	{
-		return OpenEditorForAsset(Assets[0], ToolkitMode, OpenedFromLevelEditor);
+		return OpenEditorForAsset(Assets[0], ToolkitMode, OpenedFromLevelEditor, true, OpenedMethod);
 	}
 	else if (Assets.Num() > 0)
 	{
@@ -484,7 +485,7 @@ bool UAssetEditorSubsystem::OpenEditorForAssets_Advanced(const TArray <UObject*
 			// If any of the assets are already open or the package is cooked,
 			// remove them from the list of assets to open an editor for
 			UPackage* Package = Asset->GetOutermost();
-			if (FindEditorForAsset(Asset, true) != nullptr || (Package && Package->bIsCookedForEditor))
+			if (FindEditorForAsset(Asset, true) != nullptr || (OpenedMethod == EAssetTypeActivationOpenedMethod::Edit && Package && Package->bIsCookedForEditor))
 			{
 				SkipOpenAssets.Add(Asset);
 			}
@@ -508,7 +509,7 @@ bool UAssetEditorSubsystem::OpenEditorForAssets_Advanced(const TArray <UObject*
 			FAssetToolsModule& AssetToolsModule = FModuleManager::LoadModuleChecked<FAssetToolsModule>(TEXT("AssetTools"));
 			TWeakPtr<IAssetTypeActions> AssetTypeActions = AssetToolsModule.Get().GetAssetTypeActionsForClass(AssetClass);
 
-			if (AssetTypeActions.IsValid())
+			if (AssetTypeActions.IsValid() && AssetTypeActions.Pin()->SupportsOpenedMethod(OpenedMethod))
 			{
 				GWarn->BeginSlowTask(LOCTEXT("OpenEditors", "Opening Editor(s)..."), true);
 
@@ -551,7 +552,7 @@ bool UAssetEditorSubsystem::OpenEditorForAssets_Advanced(const TArray <UObject*
 				}
 
 				// How to handle multiple assets is left up to the type actions (i.e. open a single shared editor or an editor for each)
-				AssetTypeActions.Pin()->OpenAssetEditor(Assets, ActualToolkitMode == EToolkitMode::WorldCentric ? OpenedFromLevelEditor : TSharedPtr<IToolkitHost>());
+				AssetTypeActions.Pin()->OpenAssetEditor(Assets, OpenedMethod, ActualToolkitMode == EToolkitMode::WorldCentric ? OpenedFromLevelEditor : TSharedPtr<IToolkitHost>());
 
 				// If any assets were destroyed, attempt to find them if they were recreated
 				for (int32 i = 0; i < Assets.Num(); i++)
@@ -578,7 +579,7 @@ bool UAssetEditorSubsystem::OpenEditorForAssets_Advanced(const TArray <UObject*
 			{
 				if (!SkipOpenAssets.Contains(Asset))
 				{
-					OpenEditorForAsset(Asset, ToolkitMode, OpenedFromLevelEditor);
+					OpenEditorForAsset(Asset, ToolkitMode, OpenedFromLevelEditor, true, OpenedMethod);
 				}
 			}
 		}
@@ -587,9 +588,9 @@ bool UAssetEditorSubsystem::OpenEditorForAssets_Advanced(const TArray <UObject*
 	return NumNullAssets == 0;
 }
 
-bool UAssetEditorSubsystem::OpenEditorForAssets(const TArray<UObject*>& Assets)
+bool UAssetEditorSubsystem::OpenEditorForAssets(const TArray<UObject*>& Assets, const EAssetTypeActivationOpenedMethod OpenedMethod)
 {
-	return OpenEditorForAssets_Advanced(Assets, EToolkitMode::Standalone, TSharedPtr<IToolkitHost>());
+	return OpenEditorForAssets_Advanced(Assets, EToolkitMode::Standalone, TSharedPtr<IToolkitHost>(), OpenedMethod);
 }
 
 void UAssetEditorSubsystem::HandleRequestOpenAssetMessage(const FAssetEditorRequestOpenAsset& Message, const TSharedRef<IMessageContext, ESPMode::ThreadSafe>& Context)
@@ -597,12 +598,7 @@ void UAssetEditorSubsystem::HandleRequestOpenAssetMessage(const FAssetEditorRequ
 	OpenEditorForAsset(Message.AssetName);
 }
 
-void UAssetEditorSubsystem::OpenEditorForAsset(const FSoftObjectPath& AssetPath)
-{
-	OpenEditorForAsset(AssetPath.ToString());
-}
-
-void UAssetEditorSubsystem::OpenEditorForAsset(const FString& AssetPathName)
+void UAssetEditorSubsystem::OpenEditorForAsset(const FString& AssetPathName, const EAssetTypeActivationOpenedMethod OpenedMethod)
 {
 	// An asset needs loading
 
@@ -613,7 +609,7 @@ void UAssetEditorSubsystem::OpenEditorForAsset(const FString& AssetPathName)
 
 		FString AssetName = FPaths::GetBaseFilename(AssetPathName);
 		UObject* Object = FindObject<UObject>(Package, *AssetName);
-
+		
 		if (Object != nullptr)
 		{
 			OpenEditorForAsset(Object);
@@ -626,7 +622,7 @@ void UAssetEditorSubsystem::OpenEditorForAsset(const FString& AssetPathName)
 		
 		if (Object != nullptr)
 		{
-			OpenEditorForAsset(Object);
+			OpenEditorForAsset(Object, EToolkitMode::Standalone, TSharedPtr<IToolkitHost>(), true, OpenedMethod);
 		}
 	}
 }
@@ -1025,27 +1021,19 @@ void UAssetEditorSubsystem::HandlePackageReloaded(const EPackageReloadPhase InPa
 	}
 }
 
-void UAssetEditorSubsystem::OpenEditorsForAssets(const TArray<FSoftObjectPath>& AssetsToOpen)
-{
-	for (const FSoftObjectPath& AssetName : AssetsToOpen)
-	{
-		OpenEditorForAsset(AssetName);
-	}
-}
-
-void UAssetEditorSubsystem::OpenEditorsForAssets(const TArray<FString>& AssetsToOpen)
+void UAssetEditorSubsystem::OpenEditorsForAssets(const TArray<FString>& AssetsToOpen, const EAssetTypeActivationOpenedMethod OpenedMethod)
 {
 	for (const FString& AssetName : AssetsToOpen)
 	{
-		OpenEditorForAsset(AssetName);
+		OpenEditorForAsset(AssetName, OpenedMethod);
 	}
 }
 
-void UAssetEditorSubsystem::OpenEditorsForAssets(const TArray<FName>& AssetsToOpen)
+void UAssetEditorSubsystem::OpenEditorsForAssets(const TArray<FName>& AssetsToOpen, const EAssetTypeActivationOpenedMethod OpenedMethod)
 {
 	for (const FName& AssetName : AssetsToOpen)
 	{
-		OpenEditorForAsset(AssetName.ToString());
+		OpenEditorForAsset(AssetName.ToString(), OpenedMethod);
 	}
 }
 
diff --git a/Engine/Source/Editor/UnrealEd/Public/Subsystems/AssetEditorSubsystem.h b/Engine/Source/Editor/UnrealEd/Public/Subsystems/AssetEditorSubsystem.h
--- a/Engine/Source/Editor/UnrealEd/Public/Subsystems/AssetEditorSubsystem.h
+++ b/Engine/Source/Editor/UnrealEd/Public/Subsystems/AssetEditorSubsystem.h
@@ -9,6 +9,7 @@
 #include "Containers/Ticker.h"
 #include "Tools/Modes.h"
 #include "Misc/NamePermissionList.h"
+#include "AssetTypeActivationOpenedMethod.h"
 #include "AssetEditorSubsystem.generated.h"
 
 class UAssetEditor;
@@ -73,13 +74,14 @@ public:
 	bool IsAssetEditable(const UObject* Asset);
 
 	/** Opens an asset by path */
-	void OpenEditorForAsset(const FString& AssetPath);
+	void OpenEditorForAsset(const FString& AssetPathName, const EAssetTypeActivationOpenedMethod OpenedMethod = EAssetTypeActivationOpenedMethod::Edit);
 	void OpenEditorForAsset(const FSoftObjectPath& AssetPath);
 
 	/**
 	 * Tries to open an editor for the specified asset.  Returns true if the asset is open in an editor.
 	 * If the file is already open in an editor, it will not create another editor window but instead bring it to front
 	 */
+	bool OpenEditorForAsset(UObject* Asset, const EToolkitMode::Type ToolkitMode = EToolkitMode::Standalone, TSharedPtr<IToolkitHost> OpenedFromLevelEditor = TSharedPtr<IToolkitHost>(), const bool bShowProgressWindow = true, EAssetTypeActivationOpenedMethod OpenedMethod = EAssetTypeActivationOpenedMethod::Edit);
 	bool OpenEditorForAsset(UObject* Asset, const EToolkitMode::Type ToolkitMode = EToolkitMode::Standalone, TSharedPtr<IToolkitHost> OpenedFromLevelEditor = TSharedPtr<IToolkitHost>(), const bool bShowProgressWindow = true);
 	
 	template<typename ObjectType>
@@ -94,12 +96,12 @@ public:
 	 * If all assets are of the same type, the supporting AssetTypeAction (if it exists) is responsible for the details of how to handle opening multiple assets at once.
 	 */
 	UFUNCTION(BlueprintCallable, Category = "Editor Scripting | Asset Tools")
-	bool OpenEditorForAssets(const TArray<UObject*>& Assets);
-	bool OpenEditorForAssets_Advanced(const TArray<UObject*>& Assets, const EToolkitMode::Type ToolkitMode = EToolkitMode::Standalone, TSharedPtr<IToolkitHost> OpenedFromLevelEditor = TSharedPtr<IToolkitHost>());
+	bool OpenEditorForAssets(const TArray<UObject*>& Assets, const EAssetTypeActivationOpenedMethod OpenedMethod = EAssetTypeActivationOpenedMethod::Edit);
+	bool OpenEditorForAssets_Advanced(const TArray<UObject*>& InAssets, const EToolkitMode::Type ToolkitMode = EToolkitMode::Standalone, TSharedPtr<IToolkitHost> OpenedFromLevelEditor = TSharedPtr<IToolkitHost>(), const EAssetTypeActivationOpenedMethod OpenedMethod = EAssetTypeActivationOpenedMethod::Edit);
 
 	/** Opens editors for the supplied assets (via OpenEditorForAsset) */
-	void OpenEditorsForAssets(const TArray<FString>& AssetsToOpen);
-	void OpenEditorsForAssets(const TArray<FName>& AssetsToOpen);
+	void OpenEditorsForAssets(const TArray<FString>& AssetsToOpen, const EAssetTypeActivationOpenedMethod OpenedMethod = EAssetTypeActivationOpenedMethod::Edit);
+	void OpenEditorsForAssets(const TArray<FName>& AssetsToOpen, const EAssetTypeActivationOpenedMethod OpenedMethod = EAssetTypeActivationOpenedMethod::Edit);
 	void OpenEditorsForAssets(const TArray<FSoftObjectPath>& AssetsToOpen);
 
 	/** Returns the primary editor if one is already open for the specified asset.
diff --git a/Engine/Source/Editor/UnrealEd/UnrealEd.Build.cs b/Engine/Source/Editor/UnrealEd/UnrealEd.Build.cs
--- a/Engine/Source/Editor/UnrealEd/UnrealEd.Build.cs
+++ b/Engine/Source/Editor/UnrealEd/UnrealEd.Build.cs
@@ -121,6 +121,7 @@ public class UnrealEd : ModuleRules
 				"SubobjectEditor",
 				"PhysicsUtilities",
 				"ToolWidgets",
+				"AssetTools",
 			}
 		);
 
@@ -229,7 +230,6 @@ public class UnrealEd : ModuleRules
 				"TextureEditor",
 				"Cascade",
 				"UMGEditor",
-				"AssetTools",
 				"ClassViewer",
 				"StructViewer",
 				"CollectionManager",
@@ -326,7 +326,6 @@ public class UnrealEd : ModuleRules
 				"BlueprintGraph",
 				"AddContentDialog",
 				"MeshUtilities",
-				"AssetTools",
 				"KismetCompiler",
 				"NavigationSystem",
 				"GameplayTasks",
