Fix for https://horde.devtools.epicgames.com/log/635979d6767e539ccacc08a8?lineindex=182 caused by bad emrge leaving double function declarations

diff --git a/Engine/Source/Editor/UnrealEd/Public/Subsystems/AssetEditorSubsystem.h b/Engine/Source/Editor/UnrealEd/Public/Subsystems/AssetEditorSubsystem.h
--- a/Engine/Source/Editor/UnrealEd/Public/Subsystems/AssetEditorSubsystem.h
+++ b/Engine/Source/Editor/UnrealEd/Public/Subsystems/AssetEditorSubsystem.h
@@ -82,12 +82,11 @@ public:
 	 * If the file is already open in an editor, it will not create another editor window but instead bring it to front
 	 */
 	bool OpenEditorForAsset(UObject* Asset, const EToolkitMode::Type ToolkitMode = EToolkitMode::Standalone, TSharedPtr<IToolkitHost> OpenedFromLevelEditor = TSharedPtr<IToolkitHost>(), const bool bShowProgressWindow = true, EAssetTypeActivationOpenedMethod OpenedMethod = EAssetTypeActivationOpenedMethod::Edit);
-	bool OpenEditorForAsset(UObject* Asset, const EToolkitMode::Type ToolkitMode = EToolkitMode::Standalone, TSharedPtr<IToolkitHost> OpenedFromLevelEditor = TSharedPtr<IToolkitHost>(), const bool bShowProgressWindow = true);
 	
 	template<typename ObjectType>
-	bool OpenEditorForAsset(TObjectPtr<ObjectType> Asset, const EToolkitMode::Type ToolkitMode = EToolkitMode::Standalone, TSharedPtr<IToolkitHost> OpenedFromLevelEditor = TSharedPtr<IToolkitHost>(), const bool bShowProgressWindow = true)
+	bool OpenEditorForAsset(TObjectPtr<ObjectType> Asset, const EToolkitMode::Type ToolkitMode = EToolkitMode::Standalone, TSharedPtr<IToolkitHost> OpenedFromLevelEditor = TSharedPtr<IToolkitHost>(), const bool bShowProgressWindow = true, EAssetTypeActivationOpenedMethod OpenedMethod = EAssetTypeActivationOpenedMethod::Edit)
 	{
-		return OpenEditorForAsset(ToRawPtr(Asset), ToolkitMode, OpenedFromLevelEditor, bShowProgressWindow);
+		return OpenEditorForAsset(ToRawPtr(Asset), ToolkitMode, OpenedFromLevelEditor, bShowProgressWindow, OpenedMethod);
 	}
 
 	/**
