[Backout] - CL22782944
Original CL Desc
-----------------------------------------------------------------
Addition of 'view' method to AssetTypeActions

Iteration on Jason Stasiks CL 17314610. This is a minimal version of CL 17314610 which adds a 'view' method to AssetTypeActions which is now attempted whenever an activated item's 'edit' attempt fails.

The matching CL for the restricted FAssetTypeActions_VerseClass changes is 22472079.

diff --git a/Engine/Plugins/Editor/ContentBrowser/ContentBrowserAliasDataSoure/Source/ContentBrowserAliasDataSource/Private/ContentBrowserAliasDataSource.cpp b/Engine/Plugins/Editor/ContentBrowser/ContentBrowserAliasDataSoure/Source/ContentBrowserAliasDataSource/Private/ContentBrowserAliasDataSource.cpp
--- a/Engine/Plugins/Editor/ContentBrowser/ContentBrowserAliasDataSoure/Source/ContentBrowserAliasDataSource/Private/ContentBrowserAliasDataSource.cpp
+++ b/Engine/Plugins/Editor/ContentBrowser/ContentBrowserAliasDataSoure/Source/ContentBrowserAliasDataSource/Private/ContentBrowserAliasDataSource.cpp
@@ -692,21 +692,6 @@ bool UContentBrowserAliasDataSource::BulkEditItems(TArrayView<const FContentBrow
 	return ContentBrowserAssetData::EditItems(AssetTools, this, InItems);
 }
 
-bool UContentBrowserAliasDataSource::CanViewItem(const FContentBrowserItemData& InItem, FText* OutErrorMsg)
-{
-	return ContentBrowserAssetData::CanViewItem(AssetTools, this, InItem, OutErrorMsg);
-}
-
-bool UContentBrowserAliasDataSource::ViewItem(const FContentBrowserItemData& InItem)
-{
-	return ContentBrowserAssetData::ViewItems(AssetTools, this, MakeArrayView(&InItem, 1));
-}
-
-bool UContentBrowserAliasDataSource::BulkViewItems(TArrayView<const FContentBrowserItemData> InItems)
-{
-	return ContentBrowserAssetData::ViewItems(AssetTools, this, InItems);
-}
-
 bool UContentBrowserAliasDataSource::CanPreviewItem(const FContentBrowserItemData& InItem, FText* OutErrorMsg)
 {
 	return ContentBrowserAssetData::CanPreviewItem(AssetTools, this, InItem, OutErrorMsg);
diff --git a/Engine/Plugins/Editor/ContentBrowser/ContentBrowserAliasDataSoure/Source/ContentBrowserAliasDataSource/Public/ContentBrowserAliasDataSource.h b/Engine/Plugins/Editor/ContentBrowser/ContentBrowserAliasDataSoure/Source/ContentBrowserAliasDataSource/Public/ContentBrowserAliasDataSource.h
--- a/Engine/Plugins/Editor/ContentBrowser/ContentBrowserAliasDataSoure/Source/ContentBrowserAliasDataSource/Public/ContentBrowserAliasDataSource.h
+++ b/Engine/Plugins/Editor/ContentBrowser/ContentBrowserAliasDataSoure/Source/ContentBrowserAliasDataSource/Public/ContentBrowserAliasDataSource.h
@@ -70,10 +70,6 @@ public:
 	virtual bool CanEditItem(const FContentBrowserItemData& InItem, FText* OutErrorMsg) override;
 	virtual bool EditItem(const FContentBrowserItemData& InItem) override;
 	virtual bool BulkEditItems(TArrayView<const FContentBrowserItemData> InItems) override;
-	
-	virtual bool CanViewItem(const FContentBrowserItemData& InItem, FText* OutErrorMsg) override;
-	virtual bool ViewItem(const FContentBrowserItemData& InItem) override;
-	virtual bool BulkViewItems(TArrayView<const FContentBrowserItemData> InItems) override;
 
 	virtual bool CanPreviewItem(const FContentBrowserItemData& InItem, FText* OutErrorMsg) override;
 	virtual bool PreviewItem(const FContentBrowserItemData& InItem) override;
diff --git a/Engine/Plugins/Editor/ContentBrowser/ContentBrowserAssetDataSource/Source/ContentBrowserAssetDataSource/Private/ContentBrowserAssetDataCore.cpp b/Engine/Plugins/Editor/ContentBrowser/ContentBrowserAssetDataSource/Source/ContentBrowserAssetDataSource/Private/ContentBrowserAssetDataCore.cpp
--- a/Engine/Plugins/Editor/ContentBrowser/ContentBrowserAssetDataSource/Source/ContentBrowserAssetDataSource/Private/ContentBrowserAssetDataCore.cpp
+++ b/Engine/Plugins/Editor/ContentBrowser/ContentBrowserAssetDataSource/Source/ContentBrowserAssetDataSource/Private/ContentBrowserAssetDataCore.cpp
@@ -5,7 +5,6 @@
 #include "IAssetTools.h"
 #include "AssetRegistry/IAssetRegistry.h"
 #include "AssetViewUtils.h"
-#include "AssetTypeActivationOpenedMethod.h"
 #include "AssetPropertyTagCache.h"
 #include "ObjectTools.h"
 #include "Misc/NamePermissionList.h"
@@ -206,26 +205,6 @@ bool CanEditAssetFileItem(IAssetTools* InAssetTools, const FContentBrowserAssetF
 	return true;
 }
 
-bool CanViewItem(IAssetTools* InAssetTools, const UContentBrowserDataSource* InOwnerDataSource, const FContentBrowserItemData& InItem, FText* OutErrorMsg)
-{
-	if (TSharedPtr<const FContentBrowserAssetFileItemDataPayload> AssetPayload = GetAssetFileItemPayload(InOwnerDataSource, InItem))
-	{
-		return CanViewAssetFileItem(InAssetTools, *AssetPayload, OutErrorMsg);
-	}
-
-	return false;
-}
-
-bool CanViewAssetFileItem(IAssetTools* InAssetTools, const FContentBrowserAssetFileItemDataPayload& InAssetPayload, FText* OutErrorMsg)
-{
-	const TWeakPtr<IAssetTypeActions> AssetTypeActions = InAssetTools->GetAssetTypeActionsForClass(InAssetPayload.GetAssetData().GetClass());
-	if (AssetTypeActions.IsValid())
-	{
-		return AssetTypeActions.Pin()->SupportsOpenedMethod(EAssetTypeActivationOpenedMethod::View);
-	}
-	return false;
-}
-
 bool CanPreviewItem(IAssetTools* InAssetTools, const UContentBrowserDataSource* InOwnerDataSource, const FContentBrowserItemData& InItem, FText* OutErrorMsg)
 {
 	if (TSharedPtr<const FContentBrowserAssetFileItemDataPayload> AssetPayload = GetAssetFileItemPayload(InOwnerDataSource, InItem))
@@ -241,15 +220,14 @@ bool CanPreviewAssetFileItem(IAssetTools* InAssetTools, const FContentBrowserAss
 	return true;
 }
 
-bool EditOrPreviewAssetFileItems(TArrayView<const TSharedRef<const FContentBrowserAssetFileItemDataPayload>> InAssetPayloads, EAssetTypeActivationMethod::Type ActivationMethod, const EAssetTypeActivationOpenedMethod OpenedMethod)
+bool EditOrPreviewAssetFileItems(TArrayView<const TSharedRef<const FContentBrowserAssetFileItemDataPayload>> InAssetPayloads, const bool bIsPreview)
 {
 	if (InAssetPayloads.Num() == 0)
 	{
 		return false;
 	}
 
-	ensure(ActivationMethod != EAssetTypeActivationMethod::Type::DoubleClicked);
-
+	const EAssetTypeActivationMethod::Type ActivationMethod = bIsPreview ? EAssetTypeActivationMethod::Previewed : EAssetTypeActivationMethod::Opened;
 	UAssetEditorSubsystem* AssetEditorSubsystem = GEditor->GetEditorSubsystem<UAssetEditorSubsystem>();
 
 	TMap<TSharedPtr<IAssetTypeActions>, TArray<FAssetData>> TypeActionsToAssetData;
@@ -281,7 +259,7 @@ bool EditOrPreviewAssetFileItems(TArrayView<const TSharedRef<const FContentBrows
 		TArray<FAssetData>& AssetsToLoad = TypeActionToObjectsPair.Value;
 		if (TypeActions.IsValid())
 		{
-			AssetsToLoad = TypeActions->GetValidAssetsForPreviewOrEdit(AssetsToLoad, ActivationMethod == EAssetTypeActivationMethod::Type::Previewed);
+			AssetsToLoad = TypeActions->GetValidAssetsForPreviewOrEdit(AssetsToLoad, bIsPreview);
 		}
 
 		TArray<UObject*> ObjList;
@@ -307,84 +285,47 @@ bool EditOrPreviewAssetFileItems(TArrayView<const TSharedRef<const FContentBrows
 
 		if (bOpenEditorForAssets)
 		{
-			bSuccessfulEditorOpen &= AssetEditorSubsystem->OpenEditorForAssets(ObjList, OpenedMethod);
+			bSuccessfulEditorOpen &= AssetEditorSubsystem->OpenEditorForAssets(ObjList);
 		}
 	}
 
 	return bSuccessfulEditorOpen;
 }
 
-bool EditOrPreviewItems(IAssetTools* InAssetTools, const UContentBrowserDataSource* InOwnerDataSource, TArrayView<const FContentBrowserItemData> InItems, EAssetTypeActivationMethod::Type ActivationMethod, EAssetTypeActivationOpenedMethod OpenedMethod)
+bool EditOrPreviewItems(IAssetTools* InAssetTools, const UContentBrowserDataSource* InOwnerDataSource, TArrayView<const FContentBrowserItemData> InItems, const bool bIsPreview)
 {
-	TArray<TSharedRef<const FContentBrowserAssetFileItemDataPayload>, TInlineAllocator<16>> EditOrPreviewPayloads;
-	TArray<TSharedRef<const FContentBrowserAssetFileItemDataPayload>, TInlineAllocator<16>> ViewPayloads;
-
-	ensure(ActivationMethod != EAssetTypeActivationMethod::Type::DoubleClicked);
+	TArray<TSharedRef<const FContentBrowserAssetFileItemDataPayload>, TInlineAllocator<16>> AssetPayloads;
 
-	const bool bIsPreview = ActivationMethod == EAssetTypeActivationMethod::Type::Previewed;
-	EnumerateAssetFileItemPayloads(InOwnerDataSource, InItems, [InAssetTools, bIsPreview, OpenedMethod, &EditOrPreviewPayloads, &ViewPayloads](const TSharedRef<const FContentBrowserAssetFileItemDataPayload>& InAssetPayload)
+	EnumerateAssetFileItemPayloads(InOwnerDataSource, InItems, [InAssetTools, bIsPreview , &AssetPayloads](const TSharedRef<const FContentBrowserAssetFileItemDataPayload>& InAssetPayload)
 	{
-		if (bIsPreview)
-		{
-			if (CanPreviewAssetFileItem(InAssetTools, *InAssetPayload, nullptr))
-			{
-				EditOrPreviewPayloads.Add(InAssetPayload);
-			}
-		}
-		else if (OpenedMethod == EAssetTypeActivationOpenedMethod::Edit)
-		{
-			if (CanEditAssetFileItem(InAssetTools, *InAssetPayload, nullptr))
-			{
-				EditOrPreviewPayloads.Add(InAssetPayload);
-			}
-			else if (CanViewAssetFileItem(InAssetTools, *InAssetPayload, nullptr))
-			{
-				ViewPayloads.Add(InAssetPayload);
-			}
-		}
-		else if (OpenedMethod == EAssetTypeActivationOpenedMethod::View)
+		if ((bIsPreview ? CanPreviewAssetFileItem(InAssetTools, *InAssetPayload, nullptr) : CanEditAssetFileItem(InAssetTools, *InAssetPayload, nullptr)))
 		{
-			if (CanViewAssetFileItem(InAssetTools, *InAssetPayload, nullptr))
-			{
-				ViewPayloads.Add(InAssetPayload);
-			}
+			AssetPayloads.Add(InAssetPayload);
 		}
 		return true;
 	});
 
-	const bool bEditItems = EditOrPreviewAssetFileItems(EditOrPreviewPayloads, ActivationMethod, EAssetTypeActivationOpenedMethod::Edit);
-	const bool bViewItems = EditOrPreviewAssetFileItems(ViewPayloads, ActivationMethod, EAssetTypeActivationOpenedMethod::View);
-	return bEditItems || bViewItems;
+	return EditOrPreviewAssetFileItems(AssetPayloads, bIsPreview);
 }
 
 bool EditItems(IAssetTools* InAssetTools, const UContentBrowserDataSource* InOwnerDataSource, TArrayView<const FContentBrowserItemData> InItems)
 {
-	return EditOrPreviewItems(InAssetTools, InOwnerDataSource, InItems, EAssetTypeActivationMethod::Type::Opened, EAssetTypeActivationOpenedMethod::Edit);
+	return EditOrPreviewItems(InAssetTools, InOwnerDataSource, InItems, /*bIsPreview*/false);
 }
 
 bool EditAssetFileItems(TArrayView<const TSharedRef<const FContentBrowserAssetFileItemDataPayload>> InAssetPayloads)
 {
-	return EditOrPreviewAssetFileItems(InAssetPayloads, EAssetTypeActivationMethod::Type::Opened, EAssetTypeActivationOpenedMethod::Edit);
-}
-
-bool ViewItems(IAssetTools* InAssetTools, const UContentBrowserDataSource* InOwnerDataSource, TArrayView<const FContentBrowserItemData> InItems)
-{
-	return EditOrPreviewItems(InAssetTools, InOwnerDataSource, InItems, EAssetTypeActivationMethod::Type::Opened, EAssetTypeActivationOpenedMethod::View);
-}
-
-bool ViewAssetFileItems(TArrayView<const TSharedRef<const FContentBrowserAssetFileItemDataPayload>> InAssetPayloads)
-{
-	return EditOrPreviewAssetFileItems(InAssetPayloads, EAssetTypeActivationMethod::Type::Opened, EAssetTypeActivationOpenedMethod::View);
+	return EditOrPreviewAssetFileItems(InAssetPayloads, /*bIsPreview*/false);
 }
 
 bool PreviewItems(IAssetTools* InAssetTools, const UContentBrowserDataSource* InOwnerDataSource, TArrayView<const FContentBrowserItemData> InItems)
 {
-	return EditOrPreviewItems(InAssetTools, InOwnerDataSource, InItems, EAssetTypeActivationMethod::Type::Previewed, EAssetTypeActivationOpenedMethod::Edit);
+	return EditOrPreviewItems(InAssetTools, InOwnerDataSource, InItems, /*bIsPreview*/true);
 }
 
 bool PreviewAssetFileItems(TArrayView<const TSharedRef<const FContentBrowserAssetFileItemDataPayload>> InAssetPayloads)
 {
-	return EditOrPreviewAssetFileItems(InAssetPayloads, EAssetTypeActivationMethod::Type::Previewed, EAssetTypeActivationOpenedMethod::Edit);
+	return EditOrPreviewAssetFileItems(InAssetPayloads, /*bIsPreview*/true);
 }
 
 bool CanDuplicateItem(IAssetTools* InAssetTools, const UContentBrowserDataSource* InOwnerDataSource, const FContentBrowserItemData& InItem, FText* OutErrorMsg)
diff --git a/Engine/Plugins/Editor/ContentBrowser/ContentBrowserAssetDataSource/Source/ContentBrowserAssetDataSource/Private/ContentBrowserAssetDataSource.cpp b/Engine/Plugins/Editor/ContentBrowser/ContentBrowserAssetDataSource/Source/ContentBrowserAssetDataSource/Private/ContentBrowserAssetDataSource.cpp
--- a/Engine/Plugins/Editor/ContentBrowser/ContentBrowserAssetDataSource/Source/ContentBrowserAssetDataSource/Private/ContentBrowserAssetDataSource.cpp
+++ b/Engine/Plugins/Editor/ContentBrowser/ContentBrowserAssetDataSource/Source/ContentBrowserAssetDataSource/Private/ContentBrowserAssetDataSource.cpp
@@ -1237,21 +1237,6 @@ bool UContentBrowserAssetDataSource::BulkEditItems(TArrayView<const FContentBrow
 	return ContentBrowserAssetData::EditItems(AssetTools, this, InItems);
 }
 
-bool UContentBrowserAssetDataSource::CanViewItem(const FContentBrowserItemData& InItem, FText* OutErrorMsg)
-{
-	return ContentBrowserAssetData::CanViewItem(AssetTools, this, InItem, OutErrorMsg);
-}
-
-bool UContentBrowserAssetDataSource::ViewItem(const FContentBrowserItemData& InItem)
-{
-	return ContentBrowserAssetData::ViewItems(AssetTools, this, MakeArrayView(&InItem, 1));
-}
-
-bool UContentBrowserAssetDataSource::BulkViewItems(TArrayView<const FContentBrowserItemData> InItems)
-{
-	return ContentBrowserAssetData::ViewItems(AssetTools, this, InItems);
-}
-
 bool UContentBrowserAssetDataSource::CanPreviewItem(const FContentBrowserItemData& InItem, FText* OutErrorMsg)
 {
 	return ContentBrowserAssetData::CanPreviewItem(AssetTools, this, InItem, OutErrorMsg);
diff --git a/Engine/Plugins/Editor/ContentBrowser/ContentBrowserAssetDataSource/Source/ContentBrowserAssetDataSource/Public/ContentBrowserAssetDataCore.h b/Engine/Plugins/Editor/ContentBrowser/ContentBrowserAssetDataSource/Source/ContentBrowserAssetDataSource/Public/ContentBrowserAssetDataCore.h
--- a/Engine/Plugins/Editor/ContentBrowser/ContentBrowserAssetDataSource/Source/ContentBrowserAssetDataSource/Public/ContentBrowserAssetDataCore.h
+++ b/Engine/Plugins/Editor/ContentBrowser/ContentBrowserAssetDataSource/Source/ContentBrowserAssetDataSource/Public/ContentBrowserAssetDataCore.h
@@ -53,14 +53,6 @@ namespace ContentBrowserAssetData
 
 	CONTENTBROWSERASSETDATASOURCE_API bool EditAssetFileItems(TArrayView<const TSharedRef<const FContentBrowserAssetFileItemDataPayload>> InAssetPayloads);
 
-	CONTENTBROWSERASSETDATASOURCE_API bool CanViewItem(IAssetTools* InAssetTools, const UContentBrowserDataSource* InOwnerDataSource, const FContentBrowserItemData& InItem, FText* OutErrorMsg);
-
-	CONTENTBROWSERASSETDATASOURCE_API bool CanViewAssetFileItem(IAssetTools* InAssetTools, const FContentBrowserAssetFileItemDataPayload& InAssetPayload, FText* OutErrorMsg);
-
-	CONTENTBROWSERASSETDATASOURCE_API bool ViewItems(IAssetTools* InAssetTools, const UContentBrowserDataSource* InOwnerDataSource, TArrayView<const FContentBrowserItemData> InItems);
-
-	CONTENTBROWSERASSETDATASOURCE_API bool ViewAssetFileItems(TArrayView<const TSharedRef<const FContentBrowserAssetFileItemDataPayload>> InAssetPayloads);
-
 	CONTENTBROWSERASSETDATASOURCE_API bool CanPreviewItem(IAssetTools* InAssetTools, const UContentBrowserDataSource* InOwnerDataSource, const FContentBrowserItemData& InItem, FText* OutErrorMsg);
 
 	CONTENTBROWSERASSETDATASOURCE_API bool CanPreviewAssetFileItem(IAssetTools* InAssetTools, const FContentBrowserAssetFileItemDataPayload& InAssetPayload, FText* OutErrorMsg);
diff --git a/Engine/Plugins/Editor/ContentBrowser/ContentBrowserAssetDataSource/Source/ContentBrowserAssetDataSource/Public/ContentBrowserAssetDataSource.h b/Engine/Plugins/Editor/ContentBrowser/ContentBrowserAssetDataSource/Source/ContentBrowserAssetDataSource/Public/ContentBrowserAssetDataSource.h
--- a/Engine/Plugins/Editor/ContentBrowser/ContentBrowserAssetDataSource/Source/ContentBrowserAssetDataSource/Public/ContentBrowserAssetDataSource.h
+++ b/Engine/Plugins/Editor/ContentBrowser/ContentBrowserAssetDataSource/Source/ContentBrowserAssetDataSource/Public/ContentBrowserAssetDataSource.h
@@ -187,12 +187,6 @@ public:
 	virtual bool EditItem(const FContentBrowserItemData& InItem) override;
 
 	virtual bool BulkEditItems(TArrayView<const FContentBrowserItemData> InItems) override;
-	
-	virtual bool CanViewItem(const FContentBrowserItemData& InItem, FText* OutErrorMsg) override;
-
-	virtual bool ViewItem(const FContentBrowserItemData& InItem) override;
-
-	virtual bool BulkViewItems(TArrayView<const FContentBrowserItemData> InItems) override;
 
 	virtual bool CanPreviewItem(const FContentBrowserItemData& InItem, FText* OutErrorMsg) override;
 
diff --git a/Engine/Source/Developer/AssetTools/Public/AssetTypeActions_Base.h b/Engine/Source/Developer/AssetTools/Public/AssetTypeActions_Base.h
--- a/Engine/Source/Developer/AssetTools/Public/AssetTypeActions_Base.h
+++ b/Engine/Source/Developer/AssetTools/Public/AssetTypeActions_Base.h
@@ -11,7 +11,6 @@
 #include "IAssetTools.h"
 #include "IAssetTypeActions.h"
 #include "AssetToolsModule.h"
-#include "AssetTypeActivationOpenedMethod.h"
 #include "Toolkits/SimpleAssetEditor.h"
 #include "AssetRegistry/ARFilter.h"
 
@@ -36,14 +35,6 @@ public:
 		FSimpleAssetEditor::CreateEditor(EToolkitMode::Standalone, EditWithinLevelEditor, InObjects);
 	}
 
-	virtual void OpenAssetEditor(const TArray<UObject*>& InObjects, const EAssetTypeActivationOpenedMethod OpenedMethod, TSharedPtr<IToolkitHost> EditWithinLevelEditor = TSharedPtr<IToolkitHost>()) override
-	{
-		if (OpenedMethod == EAssetTypeActivationOpenedMethod::Edit)
-		{
-			OpenAssetEditor(InObjects, EditWithinLevelEditor);
-		}
-	}
-
 	virtual bool AssetsActivatedOverride(const TArray<UObject*>& InObjects, EAssetTypeActivationMethod::Type ActivationType) override
 	{
 		return false;
@@ -171,11 +162,6 @@ public:
 		return bIsSupported;
 	}
 
-	virtual bool SupportsOpenedMethod(const EAssetTypeActivationOpenedMethod OpenedMethod) const override
-	{
-		return OpenedMethod == EAssetTypeActivationOpenedMethod::Edit;
-	}
-
 	virtual FName GetFilterName() const override
 	{
 		return FName(*GetSupportedClass()->GetPathName());
diff --git a/Engine/Source/Developer/AssetTools/Public/AssetTypeActivationOpenedMethod.h b/Engine/Source/Developer/AssetTools/Public/AssetTypeActivationOpenedMethod.h
deleted file mode 100644
diff --git a/Engine/Source/Developer/AssetTools/Public/IAssetTypeActions.h b/Engine/Source/Developer/AssetTools/Public/IAssetTypeActions.h
--- a/Engine/Source/Developer/AssetTools/Public/IAssetTypeActions.h
+++ b/Engine/Source/Developer/AssetTools/Public/IAssetTypeActions.h
@@ -10,7 +10,6 @@
 #include "IAssetTypeActions.generated.h"
 
 class IToolkitHost;
-enum class EAssetTypeActivationOpenedMethod : int32;
 
 namespace EAssetTypeActivationMethod
 {
@@ -22,6 +21,8 @@ namespace EAssetTypeActivationMethod
 	};
 }
 
+class IToolkitHost;
+
 /* Revision information for a single revision of a file in source control */
 USTRUCT(BlueprintType)
 struct FRevisionInfo
@@ -95,9 +96,6 @@ public:
 	/** Opens the asset editor for the specified objects. If EditWithinLevelEditor is valid, the world-centric editor will be used. */
 	virtual void OpenAssetEditor( const TArray<UObject*>& InObjects, TSharedPtr<IToolkitHost> EditWithinLevelEditor = TSharedPtr<IToolkitHost>() ) = 0;
 
-	/** Opens the asset editor for the specified objects. If EditWithinLevelEditor is valid, the world-centric editor will be used. */
-	virtual void OpenAssetEditor( const TArray<UObject*>& InObjects, const EAssetTypeActivationOpenedMethod OpenedMethod, TSharedPtr<IToolkitHost> EditWithinLevelEditor = TSharedPtr<IToolkitHost>() ) = 0;
-
 	/** Allows overriding asset activation to perform asset type specific activation for the supplied assets. This happens when the user double clicks, presses enter, or presses space. Return true if you have overridden the behavior. */
 	virtual bool AssetsActivatedOverride(const TArray<UObject*>& InObjects, EAssetTypeActivationMethod::Type ActivationType) = 0;
 
@@ -179,9 +177,6 @@ public:
 	/** Returns class path name as a package + class FName pair */
 	virtual FTopLevelAssetPath GetClassPathName() const = 0;
 
-	/** Does this asset support edit or view methods? */
-	virtual bool SupportsOpenedMethod(const EAssetTypeActivationOpenedMethod OpenedMethod) const = 0;
-
 	/** Returns thumbnail brush unique for the given asset data.  Returning null falls back to class thumbnail brush. */
 	virtual const FSlateBrush* GetThumbnailBrush(const FAssetData& InAssetData, const FName InClassName) const = 0;
 
diff --git a/Engine/Source/Editor/ContentBrowser/Private/SContentBrowser.cpp b/Engine/Source/Editor/ContentBrowser/Private/SContentBrowser.cpp
--- a/Engine/Source/Editor/ContentBrowser/Private/SContentBrowser.cpp
+++ b/Engine/Source/Editor/ContentBrowser/Private/SContentBrowser.cpp
@@ -3011,7 +3011,7 @@ void SContentBrowser::OnItemsActivated(TArrayView<const FContentBrowserItem> Act
 			for (const FContentBrowserItemData& ItemToEdit : SourceAndItemsPair.Value)
 			{
 				FText EditErrorMsg;
-				if (!SourceAndItemsPair.Key->CanEditItem(ItemToEdit, &EditErrorMsg) && !SourceAndItemsPair.Key->CanViewItem(ItemToEdit, &EditErrorMsg))
+				if (!SourceAndItemsPair.Key->CanEditItem(ItemToEdit, &EditErrorMsg))
 				{
 					AssetViewUtils::ShowErrorNotifcation(EditErrorMsg);
 				}
diff --git a/Engine/Source/Editor/ContentBrowserData/Private/ContentBrowserDataSource.cpp b/Engine/Source/Editor/ContentBrowserData/Private/ContentBrowserDataSource.cpp
--- a/Engine/Source/Editor/ContentBrowserData/Private/ContentBrowserDataSource.cpp
+++ b/Engine/Source/Editor/ContentBrowserData/Private/ContentBrowserDataSource.cpp
@@ -305,26 +305,6 @@ bool UContentBrowserDataSource::BulkEditItems(TArrayView<const FContentBrowserIt
 	return bSuccess;
 }
 
-bool UContentBrowserDataSource::CanViewItem(const FContentBrowserItemData& InItem, FText* OutErrorMsg)
-{
-	return false;
-}
-
-bool UContentBrowserDataSource::ViewItem(const FContentBrowserItemData& InItem)
-{
-	return false;
-}
-
-bool UContentBrowserDataSource::BulkViewItems(TArrayView<const FContentBrowserItemData> InItems)
-{
-	bool bSuccess = false;
-	for (const FContentBrowserItemData& Item : InItems)
-	{
-		bSuccess |= ViewItem(Item);
-	}
-	return bSuccess;
-}
-
 bool UContentBrowserDataSource::CanPreviewItem(const FContentBrowserItemData& InItem, FText* OutErrorMsg)
 {
 	return false;
diff --git a/Engine/Source/Editor/ContentBrowserData/Public/ContentBrowserDataSource.h b/Engine/Source/Editor/ContentBrowserData/Public/ContentBrowserDataSource.h
--- a/Engine/Source/Editor/ContentBrowserData/Public/ContentBrowserDataSource.h
+++ b/Engine/Source/Editor/ContentBrowserData/Public/ContentBrowserDataSource.h
@@ -356,35 +356,6 @@ public:
 	 */
 	virtual bool BulkEditItems(TArrayView<const FContentBrowserItemData> InItems);
 
-	/**
-     * Query whether the given item is can be viewed (a read-only asset editor), optionally providing error information if it cannot.
-     *
-     * @param InItem The item to query.
-     * @param OutErrorMessage Optional error message to fill on failure.
-     *
-     * @return True if the item can be viewed in a read-only editor, false otherwise.
-     */
-    virtual bool CanViewItem(const FContentBrowserItemData& InItem, FText* OutErrorMsg);
-
-    /**
-     * Attempt to open the given item for read-only viewing.
-     *
-     * @param InItem The item to view.
-     *
-     * @return True if the item was opened for read-only viewing, false otherwise.
-     */
-    virtual bool ViewItem(const FContentBrowserItemData& InItem);
-
-    /**
-     * Attempt to open the given items for read-only viewing.
-     * @note The default implementation of this will call ViewItem for each item. Override if you can provide a more efficient implementation.
-     *
-     * @param InItems The items to view.
-     *
-     * @return True if any items were opened for read-only viewing, false otherwise.
-     */
-    virtual bool BulkViewItems(TArrayView<const FContentBrowserItemData> InItems);
-
 	/**
 	 * Query whether the given item is can be previewed, optionally providing error information if it cannot.
 	 *
diff --git a/Engine/Source/Editor/UnrealEd/Private/Subsystems/AssetEditorSubsystem.cpp b/Engine/Source/Editor/UnrealEd/Private/Subsystems/AssetEditorSubsystem.cpp
--- a/Engine/Source/Editor/UnrealEd/Private/Subsystems/AssetEditorSubsystem.cpp
+++ b/Engine/Source/Editor/UnrealEd/Private/Subsystems/AssetEditorSubsystem.cpp
@@ -151,6 +151,7 @@ int32 UAssetEditorSubsystem::CloseAllEditorsForAsset(UObject* Asset)
 	return EditorInstances.Num();
 }
 
+
 void UAssetEditorSubsystem::RemoveAssetFromAllEditors(UObject* Asset)
 {
 	TArray<IAssetEditorInstance*> EditorInstances = FindEditorsForAsset(Asset);
@@ -310,7 +311,7 @@ bool UAssetEditorSubsystem::IsAssetEditable(const UObject* Asset)
 	return true;
 }
 
-bool UAssetEditorSubsystem::OpenEditorForAsset(UObject* Asset, const EToolkitMode::Type ToolkitMode, TSharedPtr< IToolkitHost > OpenedFromLevelEditor, const bool bShowProgressWindow, const EAssetTypeActivationOpenedMethod OpenedMethod)
+bool UAssetEditorSubsystem::OpenEditorForAsset(UObject* Asset, const EToolkitMode::Type ToolkitMode, TSharedPtr< IToolkitHost > OpenedFromLevelEditor, const bool bShowProgressWindow)
 {
 	SCOPE_STALL_REPORTER(UAssetEditorSubsystem::OpenEditorForAsset, 2.0);
 
@@ -330,21 +331,18 @@ bool UAssetEditorSubsystem::OpenEditorForAsset(UObject* Asset, const EToolkitMod
 
 	const bool bBringToFrontIfOpen = true;
 
-	if (OpenedMethod == EAssetTypeActivationOpenedMethod::Edit)
+	if (UPackage* Package = Asset->GetOutermost())
 	{
-		if (UPackage* Package = Asset->GetOutermost())
+		// Don't open asset editors for cooked packages
+		if (Package->bIsCookedForEditor)
 		{
-			// Don't open asset editors for cooked packages
-			if (Package->bIsCookedForEditor)
-			{
-				return false;
-			}
+			return false;
+		}
 
-			if (!AssetToolsModule.Get().GetWritableFolderPermissionList()->PassesStartsWithFilter(Package->GetName()))
-			{
-				AssetToolsModule.Get().NotifyBlockedByWritableFolderFilter();
-				return false;
-			}
+		if (!AssetToolsModule.Get().GetWritableFolderPermissionList()->PassesStartsWithFilter(Package->GetName()))
+		{
+			AssetToolsModule.Get().NotifyBlockedByWritableFolderFilter();
+			return false;
 		}
 	}
 
@@ -393,7 +391,7 @@ bool UAssetEditorSubsystem::OpenEditorForAsset(UObject* Asset, const EToolkitMod
 	// Disallow opening an asset editor for classes
 	bool bCanSummonSimpleAssetEditor = !Asset->IsA<UClass>();
 
-	if (AssetTypeActions.IsValid() && AssetTypeActions.Pin()->SupportsOpenedMethod(OpenedMethod))
+	if (AssetTypeActions.IsValid())
 	{
 		TArray<UObject*> AssetsToEdit;
 		AssetsToEdit.Add(Asset);
@@ -402,7 +400,7 @@ bool UAssetEditorSubsystem::OpenEditorForAsset(UObject* Asset, const EToolkitMod
 		TWeakObjectPtr<UObject> WeakAsset = Asset;
 		const FString AssetPath = Asset->GetPathName();
 
-		AssetTypeActions.Pin()->OpenAssetEditor(AssetsToEdit, OpenedMethod, ActualToolkitMode == EToolkitMode::WorldCentric ? OpenedFromLevelEditor : TSharedPtr<IToolkitHost>());
+		AssetTypeActions.Pin()->OpenAssetEditor(AssetsToEdit, ActualToolkitMode == EToolkitMode::WorldCentric ? OpenedFromLevelEditor : TSharedPtr<IToolkitHost>());
 
 		// If the Asset was destroyed, attempt to find it if it was recreated
 		if (!WeakAsset.IsValid() && !AssetPath.IsEmpty())
@@ -441,13 +439,14 @@ bool UAssetEditorSubsystem::OpenEditorForAsset(UObject* Asset, const EToolkitMod
 		FStudioAnalytics::FireEvent_Loading(TEXT("OpenAssetEditor"), OpenTime, {
 			FAnalyticsEventAttribute(TEXT("AssetPath"), Asset->GetFullName()),
 			FAnalyticsEventAttribute(TEXT("AssetType"), Asset->GetClass()->GetName())
-			});
+		});
 	}
-	
+
 	return true;
 }
 
-bool UAssetEditorSubsystem::OpenEditorForAssets_Advanced(const TArray <UObject*>& InAssets, const EToolkitMode::Type ToolkitMode, TSharedPtr< IToolkitHost > OpenedFromLevelEditor, const EAssetTypeActivationOpenedMethod OpenedMethod)
+
+bool UAssetEditorSubsystem::OpenEditorForAssets_Advanced(const TArray <UObject* >& InAssets, const EToolkitMode::Type ToolkitMode, TSharedPtr< IToolkitHost > OpenedFromLevelEditor)
 {
 	TArray<UObject*> Assets;
 	Assets.Reserve(InAssets.Num());
@@ -475,7 +474,7 @@ bool UAssetEditorSubsystem::OpenEditorForAssets_Advanced(const TArray <UObject*>
 
 	if (Assets.Num() == 1)
 	{
-		return OpenEditorForAsset(Assets[0], ToolkitMode, OpenedFromLevelEditor, true, OpenedMethod);
+		return OpenEditorForAsset(Assets[0], ToolkitMode, OpenedFromLevelEditor);
 	}
 	else if (Assets.Num() > 0)
 	{
@@ -485,7 +484,7 @@ bool UAssetEditorSubsystem::OpenEditorForAssets_Advanced(const TArray <UObject*>
 			// If any of the assets are already open or the package is cooked,
 			// remove them from the list of assets to open an editor for
 			UPackage* Package = Asset->GetOutermost();
-			if (FindEditorForAsset(Asset, true) != nullptr || (OpenedMethod == EAssetTypeActivationOpenedMethod::Edit && Package && Package->bIsCookedForEditor))
+			if (FindEditorForAsset(Asset, true) != nullptr || (Package && Package->bIsCookedForEditor))
 			{
 				SkipOpenAssets.Add(Asset);
 			}
@@ -509,7 +508,7 @@ bool UAssetEditorSubsystem::OpenEditorForAssets_Advanced(const TArray <UObject*>
 			FAssetToolsModule& AssetToolsModule = FModuleManager::LoadModuleChecked<FAssetToolsModule>(TEXT("AssetTools"));
 			TWeakPtr<IAssetTypeActions> AssetTypeActions = AssetToolsModule.Get().GetAssetTypeActionsForClass(AssetClass);
 
-			if (AssetTypeActions.IsValid() && AssetTypeActions.Pin()->SupportsOpenedMethod(OpenedMethod))
+			if (AssetTypeActions.IsValid())
 			{
 				GWarn->BeginSlowTask(LOCTEXT("OpenEditors", "Opening Editor(s)..."), true);
 
@@ -552,7 +551,7 @@ bool UAssetEditorSubsystem::OpenEditorForAssets_Advanced(const TArray <UObject*>
 				}
 
 				// How to handle multiple assets is left up to the type actions (i.e. open a single shared editor or an editor for each)
-				AssetTypeActions.Pin()->OpenAssetEditor(Assets, OpenedMethod, ActualToolkitMode == EToolkitMode::WorldCentric ? OpenedFromLevelEditor : TSharedPtr<IToolkitHost>());
+				AssetTypeActions.Pin()->OpenAssetEditor(Assets, ActualToolkitMode == EToolkitMode::WorldCentric ? OpenedFromLevelEditor : TSharedPtr<IToolkitHost>());
 
 				// If any assets were destroyed, attempt to find them if they were recreated
 				for (int32 i = 0; i < Assets.Num(); i++)
@@ -579,7 +578,7 @@ bool UAssetEditorSubsystem::OpenEditorForAssets_Advanced(const TArray <UObject*>
 			{
 				if (!SkipOpenAssets.Contains(Asset))
 				{
-					OpenEditorForAsset(Asset, ToolkitMode, OpenedFromLevelEditor, true, OpenedMethod);
+					OpenEditorForAsset(Asset, ToolkitMode, OpenedFromLevelEditor);
 				}
 			}
 		}
@@ -588,9 +587,9 @@ bool UAssetEditorSubsystem::OpenEditorForAssets_Advanced(const TArray <UObject*>
 	return NumNullAssets == 0;
 }
 
-bool UAssetEditorSubsystem::OpenEditorForAssets(const TArray<UObject*>& Assets, const EAssetTypeActivationOpenedMethod OpenedMethod)
+bool UAssetEditorSubsystem::OpenEditorForAssets(const TArray<UObject*>& Assets)
 {
-	return OpenEditorForAssets_Advanced(Assets, EToolkitMode::Standalone, TSharedPtr<IToolkitHost>(), OpenedMethod);
+	return OpenEditorForAssets_Advanced(Assets, EToolkitMode::Standalone, TSharedPtr<IToolkitHost>());
 }
 
 void UAssetEditorSubsystem::HandleRequestOpenAssetMessage(const FAssetEditorRequestOpenAsset& Message, const TSharedRef<IMessageContext, ESPMode::ThreadSafe>& Context)
@@ -598,7 +597,12 @@ void UAssetEditorSubsystem::HandleRequestOpenAssetMessage(const FAssetEditorRequ
 	OpenEditorForAsset(Message.AssetName);
 }
 
-void UAssetEditorSubsystem::OpenEditorForAsset(const FString& AssetPathName, const EAssetTypeActivationOpenedMethod OpenedMethod)
+void UAssetEditorSubsystem::OpenEditorForAsset(const FSoftObjectPath& AssetPath)
+{
+	OpenEditorForAsset(AssetPath.ToString());
+}
+
+void UAssetEditorSubsystem::OpenEditorForAsset(const FString& AssetPathName)
 {
 	// An asset needs loading
 
@@ -609,7 +613,7 @@ void UAssetEditorSubsystem::OpenEditorForAsset(const FString& AssetPathName, con
 
 		FString AssetName = FPaths::GetBaseFilename(AssetPathName);
 		UObject* Object = FindObject<UObject>(Package, *AssetName);
-		
+
 		if (Object != nullptr)
 		{
 			OpenEditorForAsset(Object);
@@ -622,7 +626,6 @@ void UAssetEditorSubsystem::OpenEditorForAsset(const FString& AssetPathName, con
 		
 		if (Object != nullptr)
 		{
-			OpenEditorForAsset(Object, EToolkitMode::Standalone, TSharedPtr<IToolkitHost>(), true, OpenedMethod);
 		}
 	}
 }
@@ -1021,19 +1024,27 @@ void UAssetEditorSubsystem::HandlePackageReloaded(const EPackageReloadPhase InPa
 	}
 }
 
-void UAssetEditorSubsystem::OpenEditorsForAssets(const TArray<FString>& AssetsToOpen, const EAssetTypeActivationOpenedMethod OpenedMethod)
+void UAssetEditorSubsystem::OpenEditorsForAssets(const TArray<FSoftObjectPath>& AssetsToOpen)
+{
+	for (const FSoftObjectPath& AssetName : AssetsToOpen)
+	{
+		OpenEditorForAsset(AssetName);
+	}
+}
+
+void UAssetEditorSubsystem::OpenEditorsForAssets(const TArray<FString>& AssetsToOpen)
 {
 	for (const FString& AssetName : AssetsToOpen)
 	{
-		OpenEditorForAsset(AssetName, OpenedMethod);
+		OpenEditorForAsset(AssetName);
 	}
 }
 
-void UAssetEditorSubsystem::OpenEditorsForAssets(const TArray<FName>& AssetsToOpen, const EAssetTypeActivationOpenedMethod OpenedMethod)
+void UAssetEditorSubsystem::OpenEditorsForAssets(const TArray<FName>& AssetsToOpen)
 {
 	for (const FName& AssetName : AssetsToOpen)
 	{
-		OpenEditorForAsset(AssetName.ToString(), OpenedMethod);
+		OpenEditorForAsset(AssetName.ToString());
 	}
 }
 
diff --git a/Engine/Source/Editor/UnrealEd/Public/Subsystems/AssetEditorSubsystem.h b/Engine/Source/Editor/UnrealEd/Public/Subsystems/AssetEditorSubsystem.h
--- a/Engine/Source/Editor/UnrealEd/Public/Subsystems/AssetEditorSubsystem.h
+++ b/Engine/Source/Editor/UnrealEd/Public/Subsystems/AssetEditorSubsystem.h
@@ -9,7 +9,6 @@
 #include "Containers/Ticker.h"
 #include "Tools/Modes.h"
 #include "Misc/NamePermissionList.h"
-#include "AssetTypeActivationOpenedMethod.h"
 #include "AssetEditorSubsystem.generated.h"
 
 class UAssetEditor;
@@ -74,19 +73,19 @@ public:
 	bool IsAssetEditable(const UObject* Asset);
 
 	/** Opens an asset by path */
-	void OpenEditorForAsset(const FString& AssetPathName, const EAssetTypeActivationOpenedMethod OpenedMethod = EAssetTypeActivationOpenedMethod::Edit);
+	void OpenEditorForAsset(const FString& AssetPath);
 	void OpenEditorForAsset(const FSoftObjectPath& AssetPath);
 
 	/**
 	 * Tries to open an editor for the specified asset.  Returns true if the asset is open in an editor.
 	 * If the file is already open in an editor, it will not create another editor window but instead bring it to front
 	 */
-	bool OpenEditorForAsset(UObject* Asset, const EToolkitMode::Type ToolkitMode = EToolkitMode::Standalone, TSharedPtr<IToolkitHost> OpenedFromLevelEditor = TSharedPtr<IToolkitHost>(), const bool bShowProgressWindow = true, EAssetTypeActivationOpenedMethod OpenedMethod = EAssetTypeActivationOpenedMethod::Edit);
+	bool OpenEditorForAsset(UObject* Asset, const EToolkitMode::Type ToolkitMode = EToolkitMode::Standalone, TSharedPtr<IToolkitHost> OpenedFromLevelEditor = TSharedPtr<IToolkitHost>(), const bool bShowProgressWindow = true);
 	
 	template<typename ObjectType>
-	bool OpenEditorForAsset(TObjectPtr<ObjectType> Asset, const EToolkitMode::Type ToolkitMode = EToolkitMode::Standalone, TSharedPtr<IToolkitHost> OpenedFromLevelEditor = TSharedPtr<IToolkitHost>(), const bool bShowProgressWindow = true, EAssetTypeActivationOpenedMethod OpenedMethod = EAssetTypeActivationOpenedMethod::Edit)
+	bool OpenEditorForAsset(TObjectPtr<ObjectType> Asset, const EToolkitMode::Type ToolkitMode = EToolkitMode::Standalone, TSharedPtr<IToolkitHost> OpenedFromLevelEditor = TSharedPtr<IToolkitHost>(), const bool bShowProgressWindow = true)
 	{
-		return OpenEditorForAsset(ToRawPtr(Asset), ToolkitMode, OpenedFromLevelEditor, bShowProgressWindow, OpenedMethod);
+		return OpenEditorForAsset(ToRawPtr(Asset), ToolkitMode, OpenedFromLevelEditor, bShowProgressWindow);
 	}
 
 	/**
@@ -95,12 +94,12 @@ public:
 	 * If all assets are of the same type, the supporting AssetTypeAction (if it exists) is responsible for the details of how to handle opening multiple assets at once.
 	 */
 	UFUNCTION(BlueprintCallable, Category = "Editor Scripting | Asset Tools")
-	bool OpenEditorForAssets(const TArray<UObject*>& Assets, const EAssetTypeActivationOpenedMethod OpenedMethod = EAssetTypeActivationOpenedMethod::Edit);
-	bool OpenEditorForAssets_Advanced(const TArray<UObject*>& InAssets, const EToolkitMode::Type ToolkitMode = EToolkitMode::Standalone, TSharedPtr<IToolkitHost> OpenedFromLevelEditor = TSharedPtr<IToolkitHost>(), const EAssetTypeActivationOpenedMethod OpenedMethod = EAssetTypeActivationOpenedMethod::Edit);
+	bool OpenEditorForAssets(const TArray<UObject*>& Assets);
+	bool OpenEditorForAssets_Advanced(const TArray<UObject*>& Assets, const EToolkitMode::Type ToolkitMode = EToolkitMode::Standalone, TSharedPtr<IToolkitHost> OpenedFromLevelEditor = TSharedPtr<IToolkitHost>());
 
 	/** Opens editors for the supplied assets (via OpenEditorForAsset) */
-	void OpenEditorsForAssets(const TArray<FString>& AssetsToOpen, const EAssetTypeActivationOpenedMethod OpenedMethod = EAssetTypeActivationOpenedMethod::Edit);
-	void OpenEditorsForAssets(const TArray<FName>& AssetsToOpen, const EAssetTypeActivationOpenedMethod OpenedMethod = EAssetTypeActivationOpenedMethod::Edit);
+	void OpenEditorsForAssets(const TArray<FString>& AssetsToOpen);
+	void OpenEditorsForAssets(const TArray<FName>& AssetsToOpen);
 	void OpenEditorsForAssets(const TArray<FSoftObjectPath>& AssetsToOpen);
 
 	/** Returns the primary editor if one is already open for the specified asset.
diff --git a/Engine/Source/Editor/UnrealEd/UnrealEd.Build.cs b/Engine/Source/Editor/UnrealEd/UnrealEd.Build.cs
--- a/Engine/Source/Editor/UnrealEd/UnrealEd.Build.cs
+++ b/Engine/Source/Editor/UnrealEd/UnrealEd.Build.cs
@@ -121,7 +121,6 @@ public class UnrealEd : ModuleRules
 				"SubobjectEditor",
 				"PhysicsUtilities",
 				"ToolWidgets",
-				"AssetTools",
 			}
 		);
 
@@ -230,6 +229,7 @@ public class UnrealEd : ModuleRules
 				"TextureEditor",
 				"Cascade",
 				"UMGEditor",
+				"AssetTools",
 				"ClassViewer",
 				"StructViewer",
 				"CollectionManager",
@@ -326,6 +326,7 @@ public class UnrealEd : ModuleRules
 				"BlueprintGraph",
 				"AddContentDialog",
 				"MeshUtilities",
+				"AssetTools",
 				"KismetCompiler",
 				"NavigationSystem",
 				"GameplayTasks",
