Fbx scene import: fix memory corruption.

#rnx

diff --git a/Engine/Source/Editor/UnrealEd/Private/Fbx/SFbxSceneOptionWindow.cpp b/Engine/Source/Editor/UnrealEd/Private/Fbx/SFbxSceneOptionWindow.cpp
--- a/Engine/Source/Editor/UnrealEd/Private/Fbx/SFbxSceneOptionWindow.cpp
+++ b/Engine/Source/Editor/UnrealEd/Private/Fbx/SFbxSceneOptionWindow.cpp
@@ -1393,7 +1393,7 @@ bool SFbxSceneOptionWindow::CanImport()  const
 
 void SFbxSceneOptionWindow::CopyFbxOptionsToFbxOptions(UnFbx::FBXImportOptions *SourceOptions, UnFbx::FBXImportOptions *DestinationOptions)
 {
-	FMemory::BigBlockMemcpy(DestinationOptions, SourceOptions, sizeof(UnFbx::FBXImportOptions));
+	DestinationOptions = SourceOptions;
 }
 
 void SFbxSceneOptionWindow::CopyStaticMeshOptionsToFbxOptions(UnFbx::FBXImportOptions *ImportSettings, UFbxSceneImportOptionsStaticMesh* StaticMeshOptions)
