Further improve logging

diff --git a/Engine/Source/Editor/UnrealEd/Private/Kismet2/KismetReinstanceUtilities.cpp b/Engine/Source/Editor/UnrealEd/Private/Kismet2/KismetReinstanceUtilities.cpp
--- a/Engine/Source/Editor/UnrealEd/Private/Kismet2/KismetReinstanceUtilities.cpp
+++ b/Engine/Source/Editor/UnrealEd/Private/Kismet2/KismetReinstanceUtilities.cpp
@@ -106,7 +106,7 @@ struct FReplaceReferenceHelper
 			UE_LOG(LogBlueprint, Warning, TEXT("%s %s %s %s Replacing References to: %s"), 
 				*(GLogReinstancerReferenceReplacementObj->GetOutermost()->GetName()),
 				GLogReinstancerReferenceReplacementObj->GetOuter() ? *(GLogReinstancerReferenceReplacementObj->GetOuter()->GetName()) : TEXT("nullouter"),
-				*(GLogReinstancerReferenceReplacementObj->GetName()),
+				*(GLogReinstancerReferenceReplacementObj->GetPathName()),
 				*(GLogReinstancerReferenceReplacementObj->GetClass()->GetName()), 
 				*Classes);
 		}
