Add  FContentBundleEditingContext to activate a ContentBundle in a specific scope.
Allows to disable a content bundle after a plugin actors was Drag n dropped

diff --git a/Engine/Source/Editor/UnrealEd/Private/LevelEditorDragDropHandler.cpp b/Engine/Source/Editor/UnrealEd/Private/LevelEditorDragDropHandler.cpp
--- a/Engine/Source/Editor/UnrealEd/Private/LevelEditorDragDropHandler.cpp
+++ b/Engine/Source/Editor/UnrealEd/Private/LevelEditorDragDropHandler.cpp
@@ -60,4 +60,9 @@ bool ULevelEditorDragDropHandler::PreDropObjectsAtCoordinates(int32 MouseX, int3
 	return true;
 }
 
+bool ULevelEditorDragDropHandler::PostDropObjectsAtCoordinates(int32 MouseX, int32 MouseY, UWorld* World, FViewport* Viewport, const TArray<UObject*>& DroppedObjects)
+{
+	return true;
+}
+
 #undef LOCTEXT_NAMESPACE
\ No newline at end of file
diff --git a/Engine/Source/Editor/UnrealEd/Private/LevelEditorViewport.cpp b/Engine/Source/Editor/UnrealEd/Private/LevelEditorViewport.cpp
--- a/Engine/Source/Editor/UnrealEd/Private/LevelEditorViewport.cpp
+++ b/Engine/Source/Editor/UnrealEd/Private/LevelEditorViewport.cpp
@@ -1481,10 +1481,10 @@ bool FLevelEditorViewportClient::DropObjectsAtCoordinates(int32 MouseX, int32 Mo
 	bool bResult = false;
 
 	// Allow the drag drop handler to do anything pre-drop.
+	ULevelEditorDragDropHandler* DragDropHandler = GEditor->GetLevelEditorDragDropHandler();
 	if (!bCreateDropPreview)
 	{
-		ULevelEditorDragDropHandler* DragDrop = GEditor->GetLevelEditorDragDropHandler();
-		if (!DragDrop->PreDropObjectsAtCoordinates(MouseX, MouseY, GetWorld(), Viewport, DroppedObjects, OutNewActors))
+		if (!DragDropHandler->PreDropObjectsAtCoordinates(MouseX, MouseY, GetWorld(), Viewport, DroppedObjects, OutNewActors))
 		{
 			return bResult;
 		}
@@ -1696,6 +1696,11 @@ bool FLevelEditorViewportClient::DropObjectsAtCoordinates(int32 MouseX, int32 Mo
 		}
 	}
 
+	if (!bCreateDropPreview)
+	{
+		DragDropHandler->PostDropObjectsAtCoordinates(MouseX, MouseX, World, Viewport, DroppedObjects);
+	}
+
 	// Reset if creating a preview actor.
 	bIsDroppingPreviewActor = false;
 
diff --git a/Engine/Source/Editor/UnrealEd/Public/LevelEditorDragDropHandler.h b/Engine/Source/Editor/UnrealEd/Public/LevelEditorDragDropHandler.h
--- a/Engine/Source/Editor/UnrealEd/Public/LevelEditorDragDropHandler.h
+++ b/Engine/Source/Editor/UnrealEd/Public/LevelEditorDragDropHandler.h
@@ -20,6 +20,7 @@ public:
 	
 	virtual bool PreviewDropObjectsAtCoordinates(int32 MouseX, int32 MouseY, UWorld* World, FViewport* Viewport, const FAssetData& AssetData);
 	virtual bool PreDropObjectsAtCoordinates(int32 MouseX, int32 MouseY, UWorld* World, FViewport* Viewport, const TArray<UObject*>& DroppedObjects, TArray<AActor*>& OutNewActors);
+	virtual bool PostDropObjectsAtCoordinates(int32 MouseX, int32 MouseY, UWorld* World, FViewport* Viewport, const TArray<UObject*>& DroppedObjects);
 
 	bool GetCanDrop() const { return bCanDrop; }
 	FText GetPreviewDropHintText() const { return HintText; }
diff --git a/Engine/Source/Editor/WorldPartitionEditor/Private/WorldPartition/ContentBundle/ContentBundleEditorSubsystem.cpp b/Engine/Source/Editor/WorldPartitionEditor/Private/WorldPartition/ContentBundle/ContentBundleEditorSubsystem.cpp
--- a/Engine/Source/Editor/WorldPartitionEditor/Private/WorldPartition/ContentBundle/ContentBundleEditorSubsystem.cpp
+++ b/Engine/Source/Editor/WorldPartitionEditor/Private/WorldPartition/ContentBundle/ContentBundleEditorSubsystem.cpp
@@ -294,7 +294,7 @@ bool UContentBundleEditorSubsystem::ActivateContentBundleEditing(TSharedPtr<FCon
 	return ContentBundleEditionSubModule->ActivateContentBundleEditing(ContentBundleEditor);
 }
 
-bool UContentBundleEditorSubsystem::DectivateContentBundleEditing(TSharedPtr<FContentBundleEditor>& ContentBundleEditor) const
+bool UContentBundleEditorSubsystem::DeactivateContentBundleEditing(TSharedPtr<FContentBundleEditor>& ContentBundleEditor) const
 {
 	return ContentBundleEditionSubModule->DeactivateContentBundleEditing(ContentBundleEditor);
 }
diff --git a/Engine/Source/Editor/WorldPartitionEditor/Public/WorldPartition/ContentBundle/ContentBundleEditorSubsystem.h b/Engine/Source/Editor/WorldPartitionEditor/Public/WorldPartition/ContentBundle/ContentBundleEditorSubsystem.h
--- a/Engine/Source/Editor/WorldPartitionEditor/Public/WorldPartition/ContentBundle/ContentBundleEditorSubsystem.h
+++ b/Engine/Source/Editor/WorldPartitionEditor/Public/WorldPartition/ContentBundle/ContentBundleEditorSubsystem.h
@@ -117,7 +117,7 @@ public:
 	void UnreferenceAllActors(FContentBundleEditor& EditorContentBundle);
 
 	bool ActivateContentBundleEditing(TSharedPtr<FContentBundleEditor>& ContentBundleEditor) const;
-	bool DectivateContentBundleEditing(TSharedPtr<FContentBundleEditor>& ContentBundleEditor) const;
+	bool DeactivateContentBundleEditing(TSharedPtr<FContentBundleEditor>& ContentBundleEditor) const;
 	bool IsContentBundleEditingActivated(TSharedPtr<FContentBundleEditor>& ContentBundleEditor) const;
 
 	DECLARE_EVENT_OneParam(UContentBundleEditorSubsystem, FOnContentBundleChanged, const FContentBundleEditor*);
