World Partition - ResaveActorsBuilder: Added option "-ActorClassesFromFile=<file>" to read a list of classes to filter on from a file (one class per line)

diff --git a/Engine/Source/Editor/UnrealEd/Private/WorldPartition/WorldPartitionResaveActorsBuilder.cpp b/Engine/Source/Editor/UnrealEd/Private/WorldPartition/WorldPartitionResaveActorsBuilder.cpp
--- a/Engine/Source/Editor/UnrealEd/Private/WorldPartition/WorldPartitionResaveActorsBuilder.cpp
+++ b/Engine/Source/Editor/UnrealEd/Private/WorldPartition/WorldPartitionResaveActorsBuilder.cpp
@@ -16,6 +16,7 @@
 #include "Kismet2/BlueprintEditorUtils.h"
 #include "ReferenceCluster.h"
 #include "ActorFolder.h"
+#include "Misc/FileHelper.h"
 
 DEFINE_LOG_CATEGORY_STATIC(LogWorldPartitionResaveActorsBuilder, All, All);
 
@@ -71,6 +72,11 @@ bool UWorldPartitionResaveActorsBuilder::PreRun(UWorld* World, FPackageSourceCon
 			UE_LOG(LogWorldPartitionResaveActorsBuilder, Error, TEXT("SwitchActorPackagingSchemeToReduced is not compatible with ActorClassName"));
 			return false;
 		}
+		else if (!ActorClassesFromFile.IsEmpty())
+		{
+			UE_LOG(LogWorldPartitionResaveActorsBuilder, Error, TEXT("SwitchActorPackagingSchemeToReduced is not compatible with ActorClassesFromFile"));
+			return false;
+		}
 		else if (!ActorTags.IsEmpty())
 		{
 			UE_LOG(LogWorldPartitionResaveActorsBuilder, Error, TEXT("SwitchActorPackagingSchemeToReduced is not compatible with ActorTags"));
@@ -100,45 +106,96 @@ bool UWorldPartitionResaveActorsBuilder::PreRun(UWorld* World, FPackageSourceCon
 	return true;
 }
 
-bool UWorldPartitionResaveActorsBuilder::RunInternal(UWorld* World, const FCellInfo& InCellInfo, FPackageSourceControlHelper& PackageHelper)
+TArray<TSubclassOf<AActor>> UWorldPartitionResaveActorsBuilder::GetActorClassesFilter()
 {
-	UPackage* WorldPackage = World->GetPackage();
-
-	// Actor Class Filter
-	UClass* ActorClass = AActor::StaticClass();
+	TArray<TSubclassOf<AActor>> ActorClasses;
 
-	if (!ActorClassName.IsEmpty())
+	auto GetActorClassFromName = [](const FString& ClassName)
 	{
-		if (bSwitchActorPackagingSchemeToReduced)
-		{
-			UE_LOG(LogWorldPartitionResaveActorsBuilder, Error, TEXT("Changing the actor packaging scheme can't be executed on a subset of actors."));
-			return false;
-		}
-
 		// Look for native classes
-		ActorClass = UClass::TryFindTypeSlow<UClass>(ActorClassName);
+		UClass* ActorClass = UClass::TryFindTypeSlow<UClass>(ClassName);
 
 		if (!ActorClass)
 		{
 			// Look for a fully qualified BP class
-			ActorClass = LoadClass<AActor>(nullptr, *ActorClassName, nullptr, LOAD_None, nullptr);
+			ActorClass = LoadClass<AActor>(nullptr, *ClassName, nullptr, LOAD_None, nullptr);
 
 			if (!ActorClass)
 			{
 				// Look for a package BP
-				if (FPackageName::DoesPackageExist(ActorClassName))
+				if (FPackageName::DoesPackageExist(ClassName))
 				{
-					ActorClass = LoadClass<AActor>(nullptr, *FString::Printf(TEXT("%s.%s_C"), *ActorClassName, *FPackageName::GetLongPackageAssetName(ActorClassName)), nullptr, LOAD_None, nullptr);
+					ActorClass = LoadClass<AActor>(nullptr, *FString::Printf(TEXT("%s.%s_C"), *ClassName, *FPackageName::GetLongPackageAssetName(ClassName)), nullptr, LOAD_None, nullptr);
 				}
+			}
+		}
+
+		return ActorClass;
+	};
+
+	if (!ActorClassName.IsEmpty())
+	{
+		check(!bSwitchActorPackagingSchemeToReduced);
 
-				if (!ActorClass)
+		UClass* ActorClass = GetActorClassFromName(ActorClassName);
+		if (ActorClass)
+		{
+			ActorClasses.Add(ActorClass);
+		}
+		else
+		{
+			UE_LOG(LogWorldPartitionResaveActorsBuilder, Error, TEXT("Failed to find Actor Class: %s."), *ActorClassName);
+		}
+	}
+	else if (!ActorClassesFromFile.IsEmpty())
+	{
+		TArray<FString> ClassesNames;
+
+		bool bReadFile = FFileHelper::LoadFileToStringArray(ClassesNames, *ActorClassesFromFile);
+		if (bReadFile)
+		{
+			for (const FString& ClassName : ClassesNames)
+			{
+				UClass* ActorClass = GetActorClassFromName(ClassName);
+				if (ActorClass)
 				{
-					UE_LOG(LogWorldPartitionResaveActorsBuilder, Error, TEXT("Failed to find Actor Class: %s."), *ActorClassName);
-					return false;
+					ActorClasses.Add(ActorClass);
+				}
+				else
+				{
+					UE_LOG(LogWorldPartitionResaveActorsBuilder, Warning, TEXT("Failed to find Actor Class: %s."), *ClassName);
 				}
 			}
+
+			if (ActorClasses.IsEmpty())
+			{
+				UE_LOG(LogWorldPartitionResaveActorsBuilder, Error, TEXT("Failed to find any valid actor classes in file: %s."), *ActorClassesFromFile);
+			}			
+		}
+		else
+		{
+			UE_LOG(LogWorldPartitionResaveActorsBuilder, Error, TEXT("Failed to open/process actor classes in file: %s."), *ActorClassesFromFile);
 		}
 	}
+	else 
+	{
+		// No class filtering specified... all actor classes are valid
+		ActorClasses.Add(AActor::StaticClass());
+	}
+
+	return ActorClasses;
+}
+
+bool UWorldPartitionResaveActorsBuilder::RunInternal(UWorld* World, const FCellInfo& InCellInfo, FPackageSourceControlHelper& PackageHelper)
+{
+	UPackage* WorldPackage = World->GetPackage();
+
+	// Actor Class Filter
+	TArray<TSubclassOf<AActor>> ActorClasses = GetActorClassesFilter();
+	if (ActorClasses.IsEmpty())
+	{
+		return false;
+	}	
 
 	UWorldPartition* WorldPartition = World->GetWorldPartition();
 	if (!WorldPartition)
@@ -282,7 +339,7 @@ bool UWorldPartitionResaveActorsBuilder::RunInternal(UWorld* World, const FCellI
 
 		FWorldPartitionHelpers::FForEachActorWithLoadingParams ForEachActorWithLoadingParams;
 
-		ForEachActorWithLoadingParams.ActorClasses = { ActorClass };
+		ForEachActorWithLoadingParams.ActorClasses = ActorClasses;
 
 		ForEachActorWithLoadingParams.FilterActorDesc = [this](const FWorldPartitionActorDesc* ActorDesc) -> bool
 		{
diff --git a/Engine/Source/Editor/UnrealEd/Public/WorldPartition/WorldPartitionResaveActorsBuilder.h b/Engine/Source/Editor/UnrealEd/Public/WorldPartition/WorldPartitionResaveActorsBuilder.h
--- a/Engine/Source/Editor/UnrealEd/Public/WorldPartition/WorldPartitionResaveActorsBuilder.h
+++ b/Engine/Source/Editor/UnrealEd/Public/WorldPartition/WorldPartitionResaveActorsBuilder.h
@@ -20,10 +20,16 @@ public:
 	virtual bool RunInternal(UWorld* World, const FCellInfo& InCellInfo, FPackageSourceControlHelper& PackageHelper) override;
 	// UWorldPartitionBuilder interface end
 
+private:
+	TArray<TSubclassOf<AActor>> GetActorClassesFilter();
+
 private:
 	UPROPERTY()
 	FString ActorClassName;
 
+	UPROPERTY()
+	FString ActorClassesFromFile;
+
 	UPROPERTY()
 	bool bReportOnly;
 
