Asset Thumbnails
- Moved the thumbnail dirtying from FThumbnailPool to UThumbnailManager as it only needs to happen once per thumbnail, not once per pool.
- Prevent trying to dirty thumbnails in the asset package for generic thumbnails that don't get rendered.

#rnx

diff --git a/Engine/Source/Editor/UnrealEd/Classes/ThumbnailRendering/ThumbnailManager.h b/Engine/Source/Editor/UnrealEd/Classes/ThumbnailRendering/ThumbnailManager.h
--- a/Engine/Source/Editor/UnrealEd/Classes/ThumbnailRendering/ThumbnailManager.h
+++ b/Engine/Source/Editor/UnrealEd/Classes/ThumbnailRendering/ThumbnailManager.h
@@ -10,6 +10,7 @@
 #pragma once
 
 #include "CoreMinimal.h"
+#include "Delegates/Delegate.h"
 #include "UObject/ObjectMacros.h"
 #include "UObject/Object.h"
 #include "Templates/SubclassOf.h"
@@ -79,6 +80,8 @@ public:
 
 };
 
+DECLARE_EVENT_OneParam(UThumbnailManager, FOnThumbnailDirtied, const FSoftObjectPath&);
+
 UCLASS(config=Editor)
 class UThumbnailManager : public UObject
 {
@@ -155,6 +158,13 @@ public:
 	 * Unless you are rendering a huge amount of thumbnails at once this shared pool should be used
 	 */
 	UNREALED_API TSharedPtr<FAssetThumbnailPool> GetSharedThumbnailPool() const { return SharedThumbnailPool; }
+
+	/**
+	 * Event that is being broadcasted when a thumbnail gets dirtied.
+	 * Parameter is the object soft object path associated with the thumbnail.
+	 */
+	UNREALED_API FOnThumbnailDirtied& GetOnThumbnailDirtied() { return OnThumbnailDirtied; }
+
 protected:
 	/**
 	 * The array of thumbnail rendering information entries. Each type that supports
@@ -195,6 +205,8 @@ protected:
 	 */
 	bool bMapNeedsUpdate;
 
+	FOnThumbnailDirtied OnThumbnailDirtied;
+
 public:
 	/** Returns the thumbnail manager and creates it if missing */
 	UNREALED_API static UThumbnailManager& Get();
@@ -211,6 +223,15 @@ protected:
 private:
 	/** Initialize the checkerboard texture for texture thumbnails */
 	void SetupCheckerboardTexture();
+
+	/** Handler for when an object is edited. Used to dirty its thumbnail. */
+	void OnObjectPropertyChanged(UObject* Asset, FPropertyChangedEvent& PropertyChangedEvent);
+
+	/** Handler for when an actor is moved in a level. Used to update world asset thumbnails. */
+	void OnActorPostEditMove(AActor* Actor);
+
+	/** Handler to dirty cached thumbnails in packages to make sure they are re-rendered later */
+	void DirtyThumbnailForObject(UObject* ObjectBeingModified);
 };
 
 
diff --git a/Engine/Source/Editor/UnrealEd/Private/AssetThumbnail.cpp b/Engine/Source/Editor/UnrealEd/Private/AssetThumbnail.cpp
--- a/Engine/Source/Editor/UnrealEd/Private/AssetThumbnail.cpp
+++ b/Engine/Source/Editor/UnrealEd/Private/AssetThumbnail.cpp
@@ -927,12 +927,9 @@ FAssetThumbnailPool::FAssetThumbnailPool( uint32 InNumInPool, double InMaxFrameT
 	, MaxRealTimeThumbnailsPerFrame( InMaxRealTimeThumbnailsPerFrame )
 	, MaxFrameTimeAllowance( InMaxFrameTimeAllowance )
 {
-	FCoreUObjectDelegates::OnObjectPropertyChanged.AddRaw(this, &FAssetThumbnailPool::OnObjectPropertyChanged);
 	FCoreUObjectDelegates::OnAssetLoaded.AddRaw(this, &FAssetThumbnailPool::OnAssetLoaded);
-	if ( GEditor )
-	{
-		GEditor->OnActorMoved().AddRaw( this, &FAssetThumbnailPool::OnActorPostEditMove );
-	}
+
+	UThumbnailManager::Get().GetOnThumbnailDirtied().AddRaw(this, &FAssetThumbnailPool::OnThumbnailDirtied);
 
 	// Add the custom thumbnail tag to the list of tags that the asset registry can parse
 	TSet<FName>& MetaDataTagsForAssetRegistry = UObject::GetMetaDataTagsForAssetRegistry();
@@ -941,12 +938,8 @@ FAssetThumbnailPool::FAssetThumbnailPool( uint32 InNumInPool, double InMaxFrameT
 
 FAssetThumbnailPool::~FAssetThumbnailPool()
 {
-	FCoreUObjectDelegates::OnObjectPropertyChanged.RemoveAll(this);
+	UThumbnailManager::Get().GetOnThumbnailDirtied().RemoveAll(this);
 	FCoreUObjectDelegates::OnAssetLoaded.RemoveAll(this);
-	if ( GEditor )
-	{
-		GEditor->OnActorMoved().RemoveAll(this);
-	}
 
 	// Release all the texture resources
 	ReleaseResources();
@@ -1563,79 +1556,7 @@ void FAssetThumbnailPool::OnAssetLoaded( UObject* Asset )
 	}
 }
 
-void FAssetThumbnailPool::OnActorPostEditMove( AActor* Actor )
-{
-	DirtyThumbnailForObject(Actor);
-}
-
-void FAssetThumbnailPool::OnObjectPropertyChanged( UObject* ObjectBeingModified, FPropertyChangedEvent& PropertyChangedEvent )
-{
-	if(PropertyChangedEvent.ChangeType != EPropertyChangeType::Interactive)
-	{
-		DirtyThumbnailForObject(ObjectBeingModified);
-	}
-}
-
-void FAssetThumbnailPool::DirtyThumbnailForObject(UObject* ObjectBeingModified)
+void FAssetThumbnailPool::OnThumbnailDirtied( const FSoftObjectPath& ObjectPath )
 {
-	if (!ObjectBeingModified)
-	{
-		return;
-	}
-
-	if (ObjectBeingModified->HasAnyFlags(RF_ClassDefaultObject))
-	{
-		if (ObjectBeingModified->GetClass()->ClassGeneratedBy != nullptr)
-		{
-			// This is a blueprint modification. Check to see if this thumbnail is the blueprint of the modified CDO
-			ObjectBeingModified = ObjectBeingModified->GetClass()->ClassGeneratedBy;
-		}
-	}
-	else if (AActor* ActorBeingModified = Cast<AActor>(ObjectBeingModified))
-	{
-		// This is a non CDO actor getting modified. Update the actor's world's thumbnail.
-		ObjectBeingModified = ActorBeingModified->GetWorld();
-	}
-
-	if (ObjectBeingModified)
-	{
-		// An object in memory was modified.  We'll mark it's thumbnail as dirty so that it'll be
-		// regenerated on demand later. (Before being displayed in the browser, or package saves, etc.)
-		FObjectThumbnail* Thumbnail = ThumbnailTools::GetThumbnailForObject(ObjectBeingModified);
-
-		// If we don't yet have a thumbnail map, load one from disk if possible
-		if (Thumbnail == nullptr)
-		{
-			UPackage* ObjectPackage = ObjectBeingModified->GetOutermost();
-
-			const bool bMemoryPackage = FPackageName::IsMemoryPackage(ObjectBeingModified->GetPathName());	// Don't try to load from disk if the package is a memory package
-			const bool bUnsavedPackage = ObjectPackage->HasAnyPackageFlags(PKG_NewlyCreated);				// Don't try loading thumbnails for package that have never been saved
-			const bool bIsGarbageCollecting = IsGarbageCollecting();										// Don't attempt to do this while garbage collecting since loading or finding objects during GC is illegal
-
-			const bool bTryLoadThumbnailFromDisk = !bIsGarbageCollecting && !bMemoryPackage && !bUnsavedPackage;
-			if (bTryLoadThumbnailFromDisk)
-			{
-				FName ObjectFullName = FName(*ObjectBeingModified->GetFullName());
-
-				FThumbnailMap LoadedThumbnails;
-				if (ThumbnailTools::ConditionallyLoadThumbnailsForObjects({ ObjectFullName }, LoadedThumbnails))
-				{
-					Thumbnail = LoadedThumbnails.Find(ObjectFullName);
-
-					if (Thumbnail != nullptr)
-					{
-						Thumbnail = ThumbnailTools::CacheThumbnail(ObjectBeingModified->GetFullName(), Thumbnail, ObjectPackage);
-					}
-				}
-			}
-		}
-
-		if (Thumbnail != nullptr)
-		{
-			// Mark the thumbnail as dirty
-			Thumbnail->MarkAsDirty();
-		}
-
-		RefreshThumbnailsFor(FSoftObjectPath(ObjectBeingModified));
-	}
-}
+	RefreshThumbnailsFor( ObjectPath );
+}
\ No newline at end of file
diff --git a/Engine/Source/Editor/UnrealEd/Private/ThumbnailManager.cpp b/Engine/Source/Editor/UnrealEd/Private/ThumbnailManager.cpp
--- a/Engine/Source/Editor/UnrealEd/Private/ThumbnailManager.cpp
+++ b/Engine/Source/Editor/UnrealEd/Private/ThumbnailManager.cpp
@@ -1,9 +1,12 @@
 // Copyright Epic Games, Inc. All Rights Reserved.
 
 #include "ThumbnailRendering/ThumbnailManager.h"
+
+#include "Editor.h"
 #include "HAL/FileManager.h"
 #include "Misc/FileHelper.h"
 #include "Misc/Paths.h"
+#include "ObjectTools.h"
 #include "UObject/ConstructorHelpers.h"
 #include "Materials/Material.h"
 #include "ISourceControlOperation.h"
@@ -73,6 +76,13 @@ void UThumbnailManager::Initialize(void)
 
 		SharedThumbnailPool = MakeShared<FAssetThumbnailPool>(100);
 
+		FCoreUObjectDelegates::OnObjectPropertyChanged.AddUObject(this, &UThumbnailManager::OnObjectPropertyChanged);
+
+		if (GEditor)
+		{
+			GEditor->OnActorMoved().AddUObject(this, &UThumbnailManager::OnActorPostEditMove);
+		}
+
 		bIsInitialized = true;
 	}
 }
@@ -377,3 +387,92 @@ bool UThumbnailManager::CaptureProjectThumbnail(FViewport* Viewport, const FStri
 
 	return false;
 }
+
+void UThumbnailManager::OnObjectPropertyChanged(UObject* ObjectBeingModified, FPropertyChangedEvent& PropertyChangedEvent)
+{
+	if (PropertyChangedEvent.ChangeType != EPropertyChangeType::Interactive)
+	{
+		DirtyThumbnailForObject(ObjectBeingModified);
+	}
+}
+
+void UThumbnailManager::OnActorPostEditMove(AActor* Actor)
+{
+	DirtyThumbnailForObject(Actor);
+}
+
+void UThumbnailManager::DirtyThumbnailForObject(UObject* ObjectBeingModified)
+{
+	if (!ObjectBeingModified)
+	{
+		return;
+	}
+
+	if (ObjectBeingModified->HasAnyFlags(RF_ClassDefaultObject))
+	{
+		if (ObjectBeingModified->GetClass()->ClassGeneratedBy != nullptr)
+		{
+			// This is a blueprint modification. Check to see if this thumbnail is the blueprint of the modified CDO
+			ObjectBeingModified = ObjectBeingModified->GetClass()->ClassGeneratedBy;
+		}
+	}
+	else if (AActor* ActorBeingModified = Cast<AActor>(ObjectBeingModified))
+	{
+		// This is a non CDO actor getting modified. Update the actor's world's thumbnail.
+		ObjectBeingModified = ActorBeingModified->GetWorld();
+	}
+
+	if (ObjectBeingModified)
+	{
+		// An object in memory was modified.  We'll mark its thumbnail as dirty so that it'll be
+		// regenerated on demand later. (Before being displayed in the browser, or package saves, etc.)
+		FObjectThumbnail* Thumbnail = ThumbnailTools::GetThumbnailForObject(ObjectBeingModified);
+
+		// If we don't yet have a thumbnail map, load one from disk if possible
+		if (Thumbnail == nullptr)
+		{
+			UPackage* ObjectPackage = ObjectBeingModified->GetOutermost();
+
+			const bool bMemoryPackage = FPackageName::IsMemoryPackage(ObjectBeingModified->GetPathName());  // Don't try to load from disk if the package is a memory package
+			const bool bUnsavedPackage = ObjectPackage->HasAnyPackageFlags(PKG_NewlyCreated);               // Don't try loading thumbnails for package that have never been saved
+			const bool bIsGarbageCollecting = IsGarbageCollecting();                                        // Don't attempt to do this while garbage collecting since loading or finding objects during GC is illegal
+			const bool bUsesGenericThumbnail = [ObjectBeingModified, this]() -> bool                        // No need to dirty generic thumbnails
+			{
+				if (FThumbnailRenderingInfo* RenderingInfo = GetRenderingInfo(ObjectBeingModified))
+				{
+					return RenderingInfo->Renderer == nullptr;
+				}
+				else
+				{
+					return true;
+				}
+			}();
+
+			const bool bTryLoadThumbnailFromDisk = !bIsGarbageCollecting && !bMemoryPackage && !bUnsavedPackage && !bUsesGenericThumbnail;
+			if (bTryLoadThumbnailFromDisk)
+			{
+				FName ObjectFullName = FName(*ObjectBeingModified->GetFullName());
+
+				FThumbnailMap LoadedThumbnails;
+				if (ThumbnailTools::ConditionallyLoadThumbnailsForObjects({ ObjectFullName }, LoadedThumbnails))
+				{
+					Thumbnail = LoadedThumbnails.Find(ObjectFullName);
+
+					if (Thumbnail != nullptr)
+					{
+						Thumbnail = ThumbnailTools::CacheThumbnail(ObjectBeingModified->GetFullName(), Thumbnail, ObjectPackage);
+					}
+				}
+			}
+		}
+
+		if (Thumbnail != nullptr)
+		{
+			// Mark the thumbnail as dirty
+			Thumbnail->MarkAsDirty();
+		}
+
+		OnThumbnailDirtied.Broadcast(FSoftObjectPath(ObjectBeingModified));
+	}
+}
+
diff --git a/Engine/Source/Editor/UnrealEd/Public/AssetThumbnail.h b/Engine/Source/Editor/UnrealEd/Public/AssetThumbnail.h
--- a/Engine/Source/Editor/UnrealEd/Public/AssetThumbnail.h
+++ b/Engine/Source/Editor/UnrealEd/Public/AssetThumbnail.h
@@ -264,14 +264,8 @@ private:
 	/** Handler for when an asset is loaded */
 	void OnAssetLoaded( UObject* Asset );
 
-	/** Handler for when an actor is moved in a level. Used to update world asset thumbnails. */
-	void OnActorPostEditMove( AActor* Actor );
-
-	/** Handler for when an asset is loaded */
-	void OnObjectPropertyChanged( UObject* Asset, FPropertyChangedEvent& PropertyChangedEvent );
-
-	/** Handler to dirty cached thumbnails in packages to make sure they are re-rendered later */
-	void DirtyThumbnailForObject( UObject* ObjectBeingModified );
+	/** Handler for when a thumbnail gets flagged as dirty. Used to refresh the thumbnail. */
+	void OnThumbnailDirtied( const FSoftObjectPath& ObjectPath );
 
 private:
 	/** Information about a thumbnail */
