Adding messaging to users about not meeting the requirements to work with advanced rendering features. This messaging is currently specific to DX12 SM6.

diff --git a/Engine/Source/Editor/UnrealEd/Private/EditorEngine.cpp b/Engine/Source/Editor/UnrealEd/Private/EditorEngine.cpp
--- a/Engine/Source/Editor/UnrealEd/Private/EditorEngine.cpp
+++ b/Engine/Source/Editor/UnrealEd/Private/EditorEngine.cpp
@@ -256,6 +256,7 @@
 #include "Rendering/StaticLightingSystemInterface.h"
 #include "LevelEditorDragDropHandler.h"
 #include "IProjectExternalContentInterface.h"
+#include "IDocumentation.h"
 
 DEFINE_LOG_CATEGORY_STATIC(LogEditor, Log, All);
 
@@ -335,6 +336,78 @@ void DestroySelectionSets()
 */
 static TMap<UPackage*, bool> StartupPackageToWarnState;
 
+#if PLATFORM_WINDOWS
+static TWeakPtr<SNotificationItem> MissingAdvancedRenderingRequirementsNotificationPtr;
+#endif
+
+static void CheckForMissingAdvancedRenderingRequirements()
+{
+#if PLATFORM_WINDOWS
+	if (FSlateApplication::IsInitialized() && GDynamicRHIFailedToInitializeAdvancedPlatform)
+	{
+		/** Utility functions for the notification */
+		struct Local
+		{
+			static ECheckBoxState GetDontAskAgainCheckBoxState()
+			{
+				bool bSuppressNotification = false;
+				GConfig->GetBool(TEXT("WindowsEditor"), TEXT("SuppressMissingAdvancedRenderingRequirementsNotification"), bSuppressNotification, GEditorPerProjectIni);
+				return bSuppressNotification ? ECheckBoxState::Checked : ECheckBoxState::Unchecked;
+			}
+
+			static void OnDontAskAgainCheckBoxStateChanged(ECheckBoxState NewState)
+			{
+				const bool bSuppressNotification = (NewState == ECheckBoxState::Checked);
+				GConfig->SetBool(TEXT("WindowsEditor"), TEXT("SuppressMissingAdvancedRenderingRequirementsNotification"), bSuppressNotification, GEditorPerProjectIni);
+			}
+
+			static void OnMissingAdvancedRenderingRequirementsNotificationDismissed()
+			{
+				TSharedPtr<SNotificationItem> NotificationItem = MissingAdvancedRenderingRequirementsNotificationPtr.Pin();
+
+				if (NotificationItem.IsValid())
+				{
+					NotificationItem->SetCompletionState(SNotificationItem::CS_Success);
+					NotificationItem->Fadeout();
+
+					MissingAdvancedRenderingRequirementsNotificationPtr.Reset();
+				}
+			}
+		};
+
+		const ECheckBoxState DontAskAgainCheckBoxState = Local::GetDontAskAgainCheckBoxState();
+		if (DontAskAgainCheckBoxState == ECheckBoxState::Unchecked)
+		{
+			const FText TitleText = LOCTEXT("MissingAdvancedRenderingRequirementsNotificationTitle", "Missing support for advanced rendering features");
+			const FText MessageText = LOCTEXT("MissingAdvancedRenderingRequirementsNotificationText",
+				"This project attempted to launch DirectX 12 with the SM6 shader format but it is not supported by your system. This will prevent advanced rendering features like Nanite and Virtual Shadow Maps from working.\n\nMake sure your system meets the requirements for these UE5 rendering features."
+			);
+
+			FNotificationInfo Info(TitleText);
+			Info.SubText = MessageText;
+
+			Info.HyperlinkText = LOCTEXT("UnrealSoftwareRequirements", "Unreal Software Requirements");
+			Info.Hyperlink = FSimpleDelegate::CreateStatic([]() { IDocumentation::Get()->Open(TEXT("hardware-and-software-specifications-for-unreal-engine")); });
+
+			Info.bFireAndForget = false;
+			Info.FadeOutDuration = 3.0f;
+			Info.ExpireDuration = 0.0f;
+			Info.bUseLargeFont = false;
+			Info.bUseThrobber = false;
+
+			Info.ButtonDetails.Add(FNotificationButtonInfo(LOCTEXT("OK", "OK"), FText::GetEmpty(), FSimpleDelegate::CreateStatic(&Local::OnMissingAdvancedRenderingRequirementsNotificationDismissed)));
+
+			Info.CheckBoxState = TAttribute<ECheckBoxState>::Create(&Local::GetDontAskAgainCheckBoxState);
+			Info.CheckBoxStateChanged = FOnCheckStateChanged::CreateStatic(&Local::OnDontAskAgainCheckBoxStateChanged);
+			Info.CheckBoxText = NSLOCTEXT("ModalDialogs", "DefaultCheckBoxMessage", "Don't show this again");
+
+			MissingAdvancedRenderingRequirementsNotificationPtr = FSlateNotificationManager::Get().AddNotification(Info);
+			MissingAdvancedRenderingRequirementsNotificationPtr.Pin()->SetCompletionState(SNotificationItem::CS_Pending);
+		}
+	}
+#endif // PLATFORM_WINDOWS
+}
+
 
 ERHIFeatureLevel::Type FPreviewPlatformInfo::GetEffectivePreviewFeatureLevel() const
 {
@@ -1168,6 +1241,8 @@ void UEditorEngine::Init(IEngineLoop* InEngineLoop)
 
 	FEditorCommandLineUtils::ProcessEditorCommands(FCommandLine::Get());
 
+	CheckForMissingAdvancedRenderingRequirements();
+
 	// for IsInitialized()
 	bIsInitialized = true;
 };
diff --git a/Engine/Source/Runtime/RHI/Private/Windows/WindowsDynamicRHI.cpp b/Engine/Source/Runtime/RHI/Private/Windows/WindowsDynamicRHI.cpp
--- a/Engine/Source/Runtime/RHI/Private/Windows/WindowsDynamicRHI.cpp
+++ b/Engine/Source/Runtime/RHI/Private/Windows/WindowsDynamicRHI.cpp
@@ -11,6 +11,8 @@
 
 #include "Windows/WindowsPlatformApplicationMisc.h"
 
+bool GDynamicRHIFailedToInitializeAdvancedPlatform = false;
+
 static const TCHAR* GLoadedRHIModuleName;
 
 enum class EWindowsRHI
@@ -684,6 +686,11 @@ static bool HandleUnsupportedFeatureLevel(EWindowsRHI& WindowsRHI, ERHIFeatureLe
 		return false;
 	}
 
+	if (FeatureLevel == ERHIFeatureLevel::SM6)
+	{
+		GDynamicRHIFailedToInitializeAdvancedPlatform = true;
+	}
+
 	if (TOptional<ERHIFeatureLevel::Type> FallbackFeatureLevel = Config.GetNextHighestTargetedFeatureLevel(WindowsRHI, FeatureLevel))
 	{
 		UE_LOG(LogRHI, Log, TEXT("RHI %s with Feature Level %s not supported, falling back to Feature Level %s"), ModuleNameFromWindowsRHI(WindowsRHI), *LexToString(FeatureLevel), *LexToString(FallbackFeatureLevel.GetValue()));
diff --git a/Engine/Source/Runtime/RHI/Public/DynamicRHI.h b/Engine/Source/Runtime/RHI/Public/DynamicRHI.h
--- a/Engine/Source/Runtime/RHI/Public/DynamicRHI.h
+++ b/Engine/Source/Runtime/RHI/Public/DynamicRHI.h
@@ -1689,3 +1689,5 @@ FDynamicRHI* PlatformCreateDynamicRHI();
 // Name of the RHI module that will be created when PlatformCreateDynamicRHI is called
 // NOTE: This function is very slow when called before RHIInit
 extern RHI_API const TCHAR* GetSelectedDynamicRHIModuleName(bool bCleanup = true);
+
+extern RHI_API bool GDynamicRHIFailedToInitializeAdvancedPlatform;
