World Partition - ResaveActorsBuilder: Added -DiffDirtyActorDescs which will bring up your configured diff tool to compare the old & new actor descriptors

diff --git a/Engine/Source/Editor/UnrealEd/Private/WorldPartition/WorldPartitionResaveActorsBuilder.cpp b/Engine/Source/Editor/UnrealEd/Private/WorldPartition/WorldPartitionResaveActorsBuilder.cpp
--- a/Engine/Source/Editor/UnrealEd/Private/WorldPartition/WorldPartitionResaveActorsBuilder.cpp
+++ b/Engine/Source/Editor/UnrealEd/Private/WorldPartition/WorldPartitionResaveActorsBuilder.cpp
@@ -17,6 +17,8 @@
 #include "ReferenceCluster.h"
 #include "ActorFolder.h"
 #include "Misc/FileHelper.h"
+#include "Settings/EditorLoadingSavingSettings.h"
+#include "AssetToolsModule.h"
 
 DEFINE_LOG_CATEGORY_STATIC(LogWorldPartitionResaveActorsBuilder, All, All);
 
@@ -229,6 +231,9 @@ bool UWorldPartitionResaveActorsBuilder::RunInternal(UWorld* World, const FCellI
 
 	TArray<FString> PackagesToDelete;
 
+	TArray<FString> DirtyActorDescsOld;
+	TArray<FString> DirtyActorDescsNew;
+
 	if (bSwitchActorPackagingSchemeToReduced)
 	{
 		World->PersistentLevel->ActorPackagingScheme = EActorPackagingScheme::Reduced;
@@ -372,14 +377,13 @@ bool UWorldPartitionResaveActorsBuilder::RunInternal(UWorld* World, const FCellI
 			PackagesToSave.Empty();
 		};
 
-		FWorldPartitionHelpers::ForEachActorWithLoading(WorldPartition, [this, &PackagesToDelete, &PackagesToSave](const FWorldPartitionActorDesc* ActorDesc)
+		FWorldPartitionHelpers::ForEachActorWithLoading(WorldPartition, [this, &PackagesToDelete, &PackagesToSave, &DirtyActorDescsOld, &DirtyActorDescsNew](const FWorldPartitionActorDesc* ActorDesc)
 		{
 			AActor* Actor = ActorDesc->GetActor();
 
 			if (!Actor)
 			{
 				PackagesToDelete.Add(ActorDesc->GetActorPackage().ToString());
-				return true;
 			}
 			else
 			{
@@ -395,25 +399,30 @@ bool UWorldPartitionResaveActorsBuilder::RunInternal(UWorld* World, const FCellI
 				UPackage* Package = Actor->GetExternalPackage();
 				check(Package);
 
-				if (bResaveDirtyActorDescsOnly)
+				bool bActorNeedsResave = ActorDesc->IsResaveNeeded();
+				
+				if (bResaveDirtyActorDescsOnly || bDiffDirtyActorDescs)
 				{
-					if (!ActorDesc->IsResaveNeeded())
+					TUniquePtr<FWorldPartitionActorDesc> NewActorDesc = Actor->CreateActorDesc();
+
+					bool bActorDescDirty = !ActorDesc->Equals(NewActorDesc.Get());
+					if (bActorDescDirty && bDiffDirtyActorDescs)
 					{
-						TUniquePtr<FWorldPartitionActorDesc> NewActorDesc = Actor->CreateActorDesc();
+						DirtyActorDescsOld.Add(ActorDesc->ToString());
+						DirtyActorDescsNew.Add(NewActorDesc->ToString());
+					}				
 
-						if (ActorDesc->Equals(NewActorDesc.Get()))
-						{
-							return true;
-						}
-					}
+					bActorNeedsResave |= bResaveDirtyActorDescsOnly && bActorDescDirty;
+				}
 
+				if (bActorNeedsResave)
+				{
 					UE_LOG(LogWorldPartitionResaveActorsBuilder, Log, TEXT("Package %s needs to be resaved."), *Package->GetName());
 				}
 
-				if (!bReportOnly)
+				if (bActorNeedsResave && !bReportOnly)
 				{
 					PackagesToSave.Add(Package);
-					return true;
 				}
 			}
 
@@ -421,6 +430,32 @@ bool UWorldPartitionResaveActorsBuilder::RunInternal(UWorld* World, const FCellI
 		}, ForEachActorWithLoadingParams);
 	}
 
+	if (bDiffDirtyActorDescs)
+	{
+		auto WriteTempFile = [](const TArray<FString>& Lines, FString& TempFileName)
+		{
+			TempFileName = FPaths::CreateTempFilename(*FPaths::DiffDir(), TEXT("ActorsDescsDiff"));
+			return FFileHelper::SaveStringArrayToFile(Lines, *TempFileName);
+		};
+
+		FString TempFileOld;
+		FString TempFileNew;
+		bool OldResult = WriteTempFile(DirtyActorDescsOld, TempFileOld);
+		bool NewResult = WriteTempFile(DirtyActorDescsNew, TempFileNew);
+
+		if (OldResult && NewResult)
+		{
+			FString DiffCommand = GetDefault<UEditorLoadingSavingSettings>()->TextDiffToolPath.FilePath;
+
+			FAssetToolsModule& AssetToolsModule = FModuleManager::Get().LoadModuleChecked<FAssetToolsModule>("AssetTools");
+			AssetToolsModule.Get().CreateDiffProcess(DiffCommand, TempFileOld, TempFileNew);
+		}
+		else
+		{
+			UE_LOG(LogWorldPartitionResaveActorsBuilder, Error, TEXT("Failed to write actors desc diff files."));
+		}
+	}
+
 	if (!bReportOnly)
 	{
 		if (!UWorldPartitionBuilder::DeletePackages(PackagesToDelete, PackageHelper, !bSwitchActorPackagingSchemeToReduced) && bSwitchActorPackagingSchemeToReduced)
diff --git a/Engine/Source/Editor/UnrealEd/Public/WorldPartition/WorldPartitionResaveActorsBuilder.h b/Engine/Source/Editor/UnrealEd/Public/WorldPartition/WorldPartitionResaveActorsBuilder.h
--- a/Engine/Source/Editor/UnrealEd/Public/WorldPartition/WorldPartitionResaveActorsBuilder.h
+++ b/Engine/Source/Editor/UnrealEd/Public/WorldPartition/WorldPartitionResaveActorsBuilder.h
@@ -36,6 +36,9 @@ private:
 	UPROPERTY()
 	bool bResaveDirtyActorDescsOnly;
 
+	UPROPERTY()
+	bool bDiffDirtyActorDescs;
+
 	UPROPERTY()
 	bool bSwitchActorPackagingSchemeToReduced;
 
