Fixed a crash when pasting an actor element and the class of the actor can't be found in memory.

diff --git a/Engine/Source/Editor/UnrealEd/Private/Elements/Actor/ActorElementEditorCopyAndPaste.cpp b/Engine/Source/Editor/UnrealEd/Private/Elements/Actor/ActorElementEditorCopyAndPaste.cpp
--- a/Engine/Source/Editor/UnrealEd/Private/Elements/Actor/ActorElementEditorCopyAndPaste.cpp
+++ b/Engine/Source/Editor/UnrealEd/Private/Elements/Actor/ActorElementEditorCopyAndPaste.cpp
@@ -136,16 +136,20 @@ void FActorElementEditorPasteImporter::Import(FContext& InContext)
 			}
 			else if (GetBEGIN(LinePos, TEXT("Actor")))
 			{
-				CurrentImportedActor = &(ImportedActors.AddDefaulted_GetRef());
-				ParseObject<UClass>(*LinePos, TEXT("CLASS="), CurrentImportedActor->Class, nullptr);
-				FParse::Value(*LinePos, TEXT("Name="), CurrentImportedActor->Name);
-				FParse::Value(*LinePos, TEXT("Archetype="), CurrentImportedActor->Archetype);
-				FParse::Value(*LinePos, TEXT("ParentActor="), CurrentImportedActor->ParentActor);
-				FParse::Value(*LinePos, TEXT("SocketName="), CurrentImportedActor->SocketNameString);
-				FParse::Value(*LinePos, TEXT("GroupActor="), CurrentImportedActor->GroupActor);
-				FParse::Value(*LinePos, TEXT("GroupFolder="), CurrentImportedActor->GroupFolder);
-				FParse::Value(*LinePos, TEXT("ActorFolderPath="), CurrentImportedActor->ActorFolderPath);
-				FParse::Value(*LinePos, TEXT("ExportPath="), CurrentImportedActor->ExportedActorFullName);
+				UClass* ActorClass = nullptr;
+				if (ParseObject<UClass>(*LinePos, TEXT("CLASS="), ActorClass, nullptr))
+				{
+					CurrentImportedActor = &(ImportedActors.AddDefaulted_GetRef());
+					CurrentImportedActor->Class = ActorClass;
+					FParse::Value(*LinePos, TEXT("Name="), CurrentImportedActor->Name);
+					FParse::Value(*LinePos, TEXT("Archetype="), CurrentImportedActor->Archetype);
+					FParse::Value(*LinePos, TEXT("ParentActor="), CurrentImportedActor->ParentActor);
+					FParse::Value(*LinePos, TEXT("SocketName="), CurrentImportedActor->SocketNameString);
+					FParse::Value(*LinePos, TEXT("GroupActor="), CurrentImportedActor->GroupActor);
+					FParse::Value(*LinePos, TEXT("GroupFolder="), CurrentImportedActor->GroupFolder);
+					FParse::Value(*LinePos, TEXT("ActorFolderPath="), CurrentImportedActor->ActorFolderPath);
+					FParse::Value(*LinePos, TEXT("ExportPath="), CurrentImportedActor->ExportedActorFullName);
+				}
 			}
 		}
 
