Wait until the number of clients value is committed before saving the config

diff --git a/Engine/Source/Editor/UnrealEd/Private/Kismet2/DebuggerCommands.cpp b/Engine/Source/Editor/UnrealEd/Private/Kismet2/DebuggerCommands.cpp
--- a/Engine/Source/Editor/UnrealEd/Private/Kismet2/DebuggerCommands.cpp
+++ b/Engine/Source/Editor/UnrealEd/Private/Kismet2/DebuggerCommands.cpp
@@ -1823,7 +1823,11 @@ void FInternalPlayWorldCommandCallbacks::SetNumberOfClients(int32 NumClients, ET
 	PlayInSettings->SetPlayNumberOfClients(NumClients);
 
 	PlayInSettings->PostEditChange();
-	PlayInSettings->SaveConfig();
+	// ETextCommit::Default is used when the value is changed during an interactive event like mouse drag or mouse wheel, vs. when the value is actually committed
+	if(CommitInfo != ETextCommit::Default)
+	{
+		PlayInSettings->SaveConfig();
+	}
 }
 
 
