Prevent Actor attachment between different different ContentBundles and Game actors

diff --git a/Engine/Source/Editor/UnrealEd/Private/EditorEngine.cpp b/Engine/Source/Editor/UnrealEd/Private/EditorEngine.cpp
--- a/Engine/Source/Editor/UnrealEd/Private/EditorEngine.cpp
+++ b/Engine/Source/Editor/UnrealEd/Private/EditorEngine.cpp
@@ -4346,6 +4346,15 @@ bool UEditorEngine::CanParentActors( const AActor* ParentActor, const AActor* Ch
 		return false;
 	}
 
+	if (ChildActor->GetContentBundleGuid() != ParentActor->GetContentBundleGuid())
+	{
+		if (ReasonText)
+		{
+			*ReasonText = NSLOCTEXT("ActorAttachmentError", "WrongContentBundle_AttachmentError", "Actors need to be in the same content bundle!");
+		}
+		return false;
+	}
+
 	if(ParentRoot->IsAttachedTo( ChildRoot ))
 	{
 		if (ReasonText)
