Further improve logging

diff --git a/Engine/Source/Editor/UnrealEd/Private/Kismet2/KismetReinstanceUtilities.cpp b/Engine/Source/Editor/UnrealEd/Private/Kismet2/KismetReinstanceUtilities.cpp
--- a/Engine/Source/Editor/UnrealEd/Private/Kismet2/KismetReinstanceUtilities.cpp
+++ b/Engine/Source/Editor/UnrealEd/Private/Kismet2/KismetReinstanceUtilities.cpp
@@ -184,6 +184,10 @@ struct FReplaceReferenceHelper
 						const TMap<FSoftObjectPath, UObject*>& WeakReferencesMap;
 					};
 
+					if (GLogReinstancerReferenceReplacementObj)
+					{
+						UE_LOG(LogBlueprint, Warning, TEXT("Replacing References within: %s"), *(Obj->GetName()));
+					}
 					ReferenceReplace ReplaceAr(Obj, OldToNewInstanceMap, ReinstancedObjectsWeakReferenceMap);
 				}
 			}
