MPCook - remove spurious assertion that log Verbosity argument is in the range of NumVerbosity. The verbosity argument includes other information in bits outside of that range.
#rnx

diff --git a/Engine/Source/Editor/UnrealEd/Private/Cooker/CookWorkerServer.cpp b/Engine/Source/Editor/UnrealEd/Private/Cooker/CookWorkerServer.cpp
--- a/Engine/Source/Editor/UnrealEd/Private/Cooker/CookWorkerServer.cpp
+++ b/Engine/Source/Editor/UnrealEd/Private/Cooker/CookWorkerServer.cpp
@@ -881,7 +881,6 @@ FCbWriter& operator<<(FCbWriter& Writer, const FReplicatedLogData& Package)
 	Writer.BeginArray();
 	Writer << Package.Category;
 	uint8 Verbosity = static_cast<uint8>(Package.Verbosity);
-	check(Verbosity < ELogVerbosity::NumVerbosity);
 	Writer << Verbosity;
 	Writer << Package.Message;
 	Writer.EndArray();
@@ -894,7 +893,7 @@ bool LoadFromCompactBinary(FCbFieldView Field, FReplicatedLogData& OutPackage)
 	FCbFieldViewIterator It = Field.CreateViewIterator();
 	bOk = LoadFromCompactBinary(*It++, OutPackage.Category) & bOk;
 	uint8 Verbosity;
-	if (LoadFromCompactBinary(*It++, Verbosity) && Verbosity <= ELogVerbosity::NumVerbosity)
+	if (LoadFromCompactBinary(*It++, Verbosity))
 	{
 		OutPackage.Verbosity = static_cast<ELogVerbosity::Type>(Verbosity);
 	}
