Sequencer: Anim Baking. Add options for including/excluding names added to take recorder/anim recorer to the anim baking options.

diff --git a/Engine/Source/Editor/MovieSceneTools/Private/MovieSceneToolHelpers.cpp b/Engine/Source/Editor/MovieSceneTools/Private/MovieSceneToolHelpers.cpp
--- a/Engine/Source/Editor/MovieSceneTools/Private/MovieSceneToolHelpers.cpp
+++ b/Engine/Source/Editor/MovieSceneTools/Private/MovieSceneToolHelpers.cpp
@@ -4048,6 +4048,8 @@ bool MovieSceneToolHelpers::ExportToAnimSequence(UAnimSequence* AnimSequence, UA
 		RecordingSettings.bRecordAttributeCurves = ExportOptions->bExportAttributeCurves;
 		RecordingSettings.bRecordMaterialCurves = ExportOptions->bExportMaterialCurves;
 		RecordingSettings.bRecordInWorldSpace = ExportOptions->bRecordInWorldSpace;
+		RecordingSettings.IncludeAnimationNames = ExportOptions->IncludeAnimationNames;
+		RecordingSettings.ExcludeAnimationNames = ExportOptions->ExcludeAnimationNames;
 		AnimationRecorder.Init(SkelMeshComp, AnimSequence, nullptr, RecordingSettings);	
 		});
 
diff --git a/Engine/Source/Editor/UnrealEd/Classes/Exporters/AnimSeqExportOption.h b/Engine/Source/Editor/UnrealEd/Classes/Exporters/AnimSeqExportOption.h
--- a/Engine/Source/Editor/UnrealEd/Classes/Exporters/AnimSeqExportOption.h
+++ b/Engine/Source/Editor/UnrealEd/Classes/Exporters/AnimSeqExportOption.h
@@ -50,6 +50,14 @@ public:
 	UPROPERTY(EditAnywhere, BlueprintReadWrite, AdvancedDisplay, Category = Export);
 	TEnumAsByte<ERichCurveInterpMode> CurveInterpolation = ERichCurveInterpMode::RCIM_Linear;
 
+	/** Include only the animation bones/curves that match this list */
+	UPROPERTY(EditAnywhere, BlueprintReadWrite, AdvancedDisplay, Category = Export);
+	TArray<FString> IncludeAnimationNames;
+
+	/** Exclude all animation bones/curves that match this list */
+	UPROPERTY(EditAnywhere, BlueprintReadWrite, AdvancedDisplay, Category = Export);
+	TArray<FString> ExcludeAnimationNames;
+
 	/** Number of Display Rate frames to evaluate before doing the export. It will evaluate after any Delay. This will use frames before the start frame. Use it if there is some post anim BP effects you want to run before export start time.*/
 	UPROPERTY(EditAnywhere, BlueprintReadWrite, AdvancedDisplay, Category = Export);
 	FFrameNumber WarmUpFrames;
