Added a log flush to the PkgInfo commandlet to work around a race condition

GPrintLogTimes and friends are not thread-safe and are not consumed until the log is flushed, rather than when the log is queued. The added call to GLog->Flush ensures that the logs are flushed while the modifications to log printing are in effect.

#rnx

diff --git a/Engine/Source/Editor/UnrealEd/Private/Commandlets/PackageUtilities.cpp b/Engine/Source/Editor/UnrealEd/Private/Commandlets/PackageUtilities.cpp
--- a/Engine/Source/Editor/UnrealEd/Private/Commandlets/PackageUtilities.cpp
+++ b/Engine/Source/Editor/UnrealEd/Private/Commandlets/PackageUtilities.cpp
@@ -1723,6 +1723,12 @@ int32 UPkgInfoCommandlet::Main( const FString& Params )
 				Out.Logf(ELogVerbosity::Display, TEXT("Total number of Serialize calls: %lld"), TotalSerializeCalls);
 			}
 #endif // !NO_LOGGING
+
+			// Flush logs while the disabled times, category, and verbosity are in scope.
+			if (GLog)
+			{
+				GLog->Flush();
+			}
 		}
 		CollectGarbage(RF_NoFlags);
 	}
