Fixed UGatherTextFromAssetsCommandlet failing to gather text from external actor packages

UGatherTextFromAssetsCommandlet wasn't finding external actor packages for the worlds it processed, nor was it correctly loading those packages even if it did find them.

UGatherTextFromAssetsCommandlet has been updated to do a secondary scan for external actor packages that belong to any world partitioned worlds that it is going to gather localization for. Ideally these external actor packages will be gathered using their cache, but their caches are currently invalid due to the FPropertyLocalizationDataGatherer issue.

In the cases where we detect a stale external actor localization cache, we instead queue its owner world to be re-gathered with the GUID(s) of any external actors that were stale. This takes the world package through a different path in the package gather code that will load and initialize the editor world within the package, and use the world partition utils to iterate and load any of the stale actors so that they're in a valid state to be gathered from.

Note: The version bump to use the localization cache for external actors that have been re-saved after the FPropertyLocalizationDataGatherer fixes will be made when this change reaches Main.

diff --git a/Engine/Source/Editor/UnrealEd/Classes/Commandlets/GatherTextFromAssetsCommandlet.h b/Engine/Source/Editor/UnrealEd/Classes/Commandlets/GatherTextFromAssetsCommandlet.h
--- a/Engine/Source/Editor/UnrealEd/Classes/Commandlets/GatherTextFromAssetsCommandlet.h
+++ b/Engine/Source/Editor/UnrealEd/Classes/Commandlets/GatherTextFromAssetsCommandlet.h
@@ -55,6 +55,9 @@ private:
 		/** The complete set of dependencies for the package */
 		TSet<FName> Dependencies;
 
+		/** The set of external actors to process for a world partition map package */
+		TSet<FGuid> ExternalActors;
+
 		/** Localization cache state of this package */
 		EPackageLocCacheState PackageLocCacheState;
 
diff --git a/Engine/Source/Editor/UnrealEd/Private/Commandlets/GatherTextFromAssetsCommandlet.cpp b/Engine/Source/Editor/UnrealEd/Private/Commandlets/GatherTextFromAssetsCommandlet.cpp
--- a/Engine/Source/Editor/UnrealEd/Private/Commandlets/GatherTextFromAssetsCommandlet.cpp
+++ b/Engine/Source/Editor/UnrealEd/Private/Commandlets/GatherTextFromAssetsCommandlet.cpp
@@ -8,6 +8,7 @@
 #include "Misc/Paths.h"
 #include "Misc/OutputDeviceHelper.h"
 #include "Misc/FeedbackContext.h"
+#include "UObject/GCObjectScopeGuard.h"
 #include "UObject/EditorObjectVersion.h"
 #include "UObject/UObjectIterator.h"
 #include "Modules/ModuleManager.h"
@@ -21,6 +22,10 @@
 #include "Sound/DialogueWave.h"
 #include "AssetRegistry/ARFilter.h"
 #include "AssetRegistry/AssetRegistryModule.h"
+#include "WorldPartition/WorldPartition.h"
+#include "WorldPartition/WorldPartitionHelpers.h"
+#include "WorldPartition/WorldPartitionActorDescUtils.h"
+#include "EditorWorldUtils.h"
 #include "PackageHelperFunctions.h"
 #include "ShaderCompiler.h"
 #include "DistanceFieldAtlas.h"
@@ -634,6 +639,54 @@ PRAGMA_ENABLE_DEPRECATION_WARNINGS
 		return 0;
 	}
 
+	// Discover the external actors for any worlds that are pending gather
+	{
+		UE_LOG(LogGatherTextFromAssetsCommandlet, Display, TEXT("Discovering external actors to gather..."));
+		
+		TArray<FName> ExternalActorsSearchPaths;
+		AssetDataArray.RemoveAll([&ExternalActorsSearchPaths](const FAssetData& AssetData)
+		{
+			const FNameBuilder PackageNameStr(AssetData.PackageName);
+
+			if (AssetData.AssetClassPath == UWorld::StaticClass()->GetClassPathName())
+			{
+				if (ULevel::GetIsLevelPartitionedFromAsset(AssetData))
+				{
+					FString ExternalActorsPathForWorld = ULevel::GetExternalActorsPath(*PackageNameStr);
+					ExternalActorsSearchPaths.Add(*ExternalActorsPathForWorld);
+				}
+			}
+			else if (PackageNameStr.ToView().Contains(FPackagePath::GetExternalActorsFolderName()))
+			{
+				// Remove any external actors that are already in the list, as they will be re-added providing their owner world passed the gather criteria
+				return true;
+			}
+
+			return false;
+		});
+		
+		if (ExternalActorsSearchPaths.Num() > 0)
+		{
+			AssetRegistry.GetAssetsByPaths(ExternalActorsSearchPaths, AssetDataArray, /*bRecursive*/true);
+		}
+	}
+
+	auto AppendPackagePendingGather = [this](const FName PackageNameToGather) -> FPackagePendingGather*
+	{
+		FString PackageFilename;
+		if (!FPackageName::FindPackageFileWithoutExtension(FPackageName::LongPackageNameToFilename(PackageNameToGather.ToString()), PackageFilename))
+		{
+			return nullptr;
+		}
+		PackageFilename = FPaths::ConvertRelativePathToFull(PackageFilename);
+
+		FPackagePendingGather& PackagePendingGather = PackagesPendingGather.AddDefaulted_GetRef();
+		PackagePendingGather.PackageName = PackageNameToGather;
+		PackagePendingGather.PackageFilename = MoveTemp(PackageFilename);
+		PackagePendingGather.PackageLocCacheState = EPackageLocCacheState::Cached;
+		return &PackagePendingGather;
+	};
+
 	// Collect the basic information about the packages that we're going to gather from
 	{
 		// Collapse the assets down to a set of packages
@@ -649,22 +702,13 @@ PRAGMA_ENABLE_DEPRECATION_WARNINGS
 		PackagesPendingGather.Reserve(PackageNamesToGather.Num());
 		for (const FName& PackageNameToGather : PackageNamesToGather)
 		{
-			FString PackageFilename;
-			if (!FPackageName::FindPackageFileWithoutExtension(FPackageName::LongPackageNameToFilename(PackageNameToGather.ToString()), PackageFilename))
-			{
-				continue;
-			}
-			PackageFilename = FPaths::ConvertRelativePathToFull(PackageFilename);
-
-			FPackagePendingGather& PackagePendingGather = PackagesPendingGather[PackagesPendingGather.AddDefaulted()];
-			PackagePendingGather.PackageName = PackageNameToGather;
-			PackagePendingGather.PackageFilename = MoveTemp(PackageFilename);
-			PackagePendingGather.PackageLocCacheState = EPackageLocCacheState::Cached;
+			AppendPackagePendingGather(PackageNameToGather);
 		}
 	}
 
 	FAssetGatherCacheMetrics AssetGatherCacheMetrics;
 	TMap<FString, FName> AssignedPackageLocalizationIds;
+	TMap<FName, TSet<FGuid>> StaleExternalActors;
 
 	int32 NumPackagesProcessed = 0;
 	int32 PackageCount = PackagesPendingGather.Num();
@@ -684,6 +728,8 @@ PRAGMA_ENABLE_DEPRECATION_WARNINGS
 			return false;
 		}
 
+		const bool bIsExternalActorPackage = PackageNameStr.ToView().Contains(FPackagePath::GetExternalActorsFolderName());
+
 		// Read package file summary from the file.
 		FPackageFileSummary PackageFileSummary;
 		*FileReader << PackageFileSummary;
@@ -741,6 +787,15 @@ PRAGMA_ENABLE_DEPRECATION_WARNINGS
 				}
 			}
 		}
+		// TODO: We need a version bump for this, but for now consider all external actor packages dirty as they were failing their text gather on save
+		else if (bIsExternalActorPackage)
+		{
+			// Fallback on the old package flag check.
+			if (PackageFileSummary.GetPackageFlags() & PKG_RequiresLocalizationGather)
+			{
+				PackagePendingGather.PackageLocCacheState = EPackageLocCacheState::Uncached_TooOld;
+			}
+		}
 
 		// If this package doesn't have any cached data, then we have to load it for gather
 		if (PackageFileSummary.GetFileVersionUE() >= VER_UE4_SERIALIZE_TEXT_IN_PACKAGES && PackageFileSummary.GatherableTextDataOffset == 0 && (PackageFileSummary.GetPackageFlags() & PKG_RequiresLocalizationGather))
@@ -750,6 +805,23 @@ PRAGMA_ENABLE_DEPRECATION_WARNINGS
 
 		if (PackagePendingGather.PackageLocCacheState != EPackageLocCacheState::Cached)
 		{
+			// External actors actors must be gathered via their owner world rather than via a raw LoadPackage call
+			// Remove them from PackagesToGather as the owner world is merged back in below
+			if (bIsExternalActorPackage)
+			{
+				TArray<FAssetData> ActorsInPackage;
+				AssetRegistry.GetAssetsByPackageName(PackagePendingGather.PackageName, ActorsInPackage);
+				for (const FAssetData& ActorInPackage : ActorsInPackage)
+				{
+					if (TUniquePtr<FWorldPartitionActorDesc> ActorDesc = FWorldPartitionActorDescUtils::GetActorDescriptorFromAssetData(ActorInPackage))
+					{
+						FName WorldPackageName = *FPackageName::ObjectPathToPackageName(ActorDesc->GetActorSoftPath().ToString());
+						StaleExternalActors.FindOrAdd(WorldPackageName).Add(ActorDesc->GetGuid());
+					}
+				}
+				return true;
+			}
+
 			AssetGatherCacheMetrics.CountUncachedAsset(PackagePendingGather.PackageLocCacheState);
 			return false;
 		}
@@ -785,6 +857,30 @@ PRAGMA_ENABLE_DEPRECATION_WARNINGS
 
 	AssetGatherCacheMetrics.LogMetrics();
 
+	// Merge any pending WP map requests back into PackagesPendingGather
+	for (TTuple<FName, TSet<FGuid>>& StaleExternalActorsPair : StaleExternalActors)
+	{
+		FPackagePendingGather* WorldPackagePendingGather = PackagesPendingGather.FindByPredicate([&StaleExternalActorsPair](const FPackagePendingGather& PotentialPackagePendingGather)
+		{
+			return PotentialPackagePendingGather.PackageName == StaleExternalActorsPair.Key;
+		});
+		if (!WorldPackagePendingGather)
+		{
+			WorldPackagePendingGather = AppendPackagePendingGather(StaleExternalActorsPair.Key);
+		}
+
+		if (WorldPackagePendingGather)
+		{
+			WorldPackagePendingGather->ExternalActors = MoveTemp(StaleExternalActorsPair.Value);
+			WorldPackagePendingGather->PackageLocCacheState = EPackageLocCacheState::Uncached_TooOld;
+		}
+		else
+		{
+			UE_LOG(LogGatherTextFromAssetsCommandlet, Warning, TEXT("Failed to queue world package '%s' for %d external actor(s)."), *StaleExternalActorsPair.Key.ToString(), StaleExternalActorsPair.Value.Num());
+		}
+	}
+	StaleExternalActors.Reset();
+
 	NumPackagesProcessed = 0;
 	PackageCount = PackagesPendingGather.Num();
 	if (PackageCount == 0)
@@ -827,7 +923,14 @@ PRAGMA_ENABLE_DEPRECATION_WARNINGS
 		UPackage* Package = nullptr;
 		{
 			FLoadPackageLogOutputRedirector::FScopedCapture ScopedCapture(&LogOutputRedirector, *PackageNameStr);
-			Package = LoadPackage(nullptr, *PackageNameStr, LOAD_NoWarn | LOAD_Quiet);
+			if (PackagePendingGather.ExternalActors.Num() > 0)
+			{
+				Package = LoadWorldPackageForEditor(*PackageNameStr, EWorldType::Editor, LOAD_NoWarn | LOAD_Quiet);
+			}
+			else
+			{
+				Package = LoadPackage(nullptr, *PackageNameStr, LOAD_NoWarn | LOAD_Quiet);
+			}
 		}
 
 		if (!Package)
@@ -853,7 +956,7 @@ PRAGMA_ENABLE_DEPRECATION_WARNINGS
 
 		// Because packages may not have been resaved after this flagging was implemented, we may have added packages to load that weren't flagged - potential false positives.
 		// The loading process should have reflagged said packages so that only true positives will have this flag.
-		if (Package->RequiresLocalizationGather())
+		if (Package->RequiresLocalizationGather() || PackagePendingGather.ExternalActors.Num() > 0)
 		{
 			UE_LOG(LogGatherTextFromAssetsCommandlet, Display, TEXT("[%6.2f%%] Gathering package: '%s'..."), PercentageComplete, *PackageNameStr);
 
@@ -903,6 +1006,47 @@ PRAGMA_ENABLE_DEPRECATION_WARNINGS
 				}
 			}
 
+			// If this is a WP world then query the localization for any external actors actors that were determined to be stale
+			if (PackagePendingGather.ExternalActors.Num() > 0)
+			{
+				if (UWorld* World = UWorld::FindWorldInPackage(Package))
+				{
+					UWorld::InitializationValues IVS;
+					IVS.InitializeScenes(false);
+					IVS.AllowAudioPlayback(false);
+					IVS.RequiresHitProxies(false);
+					IVS.CreatePhysicsScene(false);
+					IVS.CreateNavigation(false);
+					IVS.CreateAISystem(false);
+					IVS.ShouldSimulatePhysics(false);
+					IVS.EnableTraceCollision(false);
+					IVS.SetTransactional(false);
+					IVS.CreateFXSystem(false);
+					IVS.CreateWorldPartition(true);
+
+					FScopedEditorWorld ScopeEditorWorld(World, IVS);
+
+					if (UWorldPartition* WorldPartition = World->GetWorldPartition())
+					{
+						// ForEachActorWithLoading may GC while running, so keep the world partition (and indirectly the world and its package) alive
+						TGCObjectScopeGuard<UWorldPartition> WorldPartitionGCGuard(WorldPartition);
+
+						FWorldPartitionHelpers::FForEachActorWithLoadingParams ForEachActorParams;
+						ForEachActorParams.ActorGuids = PackagePendingGather.ExternalActors.Array();
+
+						FWorldPartitionHelpers::ForEachActorWithLoading(WorldPartition, [&GatherableTextDataArray](const FWorldPartitionActorDesc* ActorDesc)
+						{
+							if (const AActor* Actor = ActorDesc->GetActor())
+							{
+								EPropertyLocalizationGathererResultFlags ActorGatherableTextResultFlags = EPropertyLocalizationGathererResultFlags::Empty;
+								FPropertyLocalizationDataGatherer(GatherableTextDataArray, Actor->GetExternalPackage(), ActorGatherableTextResultFlags);
+							}
+							return true;
+						}, ForEachActorParams);
+					}
+				}
+			}
+
 			// This package may have already been cached in cases where we're reporting or fixing assets with a stale gather cache
 			// This check prevents it being gathered a second time
 			if (PackagePendingGather.PackageLocCacheState != EPackageLocCacheState::Cached)
