Fix copy modifying level

diff --git a/Engine/Source/Editor/UnrealEd/Private/EditorExporters.cpp b/Engine/Source/Editor/UnrealEd/Private/EditorExporters.cpp
--- a/Engine/Source/Editor/UnrealEd/Private/EditorExporters.cpp
+++ b/Engine/Source/Editor/UnrealEd/Private/EditorExporters.cpp
@@ -540,7 +540,9 @@ bool ULevelExporterT3D::ExportText( const FExportObjectInnerContext* Context, UO
 
 				AActor* ParentActor = Actor->GetAttachParentActor();
 				FName SocketName = Actor->GetAttachParentSocketName();
-				Actor->DetachFromActor(FDetachmentTransformRules::KeepWorldTransform);
+				FDetachmentTransformRules DetachmentTransformRules = FDetachmentTransformRules::KeepWorldTransform;
+				DetachmentTransformRules.bCallModify = false;
+				Actor->DetachFromActor(DetachmentTransformRules);
 
 				FString ParentActorString = ( ParentActor ? FString::Printf(TEXT(" ParentActor=%s"), *ParentActor->GetName() ) : TEXT(""));
 				FString SocketNameString = ( (ParentActor && SocketName != NAME_None) ? FString::Printf(TEXT(" SocketName=%s"), *SocketName.ToString() ) : TEXT(""));
