Fixed bPromoteOutputLogWarningsDuringPIE for logs from the main thread

Declaring FOutputLogErrorsToMessageLogProxy to be usable from multiple threads makes it unbuffered. Logs coming from the game thread will once again be handled, and logs from other threads will continue to be ignored for now.

#rnx

diff --git a/Engine/Source/Editor/UnrealEd/Private/PlayLevel.cpp b/Engine/Source/Editor/UnrealEd/Private/PlayLevel.cpp
--- a/Engine/Source/Editor/UnrealEd/Private/PlayLevel.cpp
+++ b/Engine/Source/Editor/UnrealEd/Private/PlayLevel.cpp
@@ -133,7 +133,7 @@ FText GeneratePIEViewportWindowTitle(const EPlayNetMode InNetMode, const ERHIFea
 bool IsPrimaryPIEClient(const FRequestPlaySessionParams& InPlaySessionParams, const int32 InClientIndex);
 
 // This class listens to output log messages, and forwards warnings and errors to the message log
-class FOutputLogErrorsToMessageLogProxy : public FOutputDevice
+class FOutputLogErrorsToMessageLogProxy final : public FOutputDevice
 {
 public:
 	FOutputLogErrorsToMessageLogProxy()
@@ -168,6 +168,11 @@ public:
 			}
 		}
 	}
+
+	virtual bool CanBeUsedOnMultipleThreads() const final
+	{
+		return true;
+	}
 	// End of FOutputDevice interface
 };
 
@@ -2674,7 +2679,7 @@ void UEditorEngine::StartPlayInEditorSession(FRequestPlaySessionParams& InReques
 	// Register for log processing so we can promote errors/warnings to the message log
 	if (GetDefault<ULevelEditorPlaySettings>()->bPromoteOutputLogWarningsDuringPIE)
 	{
-		OutputLogErrorsToMessageLogProxyPtr = MakeShareable(new FOutputLogErrorsToMessageLogProxy());
+		OutputLogErrorsToMessageLogProxyPtr = MakeShared<FOutputLogErrorsToMessageLogProxy>();
 	}
 
 	// Notify the XRSystem that it needs to BeginPlay.
