Fixed bPromoteOutputLogWarningsDuringPIE for logs from any thread

Queue logs from outside of the game thread and flush them every tick.

#rnx

diff --git a/Engine/Source/Editor/UnrealEd/Private/PlayLevel.cpp b/Engine/Source/Editor/UnrealEd/Private/PlayLevel.cpp
--- a/Engine/Source/Editor/UnrealEd/Private/PlayLevel.cpp
+++ b/Engine/Source/Editor/UnrealEd/Private/PlayLevel.cpp
@@ -121,6 +121,8 @@
 #include "IAssetViewport.h"
 #include "IPIEAuthorizer.h"
 #include "Features/IModularFeatures.h"
+#include "Containers/DepletableMpscQueue.h"
+#include "TickableEditorObject.h"
 
 DEFINE_LOG_CATEGORY(LogPlayLevel);
 
@@ -133,7 +135,7 @@ FText GeneratePIEViewportWindowTitle(const EPlayNetMode InNetMode, const ERHIFea
 bool IsPrimaryPIEClient(const FRequestPlaySessionParams& InPlaySessionParams, const int32 InClientIndex);
 
 // This class listens to output log messages, and forwards warnings and errors to the message log
-class FOutputLogErrorsToMessageLogProxy final : public FOutputDevice
+class FOutputLogErrorsToMessageLogProxy final : public FOutputDevice, public FTickableEditorObject
 {
 public:
 	FOutputLogErrorsToMessageLogProxy()
@@ -142,29 +144,27 @@ public:
 	}
 
 	~FOutputLogErrorsToMessageLogProxy()
-{
+	{
 		GLog->RemoveOutputDevice(this);
-}
+		Tick(0.0f);
+	}
+
+	// FOutputDevice Interface
 
-	// FOutputDevice interface
 	virtual void Serialize(const TCHAR* V, ELogVerbosity::Type Verbosity, const class FName& Category) override
 	{
-		//@TODO: Remove IsInGameThread() once the message log is thread safe
-		if ((Verbosity <= ELogVerbosity::Warning) && IsInGameThread())
+		if (Verbosity <= ELogVerbosity::Warning)
 		{
-			const FText Message = FText::Format(LOCTEXT("OutputLogToMessageLog", "{0}: {1}"), FText::FromName(Category), FText::AsCultureInvariant(FString(V)));
-
-			switch (Verbosity)
+			FLine Line;
+			Line.Message = FText::Format(LOCTEXT("OutputLogToMessageLog", "{0}: {1}"), FText::FromName(Category), FText::AsCultureInvariant(FString(V)));
+			Line.Verbosity = Verbosity;
+			if (IsInGameThread())
 			{
-			case ELogVerbosity::Warning:
-				FMessageLog(NAME_CategoryPIE).SuppressLoggingToOutputLog(true).Warning(Message);
-				break;
-			case ELogVerbosity::Error:
-				FMessageLog(NAME_CategoryPIE).SuppressLoggingToOutputLog(true).Error(Message);
-				break;
-			case ELogVerbosity::Fatal:
-				checkf(false, *Message.ToString());
-				break;
+				LogLine(Line);
+			}
+			else
+			{
+				QueuedLines.Enqueue(MoveTemp(Line));
 			}
 		}
 	}
@@ -173,7 +173,48 @@ public:
 	{
 		return true;
 	}
-	// End of FOutputDevice interface
+
+	// FTickableEditorObject Interface
+
+	virtual void Tick(float DeltaTime) final
+	{
+		QueuedLines.Deplete(LogLine);
+	}
+
+	virtual ETickableTickType GetTickableTickType() const final
+	{
+		return ETickableTickType::Always;
+	}
+
+	virtual TStatId GetStatId() const final
+	{
+		RETURN_QUICK_DECLARE_CYCLE_STAT(FOutputLogErrorsToMessageLogProxy, STATGROUP_Tickables);
+	}
+
+private:
+	struct FLine
+	{
+		FText Message;
+		ELogVerbosity::Type Verbosity;
+	};
+
+	static void LogLine(const FLine& Line)
+	{
+		switch (Line.Verbosity)
+		{
+		case ELogVerbosity::Warning:
+			FMessageLog(NAME_CategoryPIE).SuppressLoggingToOutputLog(true).Warning(Line.Message);
+			break;
+		case ELogVerbosity::Error:
+			FMessageLog(NAME_CategoryPIE).SuppressLoggingToOutputLog(true).Error(Line.Message);
+			break;
+		case ELogVerbosity::Fatal:
+			checkf(false, *Line.Message.ToString());
+			break;
+		}
+	}
+
+	UE::TDepletableMpscQueue<FLine> QueuedLines;
 };
 
 void UEditorEngine::EndPlayMap()
