Fixed bPromoteOutputLogWarningsDuringPIE to flush logs on destruction

#rnx

diff --git a/Engine/Source/Editor/UnrealEd/Private/PlayLevel.cpp b/Engine/Source/Editor/UnrealEd/Private/PlayLevel.cpp
--- a/Engine/Source/Editor/UnrealEd/Private/PlayLevel.cpp
+++ b/Engine/Source/Editor/UnrealEd/Private/PlayLevel.cpp
@@ -145,6 +145,7 @@ public:
 
 	~FOutputLogErrorsToMessageLogProxy()
 	{
+		GLog->FlushThreadedLogs();
 		GLog->RemoveOutputDevice(this);
 		Tick(0.0f);
 	}
