Updated FAssetSourceFilenameCache to add the initial assets in parallel.

Before: 3.69s
After: 1.4s

diff --git a/Engine/Source/Editor/UnrealEd/Private/AutoReimport/AssetSourceFilenameCache.cpp b/Engine/Source/Editor/UnrealEd/Private/AutoReimport/AssetSourceFilenameCache.cpp
--- a/Engine/Source/Editor/UnrealEd/Private/AutoReimport/AssetSourceFilenameCache.cpp
+++ b/Engine/Source/Editor/UnrealEd/Private/AutoReimport/AssetSourceFilenameCache.cpp
@@ -6,6 +6,7 @@
 #include "AssetRegistry/AssetDataTagMap.h"
 #include "AssetRegistry/AssetRegistryModule.h"
 #include "AssetRegistry/IAssetRegistry.h"
+#include "Async/ParallelFor.h"
 #include "CoreGlobals.h"
 #include "EditorFramework/AssetImportData.h"
 #include "HAL/Platform.h"
@@ -36,10 +37,7 @@ FAssetSourceFilenameCache::FAssetSourceFilenameCache()
 	TArray<FAssetData> Assets;
 	if (AssetRegistry.GetAllAssets(Assets))
 	{
-		for (auto& Asset : Assets)
-		{
-			HandleOnAssetAdded(Asset);
-		}
+		HandleOnAssetsAdded(Assets);
 	}
 }
 
@@ -105,6 +103,29 @@ void FAssetSourceFilenameCache::HandleOnAssetAdded(const FAssetData& AssetData)
 	}
 }
 
+void FAssetSourceFilenameCache::HandleOnAssetsAdded(const TArray<FAssetData>& Assets)
+{
+	const uint32 AssetCount = Assets.Num();
+	TArray<TOptional<FAssetImportInfo>> ImportDataList;
+	ImportDataList.SetNum(AssetCount);
+	ParallelFor(Assets.Num(), [&Assets, &ImportDataList](int32 Index)
+		{
+			ImportDataList[Index] = ExtractAssetImportInfo(Assets[Index]);
+		});
+
+	for (uint32 Index = 0; Index < AssetCount; Index++)
+	{
+		const TOptional<FAssetImportInfo>& ImportData = ImportDataList[Index];
+		if (ImportData.IsSet())
+		{
+			for (const auto& SourceFile : ImportData->SourceFiles)
+			{
+				SourceFileToObjectPathCache.FindOrAdd(FPaths::GetCleanFilename(SourceFile.RelativeFilename)).Add(Assets[Index].GetSoftObjectPath());
+			}
+		}
+	}
+}
+
 void FAssetSourceFilenameCache::HandleOnAssetRemoved(const FAssetData& AssetData)
 {
 	TOptional<FAssetImportInfo> ImportData = ExtractAssetImportInfo(AssetData);
diff --git a/Engine/Source/Editor/UnrealEd/Public/AutoReimport/AssetSourceFilenameCache.h b/Engine/Source/Editor/UnrealEd/Public/AutoReimport/AssetSourceFilenameCache.h
--- a/Engine/Source/Editor/UnrealEd/Public/AutoReimport/AssetSourceFilenameCache.h
+++ b/Engine/Source/Editor/UnrealEd/Public/AutoReimport/AssetSourceFilenameCache.h
@@ -45,6 +45,7 @@ public:
 private:
 	/** Delegate bindings that keep the cache up-to-date */
 	void HandleOnAssetAdded(const FAssetData& AssetData);
+	void HandleOnAssetsAdded(const TArray<FAssetData>& Assets);
 	void HandleOnAssetRemoved(const FAssetData& AssetData);
 	void HandleOnAssetRenamed(const FAssetData& AssetData, const FString& OldPath);
 	void HandleOnAssetUpdated(const FAssetImportInfo& OldData, const UAssetImportData* ImportData);
