Use GetWorld function instead of GetTypedOuter<UWorld> in FKismetDebugUtilities::CanInspectPinValue

diff --git a/Engine/Source/Editor/UnrealEd/Private/Kismet2/KismetDebugUtilities.cpp b/Engine/Source/Editor/UnrealEd/Private/Kismet2/KismetDebugUtilities.cpp
--- a/Engine/Source/Editor/UnrealEd/Private/Kismet2/KismetDebugUtilities.cpp
+++ b/Engine/Source/Editor/UnrealEd/Private/Kismet2/KismetDebugUtilities.cpp
@@ -1691,7 +1691,7 @@ bool FKismetDebugUtilities::CanInspectPinValue(const UEdGraphPin* Pin)
 	}
 
 	// Can't inspect if not in PIE
-	const UWorld* OwningWorld = Object->GetTypedOuter<UWorld>();
+	const UWorld* OwningWorld = Object->GetWorld();
 	if (!OwningWorld || !(OwningWorld->IsPlayInEditor() || OwningWorld->IsPreviewWorld()))
 	{
 		return false;
