#UE Downgrade error to a warning temporarily while a build issue is being investigated

diff --git a/Engine/Source/Editor/UnrealEd/Private/Commandlets/AssetRegistryGenerator.cpp b/Engine/Source/Editor/UnrealEd/Private/Commandlets/AssetRegistryGenerator.cpp
--- a/Engine/Source/Editor/UnrealEd/Private/Commandlets/AssetRegistryGenerator.cpp
+++ b/Engine/Source/Editor/UnrealEd/Private/Commandlets/AssetRegistryGenerator.cpp
@@ -985,7 +985,7 @@ void FAssetRegistryGenerator::InjectEncryptionData(FAssetRegistryState& TargetSt
 						const FAssetData::FChunkArrayView ChunkIDs = AssetData->GetChunkIDs();
 						if (ChunkIDs.Num() > 1)
 						{
-							UE_LOG(LogAssetRegistryGenerator, Error, TEXT("Encrypted root asset '%s' exists in two chunks. Only secondary assets should be shared between chunks."), *AssetData->GetObjectPathString());
+							UE_LOG(LogAssetRegistryGenerator, Warning, TEXT("Encrypted root asset '%s' exists in two chunks. Only secondary assets should be shared between chunks."), *AssetData->GetObjectPathString());
 						}
 						else if (ChunkIDs.Num() == 1)
 						{
