Actual camera speed is now displayed in the viewport toolbar.

Made some camera speeds a bit more easily parsed: [REDACTED_PHONE] / 4) -> 1, [REDACTED_PHONE] / 4) -> 0.33.

diff --git a/Engine/Source/Editor/UnrealEd/Private/EditorViewportClient.cpp b/Engine/Source/Editor/UnrealEd/Private/EditorViewportClient.cpp
--- a/Engine/Source/Editor/UnrealEd/Private/EditorViewportClient.cpp
+++ b/Engine/Source/Editor/UnrealEd/Private/EditorViewportClient.cpp
@@ -307,7 +307,10 @@ const uint32 FEditorViewportClient::MaxCameraSpeeds = 8;
 
 float FEditorViewportClient::GetCameraSpeed() const
 {
-	return GetCameraSpeed(GetCameraSpeedSetting());
+	const float CameraBoost = (GetDefault<ULevelEditorViewportSettings>()->FlightCameraControlExperimentalNavigation && IsShiftPressed()) ? 2.0f : 1.0f;
+	const float SpeedSetting = GetCameraSpeed(GetCameraSpeedSetting());
+	const float FinalCameraSpeedScale = SpeedSetting * FlightCameraSpeedScale * GetCameraSpeedScalar() * CameraBoost;
+	return FinalCameraSpeedScale;
 }
 
 float FEditorViewportClient::GetCameraSpeed(int32 SpeedSetting) const
@@ -320,7 +323,7 @@ float FEditorViewportClient::GetCameraSpeed(int32 SpeedSetting) const
 	//#define MOVEMENTSPEED_VERYFAST		64	~ 16
 
 	const int32 SpeedToUse = FMath::Clamp<int32>(SpeedSetting, 1, MaxCameraSpeeds);
-	const float Speed[] = { 0.03125f, 0.09375f, 0.33f, 1.f, 3.f, 8.f, 16.f, 32.f };
+	const float Speed[] = { 0.033f, 0.1f, 0.33f, 1.f, 3.f, 8.f, 16.f, 32.f };
 
 	return Speed[SpeedToUse - 1];
 }
diff --git a/Engine/Source/Editor/UnrealEd/Private/STransformViewportToolbar.cpp b/Engine/Source/Editor/UnrealEd/Private/STransformViewportToolbar.cpp
--- a/Engine/Source/Editor/UnrealEd/Private/STransformViewportToolbar.cpp
+++ b/Engine/Source/Editor/UnrealEd/Private/STransformViewportToolbar.cpp
@@ -581,7 +581,10 @@ FText STransformViewportToolBar::GetCameraSpeedLabel() const
 	auto ViewportPin = Viewport.Pin();
 	if (ViewportPin.IsValid() && ViewportPin->GetViewportClient().IsValid())
 	{
-		return FText::AsNumber( ViewportPin->GetViewportClient()->GetCameraSpeedSetting() );
+		float CameraSpeed = ViewportPin->GetViewportClient()->GetCameraSpeed();
+		FNumberFormattingOptions FormattingOptions = FNumberFormattingOptions::DefaultNoGrouping();
+		FormattingOptions.MaximumFractionalDigits = CameraSpeed > 1 ? 1 : 3;
+		return FText::AsNumber(CameraSpeed, &FormattingOptions);
 	}
 
 	return FText();
