UWorld Initialization: CleanupPhysicsSceneThatWasInitializedForSave should skip UpdateWorldComponents if the world is no longer initialized, because it is unnecessary and can cause spurious errors if the world is not initialized.
#rnx

diff --git a/Engine/Source/Editor/UnrealEd/Private/EditorEngine.cpp b/Engine/Source/Editor/UnrealEd/Private/EditorEngine.cpp
--- a/Engine/Source/Editor/UnrealEd/Private/EditorEngine.cpp
+++ b/Engine/Source/Editor/UnrealEd/Private/EditorEngine.cpp
@@ -4503,17 +4503,21 @@ void UEditorEngine::CleanupPhysicsSceneThatWasInitializedForSave(UWorld* World,
 	// Make sure we clean up the physics scene here. If we leave too many scenes in memory, undefined behavior occurs when locking a scene for read/write.
 	World->ClearWorldComponents();
 
-	if(bForceInitialized)
+	if (bForceInitialized)
 	{
 		World->CleanupWorld(true, true, World);
 	}
 
 	World->SetPhysicsScene(nullptr);
 
-	// Update components again in case it was a world without a physics scene but did have rendered components.
-	World->UpdateWorldComponents(true, true);
+	if (World->IsInitialized())
+	{
+		// Update components again in case it was a world without a physics scene but did have rendered components.
+		World->UpdateWorldComponents(true, true);
+	}
 
 	for (UPackage* ExternalPackage : World->GetPackage()->GetExternalPackages())
+
 	{
 		if (!DirtyPackages.Contains(ExternalPackage) && ExternalPackage->IsDirty())
 		{
