The fbx import option "Build Nanite" should be hidden from the UI when doing a reimport or when showing import setting in the staticmesh editor detail panel. After the initial import the user should see only the staticmesh "Nanite settings -> Enable".

#rnx

diff --git a/Engine/Source/Editor/StaticMeshEditor/Private/StaticMeshEditorTools.cpp b/Engine/Source/Editor/StaticMeshEditor/Private/StaticMeshEditorTools.cpp
--- a/Engine/Source/Editor/StaticMeshEditor/Private/StaticMeshEditorTools.cpp
+++ b/Engine/Source/Editor/StaticMeshEditor/Private/StaticMeshEditorTools.cpp
@@ -182,6 +182,12 @@ void FStaticMeshDetails::OnInstancedFbxStaticMeshImportDataPropertyIteration(IDe
 {
 	IDetailPropertyRow* Row = nullptr;
 
+	if (Property->GetBoolMetaData(TEXT("ReimportRestrict")))
+	{
+		//Dont create a row for reimport restricted property, we do not want to show them
+		return;
+	}
+
 	if (PropertyGroup)
 	{
 		Row = &PropertyGroup->AddPropertyRow(Property);
@@ -5707,7 +5713,9 @@ FReply FNaniteSettingsLayout::OnImportHiRes()
 	{
 		StaticMesh->GetHiResSourceModel().SourceImportFilename = FString();
 		FbxMeshUtils::ImportStaticMeshHiResSourceModelDialog(StaticMesh);
-		StaticMeshEditor.RefreshTool();
+		//If we import a hires we should enable nanite
+		NaniteSettings.bEnabled = true;
+		ApplyChanges();
 	}
 	return FReply::Handled();
 }
diff --git a/Engine/Source/Editor/UnrealEd/Classes/Factories/FbxStaticMeshImportData.h b/Engine/Source/Editor/UnrealEd/Classes/Factories/FbxStaticMeshImportData.h
--- a/Engine/Source/Editor/UnrealEd/Classes/Factories/FbxStaticMeshImportData.h
+++ b/Engine/Source/Editor/UnrealEd/Classes/Factories/FbxStaticMeshImportData.h
@@ -38,7 +38,7 @@ class UFbxStaticMeshImportData : public UFbxMeshImportData
 	uint32 bBuildReversedIndexBuffer:1;
 
 	/** For static meshes, enabling this option will allow you to use Nanite rendering at runtime. Can only be used with simple opaque materials. */
-	UPROPERTY(EditAnywhere, BlueprintReadWrite, config, Category = Mesh, meta = (ToolTip = "If enabled, allows to render objects with Nanite", ImportType = "StaticMesh"))
+	UPROPERTY(EditAnywhere, BlueprintReadWrite, config, Category = Mesh, meta = (ToolTip = "If enabled, allows to render objects with Nanite", ReimportRestrict = "true", ImportType = "StaticMesh"))
 	uint32 bBuildNanite : 1;
 
 	UPROPERTY(EditAnywhere, BlueprintReadWrite, config, AdvancedDisplay, Category= Mesh, meta=(ImportType="StaticMesh", ReimportRestrict = "true"))
diff --git a/Engine/Source/Editor/UnrealEd/Private/Fbx/FbxStaticMeshImport.cpp b/Engine/Source/Editor/UnrealEd/Private/Fbx/FbxStaticMeshImport.cpp
--- a/Engine/Source/Editor/UnrealEd/Private/Fbx/FbxStaticMeshImport.cpp
+++ b/Engine/Source/Editor/UnrealEd/Private/Fbx/FbxStaticMeshImport.cpp
@@ -1710,6 +1710,10 @@ UStaticMesh* UnFbx::FFbxImporter::ImportStaticMeshAsSingle(UObject* InParent, TA
 
 		// Free any RHI resources for existing mesh before we re-create in place.
 		ExistingMesh->PreEditChange(NULL);
+
+		//When we do a reimport the nanite static mesh import setting should not affect the asset since its a "reimportrestrict" option
+		//In that case we override the import setting
+		ImportOptions->bBuildNanite = ExistingMesh->NaniteSettings.bEnabled;
 	}
 	else if (ExistingObject)
 	{
