Adding an option to specify a new commandline when restarting the editor.

diff --git a/Engine/Source/Editor/UnrealEd/Private/UnrealEdMisc.cpp b/Engine/Source/Editor/UnrealEd/Private/UnrealEdMisc.cpp
--- a/Engine/Source/Editor/UnrealEd/Private/UnrealEdMisc.cpp
+++ b/Engine/Source/Editor/UnrealEd/Private/UnrealEdMisc.cpp
@@ -525,7 +525,7 @@ void FUnrealEdMisc::OnInit()
 
 	// Give the settings editor a way to restart the editor when it needs to
 	ISettingsEditorModule& SettingsEditorModule = FModuleManager::GetModuleChecked<ISettingsEditorModule>("SettingsEditor");
-	SettingsEditorModule.SetRestartApplicationCallback(FSimpleDelegate::CreateRaw(this, &FUnrealEdMisc::RestartEditor, false));
+	SettingsEditorModule.SetRestartApplicationCallback(FSimpleDelegate::CreateLambda([this]() { RestartEditor(/*bWarn =*/false); }));
 
 	// add handler to notify about navmesh building process
 	NavigationBuildingNotificationHandler = MakeShareable(new FNavigationBuildingNotificationImpl());
@@ -1082,7 +1082,7 @@ void FUnrealEdMisc::OnExit()
 	if( PendingProjName.Len() > 0 )
 	{
 		// If there is a pending project switch, spawn that process now and use the same command line parameters that were used for this editor instance.
-		FString Cmd = FString::Printf(TEXT("%s %s"), *PendingProjName, FCommandLine::Get());
+		FString Cmd = FString::Printf(TEXT("%s %s"), *PendingProjName, *PendingCommandLine.Get(FCommandLine::Get()));
 
 		FString ExeFilename = CreateProjectPath();
 		FProcHandle Handle = FPlatformProcess::CreateProc( *ExeFilename, *Cmd, true, false, false, NULL, 0, NULL, NULL );
@@ -1686,7 +1686,7 @@ void FUnrealEdMisc::LogAssetUpdate(UObject* UpdatedAsset, FObjectPreSaveContext
 	}
 }
 
-void FUnrealEdMisc::SwitchProject(const FString& GameOrProjectFileName, bool bWarn)
+void FUnrealEdMisc::SwitchProject(const FString& GameOrProjectFileName, bool bWarn, const TOptional<FString>& NewCommandLine)
 {
 	if (GUnrealEd->WarnIfLightingBuildIsCurrentlyRunning())
 	{
@@ -1740,6 +1740,7 @@ void FUnrealEdMisc::SwitchProject(const FString& GameOrProjectFileName, bool bWa
 		}
 
 		SetPendingProjectName( PendingProjName );
+		PendingCommandLine = NewCommandLine;
 
 		// Close the editor.  This will prompt the user to save changes.  If they hit cancel, we abort the project switch
 		GEngine->DeferredCommands.Add( TEXT("CLOSE_SLATE_MAINFRAME"));
@@ -1750,7 +1751,7 @@ void FUnrealEdMisc::SwitchProject(const FString& GameOrProjectFileName, bool bWa
 	}
 }
 
-void FUnrealEdMisc::RestartEditor(bool bWarn)
+void FUnrealEdMisc::RestartEditor(bool bWarn, const TOptional<FString>& NewCommandLine)
 {
 	if (GUnrealEd->WarnIfLightingBuildIsCurrentlyRunning())
 	{
@@ -1759,15 +1760,15 @@ void FUnrealEdMisc::RestartEditor(bool bWarn)
 
 	if( FPaths::IsProjectFilePathSet() )
 	{
-		SwitchProject(FPaths::GetProjectFilePath(), bWarn);
+		SwitchProject(FPaths::GetProjectFilePath(), bWarn, NewCommandLine);
 	}
 	else if(FApp::HasProjectName())
 	{
-		SwitchProject(FApp::GetProjectName(), bWarn);
+		SwitchProject(FApp::GetProjectName(), bWarn, NewCommandLine);
 	}
 	else
 	{
-		SwitchProject(TEXT(""), bWarn);
+		SwitchProject(TEXT(""), bWarn, NewCommandLine);
 	}
 }
 
diff --git a/Engine/Source/Editor/UnrealEd/Public/UnrealEdMisc.h b/Engine/Source/Editor/UnrealEd/Public/UnrealEdMisc.h
--- a/Engine/Source/Editor/UnrealEd/Public/UnrealEdMisc.h
+++ b/Engine/Source/Editor/UnrealEd/Public/UnrealEdMisc.h
@@ -126,6 +126,7 @@ public:
 	void ClearPendingProjectName( )
 	{
 		PendingProjectName.Empty();
+		PendingCommandLine.Reset();
 	}
 
 	/**
@@ -190,10 +191,10 @@ public:
 	}
 
 	/** Opens the specified project file or game. Restarts the editor */
-	void SwitchProject(const FString& GameOrProjectFileName, bool bWarn = true);
+	void SwitchProject(const FString& GameOrProjectFileName, bool bWarn = true, const TOptional<FString>& NewCommandLine = TOptional<FString>());
 
 	/** Restarts the editor, reopening the current project, if any */
-	void RestartEditor(bool bWarn = true);
+	void RestartEditor(bool bWarn = true, const TOptional<FString>& NewCommandLine = TOptional<FString>());
 
 	/** Ticks the performance analytics used by the analytics heartbeat */
 	void TickPerformanceAnalytics();
@@ -327,6 +328,9 @@ private:
 	/** The name of a pending project.  When the editor shuts down it will switch to this project if not empty */ 
 	FString PendingProjectName;
 
+	/** Optional alt command-line to use when the editor shuts down and switches projects (if unset, the restart will use the same command-line as the current process). */
+	TOptional<FString> PendingCommandLine;
+
 	/** Map of config->backup filenames to restore on shutdown of the editor */
 	TMap<FString, FString> RestoreConfigFiles;
 
