Change deselection function in UnloadPackages to avoid potential crash

diff --git a/Engine/Source/Editor/UnrealEd/Private/PackageTools.cpp b/Engine/Source/Editor/UnrealEd/Private/PackageTools.cpp
--- a/Engine/Source/Editor/UnrealEd/Private/PackageTools.cpp
+++ b/Engine/Source/Editor/UnrealEd/Private/PackageTools.cpp
@@ -417,7 +417,7 @@ UPackageTools::UPackageTools(const FObjectInitializer& ObjectInitializer)
 			(*GFlushStreamingFunc)();
 
 			// Remove potential references to to-be deleted objects from the GB selection set.
-			GEditor->GetSelectedObjects()->DeselectAll();
+			GEditor->GetSelectedObjects()->GetElementSelectionSet()->ClearSelection(FTypedElementSelectionOptions());
 
 			bool bScriptPackageWasUnloaded = false;
 
