Added version bump to detect external actor packages with a stale localization cache

#rnx

diff --git a/Engine/Source/Editor/UnrealEd/Private/Commandlets/GatherTextFromAssetsCommandlet.cpp b/Engine/Source/Editor/UnrealEd/Private/Commandlets/GatherTextFromAssetsCommandlet.cpp
--- a/Engine/Source/Editor/UnrealEd/Private/Commandlets/GatherTextFromAssetsCommandlet.cpp
+++ b/Engine/Source/Editor/UnrealEd/Private/Commandlets/GatherTextFromAssetsCommandlet.cpp
@@ -10,6 +10,7 @@
 #include "Misc/FeedbackContext.h"
 #include "UObject/GCObjectScopeGuard.h"
 #include "UObject/EditorObjectVersion.h"
+#include "UObject/FortniteMainBranchObjectVersion.h"
 #include "UObject/UObjectIterator.h"
 #include "Modules/ModuleManager.h"
 #include "Serialization/PropertyLocalizationDataGathering.h"
@@ -762,6 +763,7 @@ PRAGMA_ENABLE_DEPRECATION_WARNINGS
 		}
 
 		const FCustomVersion* const EditorVersion = PackageFileSummary.GetCustomVersionContainer().GetVersion(FEditorObjectVersion::GUID);
+		const FCustomVersion* const FNMainVersion = PackageFileSummary.GetCustomVersionContainer().GetVersion(FFortniteMainBranchObjectVersion::GUID);
 
 		// Packages not resaved since localization gathering flagging was added to packages must be loaded.
 		if (PackageFileSummary.GetFileVersionUE() < VER_UE4_PACKAGE_REQUIRES_LOCALIZATION_GATHER_FLAGGING)
@@ -789,8 +791,7 @@ PRAGMA_ENABLE_DEPRECATION_WARNINGS
 				}
 			}
 		}
-		// TODO: We need a version bump for this, but for now consider all external actor packages dirty as they were failing their text gather on save
-		else if (bIsExternalActorPackage)
+		else if (bIsExternalActorPackage && (!FNMainVersion || FNMainVersion->Version < FFortniteMainBranchObjectVersion::FixedLocalizationGatherForExternalActorPackage))
 		{
 			// Fallback on the old package flag check.
 			if (PackageFileSummary.GetPackageFlags() & PKG_RequiresLocalizationGather)
diff --git a/Engine/Source/Runtime/Core/Public/UObject/FortniteMainBranchObjectVersion.h b/Engine/Source/Runtime/Core/Public/UObject/FortniteMainBranchObjectVersion.h
--- a/Engine/Source/Runtime/Core/Public/UObject/FortniteMainBranchObjectVersion.h
+++ b/Engine/Source/Runtime/Core/Public/UObject/FortniteMainBranchObjectVersion.h
@@ -259,6 +259,9 @@ struct CORE_API FFortniteMainBranchObjectVersion
 		// WorldPartitionStreamingSourceComponent property deprecation
 		WorldPartitionStreamingSourceComponentTargetDeprecation,
 
+		// Fixed localization gathering for external actor packages
+		FixedLocalizationGatherForExternalActorPackage,
+
 		// -----<new versions can be added above this line>-------------------------------------------------
 		VersionPlusOne,
 		LatestVersion = VersionPlusOne - 1
