Add a label for the Play Options button in the PIE transport controls.

diff --git a/Engine/Source/Editor/UnrealEd/Private/Kismet2/DebuggerCommands.cpp b/Engine/Source/Editor/UnrealEd/Private/Kismet2/DebuggerCommands.cpp
--- a/Engine/Source/Editor/UnrealEd/Private/Kismet2/DebuggerCommands.cpp
+++ b/Engine/Source/Editor/UnrealEd/Private/Kismet2/DebuggerCommands.cpp
@@ -648,7 +648,7 @@ void FPlayWorldCommands::BuildToolbar(FToolMenuSection& InSection, bool bInclude
 	SpecialPIEOptionsMenuAction.CanExecuteAction = FCanExecuteAction::CreateStatic(&FInternalPlayWorldCommandCallbacks::CanShowNonPlayWorldOnlyActions);
 	SpecialPIEOptionsMenuAction.IsActionVisibleDelegate = FIsActionButtonVisible::CreateStatic(&DebuggerCommands::AllowPlayWorldFeature);
 
-	FToolMenuEntry PIEComboEntry = FToolMenuEntry::InitComboButton("PIECombo", SpecialPIEOptionsMenuAction, FOnGetContent::CreateStatic(&GeneratePlayMenuContent, GlobalPlayWorldActions.ToSharedRef()), FText(), LOCTEXT("PIEComboToolTip", "Change Play Mode and Play Settings"));
+	FToolMenuEntry PIEComboEntry = FToolMenuEntry::InitComboButton("PIECombo", SpecialPIEOptionsMenuAction, FOnGetContent::CreateStatic(&GeneratePlayMenuContent, GlobalPlayWorldActions.ToSharedRef()), LOCTEXT("PIEComboLabel","Play Options"), LOCTEXT("PIEComboToolTip", "Change Play Mode and Play Settings"));
 	PIEComboEntry.StyleNameOverride = FName("Toolbar.BackplateRightCombo");
 
 	InSection.AddEntry(PIEComboEntry);
