Added plugin dependencies, and started generating insights profiles, in the moderation artifact generation.

diff --git a/Engine/Source/Editor/UnrealEd/Private/Commandlets/BaseIteratePackagesCommandlet.cpp b/Engine/Source/Editor/UnrealEd/Private/Commandlets/BaseIteratePackagesCommandlet.cpp
--- a/Engine/Source/Editor/UnrealEd/Private/Commandlets/BaseIteratePackagesCommandlet.cpp
+++ b/Engine/Source/Editor/UnrealEd/Private/Commandlets/BaseIteratePackagesCommandlet.cpp
@@ -590,8 +590,10 @@ void UBaseIteratePackagesCommandlet::LoadAndSaveOnePackage(const FString& Filena
 		VerboseMessage(TEXT("Post GetPackageLinker"));
 
 		bool bSavePackage = true;
-		PerformPreloadOperations(Linker, bSavePackage);
-
+		{
+			TRACE_CPUPROFILER_EVENT_SCOPE(UBaseIteratePackagesCommandlet::LoadAndSaveOnePackage::PerformPreloadOperations);
+			PerformPreloadOperations(Linker, bSavePackage);
+		}
 		VerboseMessage(FString::Printf(TEXT("Post PerformPreloadOperations, Resave? %d"), bSavePackage));
 		
 		if (bSavePackage)
@@ -608,7 +610,6 @@ void UBaseIteratePackagesCommandlet::LoadAndSaveOnePackage(const FString& Filena
 			UPackage* Package = nullptr;
 			{
 				TRACE_CPUPROFILER_EVENT_SCOPE(UBaseIteratePackagesCommandlet::LoadAndSaveOnePackage::LoadPackage);
-
 				Package = LoadWorldPackageForEditor(PackagePath.GetPackageName());
 
 				// Package = LoadPackage(NULL, *Filename, 0);
@@ -649,14 +650,18 @@ void UBaseIteratePackagesCommandlet::LoadAndSaveOnePackage(const FString& Filena
 			UWorld* World = UWorld::FindWorldInPackage(Package);
 			if (World)
 			{
+				TRACE_CPUPROFILER_EVENT_SCOPE(UBaseIteratePackagesCommandlet::LoadAndSaveOnePackage::PerformAdditionalOperations::World);
 				PerformAdditionalOperations(World, bSavePackage);
 			}
 			
 			// hook to allow performing additional checks without lumping everything into this one function
-			PerformAdditionalOperations(Package,bSavePackage);
-
+			{
+				TRACE_CPUPROFILER_EVENT_SCOPE(UBaseIteratePackagesCommandlet::LoadAndSaveOnePackage::PerformAdditionalOperations::Package);
+				PerformAdditionalOperations(Package,bSavePackage);
+			}
 			if (bUseWorldPartitionBuilder && World && UWorld::IsPartitionedWorld(World))
 			{
+				TRACE_CPUPROFILER_EVENT_SCOPE(UBaseIteratePackagesCommandlet::LoadAndSaveOnePackage::ProcessWorldPartition);
 				// Load configuration file
 				FString WorldConfigFilename = FPackageName::LongPackageNameToFilename(World->GetPackage()->GetName(), TEXT(".ini"));
 				if (FPlatformFileManager::Get().GetPlatformFile().FileExists(*WorldConfigFilename))
@@ -801,6 +806,7 @@ void UBaseIteratePackagesCommandlet::LoadAndSaveOnePackage(const FString& Filena
 			// to save this package
 			if (bSavePackage == true)
 			{
+				TRACE_CPUPROFILER_EVENT_SCOPE(UBaseIteratePackagesCommandlet::LoadAndSaveOnePackage::SavePackage);
 				if( bIsReadOnly == true && bVerifyContent == true && bAutoCheckOut == false )
 				{
 					UE_LOG(LogIteratePackagesCommandlet, Warning, TEXT("Package [%s] is read-only but needs to be resaved (UE Version: %i, Licensee Version: %i  Current UE Version: %i, Current Licensee Version: %i)"),
