Allow minimal object snapshots

diff --git a/Engine/Source/Editor/UnrealEd/Classes/Editor/Transactor.h b/Engine/Source/Editor/UnrealEd/Classes/Editor/Transactor.h
--- a/Engine/Source/Editor/UnrealEd/Classes/Editor/Transactor.h
+++ b/Engine/Source/Editor/UnrealEd/Classes/Editor/Transactor.h
@@ -128,7 +128,7 @@ protected:
 		FObjectRecord& operator=( const FObjectRecord& ) = delete;
 
 	public:
-		UE::Transaction::FDiffableObject GetDiffableObject(TArrayView<const FProperty*> PropertiesToSerialize = TArrayView<const FProperty*>()) const;
+		UE::Transaction::FDiffableObject GetDiffableObject(TArrayView<const FProperty*> PropertiesToSerialize = TArrayView<const FProperty*>(), UE::Transaction::DiffUtil::EGetDiffableObjectMode ObjectSerializationMode = UE::Transaction::DiffUtil::EGetDiffableObjectMode::SerializeObject) const;
 
 		// Functions.
 		void SerializeContents( FArchive& Ar, int32 InOper ) const;
diff --git a/Engine/Source/Editor/UnrealEd/Private/EditorTransaction.cpp b/Engine/Source/Editor/UnrealEd/Private/EditorTransaction.cpp
--- a/Engine/Source/Editor/UnrealEd/Private/EditorTransaction.cpp
+++ b/Engine/Source/Editor/UnrealEd/Private/EditorTransaction.cpp
@@ -80,18 +80,29 @@ FTransaction::FObjectRecord::FObjectRecord(FTransaction* Owner, UObject* InObjec
 	}
 }
 
-UE::Transaction::FDiffableObject FTransaction::FObjectRecord::GetDiffableObject(TArrayView<const FProperty*> PropertiesToSerialize) const
+UE::Transaction::FDiffableObject FTransaction::FObjectRecord::GetDiffableObject(TArrayView<const FProperty*> PropertiesToSerialize, UE::Transaction::DiffUtil::EGetDiffableObjectMode ObjectSerializationMode) const
 {
 	check(!Array);
+	check(ObjectSerializationMode != UE::Transaction::DiffUtil::EGetDiffableObjectMode::Custom);
 
 	if (UObject* CurrentObject = Object.Get())
 	{
 		UE::Transaction::DiffUtil::FGetDiffableObjectOptions ObjectOptions;
 		ObjectOptions.PropertiesToSerialize = PropertiesToSerialize;
 		ObjectOptions.ObjectSerializationMode = UE::Transaction::DiffUtil::EGetDiffableObjectMode::Custom;
-		ObjectOptions.CustomSerializer = [this](UE::Transaction::FDiffableObjectDataWriter& DiffWriter)
+		ObjectOptions.CustomSerializer = [this, ObjectSerializationMode](UE::Transaction::FDiffableObjectDataWriter& DiffWriter)
 		{
-			SerializeObject(DiffWriter);
+			if (ObjectSerializationMode == UE::Transaction::DiffUtil::EGetDiffableObjectMode::SerializeProperties)
+			{
+				if (UObject* ObjectToSerialize = Object.Get())
+				{
+					ObjectToSerialize->SerializeScriptProperties(DiffWriter);
+				}
+			}
+			else
+			{
+				SerializeObject(DiffWriter);
+			}
 		};
 
 		return UE::Transaction::DiffUtil::GetDiffableObject(CurrentObject, ObjectOptions);
@@ -380,9 +391,7 @@ void FTransaction::FObjectRecord::Snapshot( FTransaction* Owner, UE::Transaction
 		}
 
 		// Serialize the object so we can diff it
-		// Note: although it would be preferable to use SerializeScriptProperties, this cause a false diff between the first snapshot and the base object
-		// since they were serialized with different algo and we don't record enough context to make the comparison appropriately
-		UE::Transaction::FDiffableObject CurrentDiffableObject = GetDiffableObject(AllPropertiesSnapshot);
+		UE::Transaction::FDiffableObject CurrentDiffableObject = GetDiffableObject(AllPropertiesSnapshot, UE::Transaction::DiffUtil::EGetDiffableObjectMode::SerializeProperties);
 
 		// Diff against the correct serialized data depending on whether we already had a snapshot
 		const UE::Transaction::FDiffableObject& InitialDiffableObject = DiffableObjectSnapshot ? *DiffableObjectSnapshot : *DiffableObject;
