[BugFix] Crash when attempting to duplicate a level blueprint
- made sure GC won't destroy OldBPGCAsClass while it's still in scope

diff --git a/Engine/Source/Editor/UnrealEd/Private/Kismet2/BlueprintEditorUtils.cpp b/Engine/Source/Editor/UnrealEd/Private/Kismet2/BlueprintEditorUtils.cpp
--- a/Engine/Source/Editor/UnrealEd/Private/Kismet2/BlueprintEditorUtils.cpp
+++ b/Engine/Source/Editor/UnrealEd/Private/Kismet2/BlueprintEditorUtils.cpp
@@ -1848,6 +1848,31 @@ void FBlueprintEditorUtils::PostDuplicateBlueprint(UBlueprint* Blueprint, bool b
 			UObject* OldCDO = OldBPGC->GetDefaultObject();
 			check(OldCDO != nullptr);
 
+			// Make sure that OldBPGC isn't garbage collected within this scope
+			struct FAddToRootHelper
+			{
+				FAddToRootHelper(UBlueprintGeneratedClass* InBPGC)
+				{
+					BPGC = InBPGC;
+					bWasRoot = BPGC->IsRooted();
+					if(!bWasRoot)
+					{
+						BPGC->AddToRoot();
+					}
+				}
+
+				~FAddToRootHelper()
+				{
+					if (!bWasRoot)
+					{
+						BPGC->RemoveFromRoot();
+					}
+				}
+
+				UBlueprintGeneratedClass* BPGC;
+				bool bWasRoot;
+			} KeepBPGCAlive(OldBPGC);
+
 			if (FBlueprintDuplicationScopeFlags::HasAnyFlag(FBlueprintDuplicationScopeFlags::ValidatePinsUsingSourceClass))
 			{
 				Blueprint->OriginalClass = OldBPGC;
