Fixed ThumbnailManager static exit
- There's a circular dependency between UThumbnailManager and FThumbnailPool where UThumbnailManager owns a FThumbnailPool and FThumbnailPool is registered to a delegate from UThumbnailManager.

diff --git a/Engine/Source/Editor/UnrealEd/Classes/ThumbnailRendering/ThumbnailManager.h b/Engine/Source/Editor/UnrealEd/Classes/ThumbnailRendering/ThumbnailManager.h
--- a/Engine/Source/Editor/UnrealEd/Classes/ThumbnailRendering/ThumbnailManager.h
+++ b/Engine/Source/Editor/UnrealEd/Classes/ThumbnailRendering/ThumbnailManager.h
@@ -211,6 +211,9 @@ public:
 	/** Returns the thumbnail manager and creates it if missing */
 	UNREALED_API static UThumbnailManager& Get();
 
+	/** Returns the thumbnail manager if it exists */
+	UNREALED_API static UThumbnailManager* TryGet();
+
 	/** Writes out a png of what is currently in the specified viewport, scaled appropriately */
 	UNREALED_API static bool CaptureProjectThumbnail(FViewport* Viewport, const FString& OutputFilename, bool bUseSCCIfPossible);
 
diff --git a/Engine/Source/Editor/UnrealEd/Private/AssetThumbnail.cpp b/Engine/Source/Editor/UnrealEd/Private/AssetThumbnail.cpp
--- a/Engine/Source/Editor/UnrealEd/Private/AssetThumbnail.cpp
+++ b/Engine/Source/Editor/UnrealEd/Private/AssetThumbnail.cpp
@@ -938,7 +938,13 @@ FAssetThumbnailPool::FAssetThumbnailPool( uint32 InNumInPool, double InMaxFrameT
 
 FAssetThumbnailPool::~FAssetThumbnailPool()
 {
-	UThumbnailManager::Get().GetOnThumbnailDirtied().RemoveAll(this);
+	UThumbnailManager* ThumbnailManager = UThumbnailManager::TryGet();
+	if ( IsValid( ThumbnailManager ) &&
+		 !ThumbnailManager->HasAnyFlags( RF_BeginDestroyed | RF_FinishDestroyed ) )
+	{
+		ThumbnailManager->GetOnThumbnailDirtied().RemoveAll(this);
+	}
+
 	FCoreUObjectDelegates::OnAssetLoaded.RemoveAll(this);
 
 	// Release all the texture resources
diff --git a/Engine/Source/Editor/UnrealEd/Private/ThumbnailManager.cpp b/Engine/Source/Editor/UnrealEd/Private/ThumbnailManager.cpp
--- a/Engine/Source/Editor/UnrealEd/Private/ThumbnailManager.cpp
+++ b/Engine/Source/Editor/UnrealEd/Private/ThumbnailManager.cpp
@@ -270,6 +270,11 @@ UThumbnailManager& UThumbnailManager::Get()
 	return *ThumbnailManagerSingleton;
 }
 
+UThumbnailManager* UThumbnailManager::TryGet()
+{
+	return ThumbnailManagerSingleton;
+}
+
 void UThumbnailManager::SetupCheckerboardTexture()
 {
 	if (CheckerboardTexture)
