Removed unnecessary config flushes in InternalWriteINI. When you remove a section, it flushes the config and writes it out. So InternalWriteIni would write to disk twice.
Also fixed an issue where InternalReadINI would mark the ini updated even though the file contents didn't change.

Before: 1,593ms
After: 0.7ms

diff --git a/Engine/Source/Editor/UnrealEd/Private/MRUList.cpp b/Engine/Source/Editor/UnrealEd/Private/MRUList.cpp
--- a/Engine/Source/Editor/UnrealEd/Private/MRUList.cpp
+++ b/Engine/Source/Editor/UnrealEd/Private/MRUList.cpp
@@ -143,11 +143,18 @@ void FMRUList::InternalReadINI( TArray<FString>& OutItems, const FString& INISec
 				FString NewItem;
 				if (FPackageName::TryConvertFilenameToLongPackageName(CurItem, NewItem))
 				{
-					CurItem = NewItem;
+					if (NewItem != CurItem)
+					{
+						CurItem = NewItem;
+						bConvertedToNewFormat = true;
+					}
+					
 					OutItems.AddUnique(CurItem);
 				}
-
-				bConvertedToNewFormat = true;
+				else
+				{
+					bConvertedToNewFormat = true;
+				}
 			}
 			else
 			{
@@ -165,14 +172,17 @@ void FMRUList::InternalReadINI( TArray<FString>& OutItems, const FString& INISec
 
 void FMRUList::InternalWriteINI( const TArray<FString>& InItems, const FString& INISection, const FString& INIKeyBase )
 {
-	GConfig->EmptySection( *INISection, GEditorPerProjectIni );
-
-	for ( int32 ItemIdx = 0; ItemIdx < InItems.Num(); ++ItemIdx )
+	if (FConfigFile* ConfigFile = GConfig->Find(GEditorPerProjectIni))
 	{
-		GConfig->SetString( *INISection, *FString::Printf( TEXT("%s%d"), *INIKeyBase, ItemIdx ), *InItems[ ItemIdx ], GEditorPerProjectIni );
-	}
+		ConfigFile->Remove(*INISection);
 
-	GConfig->Flush( false, GEditorPerProjectIni );
+		for (int32 ItemIdx = 0; ItemIdx < InItems.Num(); ++ItemIdx)
+		{
+			ConfigFile->SetString(*INISection, *FString::Printf(TEXT("%s%d"), *INIKeyBase, ItemIdx), *InItems[ItemIdx]);
+		}
+
+		GConfig->Flush(false, GEditorPerProjectIni);
+	}
 }
 
 
