Updated GatherTextFromAssetsCommandlet to use the new FARFilter API. The deprecated API would result in object paths specified in collection filters to fail, resulting in all assets to be gathered.
#rnx

diff --git a/Engine/Source/Editor/UnrealEd/Private/Commandlets/GatherTextFromAssetsCommandlet.cpp b/Engine/Source/Editor/UnrealEd/Private/Commandlets/GatherTextFromAssetsCommandlet.cpp
--- a/Engine/Source/Editor/UnrealEd/Private/Commandlets/GatherTextFromAssetsCommandlet.cpp
+++ b/Engine/Source/Editor/UnrealEd/Private/Commandlets/GatherTextFromAssetsCommandlet.cpp
@@ -519,9 +519,7 @@ int32 UGatherTextFromAssetsCommandlet::Main(const FString& Params)
 			ICollectionManager& CollectionManager = CollectionManagerModule.Get();
 			for (const FString& CollectionName : CollectionFilters)
 			{
-PRAGMA_DISABLE_DEPRECATION_WARNINGS
-				if (!CollectionManager.GetObjectsInCollection(FName(*CollectionName), ECollectionShareType::CST_All, FirstPassFilter.ObjectPaths, ECollectionRecursionFlags::SelfAndChildren))
-PRAGMA_ENABLE_DEPRECATION_WARNINGS
+				if (!CollectionManager.GetObjectsInCollection(FName(*CollectionName), ECollectionShareType::CST_All, FirstPassFilter.SoftObjectPaths, ECollectionRecursionFlags::SelfAndChildren))
 				{
 					UE_LOG(LogGatherTextFromAssetsCommandlet, Error, TEXT("Failed get objects in specified collection: %s"), *CollectionName);
 					HasFailedToGetACollection = true;
@@ -556,6 +554,8 @@ PRAGMA_ENABLE_DEPRECATION_WARNINGS
 		// Apply filter if valid to do so, get all assets otherwise.
 		if (FirstPassFilter.IsEmpty())
 		{
+			// @TODOLocalization: Logging that the first path filter is empty resulting in all assets being gathered can confuse users who generally rely on the second pass.
+			// Figure out a good way to still convey the information in a log or clog.
 			AssetRegistry.GetAllAssets(AssetDataArray);
 		}
 		else
