Fixed non unity/pch compile error in PreferenceStubs.cpp

diff --git a/Engine/Source/Editor/UnrealEd/Private/PreferenceStubs.cpp b/Engine/Source/Editor/UnrealEd/Private/PreferenceStubs.cpp
--- a/Engine/Source/Editor/UnrealEd/Private/PreferenceStubs.cpp
+++ b/Engine/Source/Editor/UnrealEd/Private/PreferenceStubs.cpp
@@ -2,6 +2,7 @@
 
 
 #include "CoreMinimal.h"
+#include "Animation/Skeleton.h"
 #include "Preferences/CascadeOptions.h"
 #include "Preferences/CurveEdOptions.h"
 #include "Preferences/MaterialEditorOptions.h"
