Reinstancer now explicitly rescues objects that it has created that are renamed by subsequent reinstance operations, also prevents CPFUO created objects from taking precedence over reinstancer created objects

diff --git a/Engine/Source/Editor/UnrealEd/Private/Kismet2/KismetReinstanceUtilities.cpp b/Engine/Source/Editor/UnrealEd/Private/Kismet2/KismetReinstanceUtilities.cpp
--- a/Engine/Source/Editor/UnrealEd/Private/Kismet2/KismetReinstanceUtilities.cpp
+++ b/Engine/Source/Editor/UnrealEd/Private/Kismet2/KismetReinstanceUtilities.cpp
@@ -966,6 +966,56 @@ protected:
 	const TMap<FFieldVariant, FFieldVariant>& ReplacementMap;
 };
 
+void PreserveNewObjectsInOldOuters(TMap<UObject*, UObject*>& InOutOldToNewInstanceMap)
+{
+	// The creation of 'new' objects is not ordered in any way, so make sure that we didn't orphan any 
+	// new objects in the old outer hierarchy, e.g. when reinstancing an object before reinstancing 
+	// its outer objects. Just go through each new object, and walk its outer chain looking for a new
+	// version. If we find a new version of one of the outers, rename ourself under that new object:
+	UPackage* TransientPackage = GetTransientPackage();
+
+	TArray<UObject*> NewObjects;
+	InOutOldToNewInstanceMap.GenerateValueArray(NewObjects);
+
+	for (UObject* NewObject : NewObjects)
+	{
+		UObject* LastInner = NewObject;
+		UObject* OuterIter = NewObject->GetOuter();
+		while (OuterIter)
+		{
+			UObject** NewOuter = InOutOldToNewInstanceMap.Find(OuterIter);
+			if (NewOuter && *NewOuter)
+			{
+				UObject* NewOuterResolved = *NewOuter;
+				// find object with same name as NewInnerInOldOuter:
+				if (UObject* Collision = StaticFindObjectFast(UObject::StaticClass(),
+					NewOuterResolved,
+					LastInner->GetFName()))
+				{
+					InOutOldToNewInstanceMap.Add(Collision, LastInner);
+					Collision->Rename(
+						nullptr,
+						TransientPackage,
+						REN_DoNotDirty | REN_DontCreateRedirectors | REN_ForceNoResetLoaders | REN_NonTransactional
+					);
+				}
+				// rename, now that we've reserved the name
+				LastInner->Rename(
+					*LastInner->GetName(),
+					NewOuterResolved,
+					REN_DoNotDirty | REN_DontCreateRedirectors | REN_ForceNoResetLoaders | REN_NonTransactional
+				);
+				break;
+			}
+			else
+			{
+				LastInner = OuterIter;
+				OuterIter = OuterIter->GetOuter();
+			}
+		}
+	}
+}
+
 void FBlueprintCompileReinstancer::UpdateBytecodeReferences(
 	TSet<UBlueprint*>& OutDependentBlueprints,
 	TMap<FFieldVariant, FFieldVariant>& OutFieldMapping)
@@ -2137,6 +2187,35 @@ static void ReplaceActorHelper(AActor* OldActor, UClass* OldClass, UObject*& New
 	OldToNewInstanceMap.Add(OldActor, NewActor);
 }
 
+void AddCPFUOCreatedObjectsToOldToNew(const TMap<UObject*, UObject*>& ReplacedObjectsFromCPFUO, TMap<UObject*, UObject*>& OldToNewInstanceMap)
+{
+	// ReplacedObjectsFromCPFUO contains the list of objects CPFUO created, lets merge that into the 
+	// final OldToNewInstanceMap - but give precedence to mappings already created in OldToNewInstanceMap
+	for (TPair<UObject*, UObject*> OldToNewObject : ReplacedObjectsFromCPFUO)
+	{
+		UObject* OldObject = OldToNewObject.Key;
+		UObject* NewObject = OldToNewObject.Value;
+		// First check to see if we already have a replacement object created by the reinstancer
+		// then, include CPFUO provided mappings only if they are not null. It's possible
+		// that we may want to allow CPFUO to map *some* objects to null (but I don't believe so)
+		// however, we certainly can't allow it to null out references to objects created by the 
+		// reinstancer
+		if (UObject** ExistingNewObject = OldToNewInstanceMap.Find(OldObject))
+		{
+			// we already have an entry, remap the reinstancer's NewObject to the existing
+			// new object:
+			if (*ExistingNewObject != NewObject)
+			{
+				OldToNewInstanceMap.FindOrAdd(NewObject) = *ExistingNewObject;
+			}
+		}
+		else if (NewObject)
+		{
+			OldToNewInstanceMap.FindOrAdd(OldObject) = NewObject;
+		}
+	}
+}
+
 void FBlueprintCompileReinstancer::ReplaceInstancesOfClass_Inner(const TMap<UClass*, UClass*>& InOldToNewClassMap, const FReplaceInstancesOfClassParameters& Params)
 {
 	// If there is an original CDO, make sure we are only reinstancing a single class (legacy path, non-batch)
@@ -2404,10 +2483,14 @@ void FBlueprintCompileReinstancer::ReplaceInstancesOfClass_Inner(const TMap<UCla
 				}
 			}
 		}
+
+		PreserveNewObjectsInOldOuters(OldToNewInstanceMap);
 	}
 
 	FCoreUObjectDelegates::OnObjectsReplaced.Remove(OnObjectsReplacedHandle);
 
+	AddCPFUOCreatedObjectsToOldToNew(ObjectRemappingHelper.ReplacedObjects, OldToNewInstanceMap);
+
 	// Now replace any pointers to the old archetypes/instances with pointers to the new one
 	TArray<UObject*> SourceObjects;
 	OldToNewInstanceMap.GenerateKeyArray(SourceObjects);
@@ -2453,8 +2536,6 @@ void FBlueprintCompileReinstancer::ReplaceInstancesOfClass_Inner(const TMap<UCla
 		{
 			ReplacementActor.ApplyAttachments(ObjectRemappingHelper.ReplacedObjects, ObjectsThatShouldUseOldStuff, ObjectsReplaced, ReinstancedObjectsWeakReferenceMap);
 		}
-
-		OldToNewInstanceMap.Append(ObjectRemappingHelper.ReplacedObjects);
 	}
 
 	if(bReplaceReferencesToOldClasses)
