More clean up & sort #include paths in various private cpp files using work in progress FixIncludePath

#rnx

diff --git a/Engine/Plugins/Editor/DisplayClusterLaunch/Source/DisplayClusterLaunchEditor/Private/DisplayClusterLaunchEditorModule.cpp b/Engine/Plugins/Editor/DisplayClusterLaunch/Source/DisplayClusterLaunchEditor/Private/DisplayClusterLaunchEditorModule.cpp
--- a/Engine/Plugins/Editor/DisplayClusterLaunch/Source/DisplayClusterLaunchEditor/Private/DisplayClusterLaunchEditorModule.cpp
+++ b/Engine/Plugins/Editor/DisplayClusterLaunch/Source/DisplayClusterLaunchEditor/Private/DisplayClusterLaunchEditorModule.cpp
@@ -31,10 +31,10 @@
 #include "LevelEditor.h"
 #include "Misc/ConfigCacheIni.h"
  
-#include "ToolMenus.h"
 #include "ActorFactories/ActorFactoryBlueprint.h"
-#include "ConsoleVariablesEditorRuntime/Public/ConsoleVariablesAsset.h"
+#include "ConsoleVariablesAsset.h"
 #include "Interfaces/IPluginManager.h"
+#include "ToolMenus.h"
 
 #define LOCTEXT_NAMESPACE "FDisplayClusterLaunchEditorModule"
 
diff --git a/Engine/Plugins/Editor/WaveformEditor/Source/WaveformEditorWidgets/Private/WaveformEditorRenderData.cpp b/Engine/Plugins/Editor/WaveformEditor/Source/WaveformEditorWidgets/Private/WaveformEditorRenderData.cpp
--- a/Engine/Plugins/Editor/WaveformEditor/Source/WaveformEditorWidgets/Private/WaveformEditorRenderData.cpp
+++ b/Engine/Plugins/Editor/WaveformEditor/Source/WaveformEditorWidgets/Private/WaveformEditorRenderData.cpp
@@ -2,9 +2,9 @@
 
 #include "WaveformEditorRenderData.h"
 
-#include "AudioExtensions/Public/IWaveformTransformation.h"
 #include "DSP/BufferVectorOperations.h"
 #include "DSP/FloatArrayMath.h"
+#include "IWaveformTransformation.h"
 #include "Sound/SoundWave.h"
 
 TArrayView<const int16> FWaveformEditorRenderData::GetSampleData() const
@@ -77,4 +77,4 @@ void FWaveformEditorRenderData::UpdateRenderData(const uint8* InRawData, const u
 	{
 		OnRenderDataUpdated.Broadcast();
 	}
-}
\ No newline at end of file
+}
diff --git a/Engine/Plugins/Enterprise/DataprepEditor/Source/DataprepEditor/Private/DataprepEditorUtils.cpp b/Engine/Plugins/Enterprise/DataprepEditor/Source/DataprepEditor/Private/DataprepEditorUtils.cpp
--- a/Engine/Plugins/Enterprise/DataprepEditor/Source/DataprepEditor/Private/DataprepEditorUtils.cpp
+++ b/Engine/Plugins/Enterprise/DataprepEditor/Source/DataprepEditor/Private/DataprepEditorUtils.cpp
@@ -14,29 +14,29 @@
 #include "Widgets/Parameterization/SDataprepLinkToParameter.h"
 
 #include "AssetToolsModule.h"
-#include "Blutility/Classes/EditorUtilityBlueprint.h"
-#include "Blutility/Classes/EditorUtilityBlueprintFactory.h"
-#include "BlueprintGraph/Classes/K2Node_Event.h"
-#include "BlueprintGraph/Classes/K2Node_BreakStruct.h"
 #include "Dialogs/DlgPickPath.h"
 #include "EdGraph/EdGraph.h"
 #include "EdGraphSchema_K2.h"
 #include "EdGraphSchema_K2_Actions.h"
-#include "Styling/AppStyle.h"
+#include "EditorUtilityBlueprint.h"
+#include "EditorUtilityBlueprintFactory.h"
 #include "Engine/Blueprint.h"
-#include "EngineUtils.h"
 #include "Engine/StaticMesh.h"
+#include "EngineUtils.h"
 #include "Framework/MultiBox/MultiBoxBuilder.h"
 #include "Interfaces/IMainFrameModule.h"
+#include "K2Node_BreakStruct.h"
+#include "K2Node_Event.h"
 #include "Kismet2/BlueprintEditorUtils.h"
 #include "Kismet2/KismetEditorUtilities.h"
-#include "ScopedTransaction.h"
 #include "SchemaActions/DataprepMenuActionCollectorUtils.h"
+#include "ScopedTransaction.h"
 #include "SelectionSystem/DataprepFilter.h"
-#include "Widgets/Text/STextBlock.h"
-#include "Widgets/Input/STextComboBox.h"
-#include "Widgets/Input/SEditableTextBox.h"
+#include "Styling/AppStyle.h"
 #include "Widgets/Input/SButton.h"
+#include "Widgets/Input/SEditableTextBox.h"
+#include "Widgets/Input/STextComboBox.h"
+#include "Widgets/Text/STextBlock.h"
 
 #define LOCTEXT_NAMESPACE "DataprepEditorUtils"
 
diff --git a/Engine/Plugins/Enterprise/DatasmithFBXImporter/Source/DatasmithVREDTranslator/Private/DatasmithVREDImportData.h b/Engine/Plugins/Enterprise/DatasmithFBXImporter/Source/DatasmithVREDTranslator/Private/DatasmithVREDImportData.h
--- a/Engine/Plugins/Enterprise/DatasmithFBXImporter/Source/DatasmithVREDTranslator/Private/DatasmithVREDImportData.h
+++ b/Engine/Plugins/Enterprise/DatasmithFBXImporter/Source/DatasmithVREDTranslator/Private/DatasmithVREDImportData.h
@@ -2,7 +2,7 @@
 #pragma once
 
 #include "Curves/CurveVector.h"
-#include "Engine/Classes/EditorFramework/AssetImportData.h"
+#include "EditorFramework/AssetImportData.h"
 #include "Engine/DataTable.h"
 
 #include "DatasmithVREDImportData.generated.h"
diff --git a/Engine/Plugins/Enterprise/LidarPointCloud/Source/LidarPointCloudEditor/Private/LidarPointCloudEditorViewport.h b/Engine/Plugins/Enterprise/LidarPointCloud/Source/LidarPointCloudEditor/Private/LidarPointCloudEditorViewport.h
--- a/Engine/Plugins/Enterprise/LidarPointCloud/Source/LidarPointCloudEditor/Private/LidarPointCloudEditorViewport.h
+++ b/Engine/Plugins/Enterprise/LidarPointCloud/Source/LidarPointCloudEditor/Private/LidarPointCloudEditorViewport.h
@@ -2,12 +2,12 @@
 
 #pragma once
 
+#include "AdvancedPreviewScene.h"
 #include "CoreMinimal.h"
-#include "UObject/GCObject.h"
 #include "EditorViewportClient.h"
-#include "AdvancedPreviewScene.h"
+#include "SCommonEditorViewportToolbarBase.h"
 #include "SEditorViewport.h"
-#include "Editor/UnrealEd/Public/SCommonEditorViewportToolbarBase.h"
+#include "UObject/GCObject.h"
 #include "Widgets/Layout/SBorder.h"
 
 class FLidarPointCloudEditor;
diff --git a/Engine/Plugins/Experimental/ActorPalette/Source/ActorPalette/Private/ActorPaletteViewport.h b/Engine/Plugins/Experimental/ActorPalette/Source/ActorPalette/Private/ActorPaletteViewport.h
--- a/Engine/Plugins/Experimental/ActorPalette/Source/ActorPalette/Private/ActorPaletteViewport.h
+++ b/Engine/Plugins/Experimental/ActorPalette/Source/ActorPalette/Private/ActorPaletteViewport.h
@@ -3,10 +3,10 @@
 #pragma once
 
 #include "CoreMinimal.h"
-#include "Widgets/DeclarativeSyntaxSupport.h"
 #include "EditorViewportClient.h"
+#include "SCommonEditorViewportToolbarBase.h"
 #include "SEditorViewport.h"
-#include "Editor/UnrealEd/Public/SCommonEditorViewportToolbarBase.h"
+#include "Widgets/DeclarativeSyntaxSupport.h"
 
 class FActorPaletteViewportClient;
 class SActorPaletteViewportToolbar;
diff --git a/Engine/Plugins/Experimental/Animation/PoseSearch/Source/Editor/Private/PoseSearchDebugger.cpp b/Engine/Plugins/Experimental/Animation/PoseSearch/Source/Editor/Private/PoseSearchDebugger.cpp
--- a/Engine/Plugins/Experimental/Animation/PoseSearch/Source/Editor/Private/PoseSearchDebugger.cpp
+++ b/Engine/Plugins/Experimental/Animation/PoseSearch/Source/Editor/Private/PoseSearchDebugger.cpp
@@ -1,36 +1,36 @@
 ﻿// Copyright Epic Games, Inc. All Rights Reserved.
 
 #include "PoseSearchDebugger.h"
+#include "Algo/AllOf.h"
+#include "Animation/AnimRootMotionProvider.h"
+#include "Animation/AnimSequence.h"
+#include "Animation/AnimationPoseData.h"
+#include "Animation/BlendSpace.h"
+#include "Animation/MirrorDataTable.h"
+#include "Editor.h"
 #include "Engine/SkeletalMesh.h"
 #include "Engine/SkinnedAsset.h"
 #include "IAnimationProvider.h"
 #include "IGameplayProvider.h"
+#include "IRewindDebugger.h"
+#include "ObjectTrace.h"
 #include "PoseSearch/PoseSearch.h"
-#include "Animation/AnimationPoseData.h"
-#include "Animation/AnimSequence.h"
-#include "Animation/BlendSpace.h"
-#include "Animation/MirrorDataTable.h"
-#include "Animation/AnimRootMotionProvider.h"
 #include "Preferences/PersonaOptions.h"
 #include "PropertyEditorModule.h"
-#include "RewindDebuggerInterface/Public/IRewindDebugger.h"
-#include "ObjectTrace.h"
-#include "TraceServices/Model/AnalysisSession.h"
+#include "Styling/AppStyle.h"
+#include "Styling/SlateIconFinder.h"
+#include "Subsystems/AssetEditorSubsystem.h"
 #include "Trace/PoseSearchTraceProvider.h"
-#include "Widgets/Layout/SScrollBox.h"
-#include "Widgets/Layout/SWidgetSwitcher.h"
-#include "Widgets/SBoxPanel.h"
+#include "TraceServices/Model/AnalysisSession.h"
+#include "TraceServices/Model/Frames.h"
+#include "Widgets/Images/SImage.h"
 #include "Widgets/Input/SButton.h"
 #include "Widgets/Input/SHyperlink.h"
 #include "Widgets/Input/SNumericEntryBox.h"
 #include "Widgets/Input/SSearchBox.h"
-#include "Widgets/Images/SImage.h"
-#include "Styling/SlateIconFinder.h"
-#include "TraceServices/Model/Frames.h"
-#include "Subsystems/AssetEditorSubsystem.h"
-#include "Editor.h"
-#include "Algo/AllOf.h"
-#include "Styling/AppStyle.h"
+#include "Widgets/Layout/SScrollBox.h"
+#include "Widgets/Layout/SWidgetSwitcher.h"
+#include "Widgets/SBoxPanel.h"
 #include "Widgets/Views/SListView.h"
 
 #define LOCTEXT_NAMESPACE "PoseSearchDebugger"
@@ -3002,4 +3002,4 @@ FName FDebuggerViewCreator::GetName() const
 
 } // namespace UE::PoseSearch
 
-#undef LOCTEXT_NAMESPACE
\ No newline at end of file
+#undef LOCTEXT_NAMESPACE
diff --git a/Engine/Plugins/Experimental/Animation/PoseSearch/Source/Editor/Private/PoseSearchDebugger.h b/Engine/Plugins/Experimental/Animation/PoseSearch/Source/Editor/Private/PoseSearchDebugger.h
--- a/Engine/Plugins/Experimental/Animation/PoseSearch/Source/Editor/Private/PoseSearchDebugger.h
+++ b/Engine/Plugins/Experimental/Animation/PoseSearch/Source/Editor/Private/PoseSearchDebugger.h
@@ -2,14 +2,14 @@
 
 #pragma once
 
-#include "CoreMinimal.h"
 #include "Components/PoseableMeshComponent.h"
+#include "CoreMinimal.h"
 #include "CustomBoneIndexArray.h"
-#include "RewindDebuggerInterface/Public/IRewindDebuggerExtension.h"
-#include "RewindDebuggerInterface/Public/IRewindDebuggerView.h"
-#include "RewindDebuggerInterface/Public/IRewindDebuggerViewCreator.h"
-#include "Widgets/Views/SHeaderRow.h"
+#include "IRewindDebuggerExtension.h"
+#include "IRewindDebuggerView.h"
+#include "IRewindDebuggerViewCreator.h"
 #include "PoseSearch/PoseSearch.h"
+#include "Widgets/Views/SHeaderRow.h"
 #include "PoseSearchDebugger.generated.h"
 
 struct FPoseSearchDatabaseSequence;
diff --git a/Engine/Plugins/Experimental/Chooser/Source/ChooserEditor/Private/ChooserTableEditor.h b/Engine/Plugins/Experimental/Chooser/Source/ChooserEditor/Private/ChooserTableEditor.h
--- a/Engine/Plugins/Experimental/Chooser/Source/ChooserEditor/Private/ChooserTableEditor.h
+++ b/Engine/Plugins/Experimental/Chooser/Source/ChooserEditor/Private/ChooserTableEditor.h
@@ -5,7 +5,7 @@
 #include "CoreMinimal.h"
 #include "Toolkits/IToolkitHost.h"
 #include "Toolkits/AssetEditorToolkit.h"
-#include "Editor/PropertyEditor/Public/PropertyEditorDelegates.h"
+#include "PropertyEditorDelegates.h"
 #include "Widgets/Views/SListView.h"
 #include "Widgets/Views/STableViewBase.h"
 #include "Chooser.h"
diff --git a/Engine/Plugins/Experimental/PCG/Source/PCGEditor/Private/SPCGEditorGraphNodePalette.cpp b/Engine/Plugins/Experimental/PCG/Source/PCGEditor/Private/SPCGEditorGraphNodePalette.cpp
--- a/Engine/Plugins/Experimental/PCG/Source/PCGEditor/Private/SPCGEditorGraphNodePalette.cpp
+++ b/Engine/Plugins/Experimental/PCG/Source/PCGEditor/Private/SPCGEditorGraphNodePalette.cpp
@@ -8,8 +8,8 @@
 #include "PCGGraph.h"
 
 #include "AssetRegistry/AssetRegistryModule.h"
-#include "EditorWidgets/Public/SEnumCombo.h"
 #include "Modules/ModuleManager.h"
+#include "SEnumCombo.h"
 #include "SGraphActionMenu.h"
 
 #define LOCTEXT_NAMESPACE "SPCGEditorGraphNodePalette"
diff --git a/Engine/Plugins/Experimental/VirtualCamera/Source/VirtualCamera/Private/VCamBlueprintFunctionLibrary.cpp b/Engine/Plugins/Experimental/VirtualCamera/Source/VirtualCamera/Private/VCamBlueprintFunctionLibrary.cpp
--- a/Engine/Plugins/Experimental/VirtualCamera/Source/VirtualCamera/Private/VCamBlueprintFunctionLibrary.cpp
+++ b/Engine/Plugins/Experimental/VirtualCamera/Source/VirtualCamera/Private/VCamBlueprintFunctionLibrary.cpp
@@ -19,17 +19,17 @@
 #include "EditorAssetLibrary.h"
 #include "ITakeRecorderModule.h"
 #include "LevelEditor.h"
+#include "LevelEditorSubsystem.h"
 #include "LevelEditorViewport.h"
-#include "SceneView.h"
-#include "SLevelViewport.h"
 #include "LevelSequenceEditorBlueprintLibrary.h"
-#include "TakePreset.h"
-#include "VPUtilitiesEditor/Public/VPUtilitiesEditorBlueprintLibrary.h"
-#include "Subsystems/UnrealEditorSubsystem.h"
-#include "LevelEditor/Public/LevelEditorSubsystem.h"
 #include "Modules/ModuleManager.h"
 #include "Recorder/TakeRecorderBlueprintLibrary.h"
 #include "Recorder/TakeRecorderPanel.h"
+#include "SLevelViewport.h"
+#include "SceneView.h"
+#include "Subsystems/UnrealEditorSubsystem.h"
+#include "TakePreset.h"
+#include "VPUtilitiesEditorBlueprintLibrary.h"
 #endif
 
 bool UVCamBlueprintFunctionLibrary::IsGameRunning()
diff --git a/Engine/Plugins/Runtime/Steam/SteamVR/Source/SteamVR/Private/SteamVRPrivate.h b/Engine/Plugins/Runtime/Steam/SteamVR/Source/SteamVR/Private/SteamVRPrivate.h
--- a/Engine/Plugins/Runtime/Steam/SteamVR/Source/SteamVR/Private/SteamVRPrivate.h
+++ b/Engine/Plugins/Runtime/Steam/SteamVR/Source/SteamVR/Private/SteamVRPrivate.h
@@ -10,7 +10,7 @@
 #include "Engine/Engine.h"
 
 #include "IHeadMountedDisplay.h"
-#include "Runtime/Engine/Public/ScreenRendering.h"
+#include "ScreenRendering.h"
 
 #if PLATFORM_WINDOWS
 #include "Windows/WindowsHWrapper.h"
diff --git a/Engine/Plugins/VirtualProduction/DMX/DMXPixelMapping/Source/DMXPixelMappingEditor/Private/Views/SDMXPixelMappingHierarchyView.cpp b/Engine/Plugins/VirtualProduction/DMX/DMXPixelMapping/Source/DMXPixelMappingEditor/Private/Views/SDMXPixelMappingHierarchyView.cpp
--- a/Engine/Plugins/VirtualProduction/DMX/DMXPixelMapping/Source/DMXPixelMappingEditor/Private/Views/SDMXPixelMappingHierarchyView.cpp
+++ b/Engine/Plugins/VirtualProduction/DMX/DMXPixelMapping/Source/DMXPixelMappingEditor/Private/Views/SDMXPixelMappingHierarchyView.cpp
@@ -9,20 +9,20 @@
 #include "ViewModels/DMXPixelMappingHierarchyViewModel.h"
 #include "Widgets/SDMXPixelMappingHierarchyItem.h"
 
+#include "Components/DMXPixelMappingOutputComponent.h"
+#include "Components/DMXPixelMappingRendererComponent.h"
+#include "Components/DMXPixelMappingRootComponent.h"
+#include "Exporters/Exporter.h"
+#include "Factories.h"
 #include "Framework/Commands/GenericCommands.h"
-#include "Widgets/Layout/SBorder.h"
-#include "Framework/Views/TreeFilterHandler.h"
-#include "Widgets/Layout/SScrollBorder.h"
-#include "Misc/IFilter.h"
 #include "Framework/MultiBox/MultiBoxBuilder.h"
-#include "UnrealExporter.h"
-#include "Exporters/Exporter.h"
+#include "Framework/Views/TreeFilterHandler.h"
 #include "HAL/PlatformApplicationMisc.h"
-#include "Factories.h"
+#include "Misc/IFilter.h"
 #include "ScopedTransaction.h"
-#include "DMXPixelMappingRuntime/Public/Components/DMXPixelMappingOutputComponent.h"
-#include "DMXPixelMappingRuntime/Public/Components/DMXPixelMappingRendererComponent.h"
-#include "DMXPixelMappingRuntime/Public/Components/DMXPixelMappingRootComponent.h"
+#include "UnrealExporter.h"
+#include "Widgets/Layout/SBorder.h"
+#include "Widgets/Layout/SScrollBorder.h"
 
 #define LOCTEXT_NAMESPACE "SDMXPixelMappingHierarchyView"
 
diff --git a/Engine/Plugins/VirtualProduction/LiveLinkCamera/Source/LiveLinkCameraEditor/Private/LiveLinkCameraControllerCustomization.cpp b/Engine/Plugins/VirtualProduction/LiveLinkCamera/Source/LiveLinkCameraEditor/Private/LiveLinkCameraControllerCustomization.cpp
--- a/Engine/Plugins/VirtualProduction/LiveLinkCamera/Source/LiveLinkCameraEditor/Private/LiveLinkCameraControllerCustomization.cpp
+++ b/Engine/Plugins/VirtualProduction/LiveLinkCamera/Source/LiveLinkCameraEditor/Private/LiveLinkCameraControllerCustomization.cpp
@@ -5,12 +5,12 @@
 #include "DetailCategoryBuilder.h"
 #include "DetailLayoutBuilder.h"
 #include "DetailWidgetRow.h"
-#include "Styling/AppStyle.h"
 #include "IDetailPropertyRow.h"
 #include "LiveLinkCameraController.h"
-#include "LiveLinkComponents/Public/LiveLinkComponentController.h"
+#include "LiveLinkComponentController.h"
 #include "Modules/ModuleManager.h"
 #include "Roles/LiveLinkCameraRole.h"
+#include "Styling/AppStyle.h"
 #include "Styling/CoreStyle.h"
 #include "Widgets/Images/SImage.h"
 #include "Widgets/SBoxPanel.h"
diff --git a/Engine/Plugins/VirtualProduction/MultiUserTakes/Source/ConcertTakeRecorder/Private/ConcertTakeRecorderManager.h b/Engine/Plugins/VirtualProduction/MultiUserTakes/Source/ConcertTakeRecorder/Private/ConcertTakeRecorderManager.h
--- a/Engine/Plugins/VirtualProduction/MultiUserTakes/Source/ConcertTakeRecorder/Private/ConcertTakeRecorderManager.h
+++ b/Engine/Plugins/VirtualProduction/MultiUserTakes/Source/ConcertTakeRecorder/Private/ConcertTakeRecorderManager.h
@@ -2,14 +2,14 @@
 
 #pragma once
 
-#include "ConcertTakeRecorderMessages.h"
 #include "ConcertMessages.h"
+#include "ConcertTakeRecorderMessages.h"
 #include "Delegates/IDelegateInstance.h"
 #include "HAL/Platform.h"
+#include "IConcertClientTransactionBridge.h"
 #include "Misc/FrameNumber.h"
-#include "UObject/StrongObjectPtr.h"
 #include "TakePreset.h"
-#include "ConcertSyncClient/Public/IConcertClientTransactionBridge.h"
+#include "UObject/StrongObjectPtr.h"
 
 #include "ConcertTakeRecorderClientSessionCustomization.h"
 
diff --git a/Engine/Plugins/VirtualProduction/RemoteControl/Source/RemoteControlUI/Private/UI/Filters/SRCFilterBar.h b/Engine/Plugins/VirtualProduction/RemoteControl/Source/RemoteControlUI/Private/UI/Filters/SRCFilterBar.h
--- a/Engine/Plugins/VirtualProduction/RemoteControl/Source/RemoteControlUI/Private/UI/Filters/SRCFilterBar.h
+++ b/Engine/Plugins/VirtualProduction/RemoteControl/Source/RemoteControlUI/Private/UI/Filters/SRCFilterBar.h
@@ -4,8 +4,8 @@
 
 #include "CoreMinimal.h"
 
-#include "AssetTools/Public/AssetTypeCategories.h"
-#include "AssetTools/Public/AssetToolsModule.h"
+#include "AssetToolsModule.h"
+#include "AssetTypeCategories.h"
 
 #include "Filters/SBasicFilterBar.h"
 
diff --git a/Engine/Plugins/VirtualProduction/RemoteControl/Source/RemoteControlUI/Private/UI/SRCPanelExposedEntity.cpp b/Engine/Plugins/VirtualProduction/RemoteControl/Source/RemoteControlUI/Private/UI/SRCPanelExposedEntity.cpp
--- a/Engine/Plugins/VirtualProduction/RemoteControl/Source/RemoteControlUI/Private/UI/SRCPanelExposedEntity.cpp
+++ b/Engine/Plugins/VirtualProduction/RemoteControl/Source/RemoteControlUI/Private/UI/SRCPanelExposedEntity.cpp
@@ -3,23 +3,24 @@
 #include "SRCPanelExposedEntity.h"
 
 #include "ActorTreeItem.h"
-#include "Components/BillboardComponent.h"
 #include "Commands/RemoteControlCommands.h"
+#include "Components/ActorComponent.h"
+#include "Components/BillboardComponent.h"
 #include "Editor.h"
 #include "EditorFontGlyphs.h"
-#include "EngineUtils.h"
 #include "Engine/Selection.h"
-#include "Engine/Classes/Components/ActorComponent.h"
+#include "EngineUtils.h"
 #include "Framework/Application/SlateApplication.h"
 #include "Framework/MultiBox/MultiBoxBuilder.h"
 #include "GameFramework/Actor.h"
 #include "Interfaces/IMainFrameModule.h"
+#include "Modules/ModuleManager.h"
 #include "RemoteControlBinding.h"
 #include "RemoteControlEntity.h"
-#include "RemoteControlPreset.h"
 #include "RemoteControlField.h"
-#include "RemoteControlSettings.h"
 #include "RemoteControlPanelStyle.h"
+#include "RemoteControlPreset.h"
+#include "RemoteControlSettings.h"
 #include "SRCPanelDragHandle.h"
 #include "SceneOutlinerFilters.h"
 #include "SceneOutlinerModule.h"
@@ -27,13 +28,12 @@
 #include "Styling/CoreStyle.h"
 #include "Styling/RemoteControlStyles.h"
 #include "Styling/SlateIconFinder.h"
-#include "Modules/ModuleManager.h"
 #include "UObject/UObjectIterator.h"
 #include "Widgets/Input/SButton.h"
-#include "Widgets/Input/SComboButton.h"
 #include "Widgets/Input/SCheckBox.h"
-#include "Widgets/Layout/SBox.h"
+#include "Widgets/Input/SComboButton.h"
 #include "Widgets/Layout/SBorder.h"
+#include "Widgets/Layout/SBox.h"
 #include "Widgets/Notifications/SPopUpErrorText.h"
 #include "Widgets/SBoxPanel.h"
 #include "Widgets/Text/STextBlock.h"
diff --git a/Engine/Plugins/VirtualProduction/StageMonitoring/Source/StageDataProvider/Private/StageDataProvider.cpp b/Engine/Plugins/VirtualProduction/StageMonitoring/Source/StageDataProvider/Private/StageDataProvider.cpp
--- a/Engine/Plugins/VirtualProduction/StageMonitoring/Source/StageDataProvider/Private/StageDataProvider.cpp
+++ b/Engine/Plugins/VirtualProduction/StageMonitoring/Source/StageDataProvider/Private/StageDataProvider.cpp
@@ -9,10 +9,10 @@
 #include "Misc/App.h"
 #include "Misc/CoreDelegates.h"
 #include "StageDataProviderModule.h"
-#include "StageMonitoringSettings.h"
 #include "StageMonitorUtils.h"
+#include "StageMonitoringSettings.h"
+#include "VPRolesSubsystem.h"
 #include "VPSettings.h"
-#include "VPRoles/Public/VPRolesSubsystem.h"
 
 
 namespace StageDataProviderUtils
diff --git a/Engine/Plugins/VirtualProduction/StageMonitoring/Source/StageDataProvider/Private/StageDataProviderModule.cpp b/Engine/Plugins/VirtualProduction/StageMonitoring/Source/StageDataProvider/Private/StageDataProviderModule.cpp
--- a/Engine/Plugins/VirtualProduction/StageMonitoring/Source/StageDataProvider/Private/StageDataProviderModule.cpp
+++ b/Engine/Plugins/VirtualProduction/StageMonitoring/Source/StageDataProvider/Private/StageDataProviderModule.cpp
@@ -12,8 +12,8 @@
 #include "StageMonitoringSettings.h"
 #include "TakeRecorderStateProvider.h"
 #include "TimecodeProviderWatchdog.h"
+#include "VPRolesSubsystem.h"
 #include "VPSettings.h"
-#include "VPRoles/Public/VPRolesSubsystem.h"
 
 
 const FName IStageDataProviderModule::ModuleName = TEXT("StageDataProvider");
diff --git a/Engine/Plugins/VirtualProduction/StageMonitoring/Source/StageMonitor/Private/StageMonitorModule.cpp b/Engine/Plugins/VirtualProduction/StageMonitoring/Source/StageMonitor/Private/StageMonitorModule.cpp
--- a/Engine/Plugins/VirtualProduction/StageMonitoring/Source/StageMonitor/Private/StageMonitorModule.cpp
+++ b/Engine/Plugins/VirtualProduction/StageMonitoring/Source/StageMonitor/Private/StageMonitorModule.cpp
@@ -6,10 +6,10 @@
 #include "HAL/IConsoleManager.h"
 #include "Misc/CoreDelegates.h"
 #include "StageMonitor.h"
-#include "StageMonitoringSettings.h"
 #include "StageMonitorSessionManager.h"
+#include "StageMonitoringSettings.h"
+#include "VPRolesSubsystem.h"
 #include "VPSettings.h"
-#include "VPRoles/Public/VPRolesSubsystem.h"
 
 const FName IStageMonitorModule::ModuleName = TEXT("StageMonitor");
 
diff --git a/Engine/Plugins/VirtualProduction/StageMonitoring/Source/StageMonitorCommon/Private/StageMonitorUtils.cpp b/Engine/Plugins/VirtualProduction/StageMonitoring/Source/StageMonitorCommon/Private/StageMonitorUtils.cpp
--- a/Engine/Plugins/VirtualProduction/StageMonitoring/Source/StageMonitorCommon/Private/StageMonitorUtils.cpp
+++ b/Engine/Plugins/VirtualProduction/StageMonitoring/Source/StageMonitorCommon/Private/StageMonitorUtils.cpp
@@ -4,8 +4,8 @@
 
 #include "Engine/Engine.h"
 #include "StageMonitoringSettings.h"
+#include "VPRolesSubsystem.h"
 #include "VPSettings.h"
-#include "VPRoles/Public/VPRolesSubsystem.h"
 
 
 namespace StageMonitorUtils
diff --git a/Engine/Source/Editor/UnrealEd/Private/Commandlets/ContentCommandlets.cpp b/Engine/Source/Editor/UnrealEd/Private/Commandlets/ContentCommandlets.cpp
--- a/Engine/Source/Editor/UnrealEd/Private/Commandlets/ContentCommandlets.cpp
+++ b/Engine/Source/Editor/UnrealEd/Private/Commandlets/ContentCommandlets.cpp
@@ -4,13 +4,50 @@
 	UCContentCommandlets.cpp: Various commmandlets.
 =============================================================================*/
 
+#include "AssetCompilingManager.h"
+#include "AssetRegistry/AssetData.h"
+#include "CollectionManagerModule.h"
+#include "CommandletSourceControlUtils.h"
+#include "Commandlets/LandscapeGrassTypeCommandlet.h"
+#include "Commandlets/ListMaterialsUsedWithMeshEmittersCommandlet.h"
+#include "Commandlets/ListStaticMeshesImportedFromSpeedTreesCommandlet.h"
+#include "Commandlets/ResavePackagesCommandlet.h"
+#include "Commandlets/StaticMeshMinLodCommandlet.h"
+#include "Commandlets/WrangleContentCommandlet.h"
 #include "CoreMinimal.h"
+#include "Editor.h"
+#include "EditorWorldUtils.h"
+#include "Engine/Brush.h"
+#include "Engine/EngineTypes.h"
+#include "Engine/MapBuildDataRegistry.h"
+#include "Engine/StaticMesh.h"
+#include "EngineGlobals.h"
+#include "FileHelpers.h"
+#include "GameFramework/WorldSettings.h"
 #include "HAL/FileManager.h"
+#include "ICollectionManager.h"
+#include "ISourceControlModule.h"
+#include "ISourceControlOperation.h"
+#include "LandscapeGrassType.h"
+#include "LevelInstance/LevelInstanceInterface.h"
+#include "Materials/Material.h"
 #include "Misc/CommandLine.h"
+#include "Misc/ConfigCacheIni.h"
+#include "Misc/EngineVersion.h"
 #include "Misc/FileHelper.h"
+#include "Misc/PackageName.h"
 #include "Misc/Paths.h"
-#include "Misc/ConfigCacheIni.h"
+#include "Misc/RedirectCollector.h"
 #include "Modules/ModuleManager.h"
+#include "PackageHelperFunctions.h"
+#include "PackageTools.h"
+#include "Particles/ParticleEmitter.h"
+#include "Particles/ParticleSystem.h"
+#include "PlatformInfo.h"
+#include "SourceControlHelpers.h"
+#include "SourceControlOperations.h"
+#include "StaticMeshCompiler.h"
+#include "String/ParseTokens.h"
 #include "UObject/Class.h"
 #include "UObject/LinkerLoad.h"
 #include "UObject/MetaData.h"
@@ -18,50 +55,13 @@
 #include "UObject/ObjectMacros.h"
 #include "UObject/Package.h"
 #include "UObject/PackageResourceManager.h"
+#include "UObject/PackageTrailer.h"
+#include "UObject/SavePackage.h"
 #include "UObject/UObjectHash.h"
 #include "UObject/UObjectIterator.h"
-#include "Misc/PackageName.h"
-#include "Misc/EngineVersion.h"
-#include "Misc/RedirectCollector.h"
-#include "Engine/EngineTypes.h"
-#include "Materials/Material.h"
-#include "ISourceControlOperation.h"
-#include "SourceControlOperations.h"
-#include "SourceControlHelpers.h"
-#include "ISourceControlModule.h"
-#include "Engine/MapBuildDataRegistry.h"
-#include "Commandlets/ListMaterialsUsedWithMeshEmittersCommandlet.h"
-#include "Commandlets/ListStaticMeshesImportedFromSpeedTreesCommandlet.h"
-#include "Commandlets/StaticMeshMinLodCommandlet.h"
-#include "Commandlets/LandscapeGrassTypeCommandlet.h"
-#include "Particles/ParticleSystem.h"
-#include "Commandlets/ResavePackagesCommandlet.h"
-#include "Commandlets/WrangleContentCommandlet.h"
-#include "EngineGlobals.h"
-#include "Particles/ParticleEmitter.h"
-#include "GameFramework/WorldSettings.h"
-#include "Engine/StaticMesh.h"
-#include "Landscape/Classes/LandscapeGrassType.h"
-#include "AssetRegistry/AssetData.h"
-#include "Engine/Brush.h"
-#include "Editor.h"
-#include "EditorWorldUtils.h"
-#include "FileHelpers.h"
-#include "PlatformInfo.h"
-#include "CollectionManagerModule.h"
-#include "ICollectionManager.h"
-#include "CommandletSourceControlUtils.h"
 #include "WorldPartition/ActorDescContainer.h"
 #include "WorldPartition/WorldPartition.h"
 #include "WorldPartition/WorldPartitionHelpers.h"
-#include "LevelInstance/LevelInstanceInterface.h"
-#include "AssetCompilingManager.h"
-#include "PackageHelperFunctions.h"
-#include "PackageTools.h"
-#include "StaticMeshCompiler.h"
-#include "String/ParseTokens.h"
-#include "UObject/PackageTrailer.h"
-#include "UObject/SavePackage.h"
 
 DEFINE_LOG_CATEGORY(LogContentCommandlet);
 
