Don't consider the soon to be deleted object as part of the roots found in GatherObjectReferencersForDeletion

#rnx

diff --git a/Engine/Source/Editor/UnrealEd/Private/ObjectTools.cpp b/Engine/Source/Editor/UnrealEd/Private/ObjectTools.cpp
--- a/Engine/Source/Editor/UnrealEd/Private/ObjectTools.cpp
+++ b/Engine/Source/Editor/UnrealEd/Private/ObjectTools.cpp
@@ -344,6 +344,10 @@ namespace ObjectTools
 
 				TSet<UObject*> Roots = FindObjectsRoots(Referencers);
 
+				// Remove the object we plan on deleting if it turns out to be a root because
+				// the RF_Standalone flag is always removed later in the process.
+				Roots.Remove(InObject);
+
 				if (Roots.Contains(Transactor))
 				{
 					bOutIsReferencedInMemoryByUndo = true;
