Allow collection of Nanite FMicropolyRasterizeCS for the PSO cache

diff --git a/Engine/Source/Editor/UnrealEd/Private/Commandlets/ShaderPipelineCacheToolsCommandlet.cpp b/Engine/Source/Editor/UnrealEd/Private/Commandlets/ShaderPipelineCacheToolsCommandlet.cpp
--- a/Engine/Source/Editor/UnrealEd/Private/Commandlets/ShaderPipelineCacheToolsCommandlet.cpp
+++ b/Engine/Source/Editor/UnrealEd/Private/Commandlets/ShaderPipelineCacheToolsCommandlet.cpp
@@ -69,7 +69,7 @@ int32 GShaderPipelineCacheTools_ComputePSOInclusionMode = 2;
 static FAutoConsoleVariableRef CVarShaderPipelineCacheDoNotPrecompileComputePSO(
 	TEXT("r.ShaderPipelineCacheTools.IncludeComputePSODuringCook"),
 	GShaderPipelineCacheTools_ComputePSOInclusionMode,
-	TEXT("0 disables cook-time addition, 1 enables cook-time addition, 2 adds only Niagara PSOs."),
+	TEXT("0 disables cook-time addition, 1 enables cook-time addition, 2 adds only Niagara and Nanite Programmable Raster PSOs."),
 	ECVF_Default
 );
 
@@ -1947,14 +1947,15 @@ void AddComputePSOs(TSet<FPipelineCacheFileFormatPSO>& OutPSOs, const TMultiMap<
 
 	static FName NAME_SF_Compute("SF_Compute");
 	static FName NAME_NiagaraShader("FNiagaraShader");
+	static FName NAME_MicropolyRasterizeCS("FMicropolyRasterizeCS");	
 
 	for (TMultiMap<FStableShaderKeyAndValue, FSHAHash>::TConstIterator Iter(StableShaderMap); Iter; ++Iter)
 	{
 		if (Iter.Key().TargetFrequency == NAME_SF_Compute)
 		{
 			// add a new Compute PSO
-			// Check if we are only allowed to add Niagara PSOs
-			if (GShaderPipelineCacheTools_ComputePSOInclusionMode == 2 && Iter.Key().ShaderType != NAME_NiagaraShader)
+			// Check if we are only allowed to add Niagara and Nanite programmable raster PSOs
+			if (GShaderPipelineCacheTools_ComputePSOInclusionMode == 2 && (Iter.Key().ShaderType != NAME_NiagaraShader && Iter.Key().ShaderType != NAME_MicropolyRasterizeCS))
 			{
 				continue;
 			}
