Fix Level & Volume dirtyness:
- Avoid calling RebuildAlteredBSP() on Copy/Paste operations so that Level doesn't get dirtied if pasted actors aren't static brushes
- RebuildModelFromBrushes always checks selection of Dynamic brushes before rebuilding them

diff --git a/Engine/Source/Editor/UnrealEd/Private/EditorActor.cpp b/Engine/Source/Editor/UnrealEd/Private/EditorActor.cpp
--- a/Engine/Source/Editor/UnrealEd/Private/EditorActor.cpp
+++ b/Engine/Source/Editor/UnrealEd/Private/EditorActor.cpp
@@ -349,8 +349,7 @@ void UUnrealEdEngine::PasteActors(TArray<AActor*>& OutPastedActors, UWorld* InWo
 
 	// Reinstate old BSP update setting, and force a rebuild - any levels whose geometry has changed while pasting will be rebuilt
 	GetMutableDefault<ULevelEditorMiscSettings>()->bBSPAutoUpdate = bBSPAutoUpdate;
-	RebuildAlteredBSP();
-
+	
 	// FactoryCreateText set the selection to the new actors, so copy that into OutPastedActors and restore the original selection
 	ActorSelection->GetSelectedObjects<AActor>(OutPastedActors);
 	FObjectReader(ActorSelection, OriginalSelectionState);
@@ -360,8 +359,17 @@ void UUnrealEdEngine::PasteActors(TArray<AActor*>& OutPastedActors, UWorld* InWo
 
 	// Update the actors' locations and update the global list of visible layers.
 	ULayersSubsystem* LayersSubsystem = GEditor->GetEditorSubsystem<ULayersSubsystem>();
+	bool bRebuildBSP = false;
 	for (AActor* Actor : OutPastedActors)
 	{
+		if (!bRebuildBSP)
+		{
+			if (ABrush* Brush = Cast<ABrush>(Actor))
+			{
+				bRebuildBSP = Brush->IsStaticBrush();
+			}
+		}
+
 		if (!LocationOffset.IsZero())
 		{
 			// We only want to offset the location if this actor is the root of a selected attachment hierarchy
@@ -409,6 +417,11 @@ void UUnrealEdEngine::PasteActors(TArray<AActor*>& OutPastedActors, UWorld* InWo
 		Actor->MarkPackageDirty();
 		LevelDirtyCallback.Request();
 	}
+
+	if (bRebuildBSP)
+	{
+		RebuildAlteredBSP();
+	}
 }
 
 namespace  {
diff --git a/Engine/Source/Editor/UnrealEd/Private/EditorServer.cpp b/Engine/Source/Editor/UnrealEd/Private/EditorServer.cpp
--- a/Engine/Source/Editor/UnrealEd/Private/EditorServer.cpp
+++ b/Engine/Source/Editor/UnrealEd/Private/EditorServer.cpp
@@ -1846,11 +1846,10 @@ void UEditorEngine::RebuildModelFromBrushes(UModel* Model, bool bSelectedBrushes
 	TArray<ABrush*> DynamicBrushes;
 	if (!bTreatMovableBrushesAsStatic)
 	{
-	for( auto It(Level->Actors.CreateConstIterator()); It; ++It )
+		for( auto It(Level->Actors.CreateConstIterator()); It; ++It )
 		{
 			ABrush* DynamicBrush = Cast<ABrush>(*It);
-			if (DynamicBrush && DynamicBrush->Brush && !DynamicBrush->IsStaticBrush() &&
-				(!bSelectedBrushesOnly || DynamicBrush->IsSelected()))
+			if (DynamicBrush && DynamicBrush->Brush && !DynamicBrush->IsStaticBrush() && DynamicBrush->IsSelected())
 			{
 				DynamicBrushes.Add(DynamicBrush);
 			}
@@ -3754,8 +3753,18 @@ void UEditorEngine::PasteSelectedActorsFromClipboard( UWorld* InWorld, const FTe
 
 		RedrawLevelEditingViewports();
 
-		// If required, update the Bsp of any levels that received a pasted brush actor
-		RebuildAlteredBSP();
+		for (FSelectionIterator It(GetSelectedActorIterator()); It; ++It)
+		{
+			if (ABrush* Brush = Cast<ABrush>(*It))
+			{
+				if (Brush->IsStaticBrush())
+				{
+					// If required, update the Bsp of any levels that received a pasted brush actor
+					RebuildAlteredBSP();
+					break;
+				}
+			}
+		}
 	}
 	else
 	{
diff --git a/Engine/Source/Editor/UnrealEd/Private/Elements/Actor/ActorElementEditorWorldInterface.cpp b/Engine/Source/Editor/UnrealEd/Private/Elements/Actor/ActorElementEditorWorldInterface.cpp
--- a/Engine/Source/Editor/UnrealEd/Private/Elements/Actor/ActorElementEditorWorldInterface.cpp
+++ b/Engine/Source/Editor/UnrealEd/Private/Elements/Actor/ActorElementEditorWorldInterface.cpp
@@ -106,7 +106,13 @@ void UActorElementEditorWorldInterface::DuplicateElements(TArrayView<const FType
 		for (AActor* NewActor : NewActors)
 		{
 			// Only rebuild if the new actors will change the BSP as this is expensive
-			bRebuildBSP |= NewActor->IsA<ABrush>();
+			if (!bRebuildBSP)
+			{
+				if (ABrush* Brush = Cast<ABrush>(NewActor))
+				{
+					bRebuildBSP = Brush->IsStaticBrush();
+				}
+			}
 
 			OutNewElements.Add(UEngineElementsLibrary::AcquireEditorActorElementHandle(NewActor));
 		}
