Do not force a specific name when duplicating a component for copy to the clipboard

diff --git a/Engine/Source/Editor/UnrealEd/Private/Kismet2/ComponentEditorUtils.cpp b/Engine/Source/Editor/UnrealEd/Private/Kismet2/ComponentEditorUtils.cpp
--- a/Engine/Source/Editor/UnrealEd/Private/Kismet2/ComponentEditorUtils.cpp
+++ b/Engine/Source/Editor/UnrealEd/Private/Kismet2/ComponentEditorUtils.cpp
@@ -466,7 +466,7 @@ void FComponentEditorUtils::CopyComponents(const TArray<UActorComponent*>& Compo
 	for (UActorComponent* Component : ComponentsToCopy)
 	{
 		// Duplicate the component into a temporary object
-		UObject* DuplicatedComponent = StaticDuplicateObject(Component, GetTransientPackage(), Component->GetFName());
+		UObject* DuplicatedComponent = StaticDuplicateObject(Component, GetTransientPackage());
 		if (DuplicatedComponent)
 		{
 			// If the duplicated component is a scene component, wipe its attach parent (to prevent log warnings for referencing a private object in an external package)
