Fix a casting bug that was causing larger values of [Undo]UndoBufferSize to be truncated.

- The size of the undo buffer being loaded from the config file is a int32, but UTransBuffer::Initialize takes a SIZE_T, this was causing unintended conversions. We now cast the int32 to SIZE_T before calculating the final total in bytes.
- Added some logging so the user has some way to verify how large their undo buffer size is.

diff --git a/Engine/Source/Editor/UnrealEd/Private/EditorServer.cpp b/Engine/Source/Editor/UnrealEd/Private/EditorServer.cpp
--- a/Engine/Source/Editor/UnrealEd/Private/EditorServer.cpp
+++ b/Engine/Source/Editor/UnrealEd/Private/EditorServer.cpp
@@ -1283,8 +1283,10 @@ UTransactor* UEditorEngine::CreateTrans()
 		UndoBufferSize = 16;
 	}
 
+	UE_LOG(LogInit, Log, TEXT("Undo buffer set to %d MB"), UndoBufferSize);
+
 	UTransBuffer* TransBuffer = NewObject<UTransBuffer>();
-	TransBuffer->Initialize(UndoBufferSize * 1024 * 1024);
+	TransBuffer->Initialize((SIZE_T)UndoBufferSize * 1024 * 1024);
 	TransBuffer->OnBeforeRedoUndo().AddUObject(this, &UEditorEngine::HandleTransactorBeforeRedoUndo);
 	TransBuffer->OnRedo().AddUObject(this, &UEditorEngine::HandleTransactorRedo);
 	TransBuffer->OnUndo().AddUObject(this, &UEditorEngine::HandleTransactorUndo);
