Add cpu profile tracing to undo/redo operations for UTransBuffer.

#rnx

diff --git a/Engine/Source/Editor/UnrealEd/Private/EditorTransaction.cpp b/Engine/Source/Editor/UnrealEd/Private/EditorTransaction.cpp
--- a/Engine/Source/Editor/UnrealEd/Private/EditorTransaction.cpp
+++ b/Engine/Source/Editor/UnrealEd/Private/EditorTransaction.cpp
@@ -1382,6 +1382,8 @@ bool UTransBuffer::Undo(bool bCanRedo)
 		return false;
 	}
 
+	TRACE_CPUPROFILER_EVENT_SCOPE(UTransBuffer::Undo);
+
 	// Apply the undo changes.
 	GIsTransacting = true;
 
@@ -1447,6 +1449,8 @@ bool UTransBuffer::Redo()
 		return false;
 	}
 
+	TRACE_CPUPROFILER_EVENT_SCOPE(UTransBuffer::Redo);
+
 	// Apply the redo changes.
 	GIsTransacting = true;
 
