UE-167068 - Batch 3 of UnrealEd truncation warnings

diff --git a/Engine/Source/Editor/UnrealEd/Private/Factories/EditorFactories.cpp b/Engine/Source/Editor/UnrealEd/Private/Factories/EditorFactories.cpp
--- a/Engine/Source/Editor/UnrealEd/Private/Factories/EditorFactories.cpp
+++ b/Engine/Source/Editor/UnrealEd/Private/Factories/EditorFactories.cpp
@@ -1585,7 +1585,7 @@ UObject* UPolysFactory::FactoryCreateText
 	FFeedbackContext*	Warn
 )
 {
-	FVector PointPool[4096];
+	FVector3f PointPool[4096];
 	int32 NumPoints = 0;
 
 	GEditor->GetEditorSubsystem<UImportSubsystem>()->BroadcastAssetPreImport(this, Class, InParent, Name, Type);
@@ -1659,7 +1659,7 @@ UObject* UPolysFactory::FactoryCreateText
 					if( FParse::Command(&Str,TEXT("VERTEX")) )
 					{
 						// Start of new vertex.
-						PointPool[NumPoints] = FVector::ZeroVector;
+						PointPool[NumPoints] = FVector3f::ZeroVector;
 						Started = 1;
 						IsFace  = 0;
 					}
@@ -1712,8 +1712,10 @@ UObject* UPolysFactory::FactoryCreateText
 					else if( Code>=71 && Code<=79 && (Code-71)==NewPoly.Vertices.Num() )
 					{
 						int32 iPoint = FMath::Abs(FCString::Atoi(*ExtraLine));
-						if( iPoint>0 && iPoint<=NumPoints )
-							new(NewPoly.Vertices) FVector3f(PointPool[iPoint-1]);
+						if (iPoint > 0 && iPoint <= NumPoints)
+						{
+							NewPoly.Vertices.Emplace(PointPool[iPoint - 1]);
+						}
 						else UE_LOG(LogEditorFactories, Warning, TEXT("DXF: Invalid point index %i/%i"), iPoint, NumPoints );
 					}
 				}
@@ -2229,7 +2231,7 @@ UObject* UCurveLinearColorAtlasFactory::FactoryCreateNew(UClass* Class, UObject*
 
 	UCurveLinearColorAtlas* Object = NewObject<UCurveLinearColorAtlas>(InParent, Class, Name, Flags);
 	Object->Source.Init(Width, Height, 1, 1, TSF_RGBA16F);
-	const int32 TextureDataSize = Object->Source.CalcMipSize(0);
+	const int32 TextureDataSize = IntCastChecked<int32>(Object->Source.CalcMipSize(0));
 	Object->SrcData.AddUninitialized(TextureDataSize);
 	uint32* TextureData = (uint32*)Object->Source.LockMip(0);
 	FFloat16Color InitColor(FLinearColor::White);
@@ -3633,7 +3635,7 @@ UTexture* UTextureFactory::ImportTexture(UClass* Class, UObject* InParent, FName
 	if( FCString::Stricmp(Type, TEXT("ies")) == 0)
 	{
 		// checks for .IES extension to avoid wasting loading large assets just to reject them during header parsing
-		FIESConverter IESConverter(Buffer, Length);
+		FIESConverter IESConverter(Buffer, IntCastChecked<uint32>(Length));
 
 		if(IESConverter.IsValid())
 		{
@@ -4605,8 +4607,8 @@ bool UTextureExporterPCX::ExportBinary( UObject* Object, const TCHAR* Type, FArc
 		return false;
 	}
 
-	int32 SizeX = Texture->Source.GetSizeX();
-	int32 SizeY = Texture->Source.GetSizeY();
+	uint16 SizeX = IntCastChecked<uint16>(Texture->Source.GetSizeX());
+	uint16 SizeY = IntCastChecked<uint16>(Texture->Source.GetSizeY());
 	TArray64<uint8> RawData;
 	Texture->Source.GetMipData(RawData, 0);
 
@@ -5300,8 +5302,8 @@ bool UTextureExporterTGA::ExportBinary( UObject* Object, const TCHAR* Type, FArc
 		}
 	}
 
-	const int32 OriginalWidth = SizeX;
-	const int32 OriginalHeight = SizeY;		
+	const uint16 OriginalWidth = IntCastChecked<uint16>(SizeX);
+	const uint16 OriginalHeight = IntCastChecked<uint16>(SizeY);
 
 	FTGAFileHeader TGA;
 	FMemory::Memzero( &TGA, sizeof(TGA) );
@@ -5457,7 +5459,8 @@ UObject* UFontFileImportFactory::FactoryCreateBinary(UClass* InClass, UObject* I
 		FontFace->SourceFilename = GetCurrentFilename();
 
 		TArray<uint8> FontData;
-		FontData.Append(InBuffer, InBufferEnd - InBuffer);
+		int32 BufferSize = IntCastChecked<int32>(InBufferEnd - InBuffer);
+		FontData.Append(InBuffer, BufferSize);
 		FontFace->FontFaceData->SetData(MoveTemp(FontData));
 		FontFace->CacheSubFaces();
 	}
diff --git a/Engine/Source/Editor/UnrealEd/Private/Factories/PackFactory.cpp b/Engine/Source/Editor/UnrealEd/Private/Factories/PackFactory.cpp
--- a/Engine/Source/Editor/UnrealEd/Private/Factories/PackFactory.cpp
+++ b/Engine/Source/Editor/UnrealEd/Private/Factories/PackFactory.cpp
@@ -76,7 +76,7 @@ namespace PackFactoryHelper
 				FAES::FAESKey Key;
 				FPakPlatformFile::GetPakEncryptionKey(Key, PakFile.GetInfo().EncryptionKeyGuid);
 				checkf(Key.IsValid(), TEXT("Trying to copy an encrypted file between pak files, but no decryption key is available"));
-				FAES::DecryptData(Buffer.GetData(), SizeToRead, Key);
+				FAES::DecryptData(Buffer.GetData(), IntCastChecked<uint32>(SizeToRead), Key);
 			}
 			DestAr.Serialize(Buffer.GetData(), SizeToCopy);
 			RemainingSizeToCopy -= SizeToRead;
@@ -92,8 +92,9 @@ namespace PackFactoryHelper
 			return false;
 		}
 
-		int64 WorkingSize = Entry.CompressionBlockSize;
-		int32 MaxCompressionBlockSize = FCompression::GetMaximumCompressedSize(PakFile.GetInfo().GetCompressionMethod(Entry.CompressionMethodIndex), WorkingSize);
+		int32 WorkingSize = Entry.CompressionBlockSize;
+		uint8 CompressionMethodIndex = IntCastChecked<uint8>(Entry.CompressionMethodIndex);
+		int32 MaxCompressionBlockSize = FCompression::GetMaximumCompressedSize(PakFile.GetInfo().GetCompressionMethod(CompressionMethodIndex), WorkingSize);
 		WorkingSize += MaxCompressionBlockSize;
 		if (PersistentBuffer.Num() < WorkingSize)
 		{
@@ -104,10 +105,10 @@ namespace PackFactoryHelper
 
 		for (uint32 BlockIndex = 0, BlockIndexNum = Entry.CompressionBlocks.Num(); BlockIndex < BlockIndexNum; ++BlockIndex)
 		{
-			uint32 CompressedBlockSize = Entry.CompressionBlocks[BlockIndex].CompressedEnd - Entry.CompressionBlocks[BlockIndex].CompressedStart;
+			int32 CompressedBlockSize = IntCastChecked<int32>(Entry.CompressionBlocks[BlockIndex].CompressedEnd - Entry.CompressionBlocks[BlockIndex].CompressedStart);
 			uint32 UncompressedBlockSize = (uint32)FMath::Min<int64>(Entry.UncompressedSize - Entry.CompressionBlockSize*BlockIndex, Entry.CompressionBlockSize);
 			Source.Seek(Entry.CompressionBlocks[BlockIndex].CompressedStart + (PakFile.GetInfo().HasRelativeCompressedChunkOffsets() ? Entry.Offset : 0));
-			uint32 SizeToRead = Entry.IsEncrypted() ? Align(CompressedBlockSize, FAES::AESBlockSize) : CompressedBlockSize;
+			int32 SizeToRead = Entry.IsEncrypted() ? Align(CompressedBlockSize, FAES::AESBlockSize) : CompressedBlockSize;
 			Source.Serialize(PersistentBuffer.GetData(), SizeToRead);
 
 			if (Entry.IsEncrypted())
@@ -118,7 +119,7 @@ namespace PackFactoryHelper
 				FAES::DecryptData(PersistentBuffer.GetData(), SizeToRead, Key);
 			}
 
-			if (!FCompression::UncompressMemory(PakFile.GetInfo().GetCompressionMethod(Entry.CompressionMethodIndex), UncompressedBuffer, UncompressedBlockSize, PersistentBuffer.GetData(), CompressedBlockSize))
+			if (!FCompression::UncompressMemory(PakFile.GetInfo().GetCompressionMethod(CompressionMethodIndex), UncompressedBuffer, UncompressedBlockSize, PersistentBuffer.GetData(), CompressedBlockSize))
 			{
 				return false;
 			}
diff --git a/Engine/Source/Editor/UnrealEd/Private/Factories/TTFontImport.cpp b/Engine/Source/Editor/UnrealEd/Private/Factories/TTFontImport.cpp
--- a/Engine/Source/Editor/UnrealEd/Private/Factories/TTFontImport.cpp
+++ b/Engine/Source/Editor/UnrealEd/Private/Factories/TTFontImport.cpp
@@ -498,7 +498,7 @@ void FTextureAlphaToDistanceField::Generate(int32 ScaleFactor, int32 ScanRadius)
 
 	// destination texture
 	// note that destination format can be different from source format	
-	SIZE_T NumBytes = CalculateImageBytes(DstSizeX,DstSizeY,0,DstFormat);	
+	int32 NumBytes = IntCastChecked<int32>(CalculateImageBytes(DstSizeX,DstSizeY,0,DstFormat));
 	DstTexture.Empty(NumBytes);	
 	DstTexture.AddZeroed(NumBytes);
 	
@@ -584,7 +584,7 @@ void FTextureAlphaToDistanceField::Generate(int32 ScaleFactor, int32 ScanRadius)
 		{
 			for( int32 y=0; y < DstSizeY; y++ )
 			{
-				DstTexture[x + y * DstSizeX] = SignedDistance[x + y * DstSizeX] * 255;
+				DstTexture[x + y * DstSizeX] = static_cast<uint8>(SignedDistance[x + y * DstSizeX] * 255);
 			}
 		}
 	}
@@ -598,7 +598,7 @@ void FTextureAlphaToDistanceField::Generate(int32 ScaleFactor, int32 ScanRadius)
 				DstTexture[4 * (x + y * DstSizeX) + 0] = SrcColor.B;
 				DstTexture[4 * (x + y * DstSizeX) + 1] = SrcColor.G;
 				DstTexture[4 * (x + y * DstSizeX) + 2] = SrcColor.R;
-				DstTexture[4 * (x + y * DstSizeX) + 3] = SignedDistance[x + y * DstSizeX] * 255;
+				DstTexture[4 * (x + y * DstSizeX) + 3] = static_cast<uint8>(SignedDistance[x + y * DstSizeX] * 255);
 			}
 		}	
 	}
@@ -931,7 +931,7 @@ UTexture2D* UTrueTypeFontFactory::CreateTextureFromDC( UFont* Font, HDC dc, int3
 				MipData[4 * (i + j * SizeX) + 0] = ( uint8 )( FontColor8Bit.B * ( 1.0f - fDropShadowAlpha ) );
 				MipData[4 * (i + j * SizeX) + 1] = ( uint8 )( FontColor8Bit.G * ( 1.0f - fDropShadowAlpha ) );
 				MipData[4 * (i + j * SizeX) + 2] = ( uint8 )( FontColor8Bit.R * ( 1.0f - fDropShadowAlpha ) );
-				MipData[4 * (i + j * SizeX) + 3] = CharAlpha + DropShadowAlpha;
+				MipData[4 * (i + j * SizeX) + 3] = ( uint8 )( CharAlpha + DropShadowAlpha );
 			}
 		}
 	}
@@ -949,7 +949,7 @@ UTexture2D* UTrueTypeFontFactory::CreateTextureFromDC( UFont* Font, HDC dc, int3
 			PF_B8G8R8A8
 			);
 		// estimate scan radius based on half font height scaled by bitmap scale factor
-		const int32 ScanRadius = ImportOptions->Data.Height/2 * ImportOptions->Data.DistanceFieldScaleFactor * ImportOptions->Data.DistanceFieldScanRadiusScale;
+		const int32 ScanRadius = static_cast<int32>((ImportOptions->Data.Height / 2.0f) * ImportOptions->Data.DistanceFieldScaleFactor * ImportOptions->Data.DistanceFieldScanRadiusScale);
 		// generate downsampled distance field using high res source bitmap
 		DistanceFieldTex.Generate(ImportOptions->Data.DistanceFieldScaleFactor,ScanRadius);
 		check(DistanceFieldTex.GetResultTextureSize() > 0);
@@ -1337,7 +1337,7 @@ bool UTrueTypeFontFactory::CreateFontTexture(
 			NewCharacterRef.VSize =
 				FMath::Clamp<int32>( FontHeight + ImportOptions->Data.ExtendBoxTop + ImportOptions->Data.ExtendBoxBottom,
 							0, ImportOptions->Data.TexturePageMaxHeight - NewCharacterRef.StartV );
-			NewCharacterRef.TextureIndex = CurrentTexture;
+			NewCharacterRef.TextureIndex = static_cast<uint8>(CurrentTexture);
 			NewCharacterRef.VerticalOffset = VerticalOffset;
 
 			// Draw character into font and advance.
@@ -2139,7 +2139,7 @@ bool UTrueTypeFontFactory::ImportTrueTypeFont(
 		Font->IsRemapped = 1;
 
 		// Only include ASCII characters if we were asked to
-		int32 MinRangeCharacter = 0;
+		uint16 MinRangeCharacter = 0;
 		if( ImportOptions->Data.bIncludeASCIIRange )
 		{
 			// Map (ASCII)
@@ -2224,9 +2224,10 @@ bool UTrueTypeFontFactory::ImportTrueTypeFont(
 
 		}
 
-		int32 j=MinRangeCharacter;
-		int32 Min=65536, Max=0;
-		for( int32 i=MinRangeCharacter; i<65536; i++ )
+		uint16 j = MinRangeCharacter;
+		uint16 Min = std::numeric_limits<uint16>::max();
+		uint16 Max = 0;
+		for( uint16 i = MinRangeCharacter; i <= std::numeric_limits<uint16>::max(); ++i )
 		{
 			if( Chars[i] )
 			{
diff --git a/Engine/Source/Editor/UnrealEd/Private/StaticLightingSystem/StaticLightingDebug.cpp b/Engine/Source/Editor/UnrealEd/Private/StaticLightingSystem/StaticLightingDebug.cpp
--- a/Engine/Source/Editor/UnrealEd/Private/StaticLightingSystem/StaticLightingDebug.cpp
+++ b/Engine/Source/Editor/UnrealEd/Private/StaticLightingSystem/StaticLightingDebug.cpp
@@ -19,6 +19,8 @@
 #include "Engine/StaticMesh.h"
 #include "TextureResource.h"
 
+#include <limits>
+
 /** Information about the texel that is selected */
 FSelectedLightmapSample GCurrentSelectedLightmapSample;
 
@@ -96,16 +98,16 @@ static bool UpdateSelectedTexel(
 			UTexture2D* CurrentLightmap = Lightmap2D->GetTexture( LightmapIndex );
 			{
 				// UV's in the lightmap atlas
-				int32 LightmapX = FMath::TruncToInt(LightmapUV.X * CurrentLightmap->GetSizeX());
-				int32 LightmapY = FMath::TruncToInt(LightmapUV.Y * .5f * CurrentLightmap->GetSizeY());
+				int32 LightmapX = FMath::TruncToInt32(LightmapUV.X * CurrentLightmap->GetSizeX());
+				int32 LightmapY = FMath::TruncToInt32(LightmapUV.Y * .5f * CurrentLightmap->GetSizeY());
 				// Write the selection color to the selected lightmap texel
 				WriteTexel(CurrentLightmap, LightmapX, LightmapY, GTexelSelectionColor);
 			}
 
 			{
 				// UV's in the lightmap atlas
-				int32 LightmapX = FMath::TruncToInt(LightmapUV.X * CurrentLightmap->GetSizeX());
-				int32 LightmapY = FMath::TruncToInt((LightmapUV.Y * .5f + .5f) * CurrentLightmap->GetSizeY());
+				int32 LightmapX = FMath::TruncToInt32(LightmapUV.X * CurrentLightmap->GetSizeX());
+				int32 LightmapY = FMath::TruncToInt32((LightmapUV.Y * .5f + .5f) * CurrentLightmap->GetSizeY());
 				// Write the selection color to the selected lightmap texel
 				WriteTexel(CurrentLightmap, LightmapX, LightmapY, GTexelSelectionColor);
 			}
@@ -127,13 +129,13 @@ static bool GetBarycentricWeights(
 	const FVector& Position2,
 	FVector InterpolatePosition,
 	float Tolerance,
-	float& PlaneDistance,
+	double& PlaneDistance,
 	FVector& BarycentricWeights
 	)
 {
 	BarycentricWeights = FVector::ZeroVector;
 	FVector TriangleNormal = (Position0 - Position1) ^ (Position2 - Position0);
-	float ParallelogramArea = TriangleNormal.Size();
+	FVector::FReal ParallelogramArea = TriangleNormal.Size();
 	FVector UnitTriangleNormal = TriangleNormal / ParallelogramArea;
 	PlaneDistance = UnitTriangleNormal | (InterpolatePosition - Position0);
 
@@ -143,14 +145,14 @@ static bool GetBarycentricWeights(
 
 	FVector NormalU = (InterpolatePosition - Position1) ^ (Position2 - InterpolatePosition);
 	// Signed area, if negative then InterpolatePosition is not in the triangle
-	float ParallelogramAreaU = NormalU.Size() * FMath::FloatSelect(NormalU | TriangleNormal, (FVector::FReal)1.0f, (FVector::FReal)-1.0f);
-	float BaryCentricU = ParallelogramAreaU / ParallelogramArea;
+	FVector::FReal ParallelogramAreaU = NormalU.Size() * FMath::FloatSelect(NormalU | TriangleNormal, (FVector::FReal)1.0f, (FVector::FReal)-1.0f);
+	FVector::FReal BaryCentricU = ParallelogramAreaU / ParallelogramArea;
 
 	FVector NormalV = (InterpolatePosition - Position2) ^ (Position0 - InterpolatePosition);
-	float ParallelogramAreaV = NormalV.Size() * FMath::FloatSelect(NormalV | TriangleNormal, (FVector::FReal)1.0f, (FVector::FReal)-1.0f);
-	float BaryCentricV = ParallelogramAreaV / ParallelogramArea;
+	FVector::FReal ParallelogramAreaV = NormalV.Size() * FMath::FloatSelect(NormalV | TriangleNormal, (FVector::FReal)1.0f, (FVector::FReal)-1.0f);
+	FVector::FReal BaryCentricV = ParallelogramAreaV / ParallelogramArea;
 
-	float BaryCentricW = 1.0f - BaryCentricU - BaryCentricV;
+	FVector::FReal BaryCentricW = 1.0f - BaryCentricU - BaryCentricV;
 	if (BaryCentricU > -Tolerance && BaryCentricV > -Tolerance && BaryCentricW > -Tolerance)
 	{
 		BarycentricWeights = FVector(BaryCentricU, BaryCentricV, BaryCentricW);
@@ -216,7 +218,7 @@ void SetDebugLightmapSample(TArray<UActorComponent*>* Components, UModel* Model,
 
 				if (bUseTextureMap)
 				{
-					float ClosestPlaneDistance = FLT_MAX;
+					double ClosestPlaneDistance = std::numeric_limits<double>::max();
 					FVector ClosestPlaneBaryCentricWeights;
 					int32 ClosestPlaneTriangleIndex = -1;
 
@@ -232,7 +234,7 @@ void SetDebugLightmapSample(TArray<UActorComponent*>* Components, UModel* Model,
 						FVector Position1 = SMComponent->GetComponentTransform().TransformPosition((FVector)LODModel.VertexBuffers.PositionVertexBuffer.VertexPosition(Index1));
 						FVector Position2 = SMComponent->GetComponentTransform().TransformPosition((FVector)LODModel.VertexBuffers.PositionVertexBuffer.VertexPosition(Index2));
 
-						float PlaneDistance;
+						double PlaneDistance;
 						FVector BaryCentricWeights;
 						// Continue if click location is in the triangle and get its barycentric weights
 						if (GetBarycentricWeights(Position0, Position1, Position2, ClickLocation, TriangleTolerance, PlaneDistance, BaryCentricWeights))
@@ -267,8 +269,8 @@ void SetDebugLightmapSample(TArray<UActorComponent*>* Components, UModel* Model,
 							PaddedSizeY -= 2;
 						}
 
-						const int32 LocalX = FMath::TruncToInt(InterpolatedUV.X * PaddedSizeX);
-						const int32 LocalY = FMath::TruncToInt(InterpolatedUV.Y * PaddedSizeY);
+						const int32 LocalX = FMath::TruncToInt32(InterpolatedUV.X * PaddedSizeX);
+						const int32 LocalY = FMath::TruncToInt32(InterpolatedUV.Y * PaddedSizeY);
 						if (LocalX < 0 || LocalX >= PaddedSizeX
 							|| LocalY < 0 || LocalY >= PaddedSizeY)
 						{
@@ -322,8 +324,8 @@ void SetDebugLightmapSample(TArray<UActorComponent*>* Components, UModel* Model,
 						PaddedSizeY -= 2;
 					}
 
-					const int32 LocalX = FMath::TruncToInt(InterpolatedUV.X * PaddedSizeX);
-					const int32 LocalY = FMath::TruncToInt(InterpolatedUV.Y * PaddedSizeY);
+					const int32 LocalX = FMath::TruncToInt32(InterpolatedUV.X * PaddedSizeX);
+					const int32 LocalY = FMath::TruncToInt32(InterpolatedUV.Y * PaddedSizeY);
 					if (LocalX < 0 || LocalX >= PaddedSizeX
 						|| LocalY < 0 || LocalY >= PaddedSizeY)
 					{
@@ -350,7 +352,7 @@ void SetDebugLightmapSample(TArray<UActorComponent*>* Components, UModel* Model,
 			int32 ClosestElementIndex = -1;
 			uint32 ClosestTriangleIndex = 0;
 			FVector ClosestPlaneBaryCentricWeights = FVector(0);
-			float ClosestPlaneDistance = FLT_MAX;
+			double ClosestPlaneDistance = std::numeric_limits<double>::max();
 
 			for (int32 ModelIndex = 0; ModelIndex < World->GetCurrentLevel()->ModelComponents.Num(); ModelIndex++)
 			{
@@ -376,7 +378,7 @@ void SetDebugLightmapSample(TArray<UActorComponent*>* Components, UModel* Model,
 							FVector Position1 = (FVector)ModelVertices[Index1].Position;
 							FVector Position2 = (FVector)ModelVertices[Index2].Position;
 
-							float PlaneDistance;
+							double PlaneDistance;
 							FVector BaryCentricWeights;
 							// Continue if click location is in the triangle and get its barycentric weights
 							if (GetBarycentricWeights(Position0, Position1, Position2, ClickLocation, .001f, PlaneDistance, BaryCentricWeights))
@@ -471,8 +473,8 @@ void SetDebugLightmapSample(TArray<UActorComponent*>* Components, UModel* Model,
 
 				// Apply the transform to the intersection position to find the local texel coordinates
 				const FVector4 StaticLightingTextureCoordinate = WorldToMap.TransformPosition(ClickLocation);
-				const int32 LocalX = FMath::TruncToInt(StaticLightingTextureCoordinate.X * PaddedSizeX);
-				const int32 LocalY = FMath::TruncToInt(StaticLightingTextureCoordinate.Y * PaddedSizeY);
+				const int32 LocalX = FMath::TruncToInt32(StaticLightingTextureCoordinate.X * PaddedSizeX);
+				const int32 LocalY = FMath::TruncToInt32(StaticLightingTextureCoordinate.Y * PaddedSizeY);
 				check(LocalX >= 0 && LocalX < PaddedSizeX && LocalY >= 0 && LocalY < PaddedSizeY);
 
 				const FMeshMapBuildData* MeshMapBuildData = Element.GetMeshMapBuildData();
@@ -649,7 +651,7 @@ void DrawStaticLightingDebugInfo(const FSceneView* View, FCanvas* Canvas)
 				if(View->ScreenToPixel(View->WorldToScreen(FVector4(CurrentRecord.Vertex.VertexPosition)),PixelLocation))
 				{
 					const FColor TagColor = CurrentRecord.bAffectsSelectedTexel ? FColor(50,160,200) : FColor(120,120,120);
-					Canvas->DrawShadowedString(PixelLocation.X,PixelLocation.Y, *FString::FromInt(CurrentRecord.RecordId), GEngine->GetSmallFont(), TagColor);
+					Canvas->DrawShadowedString(static_cast<float>(PixelLocation.X), static_cast<float>(PixelLocation.Y), *FString::FromInt(CurrentRecord.RecordId), GEngine->GetSmallFont(), TagColor);
 				}
 			}
 		}
@@ -661,7 +663,7 @@ void DrawStaticLightingDebugInfo(const FSceneView* View, FCanvas* Canvas)
 			if(View->ScreenToPixel(View->WorldToScreen(FVector4(CurrentPhoton.Position)),PixelLocation))
 			{
 				const FColor TagColor = FColor(120,120,120);
-				Canvas->DrawShadowedString(PixelLocation.X,PixelLocation.Y, *FString::FromInt(CurrentPhoton.Id), GEngine->GetSmallFont(), TagColor);
+				Canvas->DrawShadowedString(static_cast<float>(PixelLocation.X), static_cast<float>(PixelLocation.Y), *FString::FromInt(CurrentPhoton.Id), GEngine->GetSmallFont(), TagColor);
 			}
 		}
 
@@ -674,7 +676,7 @@ void DrawStaticLightingDebugInfo(const FSceneView* View, FCanvas* Canvas)
 				if(View->ScreenToPixel(View->WorldToScreen(FVector4(CurrentRay.End)),PixelLocation))
 				{
 					const FColor TagColor = FColor(180,180,120);
-					Canvas->DrawShadowedString(PixelLocation.X,PixelLocation.Y, *FString::FromInt(RayIndex), GEngine->GetSmallFont(), TagColor);
+					Canvas->DrawShadowedString(static_cast<float>(PixelLocation.X), static_cast<float>(PixelLocation.Y), *FString::FromInt(RayIndex), GEngine->GetSmallFont(), TagColor);
 				}
 			}
 		}
diff --git a/Engine/Source/Editor/UnrealEd/Private/WorldPartition/WorldPartitionBuilder.cpp b/Engine/Source/Editor/UnrealEd/Private/WorldPartition/WorldPartitionBuilder.cpp
--- a/Engine/Source/Editor/UnrealEd/Private/WorldPartition/WorldPartitionBuilder.cpp
+++ b/Engine/Source/Editor/UnrealEd/Private/WorldPartition/WorldPartitionBuilder.cpp
@@ -265,7 +265,11 @@ bool UWorldPartitionBuilder::Run(UWorld* World, FPackageSourceControlHelper& Pac
 					
 					UE_LOG(LogWorldPartitionBuilder, Display, TEXT("[%d / %d] Processing cells..."), IterationIndex, IterationCount);
 
-					FVector Min(x * IterativeCellSize, y * IterativeCellSize, z * IterativeCellSize);
+					double MinX = static_cast<double>(x * IterativeCellSize);
+					double MinY = static_cast<double>(y * IterativeCellSize);
+					double MinZ = static_cast<double>(z * IterativeCellSize);
+
+					FVector Min(MinX, MinY, MinZ);
 					FVector Max = Min + FVector(IterativeCellSize);
 
 					if (LoadingMode == ELoadingMode::IterativeCells2D)
diff --git a/Engine/Source/Editor/UnrealEd/Private/WorldPartition/WorldPartitionHLODsBuilder.cpp b/Engine/Source/Editor/UnrealEd/Private/WorldPartition/WorldPartitionHLODsBuilder.cpp
--- a/Engine/Source/Editor/UnrealEd/Private/WorldPartition/WorldPartitionHLODsBuilder.cpp
+++ b/Engine/Source/Editor/UnrealEd/Private/WorldPartition/WorldPartitionHLODsBuilder.cpp
@@ -611,10 +611,10 @@ bool UWorldPartitionHLODsBuilder::DumpStats()
 		uint64				VertexCount = 0;
 		uint64				UVChannelCount = 0;
 
-		uint32				BaseColorTextureSize = 0;
-		uint32				NormalTextureSize = 0;
-		uint32				MRSTextureSize = 0;
-		uint32				EmissiveTextureSize = 0;
+		float				BaseColorTextureSize = 0.0f;
+		float				NormalTextureSize = 0.0f;
+		float				MRSTextureSize = 0.0f;
+		float				EmissiveTextureSize = 0.0f;
 
 		uint64				MeshResourceSize = 0;
 		uint64				TexturesResourceSize = 0;
@@ -628,7 +628,7 @@ bool UWorldPartitionHLODsBuilder::DumpStats()
 
 		FString ToCSVString() const
 		{
-			return FString::Printf(TEXT("%s, %s, %s, %d, %d, %d, %s, %s, %d, %d, %d, %d, %d, %d, %d, %d, %d, %d, %d, %d, %d"), 
+			return FString::Printf(TEXT("%s, %s, %s, %d, %d, %d, %s, %s, %d, %d, %d, %d, %d, %d, %.0f, %.0f, %.0f, %.0f, %d, %d, %d"), 
 				*Name, 
 				*RuntimeGrid.ToString(),
 				bIsSpatiallyLoaded ? TEXT("true") : TEXT("false"),
