WindowsPlatformMemory: Change the behavior of FGenericPlatformMemoryStats::GetMemoryPressureStatus on Windows; it now uses QueryMemoryResourceNotification instead of the generic implementation.
Cooker: In MPCook Collect garbage based on the Operating System's MemoryPressure event, so that Standby memory does not cause spurious garbage collection.
In all cook modes, tweak the output during garbage collection to have less duplicate information.
In all cook modes, add a cooldown for GC when garbage collection is not impactful.
#rn Minor, Core

diff --git a/Engine/Config/BaseEditor.ini b/Engine/Config/BaseEditor.ini
--- a/Engine/Config/BaseEditor.ini
+++ b/Engine/Config/BaseEditor.ini
@@ -319,10 +319,23 @@ BulkDataRegistryEnabled=true
 HybridIterativeEnabled=false
 PackagesPerGC=0
 IdleTimeToGC=20
-; 2 GB memory is need to cook large textures
-;  set MinFreePhysical and MinFreeVirtual both to the same number
+; 2 GB of temporary memory is need to cook large textures, so the default for all projects is to trigger a
+; garbage collect when our free physical memory is lower than that value.
 MemoryMinFreePhysical=2048
 MemoryMinFreeVirtual=2048
+; MemoryMaxUsed is not recommended by default, but can be useful for some build setups.
+; By itself, if MemoryMinFree values are 0, it will trigger a garbage collect whenever used memory exceeds
+; its value. This can be used to enforce a maximum used memory for garbage collects.
+; When both MemoryMaxUsed and MemoryMinFree are used togther, both must be triggered to trigger a garbage
+; collect. This can be used to enforce a minimum memory capacity for machines cooking the project, while still
+; never garbage collecting until physical memory is low.
+MemoryMaxUsedPhysical=0
+MemoryMaxUsedVirtual=0
+; MemoryTriggerGCAtPressureLevel can be used as an alternative to MemoryMinFreePhysical. It allows the
+; operating system to report when physical memory is low. Values are None (disabled) or Critical (enabled).
+; Unlike MemoryMinFreePhysical, it does not interact with MemoryMaxUsed; GC will be triggered whenever the
+; event fires.
+MemoryTriggerGCAtPressureLevel=None
 MaxConcurrentShaderJobs=8192
 MaxDiffsToLog=5
 IgnoreHeaderDiffs=True
diff --git a/Engine/Source/Editor/UnrealEd/Classes/CookOnTheSide/CookOnTheFlyServer.h b/Engine/Source/Editor/UnrealEd/Classes/CookOnTheSide/CookOnTheFlyServer.h
--- a/Engine/Source/Editor/UnrealEd/Classes/CookOnTheSide/CookOnTheFlyServer.h
+++ b/Engine/Source/Editor/UnrealEd/Classes/CookOnTheSide/CookOnTheFlyServer.h
@@ -7,6 +7,7 @@
 #include "CookOnTheSide/CookLog.h"
 #include "CookPackageSplitter.h"
 #include "Engine/ICookInfo.h"
+#include "HAL/PlatformMemory.h"
 #include "INetworkFileSystemModule.h"
 #include "IPlatformFileSandboxWrapper.h"
 #include "Misc/EnumClassFlags.h"
@@ -33,6 +34,7 @@ class IPlugin;
 class ITargetPlatform;
 enum class ODSCRecompileCommand;
 struct FBeginCookContext;
+struct FGenericMemoryStats;
 struct FPropertyChangedEvent;
 struct FResourceSizeEx;
 
@@ -288,6 +290,7 @@ private:
 	uint64 MemoryMaxUsedPhysical;
 	uint64 MemoryMinFreeVirtual;
 	uint64 MemoryMinFreePhysical;
+	FGenericPlatformMemoryStats::EMemoryPressureStatus MemoryTriggerGCAtPressureLevel;
 	float MemoryExpectedFreedToSpreadRatio;
 	/** Max number of packages to save before we partial gc */
 	int32 MaxNumPackagesBeforePartialGC;
@@ -295,6 +298,8 @@ private:
 	int32 MaxConcurrentShaderJobs;
 	/** Min number of free UObject indices before the cooker should partial gc */
 	int32 MinFreeUObjectIndicesBeforeGC;
+	/** Next time at which we are allowed to declare out of memory and trigger a garbage collect. */
+	double ExceededMaxMemoryCooldownEndTimeSeconds = 0;
 	/**
 	 * The maximum number of packages that should be preloaded at once. Once this is full, packages in LoadPrepare will
 	 * remain unpreloaded in LoadPrepare until the existing preloaded packages exit {LoadPrepare,LoadReady} state.
@@ -797,9 +802,12 @@ public:
 	{
 	}
 
-	bool HasExceededMaxMemory() const;
-	void EvaluateGarbageCollectionResults(int32 NumObjectsBeforeGC, const FPlatformMemoryStats& MemStatsBeforeGC,
-		int32 NumObjectsAfterGC, const FPlatformMemoryStats& MemStatsAfterGC);
+	bool HasExceededMaxMemory();
+	void EvaluateGarbageCollectionResults(bool bWasDueToOOM, bool bWasPartialGC,
+		int32 NumObjectsBeforeGC, const FPlatformMemoryStats& MemStatsBeforeGC,
+		const FGenericMemoryStats& AllocatorStatsBeforeGC,
+		int32 NumObjectsAfterGC, const FPlatformMemoryStats& MemStatsAfterGC,
+		const FGenericMemoryStats& AllocatorStatsAfterGC);
 
 	/**
 	 * RequestPackage to be cooked
diff --git a/Engine/Source/Editor/UnrealEd/Private/Commandlets/CookCommandlet.cpp b/Engine/Source/Editor/UnrealEd/Private/Commandlets/CookCommandlet.cpp
--- a/Engine/Source/Editor/UnrealEd/Private/Commandlets/CookCommandlet.cpp
+++ b/Engine/Source/Editor/UnrealEd/Private/Commandlets/CookCommandlet.cpp
@@ -1000,61 +1000,25 @@ void UCookCommandlet::ConditionalCollectGarbage(uint32 TickResults, UCookOnTheFl
 
 	const FPlatformMemoryStats MemStatsBeforeGC = FPlatformMemory::GetStats();
 	int32 NumObjectsBeforeGC = GUObjectArray.GetObjectArrayNumMinusAvailable();
-	int32 NumObjectsAvailableBeforeGC = GUObjectArray.GetObjectArrayEstimatedAvailable();
-	UE_LOG(LogCookCommandlet, Display, TEXT("GarbageCollection...%s (%s)"), (bPartialGC ? TEXT(" partial gc") : TEXT("")), *GCReason);
-
+	FGenericMemoryStats AllocatorStatsBeforeGC;
+	GMalloc->GetAllocatorStats(AllocatorStatsBeforeGC);
 #if ENABLE_LOW_LEVEL_MEM_TRACKER
 	FLowLevelMemTracker::Get().UpdateStatsPerFrame();
 #endif
-	auto DumpMemStats = []()
-	{
-		FGenericMemoryStats MemStats;
-		GMalloc->GetAllocatorStats(MemStats);
-		for (const auto& Item : MemStats.Data)
-		{
-			UE_LOG(LogCookCommandlet, Display, TEXT("Item %s = %d"), *Item.Key, Item.Value);
-		}
-	};
-
-	if (!COTFS.IsCookOnTheFlyMode())
-	{
-		DumpMemStats();
-	}
 
+	UE_LOG(LogCookCommandlet, Display, TEXT("GarbageCollection...%s (%s)"), (bPartialGC ? TEXT(" partial gc") : TEXT("")), *GCReason);
 	CollectGarbage(RF_NoFlags);
 
-	int32 NumObjectsAfterGC = GUObjectArray.GetObjectArrayNumMinusAvailable();
-	int32 NumObjectsAvailableAfterGC = GUObjectArray.GetObjectArrayEstimatedAvailable();
 	FPlatformMemoryStats MemStatsAfterGC = FPlatformMemory::GetStats();
-	if (!COTFS.IsCookOnTheFlyMode())
-	{
-		int64 VirtualMemBeforeGC = MemStatsBeforeGC.UsedVirtual;
-		int64 VirtualMemAfterGC = MemStatsAfterGC.UsedVirtual;
-		int64 VirtualMemFreed = MemStatsBeforeGC.UsedVirtual - MemStatsAfterGC.UsedVirtual;
-		constexpr int32 BytesPerMeg = 1000000;
-		UE_LOG(LogCookCommandlet, Display, TEXT("GarbageCollection Results:\n")
-			TEXT("\tType: %s\n")
-			TEXT("\tNumObjects:\n")
-			TEXT("\t\tBefore GC:        %10d\n")
-			TEXT("\t\tAvailable Before: %10d\n")
-			TEXT("\t\tAfter GC:         %10d\n")
-			TEXT("\t\tAvailable After:  %10d\n")
-			TEXT("\t\tFreed by GC:      %10d\n")
-			TEXT("\tVirtual Memory:\n")
-			TEXT("\t\tBefore GC:        %10" INT64_FMT " MB\n")
-			TEXT("\t\tAfter GC:         %10" INT64_FMT " MB\n")
-			TEXT("\t\tFreed by GC:      %10" INT64_FMT " MB\n"),
-			(bPartialGC ? TEXT("Partial") : TEXT("Full")),
-			NumObjectsBeforeGC, NumObjectsAvailableBeforeGC, NumObjectsAfterGC, NumObjectsAvailableAfterGC,
-			NumObjectsBeforeGC - NumObjectsAfterGC,
-			VirtualMemBeforeGC / BytesPerMeg, VirtualMemAfterGC / BytesPerMeg, VirtualMemFreed / BytesPerMeg
-		);
-
-		DumpMemStats();
-	}
+	int32 NumObjectsAfterGC = GUObjectArray.GetObjectArrayNumMinusAvailable();
+	FGenericMemoryStats AllocatorStatsAfterGC;
+	GMalloc->GetAllocatorStats(AllocatorStatsAfterGC);
+#if ENABLE_LOW_LEVEL_MEM_TRACKER
+	FLowLevelMemTracker::Get().UpdateStatsPerFrame();
+#endif
 
-	if (TickResults & UCookOnTheFlyServer::COSR_RequiresGC_OOM)
-	{
-		COTFS.EvaluateGarbageCollectionResults(NumObjectsBeforeGC, MemStatsBeforeGC, NumObjectsAfterGC, MemStatsAfterGC);
-	} 
+	bool bWasDueToOOM = (TickResults & UCookOnTheFlyServer::COSR_RequiresGC_OOM) != 0;
+	COTFS.EvaluateGarbageCollectionResults(bWasDueToOOM, bPartialGC,
+		NumObjectsBeforeGC, MemStatsBeforeGC, AllocatorStatsBeforeGC,
+		NumObjectsAfterGC, MemStatsAfterGC, AllocatorStatsAfterGC);
 }
diff --git a/Engine/Source/Editor/UnrealEd/Private/CookOnTheFlyServer.cpp b/Engine/Source/Editor/UnrealEd/Private/CookOnTheFlyServer.cpp
--- a/Engine/Source/Editor/UnrealEd/Private/CookOnTheFlyServer.cpp
+++ b/Engine/Source/Editor/UnrealEd/Private/CookOnTheFlyServer.cpp
@@ -58,6 +58,7 @@
 #include "GlobalShader.h"
 #include "HAL/FileManager.h"
 #include "HAL/IConsoleManager.h"
+#include "HAL/MemoryMisc.h"
 #include "HAL/PlatformApplicationMisc.h"
 #include "HAL/PlatformFileManager.h"
 #include "HAL/PlatformProcess.h"
@@ -4398,8 +4399,18 @@ void UCookOnTheFlyServer::TickPrecacheObjectsForPlatforms(const float TimeSlice,
 
 }
 
-bool UCookOnTheFlyServer::HasExceededMaxMemory() const
+bool UCookOnTheFlyServer::HasExceededMaxMemory()
 {
+	// If we are in a cooldown period, return false.
+	if (ExceededMaxMemoryCooldownEndTimeSeconds > 0)
+	{
+		if (ExceededMaxMemoryCooldownEndTimeSeconds > FPlatformTime::Seconds())
+		{
+			return false;
+		}
+		ExceededMaxMemoryCooldownEndTimeSeconds = 0.; // Return to 0 to avoid needing to call FPlatformTime::Seconds() on every call
+	}
+
 #if UE_GC_TRACK_OBJ_AVAILABLE
 	if (GUObjectArray.GetObjectArrayEstimatedAvailable() < MinFreeUObjectIndicesBeforeGC)
 	{
@@ -4416,19 +4427,12 @@ bool UCookOnTheFlyServer::HasExceededMaxMemory() const
 #endif // UE_GC_TRACK_OBJ_AVAILABLE
 
 
-	// Only report exceeded memory if all the active memory usage triggers have fired
-	int ActiveTriggers = 0;
-	int FiredTriggers = 0;
-
 	TStringBuilder<256> TriggerMessages;
-
 	const FPlatformMemoryStats MemStats = FPlatformMemory::GetStats();
 	
+	bool bMinFreeTriggered = false;
 	if (MemoryMinFreeVirtual > 0 || MemoryMinFreePhysical > 0)
 	{
-		++ActiveTriggers;
-		bool bFired = false;
-		
 		// trigger GC if we have less than MemoryMinFreeVirtual OR MemoryMinFreePhysical
 		// the check done in AssetCompilingManager is against the min of the two :
 		//uint64 AvailableMemory = FMath::Min(MemStats.AvailablePhysical, MemStats.AvailableVirtual);
@@ -4442,62 +4446,91 @@ bool UCookOnTheFlyServer::HasExceededMaxMemory() const
 		{
 			TriggerMessages.Appendf(TEXT("\n  CookSettings.MemoryMinFreeVirtual: Available virtual memory %dMiB is less than %dMiB."),
 				static_cast<uint32>(MemStats.AvailableVirtual / 1024 / 1024), static_cast<uint32>(MemoryMinFreeVirtual / 1024 / 1024));
-			bFired = true;
+			bMinFreeTriggered = true;
 		}
 #endif
 		if (MemoryMinFreePhysical > 0 && MemStats.AvailablePhysical < MemoryMinFreePhysical)
 		{
 			TriggerMessages.Appendf(TEXT("\n  CookSettings.MemoryMinFreePhysical: Available physical memory %dMiB is less than %dMiB."),
 				static_cast<uint32>(MemStats.AvailablePhysical / 1024 / 1024), static_cast<uint32>(MemoryMinFreePhysical / 1024 / 1024));
-			bFired = true;
-		}
-		if (bFired)
-		{
-			++FiredTriggers;
+			bMinFreeTriggered = true;
 		}
 	}
 
 	// if MemoryMaxUsed is set, we won't GC until at least that much mem is used
 	// this can be useful if you demand that amount of memory as your min spec
+	bool bMaxUsedTriggered = false;
 	if (MemoryMaxUsedVirtual > 0 || MemoryMaxUsedPhysical > 0)
 	{
 		// check validity of trigger :
-		if ( FiredTriggers > 0 )
+		// if the MaxUsed config exceeds the system memory, it can never be triggered and will prevent any GC :
+		uint64 MaxMaxUsed = FMath::Max(MemoryMaxUsedVirtual,MemoryMaxUsedPhysical);
+		if ( MaxMaxUsed >= MemStats.TotalPhysical )
 		{
-			// if the MaxUsed config exceeds the system memory, it can never be triggered and will prevent any GC :
-			uint64 MaxMaxUsed = FMath::Max(MemoryMaxUsedVirtual,MemoryMaxUsedPhysical);
-			if ( MaxMaxUsed >= MemStats.TotalPhysical )
-			{
-				UE_CALL_ONCE( [&](){
-					UE_LOG(LogCook, Warning, TEXT("Warning MemoryMaxUsed condition is larger than total memory (%dMiB >= %dMiB).  System does not have enough memory to cook this project."),
-				static_cast<uint32>(MaxMaxUsed / 1024 / 1024), static_cast<uint32>(MemStats.TotalPhysical / 1024 / 1024));			
-				} );
-			}
+			UE_CALL_ONCE( [&](){
+				UE_LOG(LogCook, Warning, TEXT("Warning MemoryMaxUsed condition is larger than total memory (%dMiB >= %dMiB).  System does not have enough memory to cook this project."),
+			static_cast<uint32>(MaxMaxUsed / 1024 / 1024), static_cast<uint32>(MemStats.TotalPhysical / 1024 / 1024));			
+			} );
 		}
 
-		++ActiveTriggers;
-		bool bFired = false;
 		if (MemoryMaxUsedVirtual > 0 && MemStats.UsedVirtual >= MemoryMaxUsedVirtual)
 		{
 			TriggerMessages.Appendf(TEXT("\n  CookSettings.MemoryMaxUsedVirtual: Used virtual memory %dMiB is greater than %dMiB."),
 				static_cast<uint32>(MemStats.UsedVirtual / 1024 / 1024), static_cast<uint32>(MemoryMaxUsedVirtual / 1024 / 1024));
-			bFired = true;
+			bMaxUsedTriggered = true;
 		}
 		if (MemoryMaxUsedPhysical > 0 && MemStats.UsedPhysical >= MemoryMaxUsedPhysical)
 		{
 			TriggerMessages.Appendf(TEXT("\n  CookSettings.MemoryMaxUsedPhysical: Used physical memory %dMiB is greater than %dMiB."),
 				static_cast<uint32>(MemStats.UsedPhysical / 1024 / 1024), static_cast<uint32>(MemoryMaxUsedPhysical / 1024 / 1024));
-			bFired = true;
+			bMaxUsedTriggered = true;
+		}
+	}
+
+	bool bPressureTriggered = false;
+	if (MemoryTriggerGCAtPressureLevel != FPlatformMemoryStats::EMemoryPressureStatus::Unknown)
+	{
+		FPlatformMemoryStats::EMemoryPressureStatus PressureStatus = MemStats.GetMemoryPressureStatus();
+		if (PressureStatus == FPlatformMemoryStats::EMemoryPressureStatus::Unknown)
+		{
+			UE_CALL_ONCE([&]() {
+				UE_LOG(LogCook, Error,
+					TEXT("MemoryPressureStatus is not available from the operating system. We may run out of memory due to lack of knowledge of when to collect garbage."));
+				});
+		}
+		else
+		{
+			static_assert(FPlatformMemoryStats::EMemoryPressureStatus::Critical > FPlatformMemoryStats::EMemoryPressureStatus::Nominal,
+				"We expect higher pressure to be higher integer values");
+			int RequiredValue = static_cast<int>(MemoryTriggerGCAtPressureLevel);
+			int CurrentValue = static_cast<int>(PressureStatus);
+			if (CurrentValue >= RequiredValue)
+			{
+				bPressureTriggered = true;
+				TriggerMessages.Appendf(TEXT("\n  Operating system has signalled that memory pressure is high."));
+			}
 		}
-		if (bFired)
+	}
+
+	bool bTriggerGC = false;
+	if (bMinFreeTriggered || bMaxUsedTriggered)
+	{
+		const bool bOnlyTriggerIfBothMinFreeAndMaxUsedTrigger = true;
+		if (!bOnlyTriggerIfBothMinFreeAndMaxUsedTrigger ||
+			((bMinFreeTriggered || (MemoryMinFreeVirtual <= 0 && MemoryMinFreePhysical <= 0)) &&
+				(bMaxUsedTriggered || (MemoryMaxUsedVirtual <= 0 && MemoryMaxUsedPhysical <= 0))))
 		{
-			++FiredTriggers;
+			bTriggerGC = true;
 		}
 	}
+	if (bPressureTriggered)
+	{
+		bTriggerGC = true;
+	}
 
-	if (ActiveTriggers > 0 && FiredTriggers == ActiveTriggers)
+	if (bTriggerGC)
 	{
-		UE_LOG(LogCook, Display, TEXT("Exceeded max memory on all configured triggers:%s"), TriggerMessages.ToString());
+		UE_LOG(LogCook, Display, TEXT("Garbage collection triggered by conditions:%s"), TriggerMessages.ToString());
 		return true;
 	}
 	else
@@ -4506,8 +4539,11 @@ bool UCookOnTheFlyServer::HasExceededMaxMemory() const
 	}
 }
 
-void UCookOnTheFlyServer::EvaluateGarbageCollectionResults(int32 NumObjectsBeforeGC,
-	const FPlatformMemoryStats& MemStatsBeforeGC, int32 NumObjectsAfterGC, const FPlatformMemoryStats& MemStatsAfterGC)
+void UCookOnTheFlyServer::EvaluateGarbageCollectionResults(bool bWasDueToOOM, bool bWasPartialGC,
+	int32 NumObjectsBeforeGC, const FPlatformMemoryStats& MemStatsBeforeGC,
+	const FGenericMemoryStats& AllocatorStatsBeforeGC,
+	int32 NumObjectsAfterGC, const FPlatformMemoryStats& MemStatsAfterGC,
+	const FGenericMemoryStats& AllocatorStatsAfterGC)
 {
 	using namespace UE::Cook;
 
@@ -4515,11 +4551,16 @@ void UCookOnTheFlyServer::EvaluateGarbageCollectionResults(int32 NumObjectsBefor
 	{
 		return;
 	}
+	if (IsCookOnTheFlyMode() && !bWasDueToOOM)
+	{
+		return;
+	}
 
 	int64 NumObjectsMin = NumObjectsHistory.GetMinimum();
 	int64 NumObjectsMax = NumObjectsHistory.GetMaximum();
 	int64 NumObjectsSpread = NumObjectsMax - NumObjectsMin;
 	int64 NumObjectsFreed = NumObjectsBeforeGC - NumObjectsAfterGC;
+	int64 NumObjectsCapacity = GUObjectArray.GetObjectArrayEstimatedAvailable() + GUObjectArray.GetObjectArrayNumMinusAvailable();
 	int64 VirtualMemMin = VirtualMemoryHistory.GetMinimum();
 	int64 VirtualMemMax = VirtualMemoryHistory.GetMaximum();
 	int64 VirtualMemSpread = VirtualMemMax - VirtualMemMin;
@@ -4529,13 +4570,95 @@ void UCookOnTheFlyServer::EvaluateGarbageCollectionResults(int32 NumObjectsBefor
 
 	int64 ExpectedObjectsFreed = static_cast<int64>(MemoryExpectedFreedToSpreadRatio * static_cast<float>(NumObjectsSpread));
 	double ExpectedMemFreed = MemoryExpectedFreedToSpreadRatio * VirtualMemSpread;
-	if ((NumObjectsFreed >= ExpectedObjectsFreed || NumObjectsBeforeGC - NumObjectsMin < ExpectedObjectsFreed) &&
-		(VirtualMemFreed >= ExpectedMemFreed || VirtualMemBeforeGC - VirtualMemMin <= ExpectedMemFreed))
-	{
-		// Nothing to report, Garbage Collection was as impactful as expected
+	static bool bCookMemoryAnalysis = FParse::Param(FCommandLine::Get(), TEXT("CookMemoryAnalysis"));
+#if ENABLE_LOW_LEVEL_MEM_TRACKER
+	// When tracking memory with LLM, always show the memory status.
+	const bool bAlwaysShowAnalysis = FLowLevelMemTracker::Get().IsEnabled() || bCookMemoryAnalysis;
+#else
+	const bool bAlwaysShowAnalysis = bCookMemoryAnalysis;
+#endif		
+
+	bool bWasImpactful =
+		(NumObjectsFreed >= ExpectedObjectsFreed || NumObjectsBeforeGC - NumObjectsMin < ExpectedObjectsFreed) &&
+		(VirtualMemFreed >= ExpectedMemFreed || VirtualMemBeforeGC - VirtualMemMin <= ExpectedMemFreed);
+
+	constexpr int32 BytesPerMeg = 1000000;
+	if ((!bWasDueToOOM || bWasImpactful) && !bAlwaysShowAnalysis)
+	{
+		UE_LOG(LogCook, Display, TEXT("GarbageCollection Results:\n")
+			TEXT("\tType: %s\n")
+			TEXT("\tNumObjects:\n")
+			TEXT("\t\tCapacity:         %10d\n")
+			TEXT("\t\tBefore GC:        %10d\n")
+			TEXT("\t\tAfter GC:         %10d\n")
+			TEXT("\t\tFreed by GC:      %10d\n")
+			TEXT("\tVirtual Memory:\n")
+			TEXT("\t\tBefore GC:        %10" INT64_FMT " MB\n")
+			TEXT("\t\tAfter GC:         %10" INT64_FMT " MB\n")
+			TEXT("\t\tFreed by GC:      %10" INT64_FMT " MB\n"),
+			(bWasPartialGC ? TEXT("Partial") : TEXT("Full")),
+			NumObjectsCapacity, (int64)NumObjectsBeforeGC, (int64)NumObjectsAfterGC, NumObjectsFreed,
+			VirtualMemBeforeGC / BytesPerMeg, VirtualMemAfterGC / BytesPerMeg, VirtualMemFreed / BytesPerMeg
+		);
 		return;
 	}
 
+	// When the garbage collection is not impactful, we add a cooldown. In some cases it really was impactful and we just could not tell
+	// that it was impactful, because e.g. GMalloc is holding on to unused memory but allowing the operating system to page that unused
+	// memory out. In the case where it really was not impactful, it is better to crash on out of memory or to run slowly due to many page
+	// faults than it is to continuously collect garbage for no purpose.
+	constexpr float ExceededMaxMemoryCooldownEndTimeDuration = 60.0f;
+	FString CooldownMessage;
+	if (bWasDueToOOM && !bWasImpactful)
+	{
+		ExceededMaxMemoryCooldownEndTimeSeconds = FPlatformTime::Seconds() + ExceededMaxMemoryCooldownEndTimeDuration;
+		UE_LOG(LogCook, Display, TEXT("GarbageCollection Results: Garbage Collection was not very impactful."));
+		CooldownMessage = FString::Printf(
+			TEXT("\t\tCooldown: Further OOM GarbageCollection will be prevented for %.0f seconds and we may run out of memory.\n"),
+			ExceededMaxMemoryCooldownEndTimeDuration);
+	}
+	else
+	{
+		UE_LOG(LogCook, Display, TEXT("GarbageCollection Results:"));
+	}
+	UE_LOG(LogCook, Display, TEXT("\tMemoryAnalysis: General:\n")
+		TEXT("%s")
+		TEXT("\t\tType: %s"),
+		*CooldownMessage,
+		(bWasPartialGC ? TEXT("Partial") : TEXT("Full")));
+	UE_LOG(LogCook, Display, TEXT("\tMemoryAnalysis: NumObjects:\n")
+		TEXT("\t\tCapacity:         %10" INT64_FMT "\n")
+		TEXT("\t\tProcess Min:      %10" INT64_FMT "\n")
+		TEXT("\t\tProcess Max:      %10" INT64_FMT "\n")
+		TEXT("\t\tProcess Spread:   %10" INT64_FMT "\n")
+		TEXT("\t\tBefore GC:        %10" INT64_FMT "\n")
+		TEXT("\t\tAfter GC:         %10" INT64_FMT "\n")
+		TEXT("\t\tFreed by GC:      %10" INT64_FMT ""),
+		NumObjectsCapacity, NumObjectsMin, NumObjectsMax, NumObjectsSpread,
+		(int64)NumObjectsBeforeGC, (int64)NumObjectsAfterGC, NumObjectsFreed);
+	UE_LOG(LogCook, Display, TEXT("\tMemoryAnalysis: Virtual Memory:\n")
+		TEXT("\t\tProcess Min:      %10" INT64_FMT " MB\n")
+		TEXT("\t\tProcess Max:      %10" INT64_FMT " MB\n")
+		TEXT("\t\tProcess Spread:   %10" INT64_FMT " MB\n")
+		TEXT("\t\tBefore GC:        %10" INT64_FMT " MB\n")
+		TEXT("\t\tAfter GC:         %10" INT64_FMT " MB\n")
+		TEXT("\t\tFreed by GC:      %10" INT64_FMT " MB"),
+		VirtualMemMin / BytesPerMeg, VirtualMemMax / BytesPerMeg, VirtualMemSpread / BytesPerMeg,
+		VirtualMemBeforeGC / BytesPerMeg, VirtualMemAfterGC / BytesPerMeg, VirtualMemFreed / BytesPerMeg);
+	auto AllocatorStatsToString = [](const FGenericMemoryStats& AllocatorStats)
+	{
+		TStringBuilder<256> Writer;
+		for (const TPair<FString, SIZE_T>& Item : AllocatorStats.Data)
+		{
+			Writer << TEXT("\n\t\tItem ") << Item.Key << TEXT(" ") << (uint64)Item.Value;
+		}
+		return FString(*Writer);
+	};
+	UE_LOG(LogCook, Display, TEXT("\tMemoryAnalysis: Allocator Stats Before:%s"),
+		*AllocatorStatsToString(AllocatorStatsBeforeGC));
+	UE_LOG(LogCook, Display, TEXT("\tMemoryAnalysis: Allocator Stats After:%s"),
+		*AllocatorStatsToString(AllocatorStatsAfterGC));
+
 	TArray<UPackage*> GCKeepPackages;
 	TArray<FPackageData*> GCKeepPackageDatas;
 	TArray<UObject*> LocalGCKeepObjects;
@@ -4557,55 +4680,29 @@ void UCookOnTheFlyServer::EvaluateGarbageCollectionResults(int32 NumObjectsBefor
 	int64 NumTransitivePackages;
 	GetDirectAndTransitiveResourceSize(DirectResourceSize, TransitiveResourceSize,
 		NumDirectPackages, NumTransitivePackages, MoveTemp(DirectPackages));
-
-	constexpr int32 BytesPerMeg = 1000000;
-	UE_LOG(LogCook, Display, TEXT("Garbage Collection was not very impactful.\n")
-		TEXT("\tNumObjects:\n")
-		TEXT("\t\tProcess Min:      %10" INT64_FMT "\n")
-		TEXT("\t\tProcess Max:      %10" INT64_FMT "\n")
-		TEXT("\t\tProcess Spread:   %10" INT64_FMT "\n")
-		TEXT("\t\tBefore GC:        %10" INT64_FMT "\n")
-		TEXT("\t\tAfter GC:         %10" INT64_FMT "\n")
-		TEXT("\t\tFreed by GC:      %10" INT64_FMT "\n")
-		TEXT("\tVirtual Memory:\n")
-		TEXT("\t\tProcess Min:      %10" INT64_FMT " MB\n")
-		TEXT("\t\tProcess Max:      %10" INT64_FMT " MB\n")
-		TEXT("\t\tProcess Spread:   %10" INT64_FMT " MB\n")
-		TEXT("\t\tBefore GC:        %10" INT64_FMT " MB\n")
-		TEXT("\t\tAfter GC:         %10" INT64_FMT " MB\n")
-		TEXT("\t\tFreed by GC:      %10" INT64_FMT " MB\n")
-		TEXT("\tReferences:\n")
+	UE_LOG(LogCook, Display, TEXT("\tMemoryAnalysis: Cooker References:\n")
 		TEXT("\t\tCooker direct packages:         %10" INT64_FMT "\n")
 		TEXT("\t\tCooker transitive packages:     %10" INT64_FMT "\n")
 		TEXT("\t\tCooker direct package size:     %10" INT64_FMT " MB\n")
 		TEXT("\t\tCooker transitive package size: %10" INT64_FMT " MB\n"),
-		NumObjectsMin, NumObjectsMax, NumObjectsSpread,
-		(int64)NumObjectsBeforeGC, (int64)NumObjectsAfterGC, NumObjectsFreed,
-		VirtualMemMin / BytesPerMeg, VirtualMemMax / BytesPerMeg, VirtualMemSpread / BytesPerMeg,
-		VirtualMemBeforeGC / BytesPerMeg, VirtualMemAfterGC / BytesPerMeg, VirtualMemFreed / BytesPerMeg,
 		NumDirectPackages, NumTransitivePackages,
 		DirectResourceSize.GetTotalMemoryBytes() / BytesPerMeg,
-		TransitiveResourceSize.GetTotalMemoryBytes() / BytesPerMeg
-	);
+		TransitiveResourceSize.GetTotalMemoryBytes() / BytesPerMeg);
 
-	if (IsCookByTheBookMode())
-	{
-		UE_LOG(LogCook, Display, TEXT("See log for memory use information for UObject classes and LLM tags."));
-		UE::Cook::DumpObjClassList(CookByTheBookOptions->SessionStartupObjects);
+	UE_LOG(LogCook, Display, TEXT("See log for memory use information for UObject classes and LLM tags."));
+	UE::Cook::DumpObjClassList(CookByTheBookOptions->SessionStartupObjects);
+	GLog->Logf(TEXT("Memory Analysis: LLM Tags:"));
 #if ENABLE_LOW_LEVEL_MEM_TRACKER
-		if (FLowLevelMemTracker::Get().IsEnabled())
-		{
-			GLog->Logf(TEXT("LLM Tags:"));
-			FLowLevelMemTracker::Get().UpdateStatsPerFrame();
-			FLowLevelMemTracker::Get().DumpToLog();
-		}
-		else
+	if (FLowLevelMemTracker::Get().IsEnabled())
+	{
+		FLowLevelMemTracker::Get().DumpToLog();
+	}
+	else
 #endif
-		{
-			GLog->Logf(TEXT("LLM Tags are not displayed because llm is disabled. Run with -llm or -trace=memtag to see llm tags."));
-		}
-		GLog->Flush();
+	{
+		GLog->Logf(TEXT("LLM Tags are not displayed because llm is disabled. Run with -llm or -trace=memtag to see llm tags."));
 	}
+	GLog->Flush();
 }
 
 void UCookOnTheFlyServer::GetDirectAndTransitiveResourceSize(FResourceSizeEx& OutDirectSize,
@@ -5851,6 +5948,14 @@ void FInitializeConfigSettings::LoadLocal(const FString& InOutputDirectoryOverri
 	ReadMemorySetting(TEXT("MemoryMaxUsedPhysical"), MemoryMaxUsedPhysical);
 	ReadMemorySetting(TEXT("MemoryMinFreeVirtual"), MemoryMinFreeVirtual);
 	ReadMemorySetting(TEXT("MemoryMinFreePhysical"), MemoryMinFreePhysical);
+	FString ConfigText(TEXT("None"));
+	GConfig->GetString(TEXT("CookSettings"), TEXT("MemoryTriggerGCAtPressureLevel"), ConfigText, GEditorIni);
+	if (!LexTryParseString(MemoryTriggerGCAtPressureLevel, ConfigText))
+	{
+		UE_LOG(LogCook, Error, TEXT("Unrecognized value \"%s\" for MemoryTriggerGCAtPressureLevel. Expected None or Critical."),
+			*ConfigText);
+	}
+
 	MemoryExpectedFreedToSpreadRatio = 0.10f;
 	GConfig->GetFloat(TEXT("CookSettings"), TEXT("MemoryExpectedFreedToSpreadRatio"),
 		MemoryExpectedFreedToSpreadRatio, GEditorIni);
@@ -5864,8 +5969,10 @@ void FInitializeConfigSettings::LoadLocal(const FString& InOutputDirectoryOverri
 	
 	GConfig->GetArray(TEXT("CookSettings"), TEXT("CookOnTheFlyConfigSettingDenyList"), ConfigSettingDenyList, GEditorIni);
 
-	UE_LOG(LogCook, Display, TEXT("CookSettings for Memory: MemoryMaxUsedVirtual %dMiB, MemoryMaxUsedPhysical %dMiB, MemoryMinFreeVirtual %dMiB, MemoryMinFreePhysical %dMiB"),
-		MemoryMaxUsedVirtual / 1024 / 1024, MemoryMaxUsedPhysical / 1024 / 1024, MemoryMinFreeVirtual / 1024 / 1024, MemoryMinFreePhysical / 1024 / 1024);
+	UE_LOG(LogCook, Display, TEXT("CookSettings for Memory: MemoryMaxUsedVirtual %dMiB, MemoryMaxUsedPhysical %dMiB, ")
+		TEXT("MemoryMinFreeVirtual %dMiB, MemoryMinFreePhysical %dMiB, MemoryTriggerGCAtPressureLevel %s"),
+		MemoryMaxUsedVirtual / 1024 / 1024, MemoryMaxUsedPhysical / 1024 / 1024, MemoryMinFreeVirtual / 1024 / 1024, MemoryMinFreePhysical / 1024 / 1024,
+		*LexToString(MemoryTriggerGCAtPressureLevel));
 
 	const FConfigSection* CacheSettings = GConfig->GetSectionPrivate(TEXT("CookPlatformDataCacheSettings"), false, true, GEditorIni);
 	if (CacheSettings)
diff --git a/Engine/Source/Editor/UnrealEd/Private/Cooker/CookDirector.cpp b/Engine/Source/Editor/UnrealEd/Private/Cooker/CookDirector.cpp
--- a/Engine/Source/Editor/UnrealEd/Private/Cooker/CookDirector.cpp
+++ b/Engine/Source/Editor/UnrealEd/Private/Cooker/CookDirector.cpp
@@ -659,26 +659,19 @@ void FCookDirector::ActivateMachineResourceReduction()
 	}
 	bHasReducedMachineResources = true;
 
-	// Add MemoryInFree if it's not already set
-	if (COTFS.MemoryMinFreeVirtual > 0 || COTFS.MemoryMinFreePhysical > 0)
-	{
-		// When running a multiprocess cook, we remove the MemoryMaxUsed triggers and allow GCing based solely on MemoryMinFree
-		COTFS.MemoryMaxUsedVirtual = 0;
-		COTFS.MemoryMaxUsedPhysical = 0;
-	}
-	else
-	{
-		// If MemoryMinFree is not set, then keep MemoryMaxUsed but reduce it by the number of CookWorkers.
-		constexpr float FixedOverheadFraction = 0.10f;
-		float TargetRatio = (FixedOverheadFraction + (1 - FixedOverheadFraction) / RequestedCookWorkerCount);
-		COTFS.MemoryMaxUsedPhysical = TargetRatio * COTFS.MemoryMaxUsedPhysical;
-		COTFS.MemoryMaxUsedVirtual = TargetRatio * COTFS.MemoryMaxUsedVirtual;
-	}
+	// When running a multiprocess cook, we remove the Memory triggers and trigger GC based solely on PressureLevel
+	COTFS.MemoryMaxUsedPhysical = 0;
+	COTFS.MemoryMaxUsedVirtual = 0;
+	COTFS.MemoryMinFreeVirtual = 0;
+	COTFS.MemoryMinFreePhysical = 0;
+	COTFS.MemoryTriggerGCAtPressureLevel = FGenericPlatformMemoryStats::EMemoryPressureStatus::Critical;
 
 	UE_LOG(LogCook, Display, TEXT("CookMultiprocess changed CookSettings for Memory: MemoryMaxUsedVirtual %dMiB, MemoryMaxUsedPhysical %dMiB,")
-		TEXT("MemoryMinFreeVirtual % dMiB, MemoryMinFreePhysical % dMiB"),
+		TEXT("MemoryMinFreeVirtual % dMiB, MemoryMinFreePhysical % dMiB, MemoryTriggerGCAtPressureLevel %s"),
 		COTFS.MemoryMaxUsedVirtual / 1024 / 1024, COTFS.MemoryMaxUsedPhysical / 1024 / 1024,
-		COTFS.MemoryMinFreeVirtual / 1024 / 1024, COTFS.MemoryMinFreePhysical / 1024 / 1024);
+		COTFS.MemoryMinFreeVirtual / 1024 / 1024, COTFS.MemoryMinFreePhysical / 1024 / 1024,
+		*LexToString(COTFS.MemoryTriggerGCAtPressureLevel)
+		);
 
 	// Set CoreLimit for updating workerthreads in this process and passing to the commandline for workers
 	int32 NumProcesses = RequestedCookWorkerCount + 1;
diff --git a/Engine/Source/Editor/UnrealEd/Private/Cooker/CookProfiling.cpp b/Engine/Source/Editor/UnrealEd/Private/Cooker/CookProfiling.cpp
--- a/Engine/Source/Editor/UnrealEd/Private/Cooker/CookProfiling.cpp
+++ b/Engine/Source/Editor/UnrealEd/Private/Cooker/CookProfiling.cpp
@@ -567,7 +567,7 @@ void DumpObjClassList(TConstArrayView<FWeakObjectPtr> InitialObjects)
 		});
 
 
-	LogAr.Logf(TEXT("New Objects of each class and the top roots keeping them alive:"));
+	LogAr.Logf(TEXT("Memory Analysis: New Objects of each class and the top roots keeping them alive:"));
 	LogAr.Logf(TEXT("\t%6s %s"), TEXT("Count"), TEXT("ClassPath"));
 	LogAr.Logf(TEXT("\t\t%6s %s"), TEXT("Count"), TEXT("RootObjectAndChain"));
 	TStringBuilder<1024> RootObjectString;
diff --git a/Engine/Source/Editor/UnrealEd/Private/Cooker/CookTypes.cpp b/Engine/Source/Editor/UnrealEd/Private/Cooker/CookTypes.cpp
--- a/Engine/Source/Editor/UnrealEd/Private/Cooker/CookTypes.cpp
+++ b/Engine/Source/Editor/UnrealEd/Private/Cooker/CookTypes.cpp
@@ -184,6 +184,7 @@ FCbWriter& operator<<(FCbWriter& Writer, const UE::Cook::FInitializeConfigSettin
 	Writer << "MemoryMaxUsedPhysical" << OutValue.MemoryMaxUsedPhysical;
 	Writer << "MemoryMinFreeVirtual" << OutValue.MemoryMinFreeVirtual;
 	Writer << "MemoryMinFreePhysical" << OutValue.MemoryMinFreePhysical;
+	Writer << "MemoryTriggerGCAtPressureLevel" << static_cast<uint8>(OutValue.MemoryTriggerGCAtPressureLevel);
 	Writer << "MinFreeUObjectIndicesBeforeGC" << OutValue.MinFreeUObjectIndicesBeforeGC;
 	Writer << "MaxNumPackagesBeforePartialGC" << OutValue.MaxNumPackagesBeforePartialGC;
 	Writer << "ConfigSettingDenyList" << OutValue.ConfigSettingDenyList;
@@ -207,6 +208,16 @@ bool LoadFromCompactBinary(FCbFieldView Field, UE::Cook::FInitializeConfigSettin
 	bOk = LoadFromCompactBinary(Field["MemoryMaxUsedPhysical"], OutValue.MemoryMaxUsedPhysical) & bOk;
 	bOk = LoadFromCompactBinary(Field["MemoryMinFreeVirtual"], OutValue.MemoryMinFreeVirtual) & bOk;
 	bOk = LoadFromCompactBinary(Field["MemoryMinFreePhysical"], OutValue.MemoryMinFreePhysical) & bOk;
+	uint8 PressureLevelAsInt;
+	if (LoadFromCompactBinary(Field["MemoryTriggerGCAtPressureLevel"], PressureLevelAsInt))
+	{
+		OutValue.MemoryTriggerGCAtPressureLevel = static_cast<FGenericPlatformMemoryStats::EMemoryPressureStatus>(PressureLevelAsInt);
+	}
+	else
+	{
+		OutValue.MemoryTriggerGCAtPressureLevel = FGenericPlatformMemoryStats::EMemoryPressureStatus::Unknown;
+		bOk = false;
+	}
 	bOk = LoadFromCompactBinary(Field["MinFreeUObjectIndicesBeforeGC"], OutValue.MinFreeUObjectIndicesBeforeGC) & bOk;
 	bOk = LoadFromCompactBinary(Field["MaxNumPackagesBeforePartialGC"], OutValue.MaxNumPackagesBeforePartialGC) & bOk;
 	bOk = LoadFromCompactBinary(Field["ConfigSettingDenyList"], OutValue.ConfigSettingDenyList) & bOk;
@@ -232,6 +243,7 @@ void FInitializeConfigSettings::MoveOrCopy(SourceType&& Source, TargetType&& Tar
 	Target.MemoryMaxUsedPhysical = Source.MemoryMaxUsedPhysical;
 	Target.MemoryMinFreeVirtual = Source.MemoryMinFreeVirtual;
 	Target.MemoryMinFreePhysical = Source.MemoryMinFreePhysical;
+	Target.MemoryTriggerGCAtPressureLevel = Source.MemoryTriggerGCAtPressureLevel;
 	Target.MinFreeUObjectIndicesBeforeGC = Source.MinFreeUObjectIndicesBeforeGC;
 	Target.MaxNumPackagesBeforePartialGC = Source.MaxNumPackagesBeforePartialGC;
 	Target.ConfigSettingDenyList = MoveTempIfPossible(Source.ConfigSettingDenyList);
@@ -263,6 +275,27 @@ void FBeginCookConfigSettings::CopyFromLocal(const UCookOnTheFlyServer& COTFS)
 
 }
 
+bool LexTryParseString(FPlatformMemoryStats::EMemoryPressureStatus& OutValue, FStringView Text)
+{
+	if (Text == TEXTVIEW("None")) { OutValue = FPlatformMemoryStats::EMemoryPressureStatus::Unknown; return true; }
+	if (Text == TEXTVIEW("Unknown")) { OutValue = FPlatformMemoryStats::EMemoryPressureStatus::Unknown; return true; }
+	if (Text == TEXTVIEW("Nominal")) { OutValue = FPlatformMemoryStats::EMemoryPressureStatus::Nominal; return true; }
+	if (Text == TEXTVIEW("Critical")) { OutValue = FPlatformMemoryStats::EMemoryPressureStatus::Critical; return true; }
+	OutValue = FPlatformMemoryStats::EMemoryPressureStatus::Unknown;
+	return false;
+}
+
+FString LexToString(FPlatformMemoryStats::EMemoryPressureStatus Value)
+{
+	switch (Value)
+	{
+	case FPlatformMemoryStats::EMemoryPressureStatus::Unknown: return FString(TEXTVIEW("None"));
+	case FPlatformMemoryStats::EMemoryPressureStatus::Nominal: return FString(TEXTVIEW("Nominal"));
+	case  FPlatformMemoryStats::EMemoryPressureStatus::Critical: return FString(TEXTVIEW("Critical"));
+	default: return FString(TEXTVIEW("None"));
+	}
+}
+
 FCbWriter& operator<<(FCbWriter& Writer, const UE::Cook::FBeginCookConfigSettings& OutValue)
 {
 	Writer.BeginObject();
diff --git a/Engine/Source/Editor/UnrealEd/Private/Cooker/CookTypes.h b/Engine/Source/Editor/UnrealEd/Private/Cooker/CookTypes.h
--- a/Engine/Source/Editor/UnrealEd/Private/Cooker/CookTypes.h
+++ b/Engine/Source/Editor/UnrealEd/Private/Cooker/CookTypes.h
@@ -292,6 +292,7 @@ namespace UE::Cook
 		uint64 MemoryMaxUsedPhysical;
 		uint64 MemoryMinFreeVirtual;
 		uint64 MemoryMinFreePhysical;
+		FGenericPlatformMemoryStats::EMemoryPressureStatus MemoryTriggerGCAtPressureLevel;
 		int32 MinFreeUObjectIndicesBeforeGC;
 		int32 MaxNumPackagesBeforePartialGC;
 		TArray<FString> ConfigSettingDenyList;
@@ -353,6 +354,9 @@ namespace UE::Cook
 	bool IsCookIgnoreTimeouts();
 }
 
+bool LexTryParseString(FPlatformMemoryStats::EMemoryPressureStatus& OutValue, FStringView Text);
+FString LexToString(FPlatformMemoryStats::EMemoryPressureStatus Value);
+
 inline void RouteBeginCacheForCookedPlatformData(UObject* Obj, const ITargetPlatform* TargetPlatform)
 {
 	UE_SCOPED_TEXT_COOKTIMER(*WriteToString<128>(GetClassTraceScope(Obj), TEXT("_BeginCacheForCookedPlatformData")));
diff --git a/Engine/Source/Runtime/Core/Private/Android/AndroidPlatformMemory.cpp b/Engine/Source/Runtime/Core/Private/Android/AndroidPlatformMemory.cpp
--- a/Engine/Source/Runtime/Core/Private/Android/AndroidPlatformMemory.cpp
+++ b/Engine/Source/Runtime/Core/Private/Android/AndroidPlatformMemory.cpp
@@ -294,7 +294,7 @@ FPlatformMemoryStats FAndroidPlatformMemory::GetStats()
 	return MemoryStats;
 }
 
-FGenericPlatformMemoryStats::EMemoryPressureStatus FPlatformMemoryStats::GetMemoryPressureStatus()
+FGenericPlatformMemoryStats::EMemoryPressureStatus FPlatformMemoryStats::GetMemoryPressureStatus() const
 {
 	// convert Android's TRIM status to FGenericPlatformMemoryStats::EMemoryPressureStatus.
 	auto AndroidTRIMToMemPressureStatus = [](FAndroidPlatformMemory::ETrimValues LastTrimMemoryState)
diff --git a/Engine/Source/Runtime/Core/Private/GenericPlatform/GenericPlatformMemory.cpp b/Engine/Source/Runtime/Core/Private/GenericPlatform/GenericPlatformMemory.cpp
--- a/Engine/Source/Runtime/Core/Private/GenericPlatform/GenericPlatformMemory.cpp
+++ b/Engine/Source/Runtime/Core/Private/GenericPlatform/GenericPlatformMemory.cpp
@@ -139,7 +139,7 @@ FGenericPlatformMemoryStats::FGenericPlatformMemoryStats()
 	, PeakUsedVirtual( 0 )
 {}
 
-FGenericPlatformMemoryStats::EMemoryPressureStatus FGenericPlatformMemoryStats::GetMemoryPressureStatus()
+FGenericPlatformMemoryStats::EMemoryPressureStatus FGenericPlatformMemoryStats::GetMemoryPressureStatus() const
 {
 	if (GenericPlatformMemory::GMemoryPressureCriticalThresholdMB > 0)
 	{
diff --git a/Engine/Source/Runtime/Core/Private/Windows/WindowsPlatformMemory.cpp b/Engine/Source/Runtime/Core/Private/Windows/WindowsPlatformMemory.cpp
--- a/Engine/Source/Runtime/Core/Private/Windows/WindowsPlatformMemory.cpp
+++ b/Engine/Source/Runtime/Core/Private/Windows/WindowsPlatformMemory.cpp
@@ -26,6 +26,8 @@
 #include "Stats/Stats.h"
 #include "Windows/WindowsHWrapper.h"
 
+#include <MemoryApi.h> // Include after WindowsHWrapper.h
+
 #if ENABLE_LOW_LEVEL_MEM_TRACKER && MIMALLOC_ENABLED
 #include "ThirdParty/IncludeMimAlloc.h"
 #include "VirtualAllocPageStatus.h"
@@ -265,6 +267,8 @@ FPlatformMemoryStats FWindowsPlatformMemory::GetStats()
 
 	FPlatformMemoryStats MemoryStats;
 
+	static HANDLE MemoryResourceNotificationHandle = CreateMemoryResourceNotification(LowMemoryResourceNotification);
+
 	// Gather platform memory stats.
 	MEMORYSTATUSEX MemoryStatusEx;
 	FPlatformMemory::Memzero( &MemoryStatusEx, sizeof( MemoryStatusEx ) );
@@ -293,6 +297,17 @@ FPlatformMemoryStats FWindowsPlatformMemory::GetStats()
 	MemoryStats.UsedVirtual = ProcessMemoryCounters.PagefileUsage;
 	MemoryStats.PeakUsedVirtual = ProcessMemoryCounters.PeakPagefileUsage;
 
+	BOOL bIsInLowMemoryState;
+	if (MemoryResourceNotificationHandle && QueryMemoryResourceNotification(MemoryResourceNotificationHandle, &bIsInLowMemoryState))
+	{
+		MemoryStats.MemoryPressureStatus = bIsInLowMemoryState ? FPlatformMemoryStats::EMemoryPressureStatus::Critical :
+			FPlatformMemoryStats::EMemoryPressureStatus::Nominal;
+	}
+	else
+	{
+		MemoryStats.MemoryPressureStatus = MemoryStats.FGenericPlatformMemoryStats::GetMemoryPressureStatus();
+	}
+
 	if ( GWindowsPlatformMemoryGetStatsLimitTotalGB > 0 )
 	{
 		// if GWindowsPlatformMemoryGetStatsLimitTotalGB is set
diff --git a/Engine/Source/Runtime/Core/Public/Android/AndroidPlatformMemory.h b/Engine/Source/Runtime/Core/Public/Android/AndroidPlatformMemory.h
--- a/Engine/Source/Runtime/Core/Public/Android/AndroidPlatformMemory.h
+++ b/Engine/Source/Runtime/Core/Public/Android/AndroidPlatformMemory.h
@@ -13,7 +13,7 @@
  */
 struct FPlatformMemoryStats : public FGenericPlatformMemoryStats
 {
-	FGenericPlatformMemoryStats::EMemoryPressureStatus GetMemoryPressureStatus();
+	FGenericPlatformMemoryStats::EMemoryPressureStatus GetMemoryPressureStatus() const;
 	uint64 VMSwap = 0;
 };
 
diff --git a/Engine/Source/Runtime/Core/Public/GenericPlatform/GenericPlatformMemory.h b/Engine/Source/Runtime/Core/Public/GenericPlatform/GenericPlatformMemory.h
--- a/Engine/Source/Runtime/Core/Public/GenericPlatform/GenericPlatformMemory.h
+++ b/Engine/Source/Runtime/Core/Public/GenericPlatform/GenericPlatformMemory.h
@@ -159,7 +159,7 @@ struct CORE_API FGenericPlatformMemoryStats : public FPlatformMemoryConstants
 		Nominal, 
 		Critical, // high risk of OOM conditions
 	};
-	EMemoryPressureStatus GetMemoryPressureStatus();
+	EMemoryPressureStatus GetMemoryPressureStatus() const;
 
 	/** Default constructor, clears all variables. */
 	FGenericPlatformMemoryStats();
@@ -390,8 +390,8 @@ struct CORE_API FGenericPlatformMemory
 	static void BinnedFreeToOS( void* Ptr, SIZE_T Size );
 	
 	/**
-	 * Performs initial setup for Mimalloc.
-	 * This is a noop on platforms that do not support Mimalloc, or when MIMALLOC_ENABLED is not defined.
+	 * Performs initial setup for MiMalloc.
+	 * This is a noop on platforms that do not support MiMalloc, or when MIMALLOC_ENABLED is not defined.
 	 */
 	static void MiMallocInit()
 	{
diff --git a/Engine/Source/Runtime/Core/Public/Windows/WindowsPlatformMemory.h b/Engine/Source/Runtime/Core/Public/Windows/WindowsPlatformMemory.h
--- a/Engine/Source/Runtime/Core/Public/Windows/WindowsPlatformMemory.h
+++ b/Engine/Source/Runtime/Core/Public/Windows/WindowsPlatformMemory.h
@@ -12,8 +12,6 @@ struct FGenericMemoryStats;
 
 /**
  *	Windows implementation of the FGenericPlatformMemoryStats.
- *	At this moment it's just the same as the FGenericPlatformMemoryStats.
- *	Can be extended as shown in the following example.
  */
 struct FPlatformMemoryStats
 	: public FGenericPlatformMemoryStats
@@ -22,10 +20,18 @@ struct FPlatformMemoryStats
 	FPlatformMemoryStats()
 		: FGenericPlatformMemoryStats()
 		, WindowsSpecificMemoryStat(0)
+		, MemoryPressureStatus(EMemoryPressureStatus::Unknown)
 	{ }
 
-	/** Memory stat specific only for Windows. */
+	EMemoryPressureStatus GetMemoryPressureStatus() const
+	{
+		return MemoryPressureStatus;
+	}
+
+	/** Example of a memory stat that is specific to Windows. */
 	SIZE_T WindowsSpecificMemoryStat;
+	/** Status reported by QueryMemoryResourceNotification. */
+	EMemoryPressureStatus MemoryPressureStatus;
 };
 
 
