MPCook: Tweak logging and configuration of CookDirector to fix incorrect logs and make it easier to configure.
#rnx

diff --git a/Engine/Config/BaseEditor.ini b/Engine/Config/BaseEditor.ini
--- a/Engine/Config/BaseEditor.ini
+++ b/Engine/Config/BaseEditor.ini
@@ -343,8 +343,9 @@ IgnoreHeaderDiffs=True
 bUseFullGCAssetClassNames=true
 +FullGCAssetClassNames=/Script/Engine.World
 
-CookMultiProcessEnabled=false
-CookWorkerCount=1
+; CookProcessCount sets whether we cook singleprocess or multiprocess and how many processes there should be.
+; CookProcessCount=(1 or less) is singleprocess. CookProcessCount=(N>1) is 1 director and N-1 cookworkers.
+CookProcessCount=1
 
 +ExcludedIniParams=Game:/Script/BlueprintContext.McpContext
 +ExcludedIniParams=Editor:CookSettings
diff --git a/Engine/Source/Editor/UnrealEd/Private/CookOnTheFlyServer.cpp b/Engine/Source/Editor/UnrealEd/Private/CookOnTheFlyServer.cpp
--- a/Engine/Source/Editor/UnrealEd/Private/CookOnTheFlyServer.cpp
+++ b/Engine/Source/Editor/UnrealEd/Private/CookOnTheFlyServer.cpp
@@ -5810,13 +5810,22 @@ void UCookOnTheFlyServer::Initialize( ECookMode::Type DesiredCookMode, ECookInit
 	DirectorCookMode = WorkerRequests->GetDirectorCookMode(*this);
 	if (IsCookByTheBookMode() && !IsCookingInEditor())
 	{
-		bool bCookMultiProcess = false;;
-		GConfig->GetBool(TEXT("CookSettings"), TEXT("CookMultiProcessEnabled"), bCookMultiProcess, GEditorIni);
-		bCookMultiProcess |= FParse::Param(FCommandLine::Get(), TEXT("CookMultiProcess"));
-		bCookMultiProcess &= !FParse::Param(FCommandLine::Get(), TEXT("CookSingleProcess"));
-		if (bCookMultiProcess)
+		int32 CookProcessCount=1;
+		GConfig->GetInt(TEXT("CookSettings"), TEXT("CookProcessCount"), CookProcessCount, GEditorIni);
+		FParse::Value(FCommandLine::Get(), TEXT("-CookProcessCount="), CookProcessCount);
+		CookProcessCount = FMath::Max(1, CookProcessCount);
+		if (CookProcessCount > 1)
 		{
-			CookDirector = MakeUnique<UE::Cook::FCookDirector>(*this);
+			CookDirector = MakeUnique<UE::Cook::FCookDirector>(*this, CookProcessCount);
+			if (!CookDirector->IsMultiprocessAvailable())
+			{
+				CookDirector.Reset();
+			}
+		}
+		else
+		{
+			UE_LOG(LogCook, Display, TEXT("CookProcessCount=%d. CookMultiprocess is disabled and the cooker is running as a single process."),
+				CookProcessCount);
 		}
 	}
 
diff --git a/Engine/Source/Editor/UnrealEd/Private/Cooker/CookDirector.cpp b/Engine/Source/Editor/UnrealEd/Private/Cooker/CookDirector.cpp
--- a/Engine/Source/Editor/UnrealEd/Private/Cooker/CookDirector.cpp
+++ b/Engine/Source/Editor/UnrealEd/Private/Cooker/CookDirector.cpp
@@ -29,37 +29,44 @@ extern CORE_API int32 GNumForegroundWorkers; // TaskGraph.cpp
 namespace UE::Cook
 {
 
-FCookDirector::FCookDirector(UCookOnTheFlyServer& InCOTFS)
+FCookDirector::FCookDirector(UCookOnTheFlyServer& InCOTFS, int32 CookProcessCount)
 	: RunnableShunt(*this) 
 	, COTFS(InCOTFS)
 {
+	check(CookProcessCount > 1);
 	WorkersStalledStartTimeSeconds = MAX_flt;
 	WorkersStalledWarnTimeSeconds = MAX_flt;
 	ShutdownEvent->Reset();
 
-	ParseConfig();
+	ParseConfig(CookProcessCount);
 	ISocketSubsystem* SocketSubsystem = ISocketSubsystem::Get();
 	if (!SocketSubsystem)
 	{
-		UE_LOG(LogCook, Error, TEXT("CookDirector initialization failure: platform does not support network sockets. CookWorkers will be disabled."));
-	}
-	else
-	{
-		UE_LOG(LogCook, Display, TEXT("CookMultiprocess is enabled with %d CookWorker processes."), RequestedCookWorkerCount);
+		UE_LOG(LogCook, Error, TEXT("CookDirector initialization failure: platform does not support network sockets. CookMultiprocess is disabled and the cooker is running as a single process."));
+		bMultiprocessAvailable = false;
+		return;
 	}
+	bMultiprocessAvailable = true;
+
+	UE_LOG(LogCook, Display, TEXT("CookProcessCount=%d. CookMultiprocess is enabled with 1 CookDirector and %d %s."),
+		RequestedCookWorkerCount+1, RequestedCookWorkerCount, RequestedCookWorkerCount > 1 ? TEXT("CookWorkers") : TEXT("CookWorker"));
 
 	Register(new FLogMessagesMessageHandler());
 }
 
-void FCookDirector::ParseConfig()
+bool FCookDirector::IsMultiprocessAvailable() const
+{
+	return bMultiprocessAvailable;
+}
+
+void FCookDirector::ParseConfig(int32 CookProcessCount)
 {
 	const TCHAR* CommandLine = FCommandLine::Get();
 	FString Text;
 
 	// CookWorkerCount
-	RequestedCookWorkerCount = 3;
-	GConfig->GetInt(TEXT("CookSettings"), TEXT("CookWorkerCount"), RequestedCookWorkerCount, GEditorIni);
-	FParse::Value(CommandLine, TEXT("-CookWorkerCount="), RequestedCookWorkerCount);
+	RequestedCookWorkerCount = CookProcessCount - 1;
+	check(RequestedCookWorkerCount > 0);
 
 	// CookDirectorListenPort
 	WorkerConnectPort = Sockets::COOKDIRECTOR_DEFAULT_REQUEST_CONNECTION_PORT;
@@ -692,9 +699,9 @@ void FCookDirector::ActivateMachineResourceReduction()
 	GShaderCompilingManager->OnMachineResourcesChanged(CoreLimit, CoreIncludingHyperthreadsLimit);
 
 	UE_LOG(LogCook, Display, TEXT("CookMultiprocess changed number of cores from %d to %d."),
-		NumberOfCores, FPlatformMisc::NumberOfCores());
+		NumberOfCores, CoreLimit);
 	UE_LOG(LogCook, Display, TEXT("CookMultiprocess changed number of hyperthreads from %d to %d."),
-		HyperThreadCount, FPlatformMisc::NumberOfCoresIncludingHyperthreads());
+		HyperThreadCount, CoreIncludingHyperthreadsLimit);
 }
 
 void FCookDirector::TickWorkerConnects(ECookDirectorThread TickThread)
diff --git a/Engine/Source/Editor/UnrealEd/Private/Cooker/CookDirector.h b/Engine/Source/Editor/UnrealEd/Private/Cooker/CookDirector.h
--- a/Engine/Source/Editor/UnrealEd/Private/Cooker/CookDirector.h
+++ b/Engine/Source/Editor/UnrealEd/Private/Cooker/CookDirector.h
@@ -52,9 +52,10 @@ class FCookDirector
 {
 public:
 
-	FCookDirector(UCookOnTheFlyServer& InCOTFS);
+	FCookDirector(UCookOnTheFlyServer& InCOTFS, int32 CookProcessCount);
 	~FCookDirector();
 
+	bool IsMultiprocessAvailable() const;
 	void StartCook(const FBeginCookContext& Context);
 
 	/**
@@ -131,7 +132,7 @@ private:
 
 private:
 	/** Helper for constructor parsing. */
-	void ParseConfig();
+	void ParseConfig(int32 CookProcessCount);
 	/** Initialization helper: create the listen socket. */
 	bool TryCreateWorkerConnectSocket();
 	/**
@@ -191,6 +192,7 @@ private:
 	bool bIsFirstAssignment = true;
 	bool bCookCompleteSent = false;
 	bool bWorkersStalled = false;
+	bool bMultiprocessAvailable = false;
 
 	// Data that is read-only while the CommunicationThread is active and is readable from any thread
 	FBeginCookContextForWorker BeginCookContext;
