Show additional object changes in the undo history

#rnx

diff --git a/Engine/Source/Editor/UnrealEd/Private/EditorTransaction.cpp b/Engine/Source/Editor/UnrealEd/Private/EditorTransaction.cpp
--- a/Engine/Source/Editor/UnrealEd/Private/EditorTransaction.cpp
+++ b/Engine/Source/Editor/UnrealEd/Private/EditorTransaction.cpp
@@ -797,7 +797,7 @@ void FTransaction::Apply()
 	});
 
 	TArray<ULevel*> LevelsToCommitModelSurface;
-	for (auto ChangedObjectIt : ChangedObjects)
+	for (const TTuple<UObject*, FChangedObjectValue>& ChangedObjectIt : ChangedObjects)
 	{
 		UObject* ChangedObject = ChangedObjectIt.Key;
 		UModel* Model = Cast<UModel>(ChangedObject);
@@ -911,7 +911,7 @@ void FTransaction::Finalize()
 		return Cast<UActorComponent>(&A) != nullptr;
 	});
 
-	for (auto ChangedObjectIt : ChangedObjects)
+	for (const TTuple<UObject*, FChangedObjectValue>& ChangedObjectIt : ChangedObjects)
 	{
 		const FObjectRecord& ChangedObjectRecord = Records[ChangedObjectIt.Value.RecordIndex];
 		TSharedPtr<ITransactionObjectAnnotation> ChangedObjectTransactionAnnotation = ChangedObjectIt.Value.Annotation;		
@@ -986,11 +986,26 @@ FTransactionDiff FTransaction::GenerateDiff() const
 				if (ObjectRecord.DeltaChange.HasChanged() || ObjectRecord.SerializedObject.ObjectAnnotation)
 				{
 					// Since this transaction is not currently in an undo operation, generate a valid Guid.
-					FGuid Guid = FGuid::NewGuid();
-					TransactionDiff.DiffMap.Emplace(FName(*TransactedObject->GetPathName()), MakeShared<FTransactionObjectEvent>(
-						this->GetId(), Guid, ETransactionObjectEventType::Finalized, ETransactionObjectChangeCreatedBy::TransactionRecord, 
-						FTransactionObjectChange{ ObjectRecord.SerializedObject.ObjectId, ObjectRecord.DeltaChange }, ObjectRecord.SerializedObject.ObjectAnnotation)
-					);
+					FGuid OperationGuid = FGuid::NewGuid();
+
+					TransactionDiff.DiffMap.Emplace(
+						FName(*TransactedObject->GetPathName()), 
+						MakeShared<FTransactionObjectEvent>(this->GetId(), OperationGuid, ETransactionObjectEventType::Finalized, ETransactionObjectChangeCreatedBy::TransactionRecord, FTransactionObjectChange{ ObjectRecord.SerializedObject.ObjectId, ObjectRecord.DeltaChange }, ObjectRecord.SerializedObject.ObjectAnnotation)
+						);
+
+					if (ObjectRecord.SerializedObjectFlip.ObjectAnnotation && ObjectRecord.SerializedObject.ObjectAnnotation)
+					{
+						TMap<UObject*, FTransactionObjectChange> AdditionalObjectChanges;
+						ObjectRecord.SerializedObjectFlip.ObjectAnnotation->ComputeAdditionalObjectChanges(ObjectRecord.SerializedObject.ObjectAnnotation.Get(), AdditionalObjectChanges);
+
+						for (const TTuple<UObject*, FTransactionObjectChange>& AdditionalObjectChangePair : AdditionalObjectChanges)
+						{
+							TransactionDiff.DiffMap.Emplace(
+								FName(*AdditionalObjectChangePair.Key->GetPathName()), 
+								MakeShared<FTransactionObjectEvent>(this->GetId(), OperationGuid, ETransactionObjectEventType::Finalized, ETransactionObjectChangeCreatedBy::TransactionAnnotation, AdditionalObjectChangePair.Value, nullptr)
+								);
+						}
+					}
 				}
 			}
 		}
