- Remove Nanite FMicropolyRasterizeCS for the PSO cache again because ALL the compute shader in the shader cache for that type will be added and those are way too many - this is not driven by manual collection but will just precompile all of them.
- Fix removal of FHWRasterize shaders because of WorldGridMaterial combinations used
- Don't perform vertex declaration checks when mesh shaders are used

diff --git a/Engine/Source/Editor/UnrealEd/Private/Commandlets/ShaderPipelineCacheToolsCommandlet.cpp b/Engine/Source/Editor/UnrealEd/Private/Commandlets/ShaderPipelineCacheToolsCommandlet.cpp
--- a/Engine/Source/Editor/UnrealEd/Private/Commandlets/ShaderPipelineCacheToolsCommandlet.cpp
+++ b/Engine/Source/Editor/UnrealEd/Private/Commandlets/ShaderPipelineCacheToolsCommandlet.cpp
@@ -69,7 +69,7 @@ int32 GShaderPipelineCacheTools_ComputePSOInclusionMode = 2;
 static FAutoConsoleVariableRef CVarShaderPipelineCacheDoNotPrecompileComputePSO(
 	TEXT("r.ShaderPipelineCacheTools.IncludeComputePSODuringCook"),
 	GShaderPipelineCacheTools_ComputePSOInclusionMode,
-	TEXT("0 disables cook-time addition, 1 enables cook-time addition, 2 adds only Niagara and Nanite Programmable Raster PSOs."),
+	TEXT("0 disables cook-time addition, 1 enables cook-time addition, 2 adds only Niagara PSOs."),
 	ECVF_Default
 );
 
@@ -625,17 +625,48 @@ bool CouldBeUsedTogether(const FStableShaderKeyAndValue& A, const FStableShaderK
 	{
 		return false;
 	}
-	if (A.QualityLevel != B.QualityLevel)
+	if (A.TargetPlatform != B.TargetPlatform)
 	{
 		return false;
 	}
-	if (A.TargetPlatform != B.TargetPlatform)
+	static FName NAME_FHWRasterizeVS("FHWRasterizeVS");
+	static FName NAME_FHWRasterizeMS("FHWRasterizeMS");
+	static FName NAME_FHWRasterizePS("FHWRasterizePS");
+	if ((A.ShaderType == NAME_FHWRasterizePS && (B.ShaderType == NAME_FHWRasterizeVS || B.ShaderType == NAME_FHWRasterizeMS)) ||
+		(B.ShaderType == NAME_FHWRasterizePS && (A.ShaderType == NAME_FHWRasterizeVS || A.ShaderType == NAME_FHWRasterizeMS)))
 	{
-		return false;
+		// skip quality level and ClassNameAndObjectPath because either vertex/mesh shader or pixel shader could be from WorldGridMaterial 
+		// and then quality level could be Num and Epic, and different material name
+
+		if (A.QualityLevel != B.QualityLevel)
+		{
+			static FName NAME_NumQualityLevel("Num");
+			if (A.QualityLevel != NAME_NumQualityLevel && B.QualityLevel != NAME_NumQualityLevel)
+			{
+				return false;
+			}
+		}
+
+		if (!(A.ClassNameAndObjectPath == B.ClassNameAndObjectPath))
+		{
+			static FName NAME_WorldGridMaterial("WorldGridMaterial");
+			if (A.ClassNameAndObjectPath.ObjectClassAndPath.Num() < 3 || B.ClassNameAndObjectPath.ObjectClassAndPath.Num() < 3 ||
+				(A.ClassNameAndObjectPath.ObjectClassAndPath[2] != NAME_WorldGridMaterial && B.ClassNameAndObjectPath.ObjectClassAndPath[2] != NAME_WorldGridMaterial))
+			{
+				return false;
+			}
+		}
 	}
-	if (!(A.ClassNameAndObjectPath == B.ClassNameAndObjectPath))
+	else
 	{
-		return false;
+		if (A.QualityLevel != B.QualityLevel)
+		{
+			return false;
+		}
+		if (!(A.ClassNameAndObjectPath == B.ClassNameAndObjectPath))
+		{
+			return false;
+		}
 	}
 	return true;
 }
@@ -1781,6 +1812,11 @@ void FilterInvalidPSOs(TSet<FPipelineCacheFileFormatPSO>& InOutPSOs, const TMult
 			continue;
 		}
 
+		if (CurPSO.GraphicsDesc.MeshShader != FSHAHash())
+		{
+			continue;
+		}
+
 		if (FVertexDeclarationElementList* Existing = VSToVertexDescriptor.Find(CurPSO.GraphicsDesc.VertexShader))
 		{
 			// check if current is the same or compatible
@@ -1860,6 +1896,12 @@ void FilterInvalidPSOs(TSet<FPipelineCacheFileFormatPSO>& InOutPSOs, const TMult
 			continue;
 		}
 
+		if (CurPSO.GraphicsDesc.MeshShader != FSHAHash())
+		{
+			RetainedPSOs.Add(CurPSO);
+			continue;
+		}
+
 		if (SuspiciousVertexShaders.Contains(CurPSO.GraphicsDesc.VertexShader))
 		{
 			++NumPSOsFilteredDueToInconsistentDecls;
@@ -1947,15 +1989,14 @@ void AddComputePSOs(TSet<FPipelineCacheFileFormatPSO>& OutPSOs, const TMultiMap<
 
 	static FName NAME_SF_Compute("SF_Compute");
 	static FName NAME_NiagaraShader("FNiagaraShader");
-	static FName NAME_MicropolyRasterizeCS("FMicropolyRasterizeCS");	
 
 	for (TMultiMap<FStableShaderKeyAndValue, FSHAHash>::TConstIterator Iter(StableShaderMap); Iter; ++Iter)
 	{
 		if (Iter.Key().TargetFrequency == NAME_SF_Compute)
 		{
 			// add a new Compute PSO
-			// Check if we are only allowed to add Niagara and Nanite programmable raster PSOs
-			if (GShaderPipelineCacheTools_ComputePSOInclusionMode == 2 && (Iter.Key().ShaderType != NAME_NiagaraShader && Iter.Key().ShaderType != NAME_MicropolyRasterizeCS))
+			// Check if we are only allowed to add Niagara PSOs
+			if (GShaderPipelineCacheTools_ComputePSOInclusionMode == 2 && Iter.Key().ShaderType != NAME_NiagaraShader)
 			{
 				continue;
 			}
