BulkDataRegistry: Compute the actual PayloadId and put it into the BulkDataRegistryPayloadId cache whenever the Payload is loaded for an EditorBulkData with a PlaceholderPayloadId. Issuing Puts in this case causes future Gets to not miss; spurious misses are expensive in the DDC cache and break our analytics. Previously we would Put only when we demanded the actual Payload Id from the build system, but this did not cover all BulkDatas for which we issue Gets, since we issue Gets speculatively without knowing for sure whether the PayloadId will be needed.
#rnx

diff --git a/Engine/Source/Editor/UnrealEd/Private/EditorDomain/BulkDataRegistryImpl.cpp b/Engine/Source/Editor/UnrealEd/Private/EditorDomain/BulkDataRegistryImpl.cpp
--- a/Engine/Source/Editor/UnrealEd/Private/EditorDomain/BulkDataRegistryImpl.cpp
+++ b/Engine/Source/Editor/UnrealEd/Private/EditorDomain/BulkDataRegistryImpl.cpp
@@ -239,6 +239,7 @@ FBulkDataRegistryImpl::TryRegister(UPackage* Owner, const UE::Serialization::FEd
 			else
 			{
 				// Otherwise Keep the TargetBulkData, since it matches the location and has already calculated the PayloadId
+				// Set that it is no longer allowed to write PayloadIdToCache since it has been modified after the initial load.
 				RegisteredBulk.bAllowedToWritePayloadIdToCache = false;
 			}
 
@@ -277,6 +278,7 @@ void FBulkDataRegistryImpl::UpdateRegistrationData(UPackage* Owner, const UE::Se
 			RegisteredBulk.PackageName = PackageName;
 		}
 		RegisteredBulk.BulkData = BulkData.CopyTornOff();
+		// Set that it is no longer allowed to write PayloadIdToCache since it has been modified after the initial load.
 		RegisteredBulk.bAllowedToWritePayloadIdToCache = false;
 		RegisteredBulk.bRegistered = true;
 		RegisteredBulk.bInMemory = true;
@@ -311,6 +313,53 @@ void FBulkDataRegistryImpl::OnExitMemory(const UE::Serialization::FEditorBulkDat
 	}
 }
 
+void FBulkDataRegistryImpl::UpdatePlaceholderPayloadId(const UE::Serialization::FEditorBulkData& BulkData)
+{
+	if (!BulkData.GetIdentifier().IsValid())
+	{
+		UE_LOG(LogBulkDataRegistry, Warning, TEXT("UpdatePlaceholderPayloadId called with invalid BulkData."));
+		return;
+	}
+	if (BulkData.HasPlaceholderPayloadId())
+	{
+		UE_LOG(LogBulkDataRegistry, Warning, TEXT("UpdatePlaceholderPayloadId called with a BulkData that still has a PlaceholderPayloadId."));
+		return;
+	}
+	const FGuid& Key = BulkData.GetIdentifier();
+
+	FWriteScopeLock RegistryScopeLock(RegistryLock);
+	if (!bActive)
+	{
+		return;
+	}
+	FRegisteredBulk* Existing = Registry.Find(Key);
+	if (!Existing)
+	{
+		return;
+	}
+	if (!Existing->bRegistered ||
+		!Existing->BulkData.GetIdentifier().IsValid() ||
+		!Existing->BulkData.HasPlaceholderPayloadId() ||
+		!BulkData.LocationMatches(Existing->BulkData))
+	{
+		return;
+	}
+
+	Existing->BulkData = BulkData.CopyTornOff();
+	Existing->bAllowedToWritePayloadIdToCache = Existing->bAllowedToWritePayloadIdToCache && Existing->BulkData.CanSaveForRegistry();
+	if (Existing->bAllowedToWritePayloadIdToCache)
+	{
+		// If there is a Get still pending the PayloadId may already be in the cache. Do not Put it until we find out whether
+		// it exists. The response lambda in ReadPayloadIdsFromCache will issue the WritePayloadIdToCache if it is a miss.
+		if (!PendingPayloadIds.Find(Key))
+		{
+			WritePayloadIdToCache(Existing->PackageName, Existing->BulkData);
+			// Set that it is no longer allowed to write PayloadIdToCache since we have already written it.
+			Existing->bAllowedToWritePayloadIdToCache = false;
+		}
+	}
+}
+
 TFuture<UE::BulkDataRegistry::FMetaData> FBulkDataRegistryImpl::GetMeta(const FGuid& BulkDataId)
 {
 	bool bIsWriteLock = false;
@@ -555,18 +604,17 @@ void FBulkDataRegistryImpl::ReadPayloadIdsFromCache(FName PackageName, TArray<TR
 		UE::EditorDomain::GetBulkDataPayloadId(PackageName, NewPending->GetBulkDataId(), NewPending->GetRequestOwner(),
 			[this, PackageName, NewPending](FSharedBuffer Buffer)
 		{
-			if (Buffer.IsNull())
-			{
-				return;
-			}
-			FMemoryReaderView Reader(MakeArrayView(Buffer));
-			UE::Serialization::FEditorBulkData CachedBulkData;
-			CachedBulkData.SerializeForRegistry(Reader);
 			const FGuid& BulkDataId = NewPending->GetBulkDataId();
-			if (Reader.IsError() || CachedBulkData.GetIdentifier() != BulkDataId)
+			UE::Serialization::FEditorBulkData CachedBulkData;
+			if (!Buffer.IsNull())
 			{
-				UE_LOG(LogBulkDataRegistry, Warning, TEXT("Corrupt cache data for BulkDataPayloadId %s."), WriteToString<192>(PackageName, TEXT("/"), BulkDataId).ToString());
-				return;
+				FMemoryReaderView Reader(MakeArrayView(Buffer));
+				CachedBulkData.SerializeForRegistry(Reader);
+				if (Reader.IsError() || CachedBulkData.GetIdentifier() != BulkDataId)
+				{
+					UE_LOG(LogBulkDataRegistry, Warning, TEXT("Corrupt cache data for BulkDataPayloadId %s."), WriteToString<192>(PackageName, TEXT("/"), BulkDataId).ToString());
+					CachedBulkData = UE::Serialization::FEditorBulkData();
+				}
 			}
 
 			FWriteScopeLock RegistryScopeLock(RegistryLock);
@@ -594,10 +642,27 @@ void FBulkDataRegistryImpl::ReadPayloadIdsFromCache(FName PackageName, TArray<TR
 			}
 
 			UE::Serialization::FEditorBulkData& ExistingBulkData = ExistingRegisteredBulk->BulkData;
-			check(ExistingBulkData.GetIdentifier() == BulkDataId);
-			if (ExistingBulkData.HasPlaceholderPayloadId() && CachedBulkData.LocationMatches(ExistingBulkData))
+			if (CachedBulkData.GetIdentifier().IsValid())
+			{
+				check(ExistingBulkData.GetIdentifier() == BulkDataId);
+				if (ExistingBulkData.HasPlaceholderPayloadId() && CachedBulkData.LocationMatches(ExistingBulkData))
+				{
+					ExistingBulkData = CachedBulkData;
+					// No longer allowed to write the payloadId because we have found it already exists
+					ExistingRegisteredBulk->bAllowedToWritePayloadIdToCache = false;
+				}
+			}
+			else
 			{
-				ExistingBulkData = CachedBulkData;
+				if (!ExistingBulkData.HasPlaceholderPayloadId())
+				{
+					// In between the point where we started the cache query and we received this result,
+					// The FEditorBulkData has updated its PayloadId and informed us by calling UpdatePlaceholderPayloadId
+					// Put this locally computed PayloadId into the cache, since the cache is missing it.
+					WritePayloadIdToCache(ExistingRegisteredBulk->PackageName, ExistingBulkData);
+					// Set that it is no longer allowed to write PayloadIdToCache since we have already written it.
+					ExistingRegisteredBulk->bAllowedToWritePayloadIdToCache = false;
+				}
 			}
 		});
 	}
@@ -881,6 +946,7 @@ void FUpdatePayloadWorker::DoWork()
 			if (RegisteredBulk->bAllowedToWritePayloadIdToCache)
 			{
 				BulkDataRegistry->WritePayloadIdToCache(RegisteredBulk->PackageName, BulkData);
+				// Set that it is no longer allowed to write PayloadIdToCache since we have already written it.
 				RegisteredBulk->bAllowedToWritePayloadIdToCache = false;
 			}
 			break;
diff --git a/Engine/Source/Editor/UnrealEd/Private/EditorDomain/BulkDataRegistryImpl.h b/Engine/Source/Editor/UnrealEd/Private/EditorDomain/BulkDataRegistryImpl.h
--- a/Engine/Source/Editor/UnrealEd/Private/EditorDomain/BulkDataRegistryImpl.h
+++ b/Engine/Source/Editor/UnrealEd/Private/EditorDomain/BulkDataRegistryImpl.h
@@ -45,10 +45,37 @@ struct FRegisteredBulk
 	}
 	UE::Serialization::FEditorBulkData BulkData;
 	FName PackageName;
+	/**
+	 * True if this->BulkData has loaded its payload (because we wanted to update its PlaceholderPayloadId) and we kept
+	 * the payload in memory because we want to use it on the GetData call that we expect to come soon.
+	 */
 	bool bHasTempPayload : 1;
+	/**
+	 * For legacy bulkdatas with a PlaceholderPayloadId, and we don't know the actual PayloadId, we will load their
+	 * bulkdata and calculate the actual PayloadId when it is requested by clients such as the build system. Loading
+	 * the bulkdata to calculate this is expensive so we cache the actual PayloadId in DDC.
+	 * However, this caching is only valid for BulkDatas that are discovered during package loading and not modified
+	 * after. bAllowedToWritePayloadIdToCache is true for those bulkdatas and is set to false if they are created too
+	 * late or are modified after loading. Iff this flag is true when we load the bulkdata to calculate the PayloadId,
+	 * we write the actual PayloadId to DDC. We also set this flag to false once we know it is already in the cache.
+	 */
 	bool bAllowedToWritePayloadIdToCache : 1;
+	/**
+	 * True if the key associated with this FRegisteredBulk value has been registered by an FEditorBulkData outside
+	 * of the registry. Non-registered FRegisteredBulks can be created when we read the list of BulkDatas in the
+	 * package from DDC and get the results back before the FEditorBulkDatas are discovered during package loading.
+	 */
 	bool bRegistered : 1;
+	/**
+	 * True iff an FEditorBulkData outside of the registry has registered the key associated with this FRegisteredBulk,
+	 * AND that FEditorBulkData has not yet called OnExitMemory to indicate its destruction.
+	 */
 	bool bInMemory : 1;
+	/**
+	 * If true, then we can load the data for the BulkData to satisfy GetData requests. If false, this->BulkData is
+	 * invalid and we don't know what the data is. This can happen when a memory-only FEditorBulkData calls
+	 * OnExitMemory. We keep FRegisteredBulk around in that case solely to detect duplicate Guid bugs.
+	 */
 	bool bPayloadAvailable : 1;
 };
 
@@ -184,6 +211,7 @@ public:
 	virtual void UpdateRegistrationData(UPackage* Owner, const UE::Serialization::FEditorBulkData& BulkData) override;
 	virtual void Unregister(const UE::Serialization::FEditorBulkData& BulkData) override;
 	virtual void OnExitMemory(const UE::Serialization::FEditorBulkData& BulkData) override;
+	virtual void UpdatePlaceholderPayloadId(const UE::Serialization::FEditorBulkData& BulkData) override;
 	virtual TFuture<UE::BulkDataRegistry::FMetaData> GetMeta(const FGuid& BulkDataId) override;
 	virtual TFuture<UE::BulkDataRegistry::FData> GetData(const FGuid& BulkDataId) override;
 	virtual bool TryGetBulkData(const FGuid& BulkDataId, UE::Serialization::FEditorBulkData* OutBulk = nullptr,
@@ -218,7 +246,15 @@ private:
 	TMap<FGuid, FRegisteredBulk> Registry;
 	FResaveSizeTracker ResaveSizeTracker;
 	TMap<FGuid, FUpdatingPayload> UpdatingPayloads;
+	/**
+	 * All of the PendingPackage structs for packages that are still loading or have pending Cache Gets or Puts that
+	 * were launched when they started loading. PendingPackages are deleted when the package is done loading and the
+	 * cache operations are complete. @see FPendingPackage. These structs manage the DDC operations, both to list
+	 * the bulkdatas present in a package and to store the actual PayloadId for legacy BulkDatas with Placeholder
+	 * PayloadIds.
+	 */
 	TMap<FName, TUniquePtr<FPendingPackage>> PendingPackages;
+	/** All of the BulkDataIds for which we have a PayloadID cache Get in progress. */
 	TMap<FGuid, TRefCountPtr<FPendingPayloadId>> PendingPayloadIds;
 	TRingBuffer<FTempLoadedPayload> TempLoadedPayloads;
 	uint64 TempLoadedPayloadsSize = 0;
diff --git a/Engine/Source/Runtime/CoreUObject/Private/Serialization/BulkDataRegistry.cpp b/Engine/Source/Runtime/CoreUObject/Private/Serialization/BulkDataRegistry.cpp
--- a/Engine/Source/Runtime/CoreUObject/Private/Serialization/BulkDataRegistry.cpp
+++ b/Engine/Source/Runtime/CoreUObject/Private/Serialization/BulkDataRegistry.cpp
@@ -53,6 +53,7 @@ public:
 	virtual void UpdateRegistrationData(UPackage* Owner, const UE::Serialization::FEditorBulkData& BulkData) override {}
 	virtual void Unregister(const UE::Serialization::FEditorBulkData& BulkData) override {}
 	virtual void OnExitMemory(const UE::Serialization::FEditorBulkData& BulkData) override {}
+	virtual void UpdatePlaceholderPayloadId(const UE::Serialization::FEditorBulkData& BulkData) override {};
 	virtual TFuture<UE::BulkDataRegistry::FMetaData> GetMeta(const FGuid& BulkDataId) override
 	{
 		TPromise<UE::BulkDataRegistry::FMetaData> Promise;
diff --git a/Engine/Source/Runtime/CoreUObject/Private/Serialization/EditorBulkData.cpp b/Engine/Source/Runtime/CoreUObject/Private/Serialization/EditorBulkData.cpp
--- a/Engine/Source/Runtime/CoreUObject/Private/Serialization/EditorBulkData.cpp
+++ b/Engine/Source/Runtime/CoreUObject/Private/Serialization/EditorBulkData.cpp
@@ -1327,6 +1327,9 @@ FCompressedBuffer FEditorBulkData::LoadFromPackageFile() const
 	FCompressedBuffer PayloadFromDisk;
 	SerializeData(*BulkArchive, PayloadFromDisk, Flags);
 
+	// If PayloadId is placeholder, compute the actual now that we have the payload available to hash.
+	UpdateKeyIfNeeded(PayloadFromDisk);
+
 	return PayloadFromDisk;
 }
 
@@ -2139,14 +2142,39 @@ void FEditorBulkData::UpdateKeyIfNeeded()
 	{
 		checkf(IsDataVirtualized() == false, TEXT("Cannot have a virtualized payload if loaded from legacy BulkData")); // Sanity check
 
-		// Load the payload from disk (or memory) so that we can hash it
-		FSharedBuffer InPayload = GetDataInternal().Decompress();
-		PayloadContentId = HashPayload(InPayload);
+		// Call GetDataInternal to load the payload from disk; LoadFromPackageFile will call the version of
+		// UpdateKeyIfNeeded that takes the decompressed Payload parameter.
+		FSharedBuffer LocalPayload = GetDataInternal().Decompress();
+
+		if (EnumHasAnyFlags(Flags, EFlags::LegacyKeyWasGuidDerived))
+		{
+			ensureMsgf(false, TEXT("EditorBulkData logic error: LegacyKeyWasGuidDerived was not removed by GetDataInternal; this is unexpected and a minor performance problem."));
+			UpdateKeyIfNeeded(FCompressedBuffer::Compress(LocalPayload,
+				ECompressedBufferCompressor::NotSet, ECompressedBufferCompressionLevel::None));
+			check(!EnumHasAnyFlags(Flags, EFlags::LegacyKeyWasGuidDerived)); // UpdateKeyIfNeeded(FCompressedBuffer) guarantees the conversion, so a fatal assert is warranted
+		}
 
 		// Store as the in memory payload, since this method is only called during saving 
 		// we know it will get cleared anyway.
-		Payload = InPayload;
-		EnumRemoveFlags(Flags, EFlags::LegacyKeyWasGuidDerived);
+		Payload = LocalPayload;
+	}
+}
+
+void FEditorBulkData::UpdateKeyIfNeeded(FCompressedBuffer CompressedPayload) const
+{
+	if (EnumHasAnyFlags(Flags, EFlags::LegacyKeyWasGuidDerived))
+	{
+		// PayloadContentId and its associated flag are mutable for legacy BulkDatas.
+		// We do a const-cast instead of marking them mutable because they are immutable for non-legacy.
+		const_cast<FEditorBulkData&>(*this).PayloadContentId = CompressedPayload.GetRawHash();
+
+		EnumRemoveFlags(const_cast<FEditorBulkData&>(*this).Flags, EFlags::LegacyKeyWasGuidDerived);
+#if WITH_EDITOR
+		if (EnumHasAnyFlags(Flags, EFlags::HasRegistered))
+		{
+			IBulkDataRegistry::Get().UpdatePlaceholderPayloadId(*this);
+		}
+#endif
 	}
 }
 
diff --git a/Engine/Source/Runtime/CoreUObject/Public/Serialization/BulkDataRegistry.h b/Engine/Source/Runtime/CoreUObject/Public/Serialization/BulkDataRegistry.h
--- a/Engine/Source/Runtime/CoreUObject/Public/Serialization/BulkDataRegistry.h
+++ b/Engine/Source/Runtime/CoreUObject/Public/Serialization/BulkDataRegistry.h
@@ -89,6 +89,11 @@ public:
 	virtual void Unregister(const UE::Serialization::FEditorBulkData& BulkData) = 0;
 	/** Report that a BulkData is leaving memory and its in-memory payload (if it had one) is no longer available. */
 	virtual void OnExitMemory(const UE::Serialization::FEditorBulkData& BulkData) = 0;
+	/**
+	 * Notify that a legacy BulkData with a PlaceholderPayloadId has loaded its data and updated its PayloadId, and
+	 * the Registry should copy the updated PayloadId into the cache if it hasn't already calculated it.
+	 */
+	virtual void UpdatePlaceholderPayloadId(const UE::Serialization::FEditorBulkData& BulkData) = 0;
 
 	/** Return the metadata for the given registered BulkData; returns false if not registered. */
 	virtual TFuture<UE::BulkDataRegistry::FMetaData> GetMeta(const FGuid& BulkDataId) = 0;
diff --git a/Engine/Source/Runtime/CoreUObject/Public/Serialization/EditorBulkData.h b/Engine/Source/Runtime/CoreUObject/Public/Serialization/EditorBulkData.h
--- a/Engine/Source/Runtime/CoreUObject/Public/Serialization/EditorBulkData.h
+++ b/Engine/Source/Runtime/CoreUObject/Public/Serialization/EditorBulkData.h
@@ -403,6 +403,7 @@ private:
 	bool CanLoadDataFromDisk() const;
 
 	void UpdateKeyIfNeeded();
+	void UpdateKeyIfNeeded(FCompressedBuffer InPayload) const;
 
 	void RecompressForSerialization(FCompressedBuffer& InOutPayload, EFlags PayloadFlags) const;
 	EFlags BuildFlagsForSerialization(FArchive& Ar, bool bKeepFileDataByReference) const;
