Cooker: Skip adding empty TargetDependencies attachments; it wastefully takes space in the oplog.
#rnx

diff --git a/Engine/Source/Editor/UnrealEd/Private/CookOnTheFlyServer.cpp b/Engine/Source/Editor/UnrealEd/Private/CookOnTheFlyServer.cpp
--- a/Engine/Source/Editor/UnrealEd/Private/CookOnTheFlyServer.cpp
+++ b/Engine/Source/Editor/UnrealEd/Private/CookOnTheFlyServer.cpp
@@ -5635,10 +5635,13 @@ void FSaveCookedPackageContext::FinishPlatform()
 			Info.Status = IPackageWriter::ECommitStatus::Error;
 		}
 		Info.PackageName = Package->GetFName();
-		PRAGMA_DISABLE_DEPRECATION_WARNINGS
+		PRAGMA_DISABLE_DEPRECATION_WARNINGS;
 		Info.PackageGuid = AssetPackageData ? AssetPackageData->PackageGuid : FGuid();
-		PRAGMA_ENABLE_DEPRECATION_WARNINGS
-		Info.Attachments.Add({ "Dependencies", TargetDomainDependencies });
+		PRAGMA_ENABLE_DEPRECATION_WARNINGS;
+		if (TargetDomainDependencies)
+		{
+			Info.Attachments.Add({ "Dependencies", TargetDomainDependencies });
+		}
 		// TODO: Reenable BuildDefinitionList once FCbPackage support for empty FCbObjects is in
 		//Info.Attachments.Add({ "BuildDefinitionList", BuildDefinitionList });
 		Info.WriteOptions = IPackageWriter::EWriteOptions::Write;
