Fixed crash when exiting PIE:
- It's possible to exit PIE before destroying its Slate window.When that happen, the PIE window will try to use a FScopedSwitchWorldHack to switch GWorld to its PIE instance world and then restore GWorld to the editor world. Given that we're now out of PIE, the PIE instance world is null and thus we don't switch to any PIE world during the constructor of FScopedSwitchWorldHack but its destructor will still try to restore the editor world. At this point, the EditorWorld of PlayLevel is nullptr so we crash.
- The fix is to prevent OnSwitchWorldForSlatePieWindow to tell FScopedSwitchWorldHack that it'll need to restore the editor world if we didn't succeed the switch to the PIE world.

diff --git a/Engine/Source/Editor/UnrealEd/Private/PlayLevel.cpp b/Engine/Source/Editor/UnrealEd/Private/PlayLevel.cpp
--- a/Engine/Source/Editor/UnrealEd/Private/PlayLevel.cpp
+++ b/Engine/Source/Editor/UnrealEd/Private/PlayLevel.cpp
@@ -2238,8 +2238,13 @@ int32 UEditorEngine::OnSwitchWorldForSlatePieWindow(int32 WorldID, int32 WorldPI
 	{
 		// When we have an invalid world id we always switch to the pie world in the PIE window
 		OnSwitchWorldsForPIEInstance(WorldPIEInstance);
-		// The editor world was active restore it later
-		RestoreID = EditorWorldID;
+
+		// Make sure the switch to the PIE world was successful
+		if (GIsPlayInEditorWorld)
+		{
+			// The editor world was active restore it later
+			RestoreID = EditorWorldID;
+		}
 	}
 	else if(WorldID == PieWorldID && !GIsPlayInEditorWorld)
 	{
