Fixed issue with string processing in the moderation commandlet.
Reformated manifest so it outputs a single string json file
Prevented live lock in LoadPackageLinker code looks like it's caused by some ZenLoader incompatibility.

diff --git a/Engine/Source/Editor/UnrealEd/Private/Commandlets/BaseIteratePackagesCommandlet.cpp b/Engine/Source/Editor/UnrealEd/Private/Commandlets/BaseIteratePackagesCommandlet.cpp
--- a/Engine/Source/Editor/UnrealEd/Private/Commandlets/BaseIteratePackagesCommandlet.cpp
+++ b/Engine/Source/Editor/UnrealEd/Private/Commandlets/BaseIteratePackagesCommandlet.cpp
@@ -576,8 +576,9 @@ void UBaseIteratePackagesCommandlet::LoadAndSaveOnePackage(const FString& Filena
 
 		// Get the package linker.
 		VerboseMessage(TEXT("Pre GetPackageLinker"));
-
-		FLinkerLoad* Linker = LoadPackageLinker(nullptr, PackagePath, LOAD_NoVerify);
+#define USE_OLD_LOADLINKERPATH 0
+#if USE_OLD_LOADLINKERPATH
+		FLinkerLoad* Linker = LoadPackageLinker(nullptr, PackagePath, LOAD_NoVerify | LOAD_SkipLoadImportedPackages);
 	
 		// Bail early if we don't have a valid linker (package was out of date, etc)
 		if( !Linker )
@@ -595,6 +596,9 @@ void UBaseIteratePackagesCommandlet::LoadAndSaveOnePackage(const FString& Filena
 			PerformPreloadOperations(Linker, bSavePackage);
 		}
 		VerboseMessage(FString::Printf(TEXT("Post PerformPreloadOperations, Resave? %d"), bSavePackage));
+#else
+		bool bSavePackage = true;
+#endif
 		
 		if (bSavePackage)
 		{
@@ -626,7 +630,11 @@ void UBaseIteratePackagesCommandlet::LoadAndSaveOnePackage(const FString& Filena
 					check(Package);
 				}
 			}
-
+#if !USE_OLD_LOADLINKERPATH
+			FLinkerLoad* Linker = Package->GetLinker();
+			PerformPreloadOperations(Linker, bSavePackage); 
+			VerboseMessage(FString::Printf(TEXT("Post PerformPreloadOperations, Resave? %d"), bSavePackage));
+#endif
 			VerboseMessage(TEXT("Post LoadPackage"));
 
 			// if we are only saving dirty packages and the package is not dirty, then we do not want to save the package (remember the default behavior is to ALWAYS save the package)
@@ -725,6 +733,10 @@ void UBaseIteratePackagesCommandlet::LoadAndSaveOnePackage(const FString& Filena
 						TArray<UObject*> DependantObjects;
 						ForEachObjectWithPackage(Package, [this, &bSavePackage](UObject* Object)
 							{
+								if (!IsValid(Object))
+								{
+									return true;
+								}
 								if (!Cast<UMetaData>(Object))
 								{
 									PerformWorldBuilderAdditionalOperations(Object, bSavePackage);
