UE-167070 - Batch 4 of UnrealEd truncation warnings

diff --git a/Engine/Source/Editor/UnrealEd/Private/DataTableEditorUtils.cpp b/Engine/Source/Editor/UnrealEd/Private/DataTableEditorUtils.cpp
--- a/Engine/Source/Editor/UnrealEd/Private/DataTableEditorUtils.cpp
+++ b/Engine/Source/Editor/UnrealEd/Private/DataTableEditorUtils.cpp
@@ -813,7 +813,7 @@ void FDataTableEditorUtils::CacheDataForEditing(const UScriptStruct* RowStruct,
 			CachedColumnData->Property = Prop;
 		}
 
-		CachedColumnData->DesiredColumnWidth = FontMeasure->Measure(CachedColumnData->DisplayName, CellTextStyle.Font).X + CellPadding;
+		CachedColumnData->DesiredColumnWidth = static_cast<float>(FontMeasure->Measure(CachedColumnData->DisplayName, CellTextStyle.Font).X + CellPadding);
 
 		OutAvailableColumns.Add(CachedColumnData);
 	}
@@ -856,9 +856,9 @@ void FDataTableEditorUtils::CacheDataForEditing(const UScriptStruct* RowStruct,
 
 				const FVector2D CellTextSize = FontMeasure->Measure(CellText, CellTextStyle.Font);
 
-				CachedRowData->DesiredRowHeight = FMath::Max(CachedRowData->DesiredRowHeight, CellTextSize.Y);
+				CachedRowData->DesiredRowHeight = static_cast<float>(FMath::Max(CachedRowData->DesiredRowHeight, CellTextSize.Y));
 
-				const float CellWidth = CellTextSize.X + CellPadding;
+				const float CellWidth = static_cast<float>(CellTextSize.X + CellPadding);
 				CachedColumnData->DesiredColumnWidth = FMath::Max(CachedColumnData->DesiredColumnWidth, CellWidth);
 			}
 		}
diff --git a/Engine/Source/Editor/UnrealEd/Private/DragTool_BoxSelect.cpp b/Engine/Source/Editor/UnrealEd/Private/DragTool_BoxSelect.cpp
--- a/Engine/Source/Editor/UnrealEd/Private/DragTool_BoxSelect.cpp
+++ b/Engine/Source/Editor/UnrealEd/Private/DragTool_BoxSelect.cpp
@@ -309,11 +309,14 @@ void FDragTool_ActorBoxSelect::CalculateBox( FBox& OutBox )
 
 	FSceneView* View = LevelViewportClient->CalcSceneView(&ViewFamily);
 
-	FVector4 StartScreenPos = View->PixelToScreen(Start.X, Start.Y, 0);
-	FVector4 EndScreenPos = View->PixelToScreen(End.X, End.Y, 0);
+	FVector3f StartFloat{ Start };
+	FVector3f EndFloat{ End };
 
-	FVector TransformedStart = View->ScreenToWorld(View->PixelToScreen(Start.X, Start.Y, 0.5f));
-	FVector TransformedEnd = View->ScreenToWorld(View->PixelToScreen(End.X, End.Y, 0.5f));
+	FVector4 StartScreenPos = View->PixelToScreen(StartFloat.X, StartFloat.Y, 0);
+	FVector4 EndScreenPos = View->PixelToScreen(EndFloat.X, EndFloat.Y, 0);
+
+	FVector TransformedStart = View->ScreenToWorld(View->PixelToScreen(StartFloat.X, StartFloat.Y, 0.5f));
+	FVector TransformedEnd = View->ScreenToWorld(View->PixelToScreen(EndFloat.X, EndFloat.Y, 0.5f));
 
 	// Create a bounding box based on the start/end points (normalizes the points).
 	OutBox.Init();
diff --git a/Engine/Source/Editor/UnrealEd/Private/DragTool_FrustumSelect.cpp b/Engine/Source/Editor/UnrealEd/Private/DragTool_FrustumSelect.cpp
--- a/Engine/Source/Editor/UnrealEd/Private/DragTool_FrustumSelect.cpp
+++ b/Engine/Source/Editor/UnrealEd/Private/DragTool_FrustumSelect.cpp
@@ -153,8 +153,15 @@ void FDragTool_ActorFrustumSelect::EndDrag()
 		else
 		{
 			// Extend the endpoint of the rect to get the actual line
-			// LWC_TODO: Large world scale selection needed.
-			FIntRect BoxRect(FIntPoint(FMath::Max<double>(0.0, Start.X), FMath::Max<double>(0.0, Start.Y)), FIntPoint(FMath::Min(ViewportSizeX, FMath::TruncToInt(End.X + 1)), FMath::Min(ViewportSizeY, FMath::TruncToInt(End.Y + 1))));
+
+			const int32 MinX = UE::LWC::FloatToIntCastChecked<int32>(FMath::Max<double>(0.0, Start.X));
+			const int32 MinY = UE::LWC::FloatToIntCastChecked<int32>(FMath::Max<double>(0.0, Start.Y));
+			const int32 MaxX = FMath::Min(ViewportSizeX, FMath::TruncToInt32(End.X + 1.0));
+			const int32 MaxY = FMath::Min(ViewportSizeY, FMath::TruncToInt32(End.Y + 1.0));
+
+			const FIntPoint Min{ MinX, MinY };
+			const FIntPoint Max{ MaxX, MaxY };
+			const FIntRect BoxRect{ Min, Max };
 
 			// Typed Element selection
 			{
diff --git a/Engine/Source/Editor/UnrealEd/Private/DragTool_Measure.cpp b/Engine/Source/Editor/UnrealEd/Private/DragTool_Measure.cpp
--- a/Engine/Source/Editor/UnrealEd/Private/DragTool_Measure.cpp
+++ b/Engine/Source/Editor/UnrealEd/Private/DragTool_Measure.cpp
@@ -37,7 +37,8 @@ FVector2D FDragTool_Measure::GetSnappedPixelPos(FVector2D PixelPos)
 	FSceneView* View = ViewportClient->CalcSceneView(&ViewFamily);
 
 	// Put the mouse pos in world space
-	FVector WorldPos = View->ScreenToWorld(View->PixelToScreen(PixelPos.X, PixelPos.Y, 0.5f));;
+	FVector2f PixelPosFloat{ PixelPos };
+	FVector WorldPos = View->ScreenToWorld(View->PixelToScreen(PixelPosFloat.X, PixelPosFloat.Y, 0.5f));;
 
 	// Snap the world position
 	const float GridSize = GEditor->GetGridSize();
@@ -73,7 +74,7 @@ void FDragTool_Measure::AddDelta(const FVector& InDelta)
 void FDragTool_Measure::Render(const FSceneView* View, FCanvas* Canvas)
 {
 	const float OrthoUnitsPerPixel = ViewportClient->GetOrthoUnitsPerPixel(ViewportClient->Viewport);
-	const float Length = FMath::RoundToFloat((PixelEnd - PixelStart).Size() * OrthoUnitsPerPixel * ViewportClient->GetDPIScale());
+	const float Length = FMath::RoundToFloat(FVector2f{ PixelEnd - PixelStart }.Size() * OrthoUnitsPerPixel * ViewportClient->GetDPIScale());
 
 	if (View != nullptr && Canvas != nullptr && Length >= 1.f)
 	{
diff --git a/Engine/Source/Editor/UnrealEd/Private/DragTool_ViewportChange.cpp b/Engine/Source/Editor/UnrealEd/Private/DragTool_ViewportChange.cpp
--- a/Engine/Source/Editor/UnrealEd/Private/DragTool_ViewportChange.cpp
+++ b/Engine/Source/Editor/UnrealEd/Private/DragTool_ViewportChange.cpp
@@ -56,8 +56,8 @@ void FDragTool_ViewportChange::EndDrag()
 	}
 	else
 	{
-		float OffsetRatio = ViewOptionOffset.X / ViewOptionOffset.Y;
-		float DragAngle = FMath::RadiansToDegrees(FMath::Atan(OffsetRatio));
+		double OffsetRatio = ViewOptionOffset.X / ViewOptionOffset.Y;
+		double DragAngle = FMath::RadiansToDegrees(FMath::Atan(OffsetRatio));
 
 		if (ViewOptionOffset.Y >= 0)
 		{
@@ -99,7 +99,7 @@ void FDragTool_ViewportChange::EndDrag()
 		}
 	}
 
-	float OffsetLength = FMath::RoundToFloat((End - Start).Size());
+	double OffsetLength = FMath::RoundToFloat((End - Start).Size());
 	if (OffsetLength >= 125.15f)
 	{
 		LevelViewportClient->SetViewportTypeFromTool(ViewOption);
diff --git a/Engine/Source/Editor/UnrealEd/Private/EditorSelectUtils.cpp b/Engine/Source/Editor/UnrealEd/Private/EditorSelectUtils.cpp
--- a/Engine/Source/Editor/UnrealEd/Private/EditorSelectUtils.cpp
+++ b/Engine/Source/Editor/UnrealEd/Private/EditorSelectUtils.cpp
@@ -254,8 +254,8 @@ void UUnrealEdEngine::SetPivot( FVector NewPivot, bool bSnapPivotToGrid, bool bI
 	//if there are multiple elements selected, just use the x-axis for the "translate/rotate" or 2D widgets
 	if (NumElements == 1)
 	{
-		EditorModeTools.TranslateRotateXAxisAngle = TranslateRotateWidgetWorldXAxis.Rotation().Yaw;
-		EditorModeTools.TranslateRotate2DAngle = FMath::RadiansToDegrees(FMath::Atan2(Widget2DWorldXAxis.Z, Widget2DWorldXAxis.X));
+		EditorModeTools.TranslateRotateXAxisAngle = static_cast<float>(TranslateRotateWidgetWorldXAxis.Rotation().Yaw);
+		EditorModeTools.TranslateRotate2DAngle = static_cast<float>(FMath::RadiansToDegrees(FMath::Atan2(Widget2DWorldXAxis.Z, Widget2DWorldXAxis.X)));
 	}
 
 	// Update showing.
diff --git a/Engine/Source/Editor/UnrealEd/Private/HierarchicalLODVolume.cpp b/Engine/Source/Editor/UnrealEd/Private/HierarchicalLODVolume.cpp
--- a/Engine/Source/Editor/UnrealEd/Private/HierarchicalLODVolume.cpp
+++ b/Engine/Source/Editor/UnrealEd/Private/HierarchicalLODVolume.cpp
@@ -38,7 +38,7 @@ bool AHierarchicalLODVolume::IsActorIncluded(const AActor* InActor) const
 
 	if (bIncludeOverlappingActors)
 	{
-		return EncompassesPoint(ActorBoundingBox.GetCenter(), ActorBoundingBox.GetExtent().Size());
+		return EncompassesPoint(ActorBoundingBox.GetCenter(), static_cast<float>(ActorBoundingBox.GetExtent().Size()));
 	}
 	else
 	{
diff --git a/Engine/Source/Editor/UnrealEd/Private/PlayLevelNewProcess.cpp b/Engine/Source/Editor/UnrealEd/Private/PlayLevelNewProcess.cpp
--- a/Engine/Source/Editor/UnrealEd/Private/PlayLevelNewProcess.cpp
+++ b/Engine/Source/Editor/UnrealEd/Private/PlayLevelNewProcess.cpp
@@ -258,7 +258,7 @@ void UEditorEngine::LaunchNewProcess(const FRequestPlaySessionParams& InParams,
 			FIntPoint WindowPosition = InParams.EditorPlaySettings->NewWindowPosition;
 			
 			WindowPosition.X += FMath::Max(InInstanceNum - 1, 0) * WindowSize.X;
-			WindowPosition.Y += SWindowDefs::DefaultTitleBarSize * FPlatformApplicationMisc::GetDPIScaleFactorAtPoint(0, 0);
+			WindowPosition.Y += static_cast<int32>(SWindowDefs::DefaultTitleBarSize * FPlatformApplicationMisc::GetDPIScaleFactorAtPoint(0, 0));
 
 			// If they don't want to center the new window, we add a specific location. This will get saved to user settings
 			// via SAVEWINPOS and not end up reflected in our PlayInEditor settings.
diff --git a/Engine/Source/Editor/UnrealEd/Private/ThumbnailHelpers.cpp b/Engine/Source/Editor/UnrealEd/Private/ThumbnailHelpers.cpp
--- a/Engine/Source/Editor/UnrealEd/Private/ThumbnailHelpers.cpp
+++ b/Engine/Source/Editor/UnrealEd/Private/ThumbnailHelpers.cpp
@@ -126,7 +126,7 @@ FSceneView* FThumbnailPreviewScene::CreateView(FSceneViewFamily * ViewFamily, in
 		FPlane(0,	1,	0,	0),
 		FPlane(0,	0,	0,	1));
 
-	Origin -= ViewRotationMatrix.InverseTransformPosition( FVector::ZeroVector );
+	Origin -= ViewRotationMatrix.InverseTransformPosition(FVector::ZeroVector);
 	ViewRotationMatrix = ViewRotationMatrix.RemoveTranslation();
 
 	FSceneViewInitOptions ViewInitOptions;
@@ -148,7 +148,9 @@ FSceneView* FThumbnailPreviewScene::CreateView(FSceneViewFamily * ViewFamily, in
 	// NOTE: Sizes may not actually be in screen space depending on how the thumbnail ends up stretched by the UI.  Not a big deal though.
 	// NOTE: Textures still take a little time to stream if the view has not been re-rendered recently, so they may briefly appear blurry while mips are prepared
 	// NOTE: Content Browser only renders thumbnails for loaded assets, and only when the mouse is over the panel. They'll be frozen in their last state while the mouse cursor is not over the panel.  This is for performance reasons
-	IStreamingManager::Get().AddViewInformation( Origin, SizeX, SizeX / FMath::Tan( FOVDegrees ) );
+	float ScreenSize = static_cast<float>(SizeX);
+	float FOVScreenSize = static_cast<float>(SizeX) / FMath::Tan(FOVDegrees);
+	IStreamingManager::Get().AddViewInformation(Origin, ScreenSize, FOVScreenSize);
 
 	return NewView;
 }
@@ -166,7 +168,7 @@ TStatId FThumbnailPreviewScene::GetStatId() const
 float FThumbnailPreviewScene::GetBoundsZOffset(const FBoxSphereBounds& Bounds) const
 {
 	// Return half the height of the bounds plus one to avoid ZFighting with the floor plane
-	return Bounds.BoxExtent.Z + 1;
+	return static_cast<float>(Bounds.BoxExtent.Z + 1.0);
 }
 
 /*
@@ -390,7 +392,7 @@ void FMaterialThumbnailScene::GetViewMatrixParameters(const float InFOVDegrees,
 	const float HalfFOVRadians = FMath::DegreesToRadians<float>(InFOVDegrees) * 0.5f;
 	// Add extra size to view slightly outside of the bounds to compensate for perspective
 	const float BoundsMultiplier = 1.15f;
-	const float HalfMeshSize = PreviewActor->GetStaticMeshComponent()->Bounds.SphereRadius * BoundsMultiplier;
+	const float HalfMeshSize = static_cast<float>(PreviewActor->GetStaticMeshComponent()->Bounds.SphereRadius * BoundsMultiplier);
 	const float BoundsZOffset = GetBoundsZOffset(PreviewActor->GetStaticMeshComponent()->Bounds);
 	const float TargetDistance = HalfMeshSize / FMath::Tan(HalfFOVRadians);
 
@@ -472,7 +474,7 @@ void FSkeletalMeshThumbnailScene::GetViewMatrixParameters(const float InFOVDegre
 
 	const float HalfFOVRadians = FMath::DegreesToRadians<float>(InFOVDegrees) * 0.5f;
 	// No need to add extra size to view slightly outside of the sphere to compensate for perspective since skeletal meshes already buffer bounds.
-	const float HalfMeshSize = PreviewActor->GetSkeletalMeshComponent()->Bounds.SphereRadius; 
+	const float HalfMeshSize = static_cast<float>(PreviewActor->GetSkeletalMeshComponent()->Bounds.SphereRadius); 
 	const float BoundsZOffset = GetBoundsZOffset(PreviewActor->GetSkeletalMeshComponent()->Bounds);
 	const float TargetDistance = HalfMeshSize / FMath::Tan(HalfFOVRadians);
 
@@ -560,7 +562,7 @@ void FStaticMeshThumbnailScene::GetViewMatrixParameters(const float InFOVDegrees
 
 	const float HalfFOVRadians = FMath::DegreesToRadians<float>(InFOVDegrees) * 0.5f;
 	// Add extra size to view slightly outside of the sphere to compensate for perspective
-	const float HalfMeshSize = PreviewActor->GetStaticMeshComponent()->Bounds.SphereRadius * 1.15;
+	const float HalfMeshSize = static_cast<float>(PreviewActor->GetStaticMeshComponent()->Bounds.SphereRadius * 1.15);
 	const float BoundsZOffset = GetBoundsZOffset(PreviewActor->GetStaticMeshComponent()->Bounds);
 	const float TargetDistance = HalfMeshSize / FMath::Tan(HalfFOVRadians);
 
@@ -695,7 +697,7 @@ void FAnimationSequenceThumbnailScene::GetViewMatrixParameters(const float InFOV
 
 	const float HalfFOVRadians = FMath::DegreesToRadians<float>(InFOVDegrees) * 0.5f;
 	// No need to add extra size to view slightly outside of the sphere to compensate for perspective since skeletal meshes already buffer bounds.
-	const float HalfMeshSize = PreviewActor->GetSkeletalMeshComponent()->Bounds.SphereRadius;
+	const float HalfMeshSize = static_cast<float>(PreviewActor->GetSkeletalMeshComponent()->Bounds.SphereRadius);
 	const float BoundsZOffset = GetBoundsZOffset(PreviewActor->GetSkeletalMeshComponent()->Bounds);
 	const float TargetDistance = HalfMeshSize / FMath::Tan(HalfFOVRadians);
 
@@ -829,7 +831,7 @@ void FBlendSpaceThumbnailScene::GetViewMatrixParameters(const float InFOVDegrees
 
 	const float HalfFOVRadians = FMath::DegreesToRadians<float>(InFOVDegrees) * 0.5f;
 	// No need to add extra size to view slightly outside of the sphere to compensate for perspective since skeletal meshes already buffer bounds.
-	const float HalfMeshSize = PreviewActor->GetSkeletalMeshComponent()->Bounds.SphereRadius;
+	const float HalfMeshSize = static_cast<float>(PreviewActor->GetSkeletalMeshComponent()->Bounds.SphereRadius);
 	const float BoundsZOffset = GetBoundsZOffset(PreviewActor->GetSkeletalMeshComponent()->Bounds);
 	const float TargetDistance = HalfMeshSize / FMath::Tan(HalfFOVRadians);
 
@@ -950,7 +952,7 @@ void FAnimBlueprintThumbnailScene::GetViewMatrixParameters(const float InFOVDegr
 
 	const float HalfFOVRadians = FMath::DegreesToRadians<float>(InFOVDegrees) * 0.5f;
 	// No need to add extra size to view slightly outside of the sphere to compensate for perspective since skeletal meshes already buffer bounds.
-	const float HalfMeshSize = PreviewActor->GetSkeletalMeshComponent()->Bounds.SphereRadius;
+	const float HalfMeshSize = static_cast<float>(PreviewActor->GetSkeletalMeshComponent()->Bounds.SphereRadius);
 	const float BoundsZOffset = GetBoundsZOffset(PreviewActor->GetSkeletalMeshComponent()->Bounds);
 	const float TargetDistance = HalfMeshSize / FMath::Tan(HalfFOVRadians);
 
@@ -1025,7 +1027,7 @@ void FPhysicsAssetThumbnailScene::GetViewMatrixParameters(const float InFOVDegre
 
 	const float HalfFOVRadians = FMath::DegreesToRadians<float>(InFOVDegrees) * 0.5f;
 	// No need to add extra size to view slightly outside of the sphere to compensate for perspective since skeletal meshes already buffer bounds.
-	const float HalfMeshSize = PreviewActor->GetSkeletalMeshComponent()->Bounds.SphereRadius; 
+	const float HalfMeshSize = static_cast<float>(PreviewActor->GetSkeletalMeshComponent()->Bounds.SphereRadius);
 	const float BoundsZOffset = GetBoundsZOffset(PreviewActor->GetSkeletalMeshComponent()->Bounds);
 	const float TargetDistance = HalfMeshSize / FMath::Tan(HalfFOVRadians);
 
@@ -1161,7 +1163,7 @@ void FClassActorThumbnailScene::GetViewMatrixParameters(const float InFOVDegrees
 	// Add extra size to view slightly outside of the sphere to compensate for perspective
 	const FBoxSphereBounds Bounds = GetPreviewActorBounds();
 
-	const float HalfMeshSize = Bounds.SphereRadius * 1.15;
+	const float HalfMeshSize = static_cast<float>(Bounds.SphereRadius * 1.15);
 	const float BoundsZOffset = GetBoundsZOffset(Bounds);
 	const float TargetDistance = HalfMeshSize / FMath::Tan(HalfFOVRadians);
 
diff --git a/Engine/Source/Editor/UnrealEd/Private/Tools/LegacyEdModeWidgetHelpers.cpp b/Engine/Source/Editor/UnrealEd/Private/Tools/LegacyEdModeWidgetHelpers.cpp
--- a/Engine/Source/Editor/UnrealEd/Private/Tools/LegacyEdModeWidgetHelpers.cpp
+++ b/Engine/Source/Editor/UnrealEd/Private/Tools/LegacyEdModeWidgetHelpers.cpp
@@ -647,9 +647,9 @@ void FLegacyEdModeWidgetHelper::Render(const FSceneView* View, FViewport* Viewpo
 				const FTransform WorldWidgetTransform = LocalWidgetTransform * DisplayWidgetToWorld;
 				const FMatrix WidgetTM(WorldWidgetTransform.ToMatrixWithScale());
 
-				const float WidgetSize = 0.035f;
-				const float ZoomFactor = FMath::Min<float>(View->ViewMatrices.GetProjectionMatrix().M[0][0], View->ViewMatrices.GetProjectionMatrix().M[1][1]);
-				const float WidgetRadius = View->Project(WorldWidgetTransform.GetTranslation()).W * (WidgetSize / ZoomFactor);
+				const double WidgetSize = 0.035;
+				const double ZoomFactor = FMath::Min(View->ViewMatrices.GetProjectionMatrix().M[0][0], View->ViewMatrices.GetProjectionMatrix().M[1][1]);
+				const float WidgetRadius = static_cast<float>(View->Project(WorldWidgetTransform.GetTranslation()).W * (WidgetSize / ZoomFactor));
 
 				if (bHitTesting) PDI->SetHitProxy(new HPropertyWidgetProxy(WidgetInfo.PropertyName, WidgetInfo.PropertyIndex, WidgetInfo.bIsTransform));
 				DrawWireDiamond(PDI, WidgetTM, WidgetRadius, WidgetColor, SDPG_Foreground);
@@ -669,9 +669,14 @@ void FLegacyEdModeWidgetHelper::DrawHUD(FEditorViewportClient* ViewportClient, F
 		{
 			FEditorScriptExecutionGuard ScriptGuard;
 
-			FIntPoint SizeXYWithDPIScale = Viewport->GetSizeXY() / Canvas->GetDPIScale();
-			const int32 HalfX = 0.5f * SizeXYWithDPIScale.X;
-			const int32 HalfY = 0.5f * SizeXYWithDPIScale.Y;
+			FIntPoint SizeXY = Viewport->GetSizeXY();
+			int32 ScaledX = FMath::TruncToInt32(static_cast<float>(SizeXY.X) / Canvas->GetDPIScale());
+			int32 ScaledY = FMath::TruncToInt32(static_cast<float>(SizeXY.Y) / Canvas->GetDPIScale());
+
+			FIntPoint SizeXYWithDPIScale{ ScaledX, ScaledY };
+
+			const int32 HalfX = SizeXYWithDPIScale.X / 2;
+			const int32 HalfY = SizeXYWithDPIScale.Y / 2;
 
 			UClass* Class = BestSelectedItem->GetClass();
 			TArray<FPropertyWidgetInfo> WidgetInfos;
@@ -693,8 +698,8 @@ void FLegacyEdModeWidgetHelper::DrawHUD(FEditorViewportClient* ViewportClient, F
 					const FString WidgetDisplayName = WidgetInfo.DisplayName + ((VectorIndex != INDEX_NONE) ? FString::Printf(TEXT("[%d]"), VectorIndex) : TEXT(""));
 					const FString DisplayString = ValidationMessage.IsEmpty() ? WidgetDisplayName : ValidationMessage;
 
-					const int32 XPos = HalfX + (HalfX * Proj.X);
-					const int32 YPos = HalfY + (HalfY * (Proj.Y * -1.f));
+					const int32 XPos = static_cast<int32>(HalfX + (HalfX * Proj.X));
+					const int32 YPos = static_cast<int32>(HalfY + (HalfY * (Proj.Y * -1.f)));
 					FCanvasTextItem TextItem(FVector2D(XPos + 5, YPos), FText::FromString(DisplayString), GEngine->GetSmallFont(), FLinearColor::White);
 					TextItem.EnableShadow(FLinearColor::Black);
 					Canvas->DrawItem(TextItem);
