Fix to allow settings objects to be properly handled by LiveCoding and HotReload.

diff --git a/Engine/Source/Developer/Settings/Private/SettingsCategory.cpp b/Engine/Source/Developer/Settings/Private/SettingsCategory.cpp
--- a/Engine/Source/Developer/Settings/Private/SettingsCategory.cpp
+++ b/Engine/Source/Developer/Settings/Private/SettingsCategory.cpp
@@ -53,6 +53,16 @@ void FSettingsCategory::RemoveSection( const FName& SectionName )
 	Sections.Remove(SectionName);
 }
 
+#if WITH_RELOAD
+void FSettingsCategory::ReinstancingComplete(IReload* Reload)
+{
+	for (TTuple<FName, TSharedPtr<FSettingsSection>>& SectionPair : Sections)
+	{
+		SectionPair.Value->ReinstancingComplete(Reload);
+	}
+}
+#endif
+
 
 /* ISettingsCategory interface
  *****************************************************************************/
diff --git a/Engine/Source/Developer/Settings/Private/SettingsCategory.h b/Engine/Source/Developer/Settings/Private/SettingsCategory.h
--- a/Engine/Source/Developer/Settings/Private/SettingsCategory.h
+++ b/Engine/Source/Developer/Settings/Private/SettingsCategory.h
@@ -8,6 +8,7 @@
 #include "Misc/NamePermissionList.h"
 
 class FSettingsSection;
+class IReload;
 class SWidget;
 
 /**
@@ -71,6 +72,15 @@ public:
 	 */
 	void RemoveSection( const FName& SectionName );
 
+#if WITH_RELOAD
+	/**
+	 * Invoked when reinstancing is complete.  Allows for settings objects to update their settings object pointers.
+	 * 
+	 * @param Reload The active reload
+	 */
+	void ReinstancingComplete( IReload* Reload );
+#endif
+
 public:
 
 	// ISettingsCategory interface
diff --git a/Engine/Source/Developer/Settings/Private/SettingsContainer.cpp b/Engine/Source/Developer/Settings/Private/SettingsContainer.cpp
--- a/Engine/Source/Developer/Settings/Private/SettingsContainer.cpp
+++ b/Engine/Source/Developer/Settings/Private/SettingsContainer.cpp
@@ -67,6 +67,16 @@ void FSettingsContainer::RemoveSection( const FName& CategoryName, const FName&
 	}
 }
 
+#if WITH_RELOAD
+void FSettingsContainer::ReinstancingComplete(IReload* Reload)
+{
+	for (TTuple<FName, TSharedPtr<FSettingsCategory>>& CategoryPair : Categories)
+	{
+		CategoryPair.Value->ReinstancingComplete(Reload);
+	}
+}
+#endif
+
 /* ISettingsContainer interface
  *****************************************************************************/
 
diff --git a/Engine/Source/Developer/Settings/Private/SettingsContainer.h b/Engine/Source/Developer/Settings/Private/SettingsContainer.h
--- a/Engine/Source/Developer/Settings/Private/SettingsContainer.h
+++ b/Engine/Source/Developer/Settings/Private/SettingsContainer.h
@@ -6,6 +6,7 @@
 #include "ISettingsContainer.h"
 #include "SettingsCategory.h"
 
+class IReload;
 class SWidget;
 
 /**
@@ -61,6 +62,15 @@ public:
 	 */
 	void RemoveSection( const FName& CategoryName, const FName& SectionName );
 
+#if WITH_RELOAD
+	/**
+	 * Invoked when reinstancing is complete.  Allows for settings objects to update their settings object pointers.
+	 * 
+	 * @param Reload The active reload
+	 */
+	void ReinstancingComplete(IReload* Reload);
+#endif
+
 public:
 
 	// ISettingsContainer interface
diff --git a/Engine/Source/Developer/Settings/Private/SettingsModule.cpp b/Engine/Source/Developer/Settings/Private/SettingsModule.cpp
--- a/Engine/Source/Developer/Settings/Private/SettingsModule.cpp
+++ b/Engine/Source/Developer/Settings/Private/SettingsModule.cpp
@@ -6,7 +6,8 @@
 #include "UObject/WeakObjectPtr.h"
 #include "SettingsContainer.h"
 #include "ISettingsModule.h"
-
+#include "UObject/Reload.h"
+#include "UObject/UObjectGlobals.h"
 
 #define LOCTEXT_NAMESPACE "FSettingsModule"
 
@@ -95,9 +96,19 @@ public:
 		ProjectSettingsContainer->DescribeCategory("Editor", LOCTEXT("ProjectEditorCategoryName", "Editor"), LOCTEXT("ProjectEditorCategoryDescription", "Project settings specific to the editor"));
 		ProjectSettingsContainer->DescribeCategory("Platforms", LOCTEXT("ProjectPlatformsCategoryName", "Platforms"), LOCTEXT("ProjectPlatformsCategoryDescription", "Platform settings description text here"));
 		ProjectSettingsContainer->DescribeCategory("Plugins", LOCTEXT("ProjectPluginsCategoryName", "Plugins"), LOCTEXT("ProjectPluginsCategoryDescription", "Plugins settings description text here"));
+
+#if WITH_RELOAD
+		OnReloadHandle = FCoreUObjectDelegates::ReloadReinstancingCompleteDelegate.AddLambda([this] { ReinstancingComplete(); });
+#endif
 	}
 
-	virtual void ShutdownModule() override { }
+	virtual void ShutdownModule() override
+	{
+#if WITH_RELOAD
+		FCoreUObjectDelegates::ReloadReinstancingCompleteDelegate.Remove(OnReloadHandle);
+		OnReloadHandle.Reset();
+#endif
+	}
 
 protected:
 
@@ -119,6 +130,20 @@ protected:
 		return Container.ToSharedRef();
 	}
 
+#if WITH_RELOAD
+	void ReinstancingComplete()
+	{
+		IReload* Reload = GetActiveReloadInterface();
+		if (Reload != nullptr)
+		{
+			for (TTuple <FName, TSharedPtr<FSettingsContainer> >& ContainerPair : ContainerNamesToContainers)
+			{
+				ContainerPair.Value->ReinstancingComplete(Reload);
+			}
+		}
+	}
+#endif
+
 private:
 
 	/** Holds the default settings container. */
@@ -129,6 +154,11 @@ private:
 
 	/** Holds the collection of registered settings viewers. */
 	TMap<FName, ISettingsViewer*> ContainerNamesToViewers;
+
+#if WITH_RELOAD
+	/** Delegate handle for the re-instancing complete notification */
+	FDelegateHandle OnReloadHandle;
+#endif
 };
 
 
diff --git a/Engine/Source/Developer/Settings/Private/SettingsSection.cpp b/Engine/Source/Developer/Settings/Private/SettingsSection.cpp
--- a/Engine/Source/Developer/Settings/Private/SettingsSection.cpp
+++ b/Engine/Source/Developer/Settings/Private/SettingsSection.cpp
@@ -5,6 +5,7 @@
 #include "Misc/ConfigCacheIni.h"
 #include "Misc/ConfigContext.h"
 #include "UObject/Class.h"
+#include "UObject/Reload.h"
 
 
 /* FSettingsSection structors
@@ -27,6 +28,12 @@ FSettingsSection::FSettingsSection( const ISettingsCategoryRef& InCategory, cons
 	, Name(InName)
 { }
 
+#if WITH_RELOAD
+void FSettingsSection::ReinstancingComplete(IReload* Reload)
+{
+	SettingsObject = Reload->GetReinstancedCDO(SettingsObject.Get(true));
+}
+#endif
 
 /* ISettingsSection interface
  *****************************************************************************/
diff --git a/Engine/Source/Developer/Settings/Private/SettingsSection.h b/Engine/Source/Developer/Settings/Private/SettingsSection.h
--- a/Engine/Source/Developer/Settings/Private/SettingsSection.h
+++ b/Engine/Source/Developer/Settings/Private/SettingsSection.h
@@ -7,6 +7,7 @@
 #include "ISettingsCategory.h"
 #include "ISettingsSection.h"
 
+class IReload;
 class SWidget;
 
 /**
@@ -39,6 +40,15 @@ public:
 	 */
 	FSettingsSection( const ISettingsCategoryRef& InCategory, const FName& InName, const FText& InDisplayName, const FText& InDescription, const TSharedRef<SWidget>& InCustomWidget );
 
+#if WITH_RELOAD
+	/**
+	 * Invoked when reinstancing is complete.  Allows for settings objects to update their settings object pointers.
+	 * 
+	 * @param Reload The active reload
+	 */
+	void ReinstancingComplete(IReload* Reload);
+#endif
+
 public:
 
 	// ISettingsSection interface
diff --git a/Engine/Source/Developer/Windows/LiveCoding/Private/LiveCodingModule.cpp b/Engine/Source/Developer/Windows/LiveCoding/Private/LiveCodingModule.cpp
--- a/Engine/Source/Developer/Windows/LiveCoding/Private/LiveCodingModule.cpp
+++ b/Engine/Source/Developer/Windows/LiveCoding/Private/LiveCodingModule.cpp
@@ -129,6 +129,16 @@ public:
 	{
 	}
 
+	virtual UObject* GetReinstancedCDO(UObject* CDO) override
+	{
+		return CDO;
+	}
+
+	virtual const UObject* GetReinstancedCDO(const UObject* CDO) override
+	{
+		return CDO;
+	}
+
 	bool HasReinstancingOccurred() const
 	{
 		return bHasReinstancingOccurred;
diff --git a/Engine/Source/Editor/UnrealEd/Private/Kismet2/ReloadUtilities.cpp b/Engine/Source/Editor/UnrealEd/Private/Kismet2/ReloadUtilities.cpp
--- a/Engine/Source/Editor/UnrealEd/Private/Kismet2/ReloadUtilities.cpp
+++ b/Engine/Source/Editor/UnrealEd/Private/Kismet2/ReloadUtilities.cpp
@@ -607,7 +607,7 @@ namespace UE::Reload::Private
 		};
 
 	public:
-		FReloadClassHelper(FOutputDevice& InAr, const TSet<UObject*>& InReinstancingObjects, TSet<UBlueprint*>& InCompiledBlueprints, TMap<UObject*, UObject*>& InReconstructedCDOsMap);
+		FReloadClassHelper(FOutputDevice& InAr, const TSet<UObject*>& InReinstancingObjects, TSet<UBlueprint*>& InCompiledBlueprints, TMap<UObject*, UObject*>& InReconstructedCDOsMap, TMap<UObject*, UObject*>& InReinstancedCDOsMap);
 
 		/** 
 		 * Re-instance the collection of classes.  
@@ -632,15 +632,18 @@ namespace UE::Reload::Private
 		const TSet<UObject*>& ReinstancingObjects;
 		TSet<UBlueprint*>& CompiledBlueprints;
 		TMap<UObject*, UObject*>& ReconstructedCDOsMap;
+		TMap<UObject*, UObject*>& ReinstancedCDOsMap;
+
 
 		TMap<UClass*, ClassReinstanceState> ReinstanceStates;
 	};
 
-	FReloadClassHelper::FReloadClassHelper(FOutputDevice& InAr, const TSet<UObject*>& InReinstancingObjects, TSet<UBlueprint*>& InCompiledBlueprints, TMap<UObject*, UObject*>& InReconstructedCDOsMap)
+	FReloadClassHelper::FReloadClassHelper(FOutputDevice& InAr, const TSet<UObject*>& InReinstancingObjects, TSet<UBlueprint*>& InCompiledBlueprints, TMap<UObject*, UObject*>& InReconstructedCDOsMap, TMap<UObject*, UObject*>& InReinstancedCDOsMap)
 		: Ar(InAr)
 		, ReinstancingObjects(InReinstancingObjects)
 		, CompiledBlueprints(InCompiledBlueprints)
 		, ReconstructedCDOsMap(InReconstructedCDOsMap)
+		, ReinstancedCDOsMap(InReinstancedCDOsMap)
 	{
 	}
 
@@ -727,6 +730,7 @@ namespace UE::Reload::Private
 			// We only need to do re-instancing when we have a new UClass.
 			else
 			{
+				ReinstancedCDOsMap.Add(State.OldClassCDO, NewClass->GetDefaultObject());
 				UClass* NullableNewClass = NewClass == OldClass ? nullptr : NewClass;
 				TSharedPtr<FReloadClassReinstancer> ReinstanceHelper = MakeShareable(new FReloadClassReinstancer(
 					NullableNewClass, OldClass, State.OldClassCDO, ReinstancingObjects, CompiledBlueprints));
@@ -1127,7 +1131,7 @@ void FReload::Reinstance()
 
 	// Re-instance the classes
 	TSet<UBlueprint*> CompiledBlueprints;
-	UE::Reload::Private::FReloadClassHelper rch(Ar, ReinstancingObjects, CompiledBlueprints, ReconstructedCDOsMap);
+	UE::Reload::Private::FReloadClassHelper rch(Ar, ReinstancingObjects, CompiledBlueprints, ReconstructedCDOsMap, ReinstancedCDOsMap);
 	rch.ReinstanceClasses(ClassesToReinstance);
 
 	// Recompile blueprints if they haven't already been recompiled)
@@ -1148,6 +1152,28 @@ void FReload::Reinstance()
 	FCoreUObjectDelegates::ReloadReinstancingCompleteDelegate.Broadcast();
 }
 
+UObject* FReload::GetReinstancedCDO(UObject* CDO)
+{
+	return const_cast<UObject*>(GetReinstancedCDO(const_cast<const UObject*>(CDO)));
+}
+
+const UObject* FReload::GetReinstancedCDO(const UObject* CDO)
+{
+	UObject* const* NewCDO = ReconstructedCDOsMap.Find(CDO);
+	if (NewCDO != nullptr)
+	{
+		return *NewCDO;
+	}
+
+	NewCDO = ReinstancedCDOsMap.Find(CDO);
+	if (NewCDO != nullptr)
+	{
+		return *NewCDO;
+	}
+
+	return CDO;
+}
+
 void FReload::Finalize(bool bRunGC)
 {
 
diff --git a/Engine/Source/Editor/UnrealEd/Public/Kismet2/ReloadUtilities.h b/Engine/Source/Editor/UnrealEd/Public/Kismet2/ReloadUtilities.h
--- a/Engine/Source/Editor/UnrealEd/Public/Kismet2/ReloadUtilities.h
+++ b/Engine/Source/Editor/UnrealEd/Public/Kismet2/ReloadUtilities.h
@@ -42,6 +42,8 @@ public:
 	virtual void NotifyChange(UScriptStruct* New, UScriptStruct* Old) override;
 	virtual void NotifyChange(UPackage* New, UPackage* Old) override;
 	virtual void Reinstance() override;
+	virtual UObject* GetReinstancedCDO(UObject* CDO) override;
+	virtual const UObject* GetReinstancedCDO(const UObject* CDO) override;
 
 	/**
 	 * If you wish to reuse the same reload object, invoke this method to reset the state
@@ -141,6 +143,9 @@ private:
 	/** Map of the reconstructed CDOs during the reinstancing process */
 	TMap<UObject*, UObject*> ReconstructedCDOsMap;
 
+	/** Map of new CDOs where classes where changed */
+	TMap<UObject*, UObject*> ReinstancedCDOsMap;
+
 	/** Map from old class to new class.  New class may be null */
 	TMap<UClass*, UClass*> ReinstancedClasses;
 
diff --git a/Engine/Source/Runtime/CoreUObject/Public/UObject/Reload.h b/Engine/Source/Runtime/CoreUObject/Public/UObject/Reload.h
--- a/Engine/Source/Runtime/CoreUObject/Public/UObject/Reload.h
+++ b/Engine/Source/Runtime/CoreUObject/Public/UObject/Reload.h
@@ -53,6 +53,13 @@ public:
 	 * Perform the re-instancing 
 	 */
 	virtual void Reinstance() = 0;
+
+	/**
+	 * For a given CDO, return the new CDO if it has been re-instanced.  If it hasn't been re-instanced,
+	 * then the supplied CDO will be returned. If a non-CDO object is passed in, it will be returned as is.
+	 */
+	virtual UObject* GetReinstancedCDO(UObject* CDO) = 0;
+	virtual const UObject* GetReinstancedCDO(const UObject* CDO) = 0;
 };
 
 /**
