EditorDomain: Instrumentation and behavior fixes for Warning: Package %s received an invalid CacheRecord from DDC: size of all segments %d is not equal to FileSize in metadata %d.
Add a warning and a failure when we create the record if it is bad at that point.
Fix FEditorDomainPackageSegments::OnRecordRequestComplete to report that the corrupt record does NOT exist in the EditorDomain storage.
Bump the EditorDomain salt so we can confirm that any missing segments in the future is due to DDC replication rather than a bad record at creation.
#rnx

diff --git a/Engine/Source/Editor/UnrealEd/Private/EditorDomain/EditorDomainArchive.cpp b/Engine/Source/Editor/UnrealEd/Private/EditorDomain/EditorDomainArchive.cpp
--- a/Engine/Source/Editor/UnrealEd/Private/EditorDomain/EditorDomainArchive.cpp
+++ b/Engine/Source/Editor/UnrealEd/Private/EditorDomain/EditorDomainArchive.cpp
@@ -328,6 +328,7 @@ void FEditorDomainPackageSegments::OnRecordRequestComplete(UE::DerivedData::FCac
 					TEXT("size of all segments "), CompositeSize, TEXT(" is not equal to FileSize in metadata "), FileSize
 					) : TEXT("metadata is empty"));
 			Segments.Empty();
+			bExistsInEditorDomain = false;
 		}
 		else
 		{
diff --git a/Engine/Source/Editor/UnrealEd/Private/EditorDomain/EditorDomainUtils.cpp b/Engine/Source/Editor/UnrealEd/Private/EditorDomain/EditorDomainUtils.cpp
--- a/Engine/Source/Editor/UnrealEd/Private/EditorDomain/EditorDomainUtils.cpp
+++ b/Engine/Source/Editor/UnrealEd/Private/EditorDomain/EditorDomainUtils.cpp
@@ -201,7 +201,7 @@ FBlake3Hash GGlobalConstructClassesHash;
 int64 GMaxBulkDataSize = -1;
 
 // Change to a new guid when EditorDomain needs to be invalidated
-const TCHAR* EditorDomainVersion = TEXT("2A2A362225084BE8B1FCD539E21BDE28");
+const TCHAR* EditorDomainVersion = TEXT("0CE4958BE8484993A0399BF2CB52EF4E");
 
 // Identifier of the CacheBuckets for EditorDomain tables
 const TCHAR* EditorDomainPackageBucketName = TEXT("EditorDomainPackage");
@@ -2016,8 +2016,6 @@ bool TrySavePackage(UPackage* Package)
 		StorageResult = ESaveStorageResult::BulkDataTooLarge;
 	}
 
-	ICache& Cache = GetCache();
-
 	bool bStorageResultValid = StorageResult == ESaveStorageResult::Valid;
 	TCbWriter<16> MetaData;
 	MetaData.BeginObject();
@@ -2036,11 +2034,21 @@ bool TrySavePackage(UPackage* Package)
 	FCacheRecordBuilder RecordBuilder(GetEditorDomainPackageKey(PackageDigest.Hash));
 	if (bStorageResultValid)
 	{
+		COOK_STAT(Timer.AddMiss(FileSize + MetaData.GetSaveSize()));
+
+		uint64 TotalAttachmentSize = 0;
 		for (const FEditorDomainPackageWriter::FAttachment& Attachment : PackageWriter->GetAttachments())
 		{
 			RecordBuilder.AddValue(Attachment.ValueId, Attachment.Buffer);
+			TotalAttachmentSize += Attachment.Buffer.GetSize();
+		}
+		if (TotalAttachmentSize != FileSize)
+		{
+			UE_LOG(LogEditorDomain, Warning,
+				TEXT("Could not save %s to EditorDomain due to TrySavePackage bug: size of all segments %" UINT64_FMT " is not equal to FileSize in metadata %" UINT64_FMT "."),
+				*Package->GetName(), TotalAttachmentSize, FileSize);
+			return false;
 		}
-		COOK_STAT(Timer.AddMiss(FileSize + MetaData.GetSaveSize()));
 	}
 	else
 	{
@@ -2048,7 +2056,7 @@ bool TrySavePackage(UPackage* Package)
 	}
 	RecordBuilder.SetMeta(MetaData.Save().AsObject());
 	FRequestOwner Owner(EPriority::Normal);
-	Cache.Put({ {{Package->GetName()}, RecordBuilder.Build()} }, Owner);
+	GetCache().Put({ {{Package->GetName()}, RecordBuilder.Build()} }, Owner);
 	Owner.KeepAlive();
 
 	if (bStorageResultValid)
