BulkDataRegistry: PayloadId Puts: need to test bAllowedToWritePaylodIdToCache before calling WritePayloadIdToCache, to prevent an assertion for memory-only BulkDatas.
#rnx

diff --git a/Engine/Source/Editor/UnrealEd/Private/EditorDomain/BulkDataRegistryImpl.cpp b/Engine/Source/Editor/UnrealEd/Private/EditorDomain/BulkDataRegistryImpl.cpp
--- a/Engine/Source/Editor/UnrealEd/Private/EditorDomain/BulkDataRegistryImpl.cpp
+++ b/Engine/Source/Editor/UnrealEd/Private/EditorDomain/BulkDataRegistryImpl.cpp
@@ -654,7 +654,7 @@ void FBulkDataRegistryImpl::ReadPayloadIdsFromCache(FName PackageName, TArray<TR
 			}
 			else
 			{
-				if (!ExistingBulkData.HasPlaceholderPayloadId())
+				if (!ExistingBulkData.HasPlaceholderPayloadId() && ExistingRegisteredBulk->bAllowedToWritePayloadIdToCache)
 				{
 					// In between the point where we started the cache query and we received this result,
 					// The FEditorBulkData has updated its PayloadId and informed us by calling UpdatePlaceholderPayloadId
