- Add bDrawAxesGame which will force editor clients to draw axes when EngineShowFlags.Game is set

#rnx

diff --git a/Engine/Plugins/FX/Niagara/Source/NiagaraEditor/Private/Widgets/SNiagaraSystemViewport.cpp b/Engine/Plugins/FX/Niagara/Source/NiagaraEditor/Private/Widgets/SNiagaraSystemViewport.cpp
--- a/Engine/Plugins/FX/Niagara/Source/NiagaraEditor/Private/Widgets/SNiagaraSystemViewport.cpp
+++ b/Engine/Plugins/FX/Niagara/Source/NiagaraEditor/Private/Widgets/SNiagaraSystemViewport.cpp
@@ -97,6 +97,8 @@ FNiagaraSystemViewportClient::FNiagaraSystemViewportClient(FAdvancedPreviewScene
 	SetOrbitModeFromSettings();
 	bCaptureScreenShot = false;
 
+	bDrawAxesGame = true;
+
 	//This seems to be needed to get the correct world time in the preview.
 	FNiagaraSystemViewportClient::SetIsSimulateInEditorViewport(true);
 }
diff --git a/Engine/Source/Editor/UnrealEd/Private/EditorViewportClient.cpp b/Engine/Source/Editor/UnrealEd/Private/EditorViewportClient.cpp
--- a/Engine/Source/Editor/UnrealEd/Private/EditorViewportClient.cpp
+++ b/Engine/Source/Editor/UnrealEd/Private/EditorViewportClient.cpp
@@ -411,6 +411,7 @@ FEditorViewportClient::FEditorViewportClient(FEditorModeTools* InModeTools, FPre
 	, bUseNumpadCameraControl(true)
 	, bDisableInput(false)
 	, bDrawAxes(true)
+	, bDrawAxesGame(false)
 	, bSetListenerPosition(false)
 	, LandscapeLODOverride(-1)
 	, bDrawVertices(false)
@@ -4122,7 +4123,8 @@ void FEditorViewportClient::Draw(FViewport* InViewport, FCanvas* Canvas)
 	}
 
 	// Axes indicators
-	if (bDrawAxes && !ViewFamily.EngineShowFlags.Game && !GLevelEditorModeTools().IsViewportUIHidden() && !IsVisualizeCalibrationMaterialEnabled())
+	const bool bShouldDrawAxes = (bDrawAxes && !ViewFamily.EngineShowFlags.Game) || (bDrawAxesGame && ViewFamily.EngineShowFlags.Game);
+	if (bShouldDrawAxes && !GLevelEditorModeTools().IsViewportUIHidden() && !IsVisualizeCalibrationMaterialEnabled())
 	{
 		switch (GetViewportType())
 		{
diff --git a/Engine/Source/Editor/UnrealEd/Public/EditorViewportClient.h b/Engine/Source/Editor/UnrealEd/Public/EditorViewportClient.h
--- a/Engine/Source/Editor/UnrealEd/Public/EditorViewportClient.h
+++ b/Engine/Source/Editor/UnrealEd/Public/EditorViewportClient.h
@@ -1721,6 +1721,9 @@ public:
 	/** If true, draw the axis indicators when the viewport is perspective. */
 	bool					bDrawAxes;
 
+	/** If true, draw the axis indicators when EngineShowFlags.Game is set. */
+	bool					bDrawAxesGame;
+
 	/** If true, the listener position will be set */
 	bool					bSetListenerPosition;
 
