Return all non-package files from FindAllSubmittableProjectFiles when using UsesCustomProjectDir

To include all files in revision control which are located in well know path but don't have package extensions.

diff --git a/Engine/Source/Editor/UnrealEd/Private/FileHelpers.cpp b/Engine/Source/Editor/UnrealEd/Private/FileHelpers.cpp
--- a/Engine/Source/Editor/UnrealEd/Private/FileHelpers.cpp
+++ b/Engine/Source/Editor/UnrealEd/Private/FileHelpers.cpp
@@ -4620,8 +4620,8 @@ void FEditorFileUtils::FindAllSubmittableProjectFiles(TMap<FString, FSourceContr
 
 			if (SourceControlState->CanCheckIn() || (!SourceControlState->IsSourceControlled() && SourceControlState->CanAdd()))
 			{
-				FString PackageName;
-				if (!FPackageName::TryConvertFilenameToLongPackageName(Filename, PackageName))
+				FString Ext = FPaths::GetExtension(Filename);
+				if (!FPackageName::IsPackageExtension(*Ext) && !FPackageName::IsTextPackageExtension(*Ext))
 				{
 					OutProjectFiles.Add(Filename, MoveTemp(SourceControlState));
 				}
