MPCook: Log profiling data from CookWorkers.
Fix duplicated PrepareSave hierarchical cooktimer.
#rnx

diff --git a/Engine/Source/Editor/UnrealEd/Classes/CookOnTheSide/CookOnTheFlyServer.h b/Engine/Source/Editor/UnrealEd/Classes/CookOnTheSide/CookOnTheFlyServer.h
--- a/Engine/Source/Editor/UnrealEd/Classes/CookOnTheSide/CookOnTheFlyServer.h
+++ b/Engine/Source/Editor/UnrealEd/Classes/CookOnTheSide/CookOnTheFlyServer.h
@@ -112,6 +112,26 @@ namespace ECookMode
 		CookWorker,
 	};
 }
+inline bool IsCookByTheBookMode(ECookMode::Type CookMode)
+{
+	return CookMode == ECookMode::CookByTheBookFromTheEditor || CookMode == ECookMode::CookByTheBook;
+}
+inline bool IsRealtimeMode(ECookMode::Type CookMode)
+{
+	return CookMode == ECookMode::CookByTheBookFromTheEditor || CookMode == ECookMode::CookOnTheFlyFromTheEditor;
+}
+inline bool IsCookingInEditor(ECookMode::Type CookMode)
+{
+	return CookMode == ECookMode::CookByTheBookFromTheEditor || CookMode == ECookMode::CookOnTheFlyFromTheEditor;
+}
+inline bool IsCookOnTheFlyMode(ECookMode::Type CookMode)
+{
+	return CookMode == ECookMode::CookOnTheFly || CookMode == ECookMode::CookOnTheFlyFromTheEditor;
+}
+inline bool IsCookWorkerMode(ECookMode::Type CookMode)
+{
+	return CookMode == ECookMode::CookWorker;
+}
 
 UENUM()
 enum class ECookTickFlags : uint8
@@ -629,6 +649,9 @@ public:
 	/** Connect to the CookDirector host, Initialize this UCOTFS, and start the CookWorker session */
 	bool TryInitializeCookWorker();
 
+	/** Log stats for the CookWorker; this is called before the connection to the director is terminated. */
+	void LogCookWorkerStats();
+
 	/** Terminate the CookWorker session */
 	void ShutdownCookAsCookWorker();
 
@@ -872,6 +895,9 @@ public:
 	void GetShaderLibraryPaths(const ITargetPlatform* TargetPlatform, FString& OutShaderCodeDir,
 		FString& OutMetaDataPath, bool bUseProjectDirForDLC=false);
 
+	/** Print detailed stats from the cook. */
+	void PrintDetailedCookStats();
+
 private:
 
 	/**
diff --git a/Engine/Source/Editor/UnrealEd/Private/Commandlets/CookCommandlet.cpp b/Engine/Source/Editor/UnrealEd/Private/Commandlets/CookCommandlet.cpp
--- a/Engine/Source/Editor/UnrealEd/Private/Commandlets/CookCommandlet.cpp
+++ b/Engine/Source/Editor/UnrealEd/Private/Commandlets/CookCommandlet.cpp
@@ -10,9 +10,7 @@
 #include "Async/TaskGraphInterfaces.h"
 #include "CookOnTheSide/CookOnTheFlyServer.h"
 #include "Cooker/CookProfiling.h"
-#include "Cooker/PackageBuildDependencyTracker.h"
 #include "CookerSettings.h"
-#include "DerivedDataBuildRemoteExecutor.h"
 #include "Editor.h"
 #include "EngineGlobals.h"
 #include "GameDelegates.h"
@@ -33,6 +31,7 @@
 #include "Misc/FileHelper.h"
 #include "Misc/LocalTimestampDirectoryVisitor.h"
 #include "Misc/MessageDialog.h"
+#include "Misc/Optional.h"
 #include "Misc/Paths.h"
 #include "Misc/RedirectCollector.h"
 #include "Modules/ModuleManager.h"
@@ -50,365 +49,6 @@
 #include "UObject/UObjectIterator.h"
 DEFINE_LOG_CATEGORY_STATIC(LogCookCommandlet, Log, All);
 
-#if ENABLE_COOK_STATS
-#include "AnalyticsET.h"
-#include "AnalyticsEventAttribute.h"
-#include "IAnalyticsProviderET.h"
-#include "ProfilingDebugging/ScopedTimers.h"
-#include "StudioAnalytics.h"
-#include "Virtualization/VirtualizationSystem.h"
-
-namespace DetailedCookStats
-{
-
-FString CookProject;
-FString CookCultures;
-FString CookLabel;
-FString TargetPlatforms;
-double CookWallTimeSec = 0.0;
-double StartupWallTimeSec = 0.0;
-double CookByTheBookTimeSec = 0.0;
-double StartCookByTheBookTimeSec = 0.0;
-extern double TickCookOnTheSideTimeSec;
-extern double TickCookOnTheSideLoadPackagesTimeSec;
-extern double TickCookOnTheSideResolveRedirectorsTimeSec;
-extern double TickCookOnTheSideSaveCookedPackageTimeSec;
-extern double TickCookOnTheSidePrepareSaveTimeSec;
-extern double BlockOnAssetRegistryTimeSec;
-extern double GameCookModificationDelegateTimeSec;
-double TickLoopGCTimeSec = 0.0;
-double TickLoopRecompileShaderRequestsTimeSec = 0.0;
-double TickLoopShaderProcessAsyncResultsTimeSec = 0.0;
-double TickLoopProcessDeferredCommandsTimeSec = 0.0;
-double TickLoopTickCommandletStatsTimeSec = 0.0;
-double TickLoopFlushRenderingCommandsTimeSec = 0.0;
-bool IsCookAll = false;
-bool IsCookOnTheFly = false;
-bool IsIterativeCook = false;
-bool IsFastCook = false;
-bool IsUnversioned = false;
-
-FCookStatsManager::FAutoRegisterCallback RegisterCookStats([](FCookStatsManager::AddStatFuncRef AddStat)
-{
-	const FString StatName(TEXT("Cook.Profile"));
-	#define ADD_COOK_STAT_FLT(Path, Name) AddStat(StatName, FCookStatsManager::CreateKeyValueArray(TEXT("Path"), TEXT(Path), TEXT(#Name), Name))
-	ADD_COOK_STAT_FLT(" 0", CookWallTimeSec);
-	ADD_COOK_STAT_FLT(" 0. 0", StartupWallTimeSec);
-	ADD_COOK_STAT_FLT(" 0. 1", CookByTheBookTimeSec);
-	ADD_COOK_STAT_FLT(" 0. 1. 0", StartCookByTheBookTimeSec);
-	ADD_COOK_STAT_FLT(" 0. 1. 0. 0", BlockOnAssetRegistryTimeSec);
-	ADD_COOK_STAT_FLT(" 0. 1. 0. 1", GameCookModificationDelegateTimeSec);
-	ADD_COOK_STAT_FLT(" 0. 1. 1", TickCookOnTheSideTimeSec);
-	ADD_COOK_STAT_FLT(" 0. 1. 1. 0", TickCookOnTheSideLoadPackagesTimeSec);
-	ADD_COOK_STAT_FLT(" 0. 1. 1. 1", TickCookOnTheSideSaveCookedPackageTimeSec);
-	ADD_COOK_STAT_FLT(" 0. 1. 1. 1. 0", TickCookOnTheSideResolveRedirectorsTimeSec);
-	ADD_COOK_STAT_FLT(" 0. 1. 1. 2", TickCookOnTheSidePrepareSaveTimeSec);
-	ADD_COOK_STAT_FLT(" 0. 1. 2", TickLoopGCTimeSec);
-	ADD_COOK_STAT_FLT(" 0. 1. 3", TickLoopRecompileShaderRequestsTimeSec);
-	ADD_COOK_STAT_FLT(" 0. 1. 4", TickLoopShaderProcessAsyncResultsTimeSec);
-	ADD_COOK_STAT_FLT(" 0. 1. 5", TickLoopProcessDeferredCommandsTimeSec);
-	ADD_COOK_STAT_FLT(" 0. 1. 6", TickLoopTickCommandletStatsTimeSec);
-	ADD_COOK_STAT_FLT(" 0. 1. 7", TickLoopFlushRenderingCommandsTimeSec);
-	ADD_COOK_STAT_FLT(" 0. 1. 8", TargetPlatforms);
-	ADD_COOK_STAT_FLT(" 0. 1. 9", CookProject);
-	ADD_COOK_STAT_FLT(" 0. 1. 10", CookCultures);
-	ADD_COOK_STAT_FLT(" 0. 1. 11", IsCookAll);
-	ADD_COOK_STAT_FLT(" 0. 1. 12", IsCookOnTheFly);
-	ADD_COOK_STAT_FLT(" 0. 1. 13", IsIterativeCook);
-	ADD_COOK_STAT_FLT(" 0. 1. 14", IsUnversioned);
-	ADD_COOK_STAT_FLT(" 0. 1. 15", CookLabel);
-	ADD_COOK_STAT_FLT(" 0. 1. 16", IsFastCook);
-		
-	#undef ADD_COOK_STAT_FLT
-});
-
-static void LogCookStats(const FString& CookCmdLine)
-{
-	if (FStudioAnalytics::IsAvailable())
-	{
-
-		// convert filtered stats directly to an analytics event
-		TArray<FAnalyticsEventAttribute> StatAttrs;
-
-		// Sends each cook stat to the studio analytics system.
-		auto SendCookStatsToAnalytics = [&StatAttrs](const FString& StatName, const TArray<FCookStatsManager::StringKeyValue>& StatAttributes)
-		{
-			for (const auto& Attr : StatAttributes)
-			{
-				FString FormattedAttrName = StatName + "." + Attr.Key;
-
-				StatAttrs.Emplace(FormattedAttrName, Attr.Value);
-			}
-		};
-
-		// Now actually grab the stats 
-		FCookStatsManager::LogCookStats(SendCookStatsToAnalytics);
-
-		// Record them all under cooking event
-		FStudioAnalytics::GetProvider().RecordEvent(TEXT("Core.Cooking"), StatAttrs);
-
-		FStudioAnalytics::GetProvider().BlockUntilFlushed(60.0f);
-	}
-
-	bool bSendCookAnalytics = false;
-	GConfig->GetBool(TEXT("CookAnalytics"), TEXT("SendAnalytics"), bSendCookAnalytics, GEngineIni);
-
-	if (GIsBuildMachine || FParse::Param(FCommandLine::Get(), TEXT("SendCookAnalytics")) || bSendCookAnalytics)
-	{
-		FString APIServerET;
-		if (GConfig->GetString(TEXT("CookAnalytics"), TEXT("APIServer"), APIServerET, GEngineIni))
-		{
-			FString AppId(TEXT("Cook"));
-			bool bUseLegacyCookProtocol = !GConfig->GetString(TEXT("CookAnalytics"), TEXT("AppId"), AppId, GEngineIni);
-
-			// Optionally create an analytics provider to send stats to for central collection.
-			TSharedPtr<IAnalyticsProviderET> CookAnalytics = FAnalyticsET::Get().CreateAnalyticsProvider(FAnalyticsET::Config(AppId, APIServerET, FString(), bUseLegacyCookProtocol));
-			if (CookAnalytics.IsValid())
-			{
-				CookAnalytics->SetUserID(FString::Printf(TEXT("%s\\%s"), FPlatformProcess::ComputerName(), FPlatformProcess::UserName(false)));
-				CookAnalytics->StartSession(MakeAnalyticsEventAttributeArray(
-					TEXT("Project"), CookProject,
-					TEXT("CmdLine"), CookCmdLine,
-					TEXT("IsBuildMachine"), GIsBuildMachine,
-					TEXT("TargetPlatforms"), TargetPlatforms
-				));
-
-				TArray<FString> CookStatsToSend;
-				const bool bFilterStats = GConfig->GetArray(TEXT("CookAnalytics"), TEXT("CookStats"), CookStatsToSend, GEngineIni) > 0;
-				// Sends each cook stat to the analytics provider.
-				auto SendCookStatsToAnalytics = [CookAnalytics, &CookStatsToSend, bFilterStats](const FString& StatName, const TArray<FCookStatsManager::StringKeyValue>& StatAttributes)
-				{
-					if (!bFilterStats || CookStatsToSend.Contains(StatName))
-					{
-						// convert filtered stats directly to an analytics event
-						TArray<FAnalyticsEventAttribute> StatAttrs;
-						StatAttrs.Reset(StatAttributes.Num());
-						for (const auto& Attr : StatAttributes)
-						{
-							StatAttrs.Emplace(Attr.Key, Attr.Value);
-						}
-						CookAnalytics->RecordEvent(StatName, StatAttrs);
-					}
-					else
-					{
-						UE_LOG(LogCookCommandlet, Verbose, TEXT("[%s] not present in analytics CookStats filter"), *StatName);
-					}
-				};
-				FCookStatsManager::LogCookStats(SendCookStatsToAnalytics);
-			}
-		}
-	}
-
-	/** Used for custom logging of DDC Resource usage stats. */
-	struct FDDCResourceUsageStat
-	{
-	public:
-		FDDCResourceUsageStat(FString InAssetType, double InTotalTimeSec, bool bIsGameThreadTime, double InSizeMB, int64 InAssetsBuilt) : AssetType(MoveTemp(InAssetType)), TotalTimeSec(InTotalTimeSec), GameThreadTimeSec(bIsGameThreadTime ? InTotalTimeSec : 0.0), SizeMB(InSizeMB), AssetsBuilt(InAssetsBuilt) {}
-		void Accumulate(const FDDCResourceUsageStat& OtherStat)
-		{
-			TotalTimeSec += OtherStat.TotalTimeSec;
-			GameThreadTimeSec += OtherStat.GameThreadTimeSec;
-			SizeMB += OtherStat.SizeMB;
-			AssetsBuilt += OtherStat.AssetsBuilt;
-		}
-		FString AssetType;
-		double TotalTimeSec;
-		double GameThreadTimeSec;
-		double SizeMB;
-		int64 AssetsBuilt;
-	};
-
-	/** Used for custom TSet comparison of DDC Resource usage stats. */
-	struct FDDCResourceUsageStatKeyFuncs : BaseKeyFuncs<FDDCResourceUsageStat, FString, false>
-	{
-		static const FString& GetSetKey(const FDDCResourceUsageStat& Element) { return Element.AssetType; }
-		static bool Matches(const FString& A, const FString& B) { return A == B; }
-		static uint32 GetKeyHash(const FString& Key) { return GetTypeHash(Key); }
-	};
-
-	/** Used to store profile data for custom logging. */
-	struct FCookProfileData
-	{
-	public:
-		FCookProfileData(FString InPath, FString InKey, FString InValue) : Path(MoveTemp(InPath)), Key(MoveTemp(InKey)), Value(MoveTemp(InValue)) {}
-		FString Path;
-		FString Key;
-		FString Value;
-	};
-
-	// instead of printing the usage stats generically, we capture them so we can log a subset of them in an easy-to-read way.
-	TSet<FDDCResourceUsageStat, FDDCResourceUsageStatKeyFuncs> DDCResourceUsageStats;
-	TArray<FCookStatsManager::StringKeyValue> DDCSummaryStats;
-	TArray<FCookProfileData> CookProfileData;
-	TArray<FString> StatCategories;
-	TMap<FString, TArray<FCookStatsManager::StringKeyValue>> StatsInCategories;
-
-	/** this functor will take a collected cooker stat and log it out using some custom formatting based on known stats that are collected.. */
-	auto LogStatsFunc = [&DDCResourceUsageStats, &DDCSummaryStats, &CookProfileData, &StatCategories, &StatsInCategories]
-	(const FString& StatName, const TArray<FCookStatsManager::StringKeyValue>& StatAttributes)
-	{
-		// Some stats will use custom formatting to make a visibly pleasing summary.
-		bool bStatUsedCustomFormatting = false;
-
-		if (StatName == TEXT("DDC.Usage"))
-		{
-			// Don't even log this detailed DDC data. It's mostly only consumable by ingestion into pivot tools.
-			bStatUsedCustomFormatting = true;
-		}
-		else if (StatName.EndsWith(TEXT(".Usage"), ESearchCase::IgnoreCase))
-		{
-			// Anything that ends in .Usage is assumed to be an instance of FCookStats.FDDCResourceUsageStats. We'll log that using custom formatting.
-			FString AssetType = StatName;
-			AssetType.RemoveFromEnd(TEXT(".Usage"), ESearchCase::IgnoreCase);
-			// See if the asset has a subtype (found via the "Node" parameter")
-			const FCookStatsManager::StringKeyValue* AssetSubType = StatAttributes.FindByPredicate([](const FCookStatsManager::StringKeyValue& Item) { return Item.Key == TEXT("Node"); });
-			if (AssetSubType && AssetSubType->Value.Len() > 0)
-			{
-				AssetType += FString::Printf(TEXT(" (%s)"), *AssetSubType->Value);
-			}
-			// Pull the Time and Size attributes and AddOrAccumulate them into the set of stats. Ugly string/container manipulation code courtesy of UE/C++.
-			const FCookStatsManager::StringKeyValue* AssetTimeSecAttr = StatAttributes.FindByPredicate([](const FCookStatsManager::StringKeyValue& Item) { return Item.Key == TEXT("TimeSec"); });
-			double AssetTimeSec = 0.0;
-			if (AssetTimeSecAttr)
-			{
-				LexFromString(AssetTimeSec, *AssetTimeSecAttr->Value);
-			}
-			const FCookStatsManager::StringKeyValue* AssetSizeMBAttr = StatAttributes.FindByPredicate([](const FCookStatsManager::StringKeyValue& Item) { return Item.Key == TEXT("MB"); });
-			double AssetSizeMB = 0.0;
-			if (AssetSizeMBAttr)
-			{
-				LexFromString(AssetSizeMB, *AssetSizeMBAttr->Value);
-			}
-			const FCookStatsManager::StringKeyValue* ThreadNameAttr = StatAttributes.FindByPredicate([](const FCookStatsManager::StringKeyValue& Item) { return Item.Key == TEXT("ThreadName"); });
-			bool bIsGameThreadTime = ThreadNameAttr != nullptr && ThreadNameAttr->Value == TEXT("GameThread");
-
-			const FCookStatsManager::StringKeyValue* HitOrMissAttr = StatAttributes.FindByPredicate([](const FCookStatsManager::StringKeyValue& Item) { return Item.Key == TEXT("HitOrMiss"); });
-			bool bWasMiss = HitOrMissAttr != nullptr && HitOrMissAttr->Value == TEXT("Miss");
-			int64 AssetsBuilt = 0;
-			if (bWasMiss)
-			{
-				const FCookStatsManager::StringKeyValue* CountAttr = StatAttributes.FindByPredicate([](const FCookStatsManager::StringKeyValue& Item) { return Item.Key == TEXT("Count"); });
-				if (CountAttr)
-				{
-					LexFromString(AssetsBuilt, *CountAttr->Value);
-				}
-			}
-
-
-			FDDCResourceUsageStat Stat(AssetType, AssetTimeSec, bIsGameThreadTime, AssetSizeMB, AssetsBuilt);
-			FDDCResourceUsageStat* ExistingStat = DDCResourceUsageStats.Find(Stat.AssetType);
-			if (ExistingStat)
-			{
-				ExistingStat->Accumulate(Stat);
-			}
-			else
-			{
-				DDCResourceUsageStats.Add(Stat);
-			}
-			bStatUsedCustomFormatting = true;
-		}
-		else if (StatName == TEXT("DDC.Summary"))
-		{
-			DDCSummaryStats.Append(StatAttributes);
-			bStatUsedCustomFormatting = true;
-		}
-		else if (StatName == TEXT("Cook.Profile"))
-		{
-			if (StatAttributes.Num() >= 2)
-			{
-				CookProfileData.Emplace(StatAttributes[0].Value, StatAttributes[1].Key, StatAttributes[1].Value);
-			}
-			bStatUsedCustomFormatting = true;
-		}
-
-		// if a stat doesn't use custom formatting, just spit out the raw info.
-		if (!bStatUsedCustomFormatting)
-		{
-			TArray<FCookStatsManager::StringKeyValue>& StatsInCategory = StatsInCategories.FindOrAdd(StatName);
-			if (StatsInCategory.Num() == 0)
-			{
-				StatCategories.Add(StatName);
-			}
-			StatsInCategory.Append(StatAttributes);
-		}
-	};
-
-	FCookStatsManager::LogCookStats(LogStatsFunc);
-
-	UE_LOG(LogCookCommandlet, Display, TEXT("Misc Cook Stats"));
-	UE_LOG(LogCookCommandlet, Display, TEXT("==============="));
-	for (FString& StatCategory : StatCategories)
-	{
-		UE_LOG(LogCookCommandlet, Display, TEXT("%s"), *StatCategory);
-		TArray<FCookStatsManager::StringKeyValue>& StatsInCategory = StatsInCategories.FindOrAdd(StatCategory);
-
-		// log each key/value pair, with the equal signs lined up.
-		for (const FCookStatsManager::StringKeyValue& StatKeyValue : StatsInCategory)
-		{
-			UE_LOG(LogCookCommandlet, Display, TEXT("    %s=%s"), *StatKeyValue.Key, *StatKeyValue.Value);
-		}
-	}
-
-	// DDC Usage stats are custom formatted, and the above code just accumulated them into a TSet. Now log it with our special formatting for readability.
-	if (CookProfileData.Num() > 0)
-	{
-		UE_LOG(LogCookCommandlet, Display, TEXT(""));
-		UE_LOG(LogCookCommandlet, Display, TEXT("Cook Profile"));
-		UE_LOG(LogCookCommandlet, Display, TEXT("============"));
-		for (const auto& ProfileEntry : CookProfileData)
-		{
-			UE_LOG(LogCookCommandlet, Display, TEXT("%s.%s=%s"), *ProfileEntry.Path, *ProfileEntry.Key, *ProfileEntry.Value);
-		}
-	}
-	if (DDCSummaryStats.Num() > 0)
-	{
-		UE_LOG(LogCookCommandlet, Display, TEXT(""));
-		UE_LOG(LogCookCommandlet, Display, TEXT("DDC Summary Stats"));
-		UE_LOG(LogCookCommandlet, Display, TEXT("================="));
-		for (const auto& Attr : DDCSummaryStats)
-		{
-			UE_LOG(LogCookCommandlet, Display, TEXT("%-16s=%10s"), *Attr.Key, *Attr.Value);
-		}
-	}
-
-	DumpDerivedDataBuildRemoteExecutorStats();
-
-	if (DDCResourceUsageStats.Num() > 0)
-	{
-		// sort the list
-		TArray<FDDCResourceUsageStat> SortedDDCResourceUsageStats;
-		SortedDDCResourceUsageStats.Empty(DDCResourceUsageStats.Num());
-		for (const FDDCResourceUsageStat& Stat : DDCResourceUsageStats)
-		{
-			SortedDDCResourceUsageStats.Emplace(Stat);
-		}
-		SortedDDCResourceUsageStats.Sort([](const FDDCResourceUsageStat& LHS, const FDDCResourceUsageStat& RHS)
-			{
-				return LHS.TotalTimeSec > RHS.TotalTimeSec;
-			});
-
-		UE_LOG(LogCookCommandlet, Display, TEXT(""));
-		UE_LOG(LogCookCommandlet, Display, TEXT("DDC Resource Stats"));
-		UE_LOG(LogCookCommandlet, Display, TEXT("======================================================================================================="));
-		UE_LOG(LogCookCommandlet, Display, TEXT("Asset Type                          Total Time (Sec)  GameThread Time (Sec)  Assets Built  MB Processed"));
-		UE_LOG(LogCookCommandlet, Display, TEXT("----------------------------------  ----------------  ---------------------  ------------  ------------"));
-		for (const FDDCResourceUsageStat& Stat : SortedDDCResourceUsageStats)
-		{
-			UE_LOG(LogCookCommandlet, Display, TEXT("%-34s  %16.2f  %21.2f  %12d  %12.2f"), *Stat.AssetType, Stat.TotalTimeSec, Stat.GameThreadTimeSec, Stat.AssetsBuilt, Stat.SizeMB);
-		}
-	}
-
-	DumpBuildDependencyTrackerStats();
-
-	if (UE::Virtualization::IVirtualizationSystem::IsInitialized())
-	{
-		UE::Virtualization::IVirtualizationSystem::Get().DumpStats();
-	}
-}
-
-}
-#endif
-
 namespace UE::Cook
 {
 
@@ -524,7 +164,7 @@ bool UCookCommandlet::CookOnTheFly( FGuid InstanceId, int32 Timeout, bool bForce
 
 int32 UCookCommandlet::Main(const FString& CmdLineParams)
 {
-	COOK_STAT(double CookStartTime = FPlatformTime::Seconds());
+	COOK_STAT(DetailedCookStats::CookStartTime = FPlatformTime::Seconds());
 	Params = CmdLineParams;
 	ParseCommandLine(*Params, Tokens, Switches);
 
@@ -598,25 +238,7 @@ int32 UCookCommandlet::Main(const FString& CmdLineParams)
 	}
 	else
 	{
-		const TArray<ITargetPlatform*>& Platforms = TPM.GetActiveTargetPlatforms();
-
-		CookByTheBook(Platforms);
-		
-		if(GShaderCompilerStats)
-		{
-			GShaderCompilerStats->WriteStats();
-		}
-
-		// Use -LogCookStats to log the results to the command line after the cook (happens automatically on a build machine)
-		COOK_STAT(
-		{
-			double Now = FPlatformTime::Seconds();
-			DetailedCookStats::CookWallTimeSec = Now - GStartTime;
-			DetailedCookStats::StartupWallTimeSec = CookStartTime - GStartTime;
-			DetailedCookStats::LogCookStats(CmdLineParams);
-
-			FStudioAnalytics::FireEvent_Loading(TEXT("CookByTheBook"), DetailedCookStats::CookWallTimeSec);
-		});
+		CookByTheBook(TPM.GetActiveTargetPlatforms());
 	}
 	return 0;
 }
@@ -627,11 +249,19 @@ bool UCookCommandlet::CookByTheBook( const TArray<ITargetPlatform*>& Platforms)
 	TRACE_CPUPROFILER_EVENT_SCOPE_ON_CHANNEL(CookByTheBook, CookChannel);
 #endif // OUTPUT_COOKTIMING
 
-	COOK_STAT(FScopedDurationTimer CookByTheBookTimer(DetailedCookStats::CookByTheBookTimeSec));
-	UCookOnTheFlyServer *CookOnTheFlyServer = NewObject<UCookOnTheFlyServer>();
+	COOK_STAT(TOptional<FScopedDurationTimer> CookByTheBookTimer);
+	COOK_STAT(CookByTheBookTimer.Emplace(DetailedCookStats::CookByTheBookTimeSec));
+
+	UCookOnTheFlyServer* CookOnTheFlyServer = NewObject<UCookOnTheFlyServer>();
 	// make sure that the cookonthefly server doesn't get cleaned up while we are garbage collecting below :)
 	UE::Cook::FScopeRootObject S(CookOnTheFlyServer);
 
+	ON_SCOPE_EXIT
+	{
+		COOK_STAT(CookByTheBookTimer.Reset()); // Destruct the timer and write its duration
+		COOK_STAT(CookOnTheFlyServer->PrintDetailedCookStats());
+	};
+
 	UCookerSettings const* CookerSettings = GetDefault<UCookerSettings>();
 	const UProjectPackagingSettings* const PackagingSettings = GetDefault<UProjectPackagingSettings>();
 	ECookInitializationFlags IterateFlags = ECookInitializationFlags::Iterative;
diff --git a/Engine/Source/Editor/UnrealEd/Private/CookOnTheFlyServer.cpp b/Engine/Source/Editor/UnrealEd/Private/CookOnTheFlyServer.cpp
--- a/Engine/Source/Editor/UnrealEd/Private/CookOnTheFlyServer.cpp
+++ b/Engine/Source/Editor/UnrealEd/Private/CookOnTheFlyServer.cpp
@@ -1070,22 +1070,22 @@ bool UCookOnTheFlyServer::ContainsRedirector(const FName& PackageName, TMap<FSof
 
 bool UCookOnTheFlyServer::IsCookingInEditor() const
 {
-	return CurrentCookMode == ECookMode::CookByTheBookFromTheEditor || CurrentCookMode == ECookMode::CookOnTheFlyFromTheEditor;
+	return ::IsCookingInEditor(CurrentCookMode);
 }
 
 bool UCookOnTheFlyServer::IsRealtimeMode() const 
 {
-	return CurrentCookMode == ECookMode::CookByTheBookFromTheEditor || CurrentCookMode == ECookMode::CookOnTheFlyFromTheEditor;
+	return ::IsRealtimeMode(CurrentCookMode);
 }
 
 bool UCookOnTheFlyServer::IsCookByTheBookMode() const
 {
-	return CurrentCookMode == ECookMode::CookByTheBookFromTheEditor || CurrentCookMode == ECookMode::CookByTheBook;
+	return ::IsCookByTheBookMode(CurrentCookMode);
 }
 
 bool UCookOnTheFlyServer::IsDirectorCookByTheBook() const
 {
-	return DirectorCookMode == ECookMode::CookByTheBookFromTheEditor || DirectorCookMode == ECookMode::CookByTheBook;
+	return ::IsCookByTheBookMode(DirectorCookMode);
 }
 
 bool UCookOnTheFlyServer::IsUsingShaderCodeLibrary() const
@@ -1101,14 +1101,18 @@ bool UCookOnTheFlyServer::IsUsingZenStore() const
 
 bool UCookOnTheFlyServer::IsCookOnTheFlyMode() const
 {
-	return CurrentCookMode == ECookMode::CookOnTheFly || CurrentCookMode == ECookMode::CookOnTheFlyFromTheEditor; 
+	return ::IsCookOnTheFlyMode(CurrentCookMode);
 }
 
 bool UCookOnTheFlyServer::IsDirectorCookOnTheFly() const
 {
-	return DirectorCookMode == ECookMode::CookOnTheFly || DirectorCookMode == ECookMode::CookOnTheFlyFromTheEditor;
+	return ::IsCookOnTheFlyMode(DirectorCookMode);
 }
 
+bool UCookOnTheFlyServer::IsCookWorkerMode() const
+{
+	return ::IsCookWorkerMode(CurrentCookMode);
+}
 
 bool UCookOnTheFlyServer::IsUsingLegacyCookOnTheFlyScheduling() const
 {
@@ -1293,7 +1297,6 @@ uint32 UCookOnTheFlyServer::TickCookWorker()
 	check(IsCookWorkerMode());
 
 	LLM_SCOPE_BYTAG(Cooker);
-	COOK_STAT(FScopedDurationTimer TickTimer(DetailedCookStats::TickCookOnTheSideTimeSec));
 	UE::Cook::FTickStackData StackData(MAX_flt, ECookTickFlags::None);
 
 	TickMainCookLoop(StackData);
@@ -3419,6 +3422,7 @@ UE::Cook::EPollStatus UCookOnTheFlyServer::PrepareSave(UE::Cook::FPackageData& P
 	{
 		return EPollStatus::Error;
 	}
+	UE_SCOPED_HIERARCHICAL_COOKTIMER_AND_DURATION(PrepareSave, DetailedCookStats::TickCookOnTheSidePrepareSaveTimeSec);
 	EPollStatus Result = PrepareSaveInternal(PackageData, Timer, bPrecaching);
 	if (Result == EPollStatus::Error)
 	{
@@ -3432,7 +3436,6 @@ UE::Cook::EPollStatus UCookOnTheFlyServer::PrepareSaveInternal(UE::Cook::FPackag
 {
 	using namespace UE::Cook;
 
-	UE_SCOPED_HIERARCHICAL_COOKTIMER_AND_DURATION(PrepareSave, DetailedCookStats::TickCookOnTheSidePrepareSaveTimeSec);
 #if DEBUG_COOKONTHEFLY 
 	UE_LOG(LogCook, Display, TEXT("Caching objects for package %s"), *PackageData.GetPackageName().ToString());
 #endif
@@ -8622,7 +8625,15 @@ void UCookOnTheFlyServer::ShutdownCookSession()
 void UCookOnTheFlyServer::PrintFinishStats()
 {
 	const float TotalCookTime = (float)(FPlatformTime::Seconds() - CookByTheBookOptions->CookStartTime);
-	UE_LOG(LogCook, Display, TEXT("Cook by the book total time in tick %fs total time %f"), CookByTheBookOptions->CookTime, TotalCookTime);
+	if (IsCookByTheBookMode())
+	{
+		UE_LOG(LogCook, Display, TEXT("Cook by the book total time in tick %fs total time %f"), CookByTheBookOptions->CookTime, TotalCookTime);
+	}
+	else if (IsCookWorkerMode())
+	{
+		UE_LOG(LogCook, Display, TEXT("CookWorker total time %f"), TotalCookTime);
+	}
+
 
 	const FPlatformMemoryStats MemStats = FPlatformMemory::GetStats();
 	UE_LOG(LogCook, Display, TEXT("Peak Used virtual %u MiB Peak Used physical %u MiB"), MemStats.PeakUsedVirtual / 1024 / 1024, MemStats.PeakUsedPhysical / 1024 / 1024);
@@ -8631,6 +8642,26 @@ void UCookOnTheFlyServer::PrintFinishStats()
 		DetailedCookStats::NumPackagesSavedForCook, DetailedCookStats::NumPackagesIterativelySkipped, PackageDatas->GetNumCooked()));
 }
 
+void UCookOnTheFlyServer::PrintDetailedCookStats()
+{
+	if (GShaderCompilerStats)
+	{
+		GShaderCompilerStats->WriteStats();
+	}
+
+	COOK_STAT(
+		{
+			double Now = FPlatformTime::Seconds();
+			if (DetailedCookStats::CookStartTime <= 0.)
+			{
+				DetailedCookStats::CookStartTime = GStartTime;
+			}
+			DetailedCookStats::CookWallTimeSec = Now - GStartTime;
+			DetailedCookStats::StartupWallTimeSec = DetailedCookStats::CookStartTime - GStartTime;
+			DetailedCookStats::LogCookStats(CurrentCookMode);
+		});
+}
+
 TMap<FName, TSet<FName>> UCookOnTheFlyServer::BuildMapDependencyGraph(const ITargetPlatform* TargetPlatform)
 {
 	TMap<FName, TSet<FName>> MapDependencyGraph;
@@ -9501,6 +9532,16 @@ void UCookOnTheFlyServer::StartCookAsCookWorker()
 	}
 }
 
+void UCookOnTheFlyServer::LogCookWorkerStats()
+{
+	if (IsDirectorCookByTheBook())
+	{
+		PrintFinishStats();
+		OutputHierarchyTimers();
+		PrintDetailedCookStats();
+	}
+}
+
 void UCookOnTheFlyServer::ShutdownCookAsCookWorker()
 {
 	if (IsDirectorCookByTheBook())
@@ -9540,11 +9581,6 @@ FBeginCookContext UCookOnTheFlyServer::CreateCookWorkerContext()
 	return BeginContext;
 }
 
-bool UCookOnTheFlyServer::IsCookWorkerMode() const
-{
-	return CurrentCookMode == ECookMode::CookWorker;
-}
-
 void UCookOnTheFlyServer::GenerateInitialRequests(FBeginCookContext& BeginContext)
 {
 	TArray<ITargetPlatform*>& TargetPlatforms = BeginContext.TargetPlatforms;
diff --git a/Engine/Source/Editor/UnrealEd/Private/Cooker/CookProfiling.cpp b/Engine/Source/Editor/UnrealEd/Private/Cooker/CookProfiling.cpp
--- a/Engine/Source/Editor/UnrealEd/Private/Cooker/CookProfiling.cpp
+++ b/Engine/Source/Editor/UnrealEd/Private/Cooker/CookProfiling.cpp
@@ -11,8 +11,10 @@
 #include "Containers/UnrealString.h"
 #include "CookOnTheSide/CookLog.h"
 #include "CoreGlobals.h"
+#include "DerivedDataBuildRemoteExecutor.h"
 #include "Misc/OutputDevice.h"
 #include "Misc/StringBuilder.h"
+#include "PackageBuildDependencyTracker.h"
 #include "Serialization/ArchiveUObject.h"
 #include "Templates/Casts.h"
 #include "UObject/GCObject.h"
@@ -29,7 +31,17 @@
 
 #if OUTPUT_COOKTIMING
 #include <Containers/AllocatorFixedSizeFreeList.h>
+#endif
+
+#if ENABLE_COOK_STATS
+#include "AnalyticsET.h"
+#include "AnalyticsEventAttribute.h"
+#include "IAnalyticsProviderET.h"
+#include "StudioAnalytics.h"
+#include "Virtualization/VirtualizationSystem.h"
+#endif
 
+#if OUTPUT_COOKTIMING
 struct FHierarchicalTimerInfo
 {
 public:
@@ -193,6 +205,15 @@ void ClearHierarchyTimers()
 #if ENABLE_COOK_STATS
 namespace DetailedCookStats
 {
+	FString CookProject;
+	FString CookCultures;
+	FString CookLabel;
+	FString TargetPlatforms;
+	double CookStartTime = 0.0;
+	double CookWallTimeSec = 0.0;
+	double StartupWallTimeSec = 0.0;
+	double CookByTheBookTimeSec = 0.0;
+	double StartCookByTheBookTimeSec = 0.0;
 	double TickCookOnTheSideTimeSec = 0.0;
 	double TickCookOnTheSideLoadPackagesTimeSec = 0.0;
 	double TickCookOnTheSideResolveRedirectorsTimeSec = 0.0;
@@ -200,6 +221,18 @@ namespace DetailedCookStats
 	double TickCookOnTheSidePrepareSaveTimeSec = 0.0;
 	double BlockOnAssetRegistryTimeSec = 0.0;
 	double GameCookModificationDelegateTimeSec = 0.0;
+	double TickLoopGCTimeSec = 0.0;
+	double TickLoopRecompileShaderRequestsTimeSec = 0.0;
+	double TickLoopShaderProcessAsyncResultsTimeSec = 0.0;
+	double TickLoopProcessDeferredCommandsTimeSec = 0.0;
+	double TickLoopTickCommandletStatsTimeSec = 0.0;
+	double TickLoopFlushRenderingCommandsTimeSec = 0.0;
+	bool IsCookAll = false;
+	bool IsCookOnTheFly = false;
+	bool IsIterativeCook = false;
+	bool IsFastCook = false;
+	bool IsUnversioned = false;
+
 
 	// Stats tracked through FAutoRegisterCallback
 	int32 PeakRequestQueueSize = 0;
@@ -612,3 +645,340 @@ void DumpObjClassList(TConstArrayView<FWeakObjectPtr> InitialObjects)
 }
 
 } // namespace UE::Cook
+
+#if ENABLE_COOK_STATS
+
+namespace DetailedCookStats
+{
+
+FCookStatsManager::FAutoRegisterCallback RegisterCookStats([](FCookStatsManager::AddStatFuncRef AddStat)
+{
+	const FString StatName(TEXT("Cook.Profile"));
+	#define ADD_COOK_STAT_FLT(Path, Name) AddStat(StatName, FCookStatsManager::CreateKeyValueArray(TEXT("Path"), TEXT(Path), TEXT(#Name), Name))
+	ADD_COOK_STAT_FLT(" 0", CookWallTimeSec);
+	ADD_COOK_STAT_FLT(" 0. 0", StartupWallTimeSec);
+	ADD_COOK_STAT_FLT(" 0. 1", CookByTheBookTimeSec);
+	ADD_COOK_STAT_FLT(" 0. 1. 0", StartCookByTheBookTimeSec);
+	ADD_COOK_STAT_FLT(" 0. 1. 0. 0", BlockOnAssetRegistryTimeSec);
+	ADD_COOK_STAT_FLT(" 0. 1. 0. 1", GameCookModificationDelegateTimeSec);
+	ADD_COOK_STAT_FLT(" 0. 1. 1", TickCookOnTheSideTimeSec);
+	ADD_COOK_STAT_FLT(" 0. 1. 1. 0", TickCookOnTheSideLoadPackagesTimeSec);
+	ADD_COOK_STAT_FLT(" 0. 1. 1. 1", TickCookOnTheSideSaveCookedPackageTimeSec);
+	ADD_COOK_STAT_FLT(" 0. 1. 1. 1. 0", TickCookOnTheSideResolveRedirectorsTimeSec);
+	ADD_COOK_STAT_FLT(" 0. 1. 1. 2", TickCookOnTheSidePrepareSaveTimeSec);
+	ADD_COOK_STAT_FLT(" 0. 1. 2", TickLoopGCTimeSec);
+	ADD_COOK_STAT_FLT(" 0. 1. 3", TickLoopRecompileShaderRequestsTimeSec);
+	ADD_COOK_STAT_FLT(" 0. 1. 4", TickLoopShaderProcessAsyncResultsTimeSec);
+	ADD_COOK_STAT_FLT(" 0. 1. 5", TickLoopProcessDeferredCommandsTimeSec);
+	ADD_COOK_STAT_FLT(" 0. 1. 6", TickLoopTickCommandletStatsTimeSec);
+	ADD_COOK_STAT_FLT(" 0. 1. 7", TickLoopFlushRenderingCommandsTimeSec);
+	ADD_COOK_STAT_FLT(" 0. 1. 8", TargetPlatforms);
+	ADD_COOK_STAT_FLT(" 0. 1. 9", CookProject);
+	ADD_COOK_STAT_FLT(" 0. 1. 10", CookCultures);
+	ADD_COOK_STAT_FLT(" 0. 1. 11", IsCookAll);
+	ADD_COOK_STAT_FLT(" 0. 1. 12", IsCookOnTheFly);
+	ADD_COOK_STAT_FLT(" 0. 1. 13", IsIterativeCook);
+	ADD_COOK_STAT_FLT(" 0. 1. 14", IsUnversioned);
+	ADD_COOK_STAT_FLT(" 0. 1. 15", CookLabel);
+	ADD_COOK_STAT_FLT(" 0. 1. 16", IsFastCook);
+		
+	#undef ADD_COOK_STAT_FLT
+});
+
+void LogCookStats(ECookMode::Type CookMode)
+{
+	if (IsCookingInEditor(CookMode))
+	{
+		return;
+	}
+
+	if (FStudioAnalytics::IsAvailable() && IsCookByTheBookMode(CookMode))
+	{
+
+		// convert filtered stats directly to an analytics event
+		TArray<FAnalyticsEventAttribute> StatAttrs;
+
+		// Sends each cook stat to the studio analytics system.
+		auto SendCookStatsToAnalytics = [&StatAttrs](const FString& StatName, const TArray<FCookStatsManager::StringKeyValue>& StatAttributes)
+		{
+			for (const auto& Attr : StatAttributes)
+			{
+				FString FormattedAttrName = StatName + "." + Attr.Key;
+
+				StatAttrs.Emplace(FormattedAttrName, Attr.Value);
+			}
+		};
+
+		// Now actually grab the stats 
+		FCookStatsManager::LogCookStats(SendCookStatsToAnalytics);
+
+		// Record them all under cooking event
+		FStudioAnalytics::GetProvider().RecordEvent(TEXT("Core.Cooking"), StatAttrs);
+
+		FStudioAnalytics::GetProvider().BlockUntilFlushed(60.0f);
+	}
+
+	bool bSendCookAnalytics = false;
+	GConfig->GetBool(TEXT("CookAnalytics"), TEXT("SendAnalytics"), bSendCookAnalytics, GEngineIni);
+
+	if (IsCookByTheBookMode(CookMode) &&
+		(GIsBuildMachine || FParse::Param(FCommandLine::Get(), TEXT("SendCookAnalytics")) || bSendCookAnalytics))
+	{
+		FString APIServerET;
+		if (GConfig->GetString(TEXT("CookAnalytics"), TEXT("APIServer"), APIServerET, GEngineIni))
+		{
+			FString AppId(TEXT("Cook"));
+			bool bUseLegacyCookProtocol = !GConfig->GetString(TEXT("CookAnalytics"), TEXT("AppId"), AppId, GEngineIni);
+
+			// Optionally create an analytics provider to send stats to for central collection.
+			TSharedPtr<IAnalyticsProviderET> CookAnalytics = FAnalyticsET::Get().CreateAnalyticsProvider(FAnalyticsET::Config(AppId, APIServerET, FString(), bUseLegacyCookProtocol));
+			if (CookAnalytics.IsValid())
+			{
+				CookAnalytics->SetUserID(FString::Printf(TEXT("%s\\%s"), FPlatformProcess::ComputerName(), FPlatformProcess::UserName(false)));
+				CookAnalytics->StartSession(MakeAnalyticsEventAttributeArray(
+					TEXT("Project"), CookProject,
+					TEXT("CmdLine"), FString(FCommandLine::Get()),
+					TEXT("IsBuildMachine"), GIsBuildMachine,
+					TEXT("TargetPlatforms"), TargetPlatforms
+				));
+
+				TArray<FString> CookStatsToSend;
+				const bool bFilterStats = GConfig->GetArray(TEXT("CookAnalytics"), TEXT("CookStats"), CookStatsToSend, GEngineIni) > 0;
+				// Sends each cook stat to the analytics provider.
+				auto SendCookStatsToAnalytics = [CookAnalytics, &CookStatsToSend, bFilterStats](const FString& StatName, const TArray<FCookStatsManager::StringKeyValue>& StatAttributes)
+				{
+					if (!bFilterStats || CookStatsToSend.Contains(StatName))
+					{
+						// convert filtered stats directly to an analytics event
+						TArray<FAnalyticsEventAttribute> StatAttrs;
+						StatAttrs.Reset(StatAttributes.Num());
+						for (const auto& Attr : StatAttributes)
+						{
+							StatAttrs.Emplace(Attr.Key, Attr.Value);
+						}
+						CookAnalytics->RecordEvent(StatName, StatAttrs);
+					}
+					else
+					{
+						UE_LOG(LogCook, Verbose, TEXT("[%s] not present in analytics CookStats filter"), *StatName);
+					}
+				};
+				FCookStatsManager::LogCookStats(SendCookStatsToAnalytics);
+			}
+		}
+	}
+
+	/** Used for custom logging of DDC Resource usage stats. */
+	struct FDDCResourceUsageStat
+	{
+	public:
+		FDDCResourceUsageStat(FString InAssetType, double InTotalTimeSec, bool bIsGameThreadTime, double InSizeMB, int64 InAssetsBuilt) : AssetType(MoveTemp(InAssetType)), TotalTimeSec(InTotalTimeSec), GameThreadTimeSec(bIsGameThreadTime ? InTotalTimeSec : 0.0), SizeMB(InSizeMB), AssetsBuilt(InAssetsBuilt) {}
+		void Accumulate(const FDDCResourceUsageStat& OtherStat)
+		{
+			TotalTimeSec += OtherStat.TotalTimeSec;
+			GameThreadTimeSec += OtherStat.GameThreadTimeSec;
+			SizeMB += OtherStat.SizeMB;
+			AssetsBuilt += OtherStat.AssetsBuilt;
+		}
+		FString AssetType;
+		double TotalTimeSec;
+		double GameThreadTimeSec;
+		double SizeMB;
+		int64 AssetsBuilt;
+	};
+
+	/** Used for custom TSet comparison of DDC Resource usage stats. */
+	struct FDDCResourceUsageStatKeyFuncs : BaseKeyFuncs<FDDCResourceUsageStat, FString, false>
+	{
+		static const FString& GetSetKey(const FDDCResourceUsageStat& Element) { return Element.AssetType; }
+		static bool Matches(const FString& A, const FString& B) { return A == B; }
+		static uint32 GetKeyHash(const FString& Key) { return GetTypeHash(Key); }
+	};
+
+	/** Used to store profile data for custom logging. */
+	struct FCookProfileData
+	{
+	public:
+		FCookProfileData(FString InPath, FString InKey, FString InValue) : Path(MoveTemp(InPath)), Key(MoveTemp(InKey)), Value(MoveTemp(InValue)) {}
+		FString Path;
+		FString Key;
+		FString Value;
+	};
+
+	// instead of printing the usage stats generically, we capture them so we can log a subset of them in an easy-to-read way.
+	TSet<FDDCResourceUsageStat, FDDCResourceUsageStatKeyFuncs> DDCResourceUsageStats;
+	TArray<FCookStatsManager::StringKeyValue> DDCSummaryStats;
+	TArray<FCookProfileData> CookProfileData;
+	TArray<FString> StatCategories;
+	TMap<FString, TArray<FCookStatsManager::StringKeyValue>> StatsInCategories;
+
+	/** this functor will take a collected cooker stat and log it out using some custom formatting based on known stats that are collected.. */
+	auto LogStatsFunc = [&DDCResourceUsageStats, &DDCSummaryStats, &CookProfileData, &StatCategories, &StatsInCategories]
+	(const FString& StatName, const TArray<FCookStatsManager::StringKeyValue>& StatAttributes)
+	{
+		// Some stats will use custom formatting to make a visibly pleasing summary.
+		bool bStatUsedCustomFormatting = false;
+
+		if (StatName == TEXT("DDC.Usage"))
+		{
+			// Don't even log this detailed DDC data. It's mostly only consumable by ingestion into pivot tools.
+			bStatUsedCustomFormatting = true;
+		}
+		else if (StatName.EndsWith(TEXT(".Usage"), ESearchCase::IgnoreCase))
+		{
+			// Anything that ends in .Usage is assumed to be an instance of FCookStats.FDDCResourceUsageStats. We'll log that using custom formatting.
+			FString AssetType = StatName;
+			AssetType.RemoveFromEnd(TEXT(".Usage"), ESearchCase::IgnoreCase);
+			// See if the asset has a subtype (found via the "Node" parameter")
+			const FCookStatsManager::StringKeyValue* AssetSubType = StatAttributes.FindByPredicate([](const FCookStatsManager::StringKeyValue& Item) { return Item.Key == TEXT("Node"); });
+			if (AssetSubType && AssetSubType->Value.Len() > 0)
+			{
+				AssetType += FString::Printf(TEXT(" (%s)"), *AssetSubType->Value);
+			}
+			// Pull the Time and Size attributes and AddOrAccumulate them into the set of stats. Ugly string/container manipulation code courtesy of UE/C++.
+			const FCookStatsManager::StringKeyValue* AssetTimeSecAttr = StatAttributes.FindByPredicate([](const FCookStatsManager::StringKeyValue& Item) { return Item.Key == TEXT("TimeSec"); });
+			double AssetTimeSec = 0.0;
+			if (AssetTimeSecAttr)
+			{
+				LexFromString(AssetTimeSec, *AssetTimeSecAttr->Value);
+			}
+			const FCookStatsManager::StringKeyValue* AssetSizeMBAttr = StatAttributes.FindByPredicate([](const FCookStatsManager::StringKeyValue& Item) { return Item.Key == TEXT("MB"); });
+			double AssetSizeMB = 0.0;
+			if (AssetSizeMBAttr)
+			{
+				LexFromString(AssetSizeMB, *AssetSizeMBAttr->Value);
+			}
+			const FCookStatsManager::StringKeyValue* ThreadNameAttr = StatAttributes.FindByPredicate([](const FCookStatsManager::StringKeyValue& Item) { return Item.Key == TEXT("ThreadName"); });
+			bool bIsGameThreadTime = ThreadNameAttr != nullptr && ThreadNameAttr->Value == TEXT("GameThread");
+
+			const FCookStatsManager::StringKeyValue* HitOrMissAttr = StatAttributes.FindByPredicate([](const FCookStatsManager::StringKeyValue& Item) { return Item.Key == TEXT("HitOrMiss"); });
+			bool bWasMiss = HitOrMissAttr != nullptr && HitOrMissAttr->Value == TEXT("Miss");
+			int64 AssetsBuilt = 0;
+			if (bWasMiss)
+			{
+				const FCookStatsManager::StringKeyValue* CountAttr = StatAttributes.FindByPredicate([](const FCookStatsManager::StringKeyValue& Item) { return Item.Key == TEXT("Count"); });
+				if (CountAttr)
+				{
+					LexFromString(AssetsBuilt, *CountAttr->Value);
+				}
+			}
+
+
+			FDDCResourceUsageStat Stat(AssetType, AssetTimeSec, bIsGameThreadTime, AssetSizeMB, AssetsBuilt);
+			FDDCResourceUsageStat* ExistingStat = DDCResourceUsageStats.Find(Stat.AssetType);
+			if (ExistingStat)
+			{
+				ExistingStat->Accumulate(Stat);
+			}
+			else
+			{
+				DDCResourceUsageStats.Add(Stat);
+			}
+			bStatUsedCustomFormatting = true;
+		}
+		else if (StatName == TEXT("DDC.Summary"))
+		{
+			DDCSummaryStats.Append(StatAttributes);
+			bStatUsedCustomFormatting = true;
+		}
+		else if (StatName == TEXT("Cook.Profile"))
+		{
+			if (StatAttributes.Num() >= 2)
+			{
+				CookProfileData.Emplace(StatAttributes[0].Value, StatAttributes[1].Key, StatAttributes[1].Value);
+			}
+			bStatUsedCustomFormatting = true;
+		}
+
+		// if a stat doesn't use custom formatting, just spit out the raw info.
+		if (!bStatUsedCustomFormatting)
+		{
+			TArray<FCookStatsManager::StringKeyValue>& StatsInCategory = StatsInCategories.FindOrAdd(StatName);
+			if (StatsInCategory.Num() == 0)
+			{
+				StatCategories.Add(StatName);
+			}
+			StatsInCategory.Append(StatAttributes);
+		}
+	};
+
+	FCookStatsManager::LogCookStats(LogStatsFunc);
+
+	UE_LOG(LogCook, Display, TEXT("Misc Cook Stats"));
+	UE_LOG(LogCook, Display, TEXT("==============="));
+	for (FString& StatCategory : StatCategories)
+	{
+		UE_LOG(LogCook, Display, TEXT("%s"), *StatCategory);
+		TArray<FCookStatsManager::StringKeyValue>& StatsInCategory = StatsInCategories.FindOrAdd(StatCategory);
+
+		// log each key/value pair, with the equal signs lined up.
+		for (const FCookStatsManager::StringKeyValue& StatKeyValue : StatsInCategory)
+		{
+			UE_LOG(LogCook, Display, TEXT("    %s=%s"), *StatKeyValue.Key, *StatKeyValue.Value);
+		}
+	}
+
+	// DDC Usage stats are custom formatted, and the above code just accumulated them into a TSet. Now log it with our special formatting for readability.
+	if (CookProfileData.Num() > 0)
+	{
+		UE_LOG(LogCook, Display, TEXT(""));
+		UE_LOG(LogCook, Display, TEXT("Cook Profile"));
+		UE_LOG(LogCook, Display, TEXT("============"));
+		for (const auto& ProfileEntry : CookProfileData)
+		{
+			UE_LOG(LogCook, Display, TEXT("%s.%s=%s"), *ProfileEntry.Path, *ProfileEntry.Key, *ProfileEntry.Value);
+		}
+	}
+	if (DDCSummaryStats.Num() > 0)
+	{
+		UE_LOG(LogCook, Display, TEXT(""));
+		UE_LOG(LogCook, Display, TEXT("DDC Summary Stats"));
+		UE_LOG(LogCook, Display, TEXT("================="));
+		for (const auto& Attr : DDCSummaryStats)
+		{
+			UE_LOG(LogCook, Display, TEXT("%-16s=%10s"), *Attr.Key, *Attr.Value);
+		}
+	}
+
+	DumpDerivedDataBuildRemoteExecutorStats();
+
+	if (DDCResourceUsageStats.Num() > 0)
+	{
+		// sort the list
+		TArray<FDDCResourceUsageStat> SortedDDCResourceUsageStats;
+		SortedDDCResourceUsageStats.Empty(DDCResourceUsageStats.Num());
+		for (const FDDCResourceUsageStat& Stat : DDCResourceUsageStats)
+		{
+			SortedDDCResourceUsageStats.Emplace(Stat);
+		}
+		SortedDDCResourceUsageStats.Sort([](const FDDCResourceUsageStat& LHS, const FDDCResourceUsageStat& RHS)
+			{
+				return LHS.TotalTimeSec > RHS.TotalTimeSec;
+			});
+
+		UE_LOG(LogCook, Display, TEXT(""));
+		UE_LOG(LogCook, Display, TEXT("DDC Resource Stats"));
+		UE_LOG(LogCook, Display, TEXT("======================================================================================================="));
+		UE_LOG(LogCook, Display, TEXT("Asset Type                          Total Time (Sec)  GameThread Time (Sec)  Assets Built  MB Processed"));
+		UE_LOG(LogCook, Display, TEXT("----------------------------------  ----------------  ---------------------  ------------  ------------"));
+		for (const FDDCResourceUsageStat& Stat : SortedDDCResourceUsageStats)
+		{
+			UE_LOG(LogCook, Display, TEXT("%-34s  %16.2f  %21.2f  %12d  %12.2f"), *Stat.AssetType, Stat.TotalTimeSec, Stat.GameThreadTimeSec, Stat.AssetsBuilt, Stat.SizeMB);
+		}
+	}
+
+	DumpBuildDependencyTrackerStats();
+
+	if (UE::Virtualization::IVirtualizationSystem::IsInitialized())
+	{
+		UE::Virtualization::IVirtualizationSystem::Get().DumpStats();
+	}
+
+	if (IsCookByTheBookMode(CookMode))
+	{
+		FStudioAnalytics::FireEvent_Loading(TEXT("CookByTheBook"), DetailedCookStats::CookWallTimeSec);
+	}
+}
+
+}
+#endif
\ No newline at end of file
diff --git a/Engine/Source/Editor/UnrealEd/Private/Cooker/CookProfiling.h b/Engine/Source/Editor/UnrealEd/Private/Cooker/CookProfiling.h
--- a/Engine/Source/Editor/UnrealEd/Private/Cooker/CookProfiling.h
+++ b/Engine/Source/Editor/UnrealEd/Private/Cooker/CookProfiling.h
@@ -3,6 +3,7 @@
 #pragma once
 
 #include "Containers/ArrayView.h"
+#include "CookOnTheSide/CookOnTheFlyServer.h"
 #include "CoreTypes.h"
 #include "ProfilingDebugging/CookStats.h"
 #include "Stats/Stats.h"
@@ -123,26 +124,46 @@ inline void ClearHierarchyTimers() {}
 #if ENABLE_COOK_STATS
 namespace DetailedCookStats
 {
-	extern double TickCookOnTheSideTimeSec;
-	extern double TickCookOnTheSideLoadPackagesTimeSec;
-	extern double TickCookOnTheSideResolveRedirectorsTimeSec;
-	extern double TickCookOnTheSideSaveCookedPackageTimeSec;
-	extern double TickCookOnTheSidePrepareSaveTimeSec;
-	extern double BlockOnAssetRegistryTimeSec;
-	extern double GameCookModificationDelegateTimeSec;
-	extern double TickLoopShaderProcessAsyncResultsTimeSec;
-	extern double TickLoopProcessDeferredCommandsTimeSec;
-	extern double TickLoopTickCommandletStatsTimeSec;
-	extern double TickLoopFlushRenderingCommandsTimeSec;
-
-	// Stats tracked through FAutoRegisterCallback
-	extern uint32 NumPreloadedDependencies;
-	extern uint32 NumPackagesIterativelySkipped;
-	/** Tracks the number of packages cooked - counted only once even if multiple platforms cooked. */
-	extern uint32 NumPackagesSavedForCook;
-	extern int32 PeakRequestQueueSize;
-	extern int32 PeakLoadQueueSize;
-	extern int32 PeakSaveQueueSize;
+
+extern FString CookProject;
+extern FString CookCultures;
+extern FString CookLabel;
+extern FString TargetPlatforms;
+extern double CookStartTime;
+extern double CookWallTimeSec;
+extern double StartupWallTimeSec;
+extern double CookByTheBookTimeSec;
+extern double StartCookByTheBookTimeSec;
+extern double TickCookOnTheSideTimeSec;
+extern double TickCookOnTheSideLoadPackagesTimeSec;
+extern double TickCookOnTheSideResolveRedirectorsTimeSec;
+extern double TickCookOnTheSideSaveCookedPackageTimeSec;
+extern double TickCookOnTheSidePrepareSaveTimeSec;
+extern double BlockOnAssetRegistryTimeSec;
+extern double GameCookModificationDelegateTimeSec;
+extern double TickLoopGCTimeSec;
+extern double TickLoopRecompileShaderRequestsTimeSec;
+extern double TickLoopShaderProcessAsyncResultsTimeSec;
+extern double TickLoopProcessDeferredCommandsTimeSec;
+extern double TickLoopTickCommandletStatsTimeSec;
+extern double TickLoopFlushRenderingCommandsTimeSec;
+extern bool IsCookAll;
+extern bool IsCookOnTheFly;
+extern bool IsIterativeCook;
+extern bool IsFastCook;
+extern bool IsUnversioned;
+
+// Stats tracked through FAutoRegisterCallback
+extern uint32 NumPreloadedDependencies;
+extern uint32 NumPackagesIterativelySkipped;
+/** Tracks the number of packages cooked - counted only once even if multiple platforms cooked. */
+extern uint32 NumPackagesSavedForCook;
+extern int32 PeakRequestQueueSize;
+extern int32 PeakLoadQueueSize;
+extern int32 PeakSaveQueueSize;
+
+void LogCookStats(ECookMode::Type CookMode);
+
 }
 #endif
 
diff --git a/Engine/Source/Editor/UnrealEd/Private/Cooker/CookWorkerClient.cpp b/Engine/Source/Editor/UnrealEd/Private/Cooker/CookWorkerClient.cpp
--- a/Engine/Source/Editor/UnrealEd/Private/Cooker/CookWorkerClient.cpp
+++ b/Engine/Source/Editor/UnrealEd/Private/Cooker/CookWorkerClient.cpp
@@ -477,6 +477,7 @@ void FCookWorkerClient::PumpDisconnect(FTickStackData& StackData)
 		{
 		case EConnectStatus::FlushAndAbortFirst:
 		{
+			COTFS.LogCookWorkerStats();
 			TickCollectors(StackData, true /* bFlush */);
 			// Add code here for any waiting we need to do for the local CookOnTheFlyServer to gracefully shutdown
 			SendMessage(FAbortWorkerMessage(FAbortWorkerMessage::EType::Abort));
diff --git a/Engine/Source/Editor/UnrealEd/Private/Cooker/CookWorkerServer.cpp b/Engine/Source/Editor/UnrealEd/Private/Cooker/CookWorkerServer.cpp
--- a/Engine/Source/Editor/UnrealEd/Private/Cooker/CookWorkerServer.cpp
+++ b/Engine/Source/Editor/UnrealEd/Private/Cooker/CookWorkerServer.cpp
@@ -203,6 +203,8 @@ bool FCookWorkerServer::TryHandleConnectMessage(FWorkerConnectMessage& Message,
 	Socket = InSocket;
 
 	SendToState(EConnectStatus::Connected);
+	UE_LOG(LogCook, Display, TEXT("CookWorker %d connected after %.3fs."), WorkerId.GetRemoteIndex(),
+		static_cast<float>(FPlatformTime::Seconds() - ConnectStartTimeSeconds));
 	for (UE::CompactBinaryTCP::FMarshalledMessage& OtherMessage : OtherPacketMessages)
 	{
 		ReceiveMessages.Add(MoveTemp(OtherMessage));
