World Partiion: fix resave actors world builder not saving anything when ResaveDirtyActorDescsOnly option is not used.

diff --git a/Engine/Source/Editor/UnrealEd/Private/WorldPartition/WorldPartitionResaveActorsBuilder.cpp b/Engine/Source/Editor/UnrealEd/Private/WorldPartition/WorldPartitionResaveActorsBuilder.cpp
--- a/Engine/Source/Editor/UnrealEd/Private/WorldPartition/WorldPartitionResaveActorsBuilder.cpp
+++ b/Engine/Source/Editor/UnrealEd/Private/WorldPartition/WorldPartitionResaveActorsBuilder.cpp
@@ -399,28 +399,25 @@ bool UWorldPartitionResaveActorsBuilder::RunInternal(UWorld* World, const FCellI
 				UPackage* Package = Actor->GetExternalPackage();
 				check(Package);
 
-				bool bActorNeedsResave = ActorDesc->IsResaveNeeded();
-				
-				if (bResaveDirtyActorDescsOnly || bDiffDirtyActorDescs)
+				if (bResaveDirtyActorDescsOnly)
 				{
 					TUniquePtr<FWorldPartitionActorDesc> NewActorDesc = Actor->CreateActorDesc();
 
-					bool bActorDescDirty = !ActorDesc->Equals(NewActorDesc.Get());
-					if (bActorDescDirty && bDiffDirtyActorDescs)
+					if (!ActorDesc->IsResaveNeeded() && ActorDesc->Equals(NewActorDesc.Get()))
+					{
+						return true;
+					}
+
+					if (bDiffDirtyActorDescs)
 					{
 						DirtyActorDescsOld.Add(ActorDesc->ToString());
 						DirtyActorDescsNew.Add(NewActorDesc->ToString());
-					}				
-
-					bActorNeedsResave |= bResaveDirtyActorDescsOnly && bActorDescDirty;
-				}
+					}
 
-				if (bActorNeedsResave)
-				{
 					UE_LOG(LogWorldPartitionResaveActorsBuilder, Log, TEXT("Package %s needs to be resaved."), *Package->GetName());
 				}
 
-				if (bActorNeedsResave && !bReportOnly)
+				if (!bReportOnly)
 				{
 					PackagesToSave.Add(Package);
 				}
