Multiprocess editor support: Add the parameter -MultiprocessId=<N> to allow instances of the engine to have a unique id in their group of multiprocess instances.

MPCook and Shaders: Change ShaderSymbolExport to export to different .zip files when in a -MultiprocessId environment.
Add part of the implementation for merging the different zip files created by CookWorkers together. The rest of the implementation requires reading from .zip files and will be coming in a future change.

#rnx

diff --git a/Engine/Source/Developer/ShaderCompilerCommon/Private/ShaderCompilerCommon.cpp b/Engine/Source/Developer/ShaderCompilerCommon/Private/ShaderCompilerCommon.cpp
--- a/Engine/Source/Developer/ShaderCompilerCommon/Private/ShaderCompilerCommon.cpp
+++ b/Engine/Source/Developer/ShaderCompilerCommon/Private/ShaderCompilerCommon.cpp
@@ -4,15 +4,16 @@
 #include "ShaderParameterParser.h"
 #include "Misc/FileHelper.h"
 #include "Misc/Paths.h"
+#include "Misc/PathViews.h"
 #include "Modules/ModuleManager.h"
 #include "HlslccDefinitions.h"
 #include "HAL/FileManager.h"
 #include "String/RemoveFrom.h"
+#include "ShaderSymbolExport.h"
 #include "ShaderMinifier.h"
 
 IMPLEMENT_MODULE(FDefaultModuleImpl, ShaderCompilerCommon);
 
-
 int16 GetNumUniformBuffersUsed(const FShaderCompilerResourceTable& InSRT)
 {
 	auto CountLambda = [&](const TArray<uint32>& In)
@@ -2564,3 +2565,73 @@ namespace CrossCompiler
 	}
 
 } // namespace CrossCompiler
+
+#if WITH_ENGINE
+
+const TCHAR* FShaderSymbolExport::ZipFileBaseLeafName = TEXT("ShaderSymbols");
+const TCHAR* FShaderSymbolExport::ZipFileExtension = TEXT(".zip");
+
+inline void FShaderSymbolExport::NotifyShaderCompilersShutdown()
+{
+	constexpr bool bMergingEnabled = false; // TODO: Add ability to read zip files and turn this on
+	if (bMultiprocessOwner && ZipWriter && bMergingEnabled)
+	{
+		IPlatformFile& PlatformFile = FPlatformFileManager::Get().GetPlatformFile();
+		TArray<FString> ZipsToMergeIn;
+		PlatformFile.FindFiles(ZipsToMergeIn, *ExportPath, ZipFileExtension);
+		ZipsToMergeIn.RemoveAll([](const FString& FileName)
+			{
+				return FPathViews::GetBaseFilename(FileName) == ZipFileBaseLeafName;
+			});
+		for (const FString& ZipFile : ZipsToMergeIn)
+		{
+			struct FZipArchiveReader
+			{
+				FZipArchiveReader(IFileHandle* InFileHandle) : FileHandle(InFileHandle) {}
+				bool IsValid() { return false; }
+				TConstArrayView<FString> GetEmbeddedFileNames() const { return TConstArrayView<FString>(); }
+				bool TryReadFile(FStringView FileName, TArray<uint8>& OutData) { return false; }
+				TUniquePtr<IFileHandle> FileHandle;
+			};
+			FZipArchiveReader Reader(PlatformFile.OpenRead(*ZipFile));
+			bool bAllValid = false;
+			if (Reader.IsValid())
+			{
+				bAllValid = true;
+				for (const FString& EmbeddedFileName : Reader.GetEmbeddedFileNames())
+				{
+					TArray<uint8> Contents;
+					if (!Reader.TryReadFile(EmbeddedFileName, Contents))
+					{
+						bAllValid = false;
+						continue;
+					}
+					ZipWriter->AddFile(EmbeddedFileName, Contents, FDateTime::Now());
+				}
+			}
+			if (!bAllValid)
+			{
+				UE_LOG(LogShaderSymbolExport, Error,
+					TEXT("Failed to read from CookWorker shader symbols output file \"%s\". Some shader symbols will be missing."),
+					*ZipFile);
+			}
+			PlatformFile.DeleteFile(*ZipFile);
+		}
+	}
+	ZipWriter.Reset();
+}
+
+void FShaderSymbolExport::DeleteExistingShaderZips(IPlatformFile& PlatformFile, const FString& Directory)
+{
+	TArray<FString> ExistingZips;
+	PlatformFile.FindFiles(ExistingZips, *Directory, ZipFileExtension);
+	for (const FString& ZipFile : ExistingZips)
+	{
+		if (FPathViews::GetPathLeaf(ZipFile).StartsWith(ZipFileBaseLeafName))
+		{
+			PlatformFile.DeleteFile(*ZipFile);
+		}
+	}
+}
+
+#endif
\ No newline at end of file
diff --git a/Engine/Source/Developer/ShaderCompilerCommon/Public/ShaderSymbolExport.h b/Engine/Source/Developer/ShaderCompilerCommon/Public/ShaderSymbolExport.h
--- a/Engine/Source/Developer/ShaderCompilerCommon/Public/ShaderSymbolExport.h
+++ b/Engine/Source/Developer/ShaderCompilerCommon/Public/ShaderSymbolExport.h
@@ -27,6 +27,9 @@ public:
 	template<typename TPlatformShaderSymbolData>
 	void NotifyShaderCompiled(const TConstArrayView<uint8>& PlatformSymbolData);
 
+	/** Called at the end of a cook to free resources and finalize artifacts created during the cook. */
+	SHADERCOMPILERCOMMON_API void NotifyShaderCompilersShutdown();
+
 private:
 	void Initialize();
 
@@ -39,6 +42,15 @@ private:
 	uint64 TotalSymbolDataBytes{ 0 };
 	uint64 TotalSymbolData{ 0 };
 	bool bExportShaderSymbols{ false };
+	/**
+	 * If true, the current process is the first process in a multiprocess group, or is not in a group,
+	 * and should combine artifacts produced by the other processes. Will also be false if no combination
+	 * is necessary for given settings.
+	 */
+	bool bMultiprocessOwner{ false };
+	SHADERCOMPILERCOMMON_API static void DeleteExistingShaderZips(IPlatformFile& PlatformFile, const FString& Directory);
+	SHADERCOMPILERCOMMON_API static const TCHAR* ZipFileBaseLeafName;
+	SHADERCOMPILERCOMMON_API static const TCHAR* ZipFileExtension;
 };
 
 inline FShaderSymbolExport::FShaderSymbolExport(FName InShaderFormat)
@@ -82,7 +94,20 @@ inline void FShaderSymbolExport::Initialize()
 			bool bExportAsZip = ShouldWriteShaderSymbolsAsZip(ShaderFormat);
 			if (bExportAsZip || FParse::Param(FCommandLine::Get(), TEXT("ShaderSymbolsExportZip")))
 			{
-				FString SingleFilePath = ExportPath / TEXT("ShaderSymbols.zip");
+				uint32 MultiprocessId = 0;
+				FParse::Value(FCommandLine::Get(), TEXT("-MultiprocessId="), MultiprocessId);
+				FString LeafName;
+				bMultiprocessOwner = MultiprocessId == 0;
+				if (bMultiprocessOwner)
+				{
+					DeleteExistingShaderZips(PlatformFile, ExportPath);
+					LeafName = FString::Printf(TEXT("%s%s"), ZipFileBaseLeafName, ZipFileExtension);
+				}
+				else
+				{
+					LeafName = FString::Printf(TEXT("%s_%d%s"), ZipFileBaseLeafName, MultiprocessId, ZipFileExtension);
+				}
+				FString SingleFilePath = ExportPath / LeafName;
 
 				IFileHandle* OutputZipFile = PlatformFile.OpenWrite(*SingleFilePath);
 				if (!OutputZipFile)
diff --git a/Engine/Source/Developer/ShaderCompilerCommon/ShaderCompilerCommon.Build.cs b/Engine/Source/Developer/ShaderCompilerCommon/ShaderCompilerCommon.Build.cs
--- a/Engine/Source/Developer/ShaderCompilerCommon/ShaderCompilerCommon.Build.cs
+++ b/Engine/Source/Developer/ShaderCompilerCommon/ShaderCompilerCommon.Build.cs
@@ -13,6 +13,14 @@ public class ShaderCompilerCommon : ModuleRules
 				"RenderCore",
 			}
 			);
+		if (Target.bCompileAgainstEngine)
+		{
+			PublicDependencyModuleNames.AddRange(
+			new string[] {
+				"FileUtilities",
+			}
+			);
+		}
 
 		if (Target.Platform == UnrealTargetPlatform.Mac || Target.Platform == UnrealTargetPlatform.Win64 || Target.Platform == UnrealTargetPlatform.Linux)
 		{
diff --git a/Engine/Source/Developer/TargetPlatform/Public/Interfaces/IShaderFormat.h b/Engine/Source/Developer/TargetPlatform/Public/Interfaces/IShaderFormat.h
--- a/Engine/Source/Developer/TargetPlatform/Public/Interfaces/IShaderFormat.h
+++ b/Engine/Source/Developer/TargetPlatform/Public/Interfaces/IShaderFormat.h
@@ -113,6 +113,9 @@ public:
 	 */
 	virtual void NotifyShaderCompiled(const TConstArrayView<uint8>& PlatformDebugData, FName Format) const { }
 
+	/** Called at the end of a cook to free resources and finalize artifacts created during the cook. */
+	virtual void NotifyShaderCompilersShutdown(FName Format) const { }
+
 	/**
 	 * Appends shader key text to the provided key string for use in DDC unique key construction.
 	 * @param KeyString String that will get shader key text appended to.
diff --git a/Engine/Source/Editor/UnrealEd/Private/CookOnTheFlyServer.cpp b/Engine/Source/Editor/UnrealEd/Private/CookOnTheFlyServer.cpp
--- a/Engine/Source/Editor/UnrealEd/Private/CookOnTheFlyServer.cpp
+++ b/Engine/Source/Editor/UnrealEd/Private/CookOnTheFlyServer.cpp
@@ -8581,6 +8581,7 @@ void UCookOnTheFlyServer::CookByTheBookFinished()
 	FString ActualLibraryName = GenerateShaderCodeLibraryName(LibraryName, IsCookFlagSet(ECookInitializationFlags::IterateSharedBuild));
 	FShaderLibraryCooker::EndCookingLibrary(ActualLibraryName);
 	FShaderLibraryCooker::Shutdown();
+	ShutdownShaderCompilers(PlatformManager->GetSessionPlatforms());
 
 	if (CookByTheBookOptions->bGenerateDependenciesForMaps)
 	{
@@ -9548,15 +9549,17 @@ void UCookOnTheFlyServer::ShutdownCookAsCookWorker()
 	{
 		UnregisterCookByTheBookDelegates();
 	}
-	ShutdownCookSession();
 	FString LibraryName = GetProjectShaderLibraryName();
 	FString ActualLibraryName = GenerateShaderCodeLibraryName(LibraryName, IsCookFlagSet(ECookInitializationFlags::IterateSharedBuild));
 	FShaderLibraryCooker::EndCookingLibrary(ActualLibraryName);
 	FShaderLibraryCooker::Shutdown();
+	ShutdownShaderCompilers(PlatformManager->GetSessionPlatforms());
+
 	if (IsDirectorCookOnTheFly())
 	{
 		GShaderCompilingManager->SkipShaderCompilation(false);
 	}
+	ShutdownCookSession();
 }
 
 FBeginCookContext UCookOnTheFlyServer::CreateCookWorkerContext()
diff --git a/Engine/Source/Editor/UnrealEd/Private/Cooker/CookDirector.cpp b/Engine/Source/Editor/UnrealEd/Private/Cooker/CookDirector.cpp
--- a/Engine/Source/Editor/UnrealEd/Private/Cooker/CookDirector.cpp
+++ b/Engine/Source/Editor/UnrealEd/Private/Cooker/CookDirector.cpp
@@ -38,7 +38,14 @@ FCookDirector::FCookDirector(UCookOnTheFlyServer& InCOTFS, int32 CookProcessCoun
 	WorkersStalledWarnTimeSeconds = MAX_flt;
 	ShutdownEvent->Reset();
 
-	ParseConfig(CookProcessCount);
+	bool bConfigValid;
+	ParseConfig(CookProcessCount, bConfigValid);
+	if (!bConfigValid)
+	{
+		UE_LOG(LogCook, Error, TEXT("CookDirector initialization failure: config settings are invalid for multiprocess. CookMultiprocess is disabled and the cooker is running as a single process."));
+		bMultiprocessAvailable = false;
+		return;
+	}
 	ISocketSubsystem* SocketSubsystem = ISocketSubsystem::Get();
 	if (!SocketSubsystem)
 	{
@@ -59,8 +66,9 @@ bool FCookDirector::IsMultiprocessAvailable() const
 	return bMultiprocessAvailable;
 }
 
-void FCookDirector::ParseConfig(int32 CookProcessCount)
+void FCookDirector::ParseConfig(int32 CookProcessCount, bool& bOutValid)
 {
+	bOutValid = true;
 	const TCHAR* CommandLine = FCommandLine::Get();
 	FString Text;
 
@@ -103,6 +111,13 @@ void FCookDirector::ParseConfig(int32 CookProcessCount)
 			UE_LOG(LogCook, Warning, TEXT("Invalid selection \"%s\" for -CookLoadBalance."), *Text);
 		}
 	}
+
+	int32 MultiprocessId;
+	if (FParse::Value(CommandLine, TEXT("-MultiprocessId="), MultiprocessId))
+	{
+		bOutValid = false;
+		UE_LOG(LogCook, Error, TEXT("CookMultiprocess is incompatible with -MultiprocessId on the CookDirector's commandline. The CookDirector needs to be able to specify all MultiprocessIds."));
+	}
 }
 
 FCookDirector::~FCookDirector()
@@ -865,7 +880,7 @@ FString FCookDirector::GetWorkerCommandLine(FWorkerId WorkerId)
 				Token.StartsWith(TEXT("-CookCultures")) ||
 				Token.StartsWith(TEXT("-CookDirectorCount=")) ||
 				Token.StartsWith(TEXT("-CookDirectorHost=")) ||
-				Token.StartsWith(TEXT("-CookWorkerId=")) ||
+				Token.StartsWith(TEXT("-MultiprocessId=")) ||
 				Token.StartsWith(TEXT("-ShowCookWorker")) ||
 				Token.StartsWith(TEXT("-CoreLimit")) ||
 				Token.StartsWith(TEXT("-PhysicalCoreLimit")) ||
@@ -884,7 +899,7 @@ FString FCookDirector::GetWorkerCommandLine(FWorkerId WorkerId)
 	Tokens.Insert(TEXT("-cookworker"), 2);
 	check(!WorkerConnectAuthority.IsEmpty()); // This should have been constructed in TryCreateWorkerConnectSocket before any CookWorkerServers could exist to call GetWorkerCommandLine
 	Tokens.Add(FString::Printf(TEXT("-CookDirectorHost=%s"), *WorkerConnectAuthority));
-	Tokens.Add(FString::Printf(TEXT("-CookWorkerId=%d"), WorkerId.GetRemoteIndex()));
+	Tokens.Add(FString::Printf(TEXT("-MultiprocessId=%d"), WorkerId.GetRemoteIndex() + 1));
 	if (CoreLimit > 0)
 	{
 		Tokens.Add(FString::Printf(TEXT("-PhysicalCoreLimit=%d"), CoreLimit));
@@ -900,11 +915,19 @@ bool FDirectorConnectionInfo::TryParseCommandLine()
 		UE_LOG(LogCook, Error, TEXT("CookWorker startup failed: no CookDirector specified on commandline."));
 		return false;
 	}
-	if (!FParse::Value(FCommandLine::Get(), TEXT("-CookWorkerId="), RemoteIndex))
+	uint32 MultiprocessId;
+	if (!FParse::Value(FCommandLine::Get(), TEXT("-MultiprocessId="), MultiprocessId))
+	{
+		UE_LOG(LogCook, Error, TEXT("CookWorker startup failed: no MultiprocessId specified on commandline."));
+		return false;
+	}
+	if (MultiprocessId < 1 || 257 <= MultiprocessId)
 	{
-		UE_LOG(LogCook, Error, TEXT("CookWorker startup failed: no CookWorkerId specified on commandline."));
+		UE_LOG(LogCook, Error, TEXT("CookWorker startup failed: commandline had invalid -MultiprocessId=%d; MultiprocessId must be in the range [1, 256]."),
+			MultiprocessId);
 		return false;
 	}
+	RemoteIndex = static_cast<int32>(MultiprocessId - 1);
 	return true;
 }
 
diff --git a/Engine/Source/Editor/UnrealEd/Private/Cooker/CookDirector.h b/Engine/Source/Editor/UnrealEd/Private/Cooker/CookDirector.h
--- a/Engine/Source/Editor/UnrealEd/Private/Cooker/CookDirector.h
+++ b/Engine/Source/Editor/UnrealEd/Private/Cooker/CookDirector.h
@@ -132,7 +132,7 @@ private:
 
 private:
 	/** Helper for constructor parsing. */
-	void ParseConfig(int32 CookProcessCount);
+	void ParseConfig(int32 CookProcessCount, bool& bOutValid);
 	/** Initialization helper: create the listen socket. */
 	bool TryCreateWorkerConnectSocket();
 	/**
diff --git a/Engine/Source/Runtime/Engine/Private/ShaderCompiler/ShaderCompiler.cpp b/Engine/Source/Runtime/Engine/Private/ShaderCompiler/ShaderCompiler.cpp
--- a/Engine/Source/Runtime/Engine/Private/ShaderCompiler/ShaderCompiler.cpp
+++ b/Engine/Source/Runtime/Engine/Private/ShaderCompiler/ShaderCompiler.cpp
@@ -8068,6 +8068,25 @@ void RecompileShadersForRemote(
 	// Restore compilation state.
 	GShaderCompilingManager->SkipShaderCompilation(bPreviousState);
 }
+
+void ShutdownShaderCompilers(TConstArrayView<const ITargetPlatform*> TargetPlatforms)
+{
+	ITargetPlatformManagerModule& PlatformManager = GetTargetPlatformManagerRef();
+	for (const ITargetPlatform* TargetPlatform : TargetPlatforms)
+	{
+		TArray<FName> DesiredShaderFormats;
+		TargetPlatform->GetAllTargetedShaderFormats(DesiredShaderFormats);
+		for (FName FormatName : DesiredShaderFormats)
+		{
+			const IShaderFormat* ShaderFormat = PlatformManager.FindShaderFormat(FormatName);
+			if (ShaderFormat)
+			{
+				ShaderFormat->NotifyShaderCompilersShutdown(FormatName);
+			}
+		}
+	}
+}
+
 #endif // WITH_EDITOR
 
 void BeginRecompileGlobalShaders(const TArray<const FShaderType*>& OutdatedShaderTypes, const TArray<const FShaderPipelineType*>& OutdatedShaderPipelineTypes, EShaderPlatform ShaderPlatform, const ITargetPlatform* TargetPlatform)
diff --git a/Engine/Source/Runtime/Engine/Public/ShaderCompiler.h b/Engine/Source/Runtime/Engine/Public/ShaderCompiler.h
--- a/Engine/Source/Runtime/Engine/Public/ShaderCompiler.h
+++ b/Engine/Source/Runtime/Engine/Public/ShaderCompiler.h
@@ -1387,6 +1387,9 @@ struct FShaderRecompileData
 **/
 extern ENGINE_API void RecompileShadersForRemote(FShaderRecompileData& Args, const FString& OutputDirectory);
 
+/** Frees resources and finalizes artifacts created during the cook of the given platforms. */
+extern ENGINE_API void ShutdownShaderCompilers(TConstArrayView<const ITargetPlatform*> TargetPlatforms);
+
 #endif // WITH_EDITOR
 
 extern ENGINE_API void CompileGlobalShaderMap(bool bRefreshShaderMap=false);
