MPCook: Fix shutdown crashes due to duplicate calls to shutdown code and trying to access session data during the second shutdown after the first shutdown already exited the session.
#rnx

diff --git a/Engine/Source/Editor/UnrealEd/Classes/CookOnTheSide/CookOnTheFlyServer.h b/Engine/Source/Editor/UnrealEd/Classes/CookOnTheSide/CookOnTheFlyServer.h
--- a/Engine/Source/Editor/UnrealEd/Classes/CookOnTheSide/CookOnTheFlyServer.h
+++ b/Engine/Source/Editor/UnrealEd/Classes/CookOnTheSide/CookOnTheFlyServer.h
@@ -1043,6 +1043,7 @@ private:
 	void StartCookAsCookWorker();
 	/** Construct the on-stack data for StartCook functions, based on the CookWorker's configuration */
 	FBeginCookContext CreateCookWorkerContext();
+	void CookAsCookWorkerFinished();
 
 	//////////////////////////////////////////////////////////////////////////
 	// general functions
diff --git a/Engine/Source/Editor/UnrealEd/Private/CookOnTheFlyServer.cpp b/Engine/Source/Editor/UnrealEd/Private/CookOnTheFlyServer.cpp
--- a/Engine/Source/Editor/UnrealEd/Private/CookOnTheFlyServer.cpp
+++ b/Engine/Source/Editor/UnrealEd/Private/CookOnTheFlyServer.cpp
@@ -8605,6 +8605,11 @@ void UCookOnTheFlyServer::ShutdownCookSession()
 	{
 		CookDirector->ShutdownCookSession();
 	}
+	if (CookWorkerClient)
+	{
+		CookAsCookWorkerFinished();
+	}
+
 	for (UE::Cook::FPackageData* PackageData : *PackageDatas)
 	{
 		PackageData->DestroyGeneratorPackage();
@@ -9543,12 +9548,8 @@ void UCookOnTheFlyServer::LogCookWorkerStats()
 	}
 }
 
-void UCookOnTheFlyServer::ShutdownCookAsCookWorker()
+void UCookOnTheFlyServer::CookAsCookWorkerFinished()
 {
-	if (IsDirectorCookByTheBook())
-	{
-		UnregisterCookByTheBookDelegates();
-	}
 	FString LibraryName = GetProjectShaderLibraryName();
 	FString ActualLibraryName = GenerateShaderCodeLibraryName(LibraryName, IsCookFlagSet(ECookInitializationFlags::IterateSharedBuild));
 	FShaderLibraryCooker::EndCookingLibrary(ActualLibraryName);
@@ -9559,13 +9560,26 @@ void UCookOnTheFlyServer::ShutdownCookAsCookWorker()
 	{
 		GShaderCompilingManager->SkipShaderCompilation(false);
 	}
-	ShutdownCookSession();
+}
+
+void UCookOnTheFlyServer::ShutdownCookAsCookWorker()
+{
+	if (IsDirectorCookByTheBook())
+	{
+		UnregisterCookByTheBookDelegates();
+	}
+	if (IsInSession())
+	{
+		ShutdownCookSession();
+	}
 }
 
 FBeginCookContext UCookOnTheFlyServer::CreateCookWorkerContext()
 {
 	FBeginCookContext BeginContext(*this);
 	*CookByTheBookOptions = CookWorkerClient->ConsumeCookByTheBookOptions();
+	CookByTheBookOptions->CookTime = 0.0f;
+	CookByTheBookOptions->CookStartTime = FPlatformTime::Seconds();
 	*CookOnTheFlyOptions = CookWorkerClient->ConsumeCookOnTheFlyOptions();
 	BeginContext.TargetPlatforms = CookWorkerClient->GetTargetPlatforms();
 
diff --git a/Engine/Source/Editor/UnrealEd/Private/Cooker/CookWorkerClient.cpp b/Engine/Source/Editor/UnrealEd/Private/Cooker/CookWorkerClient.cpp
--- a/Engine/Source/Editor/UnrealEd/Private/Cooker/CookWorkerClient.cpp
+++ b/Engine/Source/Editor/UnrealEd/Private/Cooker/CookWorkerClient.cpp
@@ -498,14 +498,20 @@ void FCookWorkerClient::PumpDisconnect(FTickStackData& StackData)
 					UE_LOG(LogCook, Warning, TEXT("Timedout after %.0fs waiting to send disconnect message to CookDirector."),
 						WaitForDisconnectTimeout);
 					SendToState(EConnectStatus::LostConnection);
-					return;
+					check(ConnectStatus == EConnectStatus::LostConnection);
+					// Fall through to LostConnection
+					break;
+				}
+				else
+				{
+					return; // Exit the Pump loop for now and keep waiting
 				}
-				return; // Exit the Pump loop for now and keep waiting
 			}
 			else
 			{
 				check(ConnectStatus == EConnectStatus::LostConnection);
-				return;
+				// Fall through to LostConnection
+				break;
 			}
 		}
 		case EConnectStatus::LostConnection:
