Fix ReplaceActorHelper not taking in the flags from the Old Actor
- Flag Mask with values consistent with ReplaceObjectHelper

diff --git a/Engine/Source/Editor/UnrealEd/Private/Kismet2/KismetReinstanceUtilities.cpp b/Engine/Source/Editor/UnrealEd/Private/Kismet2/KismetReinstanceUtilities.cpp
--- a/Engine/Source/Editor/UnrealEd/Private/Kismet2/KismetReinstanceUtilities.cpp
+++ b/Engine/Source/Editor/UnrealEd/Private/Kismet2/KismetReinstanceUtilities.cpp
@@ -2024,6 +2024,8 @@ static void ReplaceActorHelper(AActor* OldActor, UClass* OldClass, UObject*& New
 	UClass* const* MappedClass = InOldToNewClassMap.Find(OldActor->GetClass());
 	UClass*  SpawnClass = MappedClass ? *MappedClass : NewClass;
 
+	const EObjectFlags FlagMask = RF_Public | RF_ArchetypeObject | RF_Transactional | RF_Transient | RF_TextExportTransient | RF_InheritableComponentTemplate | RF_Standalone;
+
 	FActorSpawnParameters SpawnInfo;
 	SpawnInfo.OverrideLevel = ActorLevel;
 	SpawnInfo.Owner = OldActor->GetOwner();
@@ -2032,6 +2034,7 @@ static void ReplaceActorHelper(AActor* OldActor, UClass* OldClass, UObject*& New
 	SpawnInfo.SpawnCollisionHandlingOverride = ESpawnActorCollisionHandlingMethod::AlwaysSpawn;
 	SpawnInfo.bDeferConstruction = true;
 	SpawnInfo.Name = OldActor->GetFName();
+	SpawnInfo.ObjectFlags |= OldActor->GetFlags() & FlagMask;
 
 	if (!OldActor->IsListedInSceneOutliner())
 	{
