Fix deadlock in Editor Asset Subsystem when waiting for loaded assets

diff --git a/Engine/Source/Editor/UnrealEd/Private/Subsystems/EditorAssetSubsystem.cpp b/Engine/Source/Editor/UnrealEd/Private/Subsystems/EditorAssetSubsystem.cpp
--- a/Engine/Source/Editor/UnrealEd/Private/Subsystems/EditorAssetSubsystem.cpp
+++ b/Engine/Source/Editor/UnrealEd/Private/Subsystems/EditorAssetSubsystem.cpp
@@ -35,16 +35,14 @@ namespace UE::EditorAssetUtils
 		FAssetRegistryModule& AssetRegistryModule = FModuleManager::LoadModuleChecked<FAssetRegistryModule>("AssetRegistry");
 		if (IAssetRegistry& AssetRegistry = AssetRegistryModule.Get(); AssetRegistry.IsLoadingAssets())
 		{
-			// Event is used like a condition variable here
-			UE::FLazyEvent AssetRegistryLoadEvent{EEventMode::ManualReset};
-			FDelegateHandle DelegateHandle = AssetRegistry.OnFilesLoaded().AddLambda([&AssetRegistryLoadEvent]
+			if (AssetRegistry.IsSearchAsync() && AssetRegistry.IsSearchAllAssets())
 			{
-				AssetRegistryLoadEvent.Trigger();
-			});
-			// open some message here
-			AssetRegistryLoadEvent.Wait();
-
-			AssetRegistry.OnFilesLoaded().Remove(DelegateHandle);
+				AssetRegistry.WaitForCompletion();
+			}
+			else
+			{
+				AssetRegistry.SearchAllAssets(true /* bSynchronousSearch */);
+			}
 		}
 		return true;
 	}
