Fixed height preview mode activation/deactivation and preview mesh cycling

robomerge[FN_EngineMerge] -UE5-RES

diff --git a/Engine/Source/Editor/UnrealEd/Private/EditorViewportClient.cpp b/Engine/Source/Editor/UnrealEd/Private/EditorViewportClient.cpp
--- a/Engine/Source/Editor/UnrealEd/Private/EditorViewportClient.cpp
+++ b/Engine/Source/Editor/UnrealEd/Private/EditorViewportClient.cpp
@@ -3239,7 +3239,7 @@ bool FEditorViewportClient::IsCommandChordPressed(const TSharedPtr<FUICommandInf
 			&& (Chord.NeedsAlt() == IsAltPressed())
 			&& (Chord.NeedsShift() == IsShiftPressed())
 			&& (Chord.NeedsCommand() == IsCmdPressed())
-			&& (InOptionalKey.IsValid() ? (Chord.Key == InOptionalKey && Viewport->KeyState(InOptionalKey)) : Viewport->KeyState(Chord.Key));
+			&& (InOptionalKey.IsValid() ? (Chord.Key == InOptionalKey) : Viewport->KeyState(Chord.Key));
 	}
 	return bIsChordPressed;
 }
diff --git a/Engine/Source/Editor/UnrealEd/Private/LevelEditorViewport.cpp b/Engine/Source/Editor/UnrealEd/Private/LevelEditorViewport.cpp
--- a/Engine/Source/Editor/UnrealEd/Private/LevelEditorViewport.cpp
+++ b/Engine/Source/Editor/UnrealEd/Private/LevelEditorViewport.cpp
@@ -3046,28 +3046,24 @@ bool FLevelEditorViewportClient::InputKey(const FInputKeyEventArgs& InEventArgs)
 
 	bool bHandled = FEditorViewportClient::InputKey(InEventArgs);
 
-	// Handle input for the player height preview mode. 
-	if (!InputState.IsMouseButtonEvent() && IsCommandChordPressed(GetLevelViewportCommands().EnablePreviewMesh, InEventArgs.Key))
+	// Handle input for the height preview mode. 
+	bool bEnablePreviewMeshChordPressed = IsCommandChordPressed(GetLevelViewportCommands().EnablePreviewMesh);
+	bool bCyclePreviewMeshChordPressed = IsCommandChordPressed(GetLevelViewportCommands().CyclePreviewMesh);
+	if (!InputState.IsMouseButtonEvent() && (bEnablePreviewMeshChordPressed || bCyclePreviewMeshChordPressed))
 	{
-		// Holding down the backslash buttons turns on the mode. 
-		if (InEventArgs.Event == IE_Pressed)
-		{
-			GEditor->SetPreviewMeshMode(true);
+		GEditor->SetPreviewMeshMode(true);
 
-			// If shift down, cycle between the preview meshes
-			if (IsCommandChordPressed(GetLevelViewportCommands().CyclePreviewMesh, InEventArgs.Key))
-			{
-				GEditor->CyclePreviewMesh();
-			}
-		}
-		// Releasing backslash turns off the mode. 
-		else if (InEventArgs.Event == IE_Released)
+		if (bCyclePreviewMeshChordPressed && (InEventArgs.Event == IE_Pressed))
 		{
-			GEditor->SetPreviewMeshMode(false);
+			GEditor->CyclePreviewMesh();
 		}
 
 		bHandled = true;
 	}
+	else
+	{
+		GEditor->SetPreviewMeshMode(false);
+	}
 
 	// Clear Duplicate Actors mode when ALT and all mouse buttons are released
 	if ( !InputState.IsAltButtonPressed() && !InputState.IsAnyMouseButtonDown() )
