Fix PSO cache expansion filtering out valid PSOs.

diff --git a/Engine/Source/Editor/UnrealEd/Private/Commandlets/ShaderPipelineCacheToolsCommandlet.cpp b/Engine/Source/Editor/UnrealEd/Private/Commandlets/ShaderPipelineCacheToolsCommandlet.cpp
--- a/Engine/Source/Editor/UnrealEd/Private/Commandlets/ShaderPipelineCacheToolsCommandlet.cpp
+++ b/Engine/Source/Editor/UnrealEd/Private/Commandlets/ShaderPipelineCacheToolsCommandlet.cpp
@@ -466,7 +466,7 @@ int32 DumpPSOSC(FString& Token, const FString& StableKeyFileDir)
 		}
 		else if (Item.Type == FPipelineCacheFileFormatPSO::DescriptorType::Graphics)
 		{
-			check(!(Item.GraphicsDesc.VertexShader == FSHAHash() || Item.GraphicsDesc.MeshShader == FSHAHash()));
+			check(Item.GraphicsDesc.VertexShader != FSHAHash() || Item.GraphicsDesc.MeshShader != FSHAHash());
 			UE_LOG(LogShaderPipelineCacheTools, Display, TEXT("%s"), *Item.GraphicsDesc.ToString());
 
 			if (InverseMap.Num())
@@ -517,7 +517,7 @@ static void PrintShaders(const TMap<FSHAHash, TArray<int32>>& InverseMap, TArray
 	}
 	for (const int32& Item : *Out)
 	{
-		UE_LOG(LogShaderPipelineCacheTools, Verbose, TEXT("    %s"), *StableArray[Item].ToString());
+		UE_LOG(LogShaderPipelineCacheTools, Display, TEXT("    %s"), *StableArray[Item].ToString());
 	}
 }
 
@@ -601,13 +601,17 @@ bool CouldBeUsedTogether(const FStableShaderKeyAndValue& A, const FStableShaderK
 
 	static FName NAME_FDeferredDecalVS("FDeferredDecalVS");
 	static FName NAME_FWriteToSliceVS("FWriteToSliceVS");
-	static FName NAME_FPostProcessVS("FPostProcessVS");
+	static FName NAME_FScreenPassVS("FScreenPassVS");
 	static FName NAME_FWriteToSliceGS("FWriteToSliceGS");
+	static FName NAME_FNaniteIndirectMaterialVS("FNaniteIndirectMaterialVS");
+	static FName NAME_FNaniteMultiViewMaterialVS("FNaniteMultiViewMaterialVS");
 	if (
 		A.ShaderType == NAME_FDeferredDecalVS || B.ShaderType == NAME_FDeferredDecalVS ||
 		A.ShaderType == NAME_FWriteToSliceVS || B.ShaderType == NAME_FWriteToSliceVS ||
-		A.ShaderType == NAME_FPostProcessVS || B.ShaderType == NAME_FPostProcessVS ||
-		A.ShaderType == NAME_FWriteToSliceGS || B.ShaderType == NAME_FWriteToSliceGS
+		A.ShaderType == NAME_FScreenPassVS || B.ShaderType == NAME_FScreenPassVS ||
+		A.ShaderType == NAME_FWriteToSliceGS || B.ShaderType == NAME_FWriteToSliceGS ||
+		A.ShaderType == NAME_FNaniteIndirectMaterialVS || B.ShaderType == NAME_FNaniteIndirectMaterialVS ||
+		A.ShaderType == NAME_FNaniteMultiViewMaterialVS || B.ShaderType == NAME_FNaniteMultiViewMaterialVS
 		)
 	{
 		// oddball mix and match with any material shader.
@@ -1014,7 +1018,8 @@ int32 ExpandPSOSC(const TArray<FString>& Tokens)
 				}
 				else
 				{
-					bool bInvertibilityResult = CheckPSOStringInveribility(TempPSO);
+					// as of UE 5.1, we do not support storing PSOs in CSV so disable the string invertibility test, as that code path isn't updated
+					bool bInvertibilityResult = true; // CheckPSOStringInveribility(TempPSO);
 					bool bVerifyResult = TempPSO.Verify();
 					if(bInvertibilityResult && bVerifyResult)
 					{
@@ -2493,7 +2498,7 @@ int32 BuildPSOSC(const TArray<FString>& Tokens, const TMap<FString, FString>& Pa
 			}
 			else if (Item.Type == FPipelineCacheFileFormatPSO::DescriptorType::Graphics)
 			{
-				check(!(Item.GraphicsDesc.VertexShader == FSHAHash() || Item.GraphicsDesc.MeshShader == FSHAHash()));
+				check(Item.GraphicsDesc.VertexShader != FSHAHash() || Item.GraphicsDesc.MeshShader != FSHAHash());
 				StringRep = Item.GraphicsDesc.ToString();
 			}
 			else if (Item.Type == FPipelineCacheFileFormatPSO::DescriptorType::RayTracing)
diff --git a/Engine/Source/Runtime/RHI/Private/PipelineFileCache.cpp b/Engine/Source/Runtime/RHI/Private/PipelineFileCache.cpp
--- a/Engine/Source/Runtime/RHI/Private/PipelineFileCache.cpp
+++ b/Engine/Source/Runtime/RHI/Private/PipelineFileCache.cpp
@@ -674,7 +674,7 @@ bool FPipelineCacheFileFormatPSO::GraphicsDescriptor::StateFromString(const FStr
 	for (int32 Index = 0; Index < MaxSimultaneousRenderTargets; Index++)
 	{
 		check(PartEnd - PartIt >= 4); //not a very robust parser
-		LexFromString((uint32&)(RenderTargetFormats[Index]), *PartIt++);
+		LexFromString((uint8&)(RenderTargetFormats[Index]), *PartIt++);
 		ETextureCreateFlags RTFlags;
 		LexFromString(RTFlags, *PartIt++);
 		// going forward, the flags will already be reduced when logging the PSOs to disk. However as of 2021-06-17 there are still old stable cache files in existence that have flags recorded as is
