Sequencer: Fix baked fbx export for local time.

Bake was transforming an integer frame by the tick resolution of the RootToLocalTransform. It needs to convert the integer to a tick resolution, perform the RootToLocalTransform and then convert back to seconds.

diff --git a/Engine/Source/Editor/UnrealEd/Private/Fbx/FbxMainExport.cpp b/Engine/Source/Editor/UnrealEd/Private/Fbx/FbxMainExport.cpp
--- a/Engine/Source/Editor/UnrealEd/Private/Fbx/FbxMainExport.cpp
+++ b/Engine/Source/Editor/UnrealEd/Private/Fbx/FbxMainExport.cpp
@@ -2815,7 +2815,8 @@ void FFbxExporter::ExportLevelSequenceBaked3DTransformTrack(IAnimTrackAdapter& A
 		}
 		else
 		{
-			FbxTime.SetSecondDouble(DisplayRate.AsSeconds(CurrentFrame * LocalToRootTransform));
+			FFrameTime CurrentTime = FFrameRate::TransformTime(CurrentFrame, DisplayRate, TickResolution) * LocalToRootTransform;
+			FbxTime.SetSecondDouble(DisplayRate.AsSeconds(FFrameRate::TransformTime(CurrentTime, TickResolution, DisplayRate)));
 		}
 
 		FbxCurveTransX->KeySet(FbxCurveTransX->KeyAdd(FbxTime), FbxTime, KeyTrans[0]);
