Reinstancer now avoids creating objects in known out of date outers

diff --git a/Engine/Source/Editor/UnrealEd/Private/Kismet2/KismetReinstanceUtilities.cpp b/Engine/Source/Editor/UnrealEd/Private/Kismet2/KismetReinstanceUtilities.cpp
--- a/Engine/Source/Editor/UnrealEd/Private/Kismet2/KismetReinstanceUtilities.cpp
+++ b/Engine/Source/Editor/UnrealEd/Private/Kismet2/KismetReinstanceUtilities.cpp
@@ -1908,8 +1908,36 @@ static void ReplaceObjectHelper(UObject*& OldObject, UClass* OldClass, UObject*&
 			NewName = OldName;
 		}
 
+		UObject* ConflictingObject = nullptr;
+		UObject* DestinationOuter = OldObject->GetOuter();
+		// Check to make sure our original outer hasn't already been reinstanced:
+		if (UObject* const* ReinstancedOuter = OldToNewInstanceMap.Find(DestinationOuter))
+		{
+			// Our outer has been replaced, use the newer object:
+			DestinationOuter = *ReinstancedOuter;
+
+			// since we're changing the destination outer, make sure that there's no chance of collision with
+			// another object:
+			UObject* ExistingObject = StaticFindObjectFast(UObject::StaticClass(), DestinationOuter, NewName);
+			if (ExistingObject)
+			{
+				// Potential bug: if the conflict is an actor (e.g. actor template, we may need to use
+				// UObject::Rename explicitly to prevent side effects)
+				ExistingObject->Rename(
+					nullptr,
+					GetTransientPackage(),
+					REN_DoNotDirty | REN_DontCreateRedirectors | REN_ForceNoResetLoaders | REN_NonTransactional);
+			}
+		}
+
 		FMakeClassSpawnableOnScope TemporarilySpawnable(NewClass);
-		NewUObject = NewObject<UObject>(OldObject->GetOuter(), NewClass, NewName, RF_NoFlags, NewArchetype);
+		NewUObject = NewObject<UObject>(DestinationOuter, NewClass, NewName, RF_NoFlags, NewArchetype);
+
+		if (ConflictingObject)
+		{
+			// Map the object that we collided with to the NewUObject
+			OldToNewInstanceMap.Add(ConflictingObject, NewUObject);
+		}
 	}
 
 	check(NewUObject != nullptr);
