Remove unused local - I tested my last reinstancer change without this code being active

diff --git a/Engine/Source/Editor/UnrealEd/Private/Kismet2/KismetReinstanceUtilities.cpp b/Engine/Source/Editor/UnrealEd/Private/Kismet2/KismetReinstanceUtilities.cpp
--- a/Engine/Source/Editor/UnrealEd/Private/Kismet2/KismetReinstanceUtilities.cpp
+++ b/Engine/Source/Editor/UnrealEd/Private/Kismet2/KismetReinstanceUtilities.cpp
@@ -1908,7 +1908,6 @@ static void ReplaceObjectHelper(UObject*& OldObject, UClass* OldClass, UObject*&
 			NewName = OldName;
 		}
 
-		UObject* ConflictingObject = nullptr;
 		UObject* DestinationOuter = OldObject->GetOuter();
 		// Check to make sure our original outer hasn't already been reinstanced:
 		if (UObject* const* ReinstancedOuter = OldToNewInstanceMap.Find(DestinationOuter))
@@ -1932,12 +1931,6 @@ static void ReplaceObjectHelper(UObject*& OldObject, UClass* OldClass, UObject*&
 
 		FMakeClassSpawnableOnScope TemporarilySpawnable(NewClass);
 		NewUObject = NewObject<UObject>(DestinationOuter, NewClass, NewName, RF_NoFlags, NewArchetype);
-
-		if (ConflictingObject)
-		{
-			// Map the object that we collided with to the NewUObject
-			OldToNewInstanceMap.Add(ConflictingObject, NewUObject);
-		}
 	}
 
 	check(NewUObject != nullptr);
