Fix PSOs with FDeferredLightVS being filtered out from PSO cache.

- This global shader is also used together with the material shaders.

diff --git a/Engine/Source/Editor/UnrealEd/Private/Commandlets/ShaderPipelineCacheToolsCommandlet.cpp b/Engine/Source/Editor/UnrealEd/Private/Commandlets/ShaderPipelineCacheToolsCommandlet.cpp
--- a/Engine/Source/Editor/UnrealEd/Private/Commandlets/ShaderPipelineCacheToolsCommandlet.cpp
+++ b/Engine/Source/Editor/UnrealEd/Private/Commandlets/ShaderPipelineCacheToolsCommandlet.cpp
@@ -600,6 +600,7 @@ bool CouldBeUsedTogether(const FStableShaderKeyAndValue& A, const FStableShaderK
 	}
 
 	static FName NAME_FDeferredDecalVS("FDeferredDecalVS");
+	static FName NAME_FDeferredLightVS("FDeferredLightVS");
 	static FName NAME_FWriteToSliceVS("FWriteToSliceVS");
 	static FName NAME_FScreenPassVS("FScreenPassVS");
 	static FName NAME_FWriteToSliceGS("FWriteToSliceGS");
@@ -607,6 +608,7 @@ bool CouldBeUsedTogether(const FStableShaderKeyAndValue& A, const FStableShaderK
 	static FName NAME_FNaniteMultiViewMaterialVS("FNaniteMultiViewMaterialVS");
 	if (
 		A.ShaderType == NAME_FDeferredDecalVS || B.ShaderType == NAME_FDeferredDecalVS ||
+		A.ShaderType == NAME_FDeferredLightVS || B.ShaderType == NAME_FDeferredLightVS ||
 		A.ShaderType == NAME_FWriteToSliceVS || B.ShaderType == NAME_FWriteToSliceVS ||
 		A.ShaderType == NAME_FScreenPassVS || B.ShaderType == NAME_FScreenPassVS ||
 		A.ShaderType == NAME_FWriteToSliceGS || B.ShaderType == NAME_FWriteToSliceGS ||
