Last temporary improvement to logging in the reinstancer

diff --git a/Engine/Source/Editor/UnrealEd/Private/Kismet2/KismetReinstanceUtilities.cpp b/Engine/Source/Editor/UnrealEd/Private/Kismet2/KismetReinstanceUtilities.cpp
--- a/Engine/Source/Editor/UnrealEd/Private/Kismet2/KismetReinstanceUtilities.cpp
+++ b/Engine/Source/Editor/UnrealEd/Private/Kismet2/KismetReinstanceUtilities.cpp
@@ -163,6 +163,15 @@ struct FReplaceReferenceHelper
 							SerializeSearchObject();
 						}
 
+						ReferenceReplace(UObject* InSearchObject, const TMap<UObject*, UObject*>& InReplacementMap, const TMap<FSoftObjectPath, UObject*>& InWeakReferencesMap, EArchiveReplaceObjectFlags Flags)
+							: FArchiveReplaceObjectAndStructPropertyRef<UObject>(InSearchObject, InReplacementMap, EArchiveReplaceObjectFlags::DelayStart), WeakReferencesMap(InWeakReferencesMap)
+						{
+							if (!(Flags & EArchiveReplaceObjectFlags::DelayStart))
+							{
+								SerializeSearchObject();
+							}
+						}
+
 						FArchive& operator<<(FSoftObjectPath& Ref) override
 						{
 							const UObject*const* PtrToObjPtr = WeakReferencesMap.Find(Ref);
@@ -186,9 +195,70 @@ struct FReplaceReferenceHelper
 
 					if (GLogReinstancerReferenceReplacementObj)
 					{
-						UE_LOG(LogBlueprint, Warning, TEXT("Replacing References within: %s"), *(Obj->GetName()));
+						// inlining the whole of UGCObjectReferencer::AddReferencedObjects because I want to track down
+						// a likely memory stomp, but don't want to slow down or destabilize other clients of UGCObjectReferencer:
+						if (UGCObjectReferencer* GCObjectReferencer = Cast<UGCObjectReferencer>(Obj))
+						{
+							ReferenceReplace ReplaceAr(Obj, OldToNewInstanceMap, ReinstancedObjectsWeakReferenceMap, EArchiveReplaceObjectFlags::DelayStart);
+							class FReplaceObjectRefCollector : public FReferenceCollector
+							{
+								FArchive& Ar;
+								bool bAllowReferenceElimination;
+							public:
+								FReplaceObjectRefCollector(FArchive& InAr)
+									: Ar(InAr)
+									, bAllowReferenceElimination(true)
+								{
+								}
+								virtual bool IsIgnoringArchetypeRef() const override
+								{
+									return Ar.IsIgnoringArchetypeRef();
+								}
+								virtual bool IsIgnoringTransient() const override
+								{
+									return false;
+								}
+								virtual void AllowEliminatingReferences(bool bAllow) override
+								{
+									bAllowReferenceElimination = bAllow;
+								}
+								virtual void HandleObjectReference(UObject*& InObject, const UObject* InReferencingObject, const FProperty* InReferencingProperty) override
+								{
+									if (bAllowReferenceElimination)
+									{
+										FProperty* NewSerializedProperty = const_cast<FProperty*>(InReferencingProperty);
+										FSerializedPropertyScope SerializedPropertyScope(Ar, NewSerializedProperty ? NewSerializedProperty : Ar.GetSerializedProperty());
+										Ar << InObject;
+									}
+								}
+							} ReplaceRefCollector(ReplaceAr);
+
+							checkSlow(!GCObjectReferencer->bIsAddingReferencedObjects);
+							GCObjectReferencer->bIsAddingReferencedObjects = true;
+							for ( int32 I = 0; I < GCObjectReferencer->ReferencedObjects.Num(); ++I )
+							{
+								FGCObject* Object = GCObjectReferencer->ReferencedObjects[I];
+								// GetReferencerName seems to be as good as i can do, but the results are poor so in an effort
+								// to provide clarity i'm also outputing the FGCObject index:
+								UE_LOG(LogBlueprint, Warning, TEXT("Replacing References within: %d %s"), I, *(Object->GetReferencerName()));
+								check(Object);
+								GCObjectReferencer->CurrentlySerializingObject = Object;
+								Object->AddReferencedObjects(ReplaceRefCollector);
+							}
+							GCObjectReferencer->CurrentlySerializingObject = nullptr;
+							UObject::AddReferencedObjects(GCObjectReferencer, ReplaceRefCollector);
+							GCObjectReferencer->bIsAddingReferencedObjects = false;
+						}
+						else
+						{
+							UE_LOG(LogBlueprint, Warning, TEXT("Replacing References within: %s"), *(Obj->GetName()));
+							ReferenceReplace ReplaceAr(Obj, OldToNewInstanceMap, ReinstancedObjectsWeakReferenceMap);
+						}
+					}
+					else
+					{
+						ReferenceReplace ReplaceAr(Obj, OldToNewInstanceMap, ReinstancedObjectsWeakReferenceMap);
 					}
-					ReferenceReplace ReplaceAr(Obj, OldToNewInstanceMap, ReinstancedObjectsWeakReferenceMap);
 				}
 			}
 		}
diff --git a/Engine/Source/Runtime/CoreUObject/Public/UObject/GCObject.h b/Engine/Source/Runtime/CoreUObject/Public/UObject/GCObject.h
--- a/Engine/Source/Runtime/CoreUObject/Public/UObject/GCObject.h
+++ b/Engine/Source/Runtime/CoreUObject/Public/UObject/GCObject.h
@@ -35,6 +35,7 @@ class COREUOBJECT_API UGCObjectReferencer : public UObject
 	bool bReferencedObjectsChangedSinceLastNameVerify = false;
 	/** Currently serializing FGCObject*, only valid if bIsAddingReferencedObjects */
 	FGCObject* CurrentlySerializingObject = nullptr;
+	friend struct FReplaceReferenceHelper;
 
 public:
 	DECLARE_CASTED_CLASS_INTRINSIC_WITH_API_NO_CTOR(UGCObjectReferencer, UObject, CLASS_Transient, TEXT("/Script/CoreUObject"), CASTCLASS_None, NO_API);
