When 1 or more "undone" transactions are removed from the undo stack because a new transaction is started, we now restore said removed transactions when the newly-inserted transaction is removed because it's considered "transient" (i.e. it contains no record), similar to what happens in case of a cancelled transaction.

diff --git a/Engine/Source/Editor/UnrealEd/Private/EditorTransaction.cpp b/Engine/Source/Editor/UnrealEd/Private/EditorTransaction.cpp
--- a/Engine/Source/Editor/UnrealEd/Private/EditorTransaction.cpp
+++ b/Engine/Source/Editor/UnrealEd/Private/EditorTransaction.cpp
@@ -1135,6 +1135,22 @@ int32 UTransBuffer::End()
 				{
 					check(UndoCount == 0);
 					UndoBuffer.Pop(false);
+					UndoBuffer.Reserve(UndoBuffer.Num() + RemovedTransactions.Num());
+
+					// Restore the transactions state to what it was before that transient (i.e. ineffective) transaction (like in the Cancel case : allows to not lose the Redo stack in case 
+					//  we had inserted a transient transaction in the middle) : 
+					if (PreviousUndoCount > 0)
+					{
+						UndoBuffer.Append(RemovedTransactions);
+					}
+					else
+					{
+						UndoBuffer.Insert(RemovedTransactions, 0);
+					}
+
+					RemovedTransactions.Reset();
+					UndoCount = PreviousUndoCount;
+
 					UndoBufferChangedDelegate.Broadcast();
 				}
 			}
