[UE-168454] Some Watch Pins Show Incorrect Values in the Player AnimBP

The source of the issue was that the key of the mapping between the original and the folded property was not able to distinguish between the instance of the node. This means that e.g. when we had several blend nodes in the graph, only one of the properties was in the map and thus all pin watches were watching the single folded property that the map contains.

This is now solved by using a graph pin as key for the map, which is unique per node instance.

An additional issue was that nodes like get variable were not compiled by the blueprint compiler and thus their pins were not part of the map as well. This got fixed by adding linked pins to the map as well.

The fix got verified in game as well. Please make sure to reload/recompile the "avg_player" and "sourcepose" anim bp before verifying.

diff --git a/Engine/Source/Editor/AnimGraph/Private/AnimBlueprintCompiler.cpp b/Engine/Source/Editor/AnimGraph/Private/AnimBlueprintCompiler.cpp
--- a/Engine/Source/Editor/AnimGraph/Private/AnimBlueprintCompiler.cpp
+++ b/Engine/Source/Editor/AnimGraph/Private/AnimBlueprintCompiler.cpp
@@ -1946,8 +1946,24 @@ void FAnimBlueprintCompilerContext::ProcessFoldedPropertyRecords()
 							Record->GeneratedProperty->SetPropertyFlags(CPF_BlueprintVisible);
 							Record->PropertyIndex = PropertyIndex++;
 
-							// Populate the original-to-folded property mapping with the newly generated property
-							AnimBlueprintDebugData.NodeToFoldedPropertyMap.Add(Record->Property, Record->GeneratedProperty);
+							// Get the source graph node of the given one as the debug utilities and other users of the graph pin to folded property map
+							// are working with the source graph nodes.
+							UEdGraphNode* GraphNode = CastChecked<UEdGraphNode>(MessageLog.FindSourceObject(Record->AnimGraphNode));
+							if (GraphNode)
+							{
+								UEdGraphPin* GraphPin = GraphNode->FindPin(Record->Property->GetName());
+								if (GraphPin)
+								{
+									// Link the graph pin with the generated (folded) property.
+									AnimBlueprintDebugData.GraphPinToFoldedPropertyMap.Add(GraphPin, Record->GeneratedProperty);
+
+									// Also add all linked pins directly here. This is needed for pins on nodes that are not processed by the compiler as e.g. get variable nodes.
+									for (UEdGraphPin* LinkedToPin : GraphPin->LinkedTo)
+									{									
+										AnimBlueprintDebugData.GraphPinToFoldedPropertyMap.Add(LinkedToPin, Record->GeneratedProperty);
+									}
+								}
+							}
 						}
 					}
 					else
diff --git a/Engine/Source/Editor/UnrealEd/Private/Kismet2/KismetDebugUtilities.cpp b/Engine/Source/Editor/UnrealEd/Private/Kismet2/KismetDebugUtilities.cpp
--- a/Engine/Source/Editor/UnrealEd/Private/Kismet2/KismetDebugUtilities.cpp
+++ b/Engine/Source/Editor/UnrealEd/Private/Kismet2/KismetDebugUtilities.cpp
@@ -1845,11 +1845,12 @@ FKismetDebugUtilities::EWatchTextResult FKismetDebugUtilities::FindDebuggingData
 			if (!PropertyBase && AnimBlueprintGeneratedClass)
 			{
 				// are we linked to an anim graph node?
+				UEdGraphPin* LinkedPin = nullptr;
 				const FProperty* LinkedProperty = Property;
 				const UAnimGraphNode_Base* Node = Cast<UAnimGraphNode_Base>(WatchPin->GetOuter());
 				if (Node == nullptr && WatchPin->LinkedTo.Num() > 0)
 				{
-					const UEdGraphPin* LinkedPin = WatchPin->LinkedTo[0];
+					LinkedPin = WatchPin->LinkedTo[0];
 					// When we change Node we *must* change Property, so it's still a sub-element of that.
 					LinkedProperty = FKismetDebugUtilities::FindClassPropertyForPin(Blueprint, LinkedPin);
 					Node = Cast<UAnimGraphNode_Base>(LinkedPin->GetOuter());
@@ -1858,9 +1859,9 @@ FKismetDebugUtilities::EWatchTextResult FKismetDebugUtilities::FindDebuggingData
 				if (Node && LinkedProperty)
 				{
 					// In case the property was folded its value has to be retrieved from the Mutable data struct on the instance rather than from the Anim Node itself
-					if (const TFieldPath<const FProperty>* FoldedPropertyPathPtr = AnimBlueprintGeneratedClass->AnimBlueprintDebugData.NodeToFoldedPropertyMap.Find(LinkedProperty))
+					if (FProperty* const* FoldedPropertyPathPtr = AnimBlueprintGeneratedClass->AnimBlueprintDebugData.GraphPinToFoldedPropertyMap.Find(LinkedPin))
 					{
-						const FProperty* FoldedProperty = FoldedPropertyPathPtr->Get();						
+						const FProperty* FoldedProperty = *FoldedPropertyPathPtr;
 						const UStruct* FoldedPropertyStruct = FoldedProperty ? FoldedProperty->GetOwnerStruct() : nullptr;
 
 						if(FoldedPropertyStruct && FoldedPropertyStruct->IsChildOf(FAnimBlueprintMutableData::StaticStruct()))
diff --git a/Engine/Source/Runtime/Engine/Classes/Animation/AnimBlueprintGeneratedClass.h b/Engine/Source/Runtime/Engine/Classes/Animation/AnimBlueprintGeneratedClass.h
--- a/Engine/Source/Runtime/Engine/Classes/Animation/AnimBlueprintGeneratedClass.h
+++ b/Engine/Source/Runtime/Engine/Classes/Animation/AnimBlueprintGeneratedClass.h
@@ -202,8 +202,9 @@ public:
 	// History of snapshots of animation data
 	TSimpleRingBuffer<FAnimationFrameSnapshot>* SnapshotBuffer;
 
-	// Mapping from animation node properties to folded versions
-	TMap<TFieldPath<const FProperty>, TFieldPath<const FProperty>> NodeToFoldedPropertyMap;
+	// Mapping from graph pins to their folded properties.
+	// Graph pins are unique per node instance and thus suitable as identifier for the properties.
+	TMap<FEdGraphPinReference, FProperty*> GraphPinToFoldedPropertyMap;
 
 	// Node visit structure
 	using FNodeVisit = FAnimBlueprintDebugData_NodeVisit;
