Fix for anim preview character being offsetted at start if root motion mode is "Loop and Reset", current animation asset is paused and then we proceed to selecting another animation asset.

diff --git a/Engine/Source/Editor/UnrealEd/Private/Animation/DebugSkelMeshComponent.cpp b/Engine/Source/Editor/UnrealEd/Private/Animation/DebugSkelMeshComponent.cpp
--- a/Engine/Source/Editor/UnrealEd/Private/Animation/DebugSkelMeshComponent.cpp
+++ b/Engine/Source/Editor/UnrealEd/Private/Animation/DebugSkelMeshComponent.cpp
@@ -581,7 +581,10 @@ void UDebugSkelMeshComponent::EnablePreview(bool bEnable, UAnimationAsset* Previ
 		    bDisableClothSimulation = bPrevDisableClothSimulation;
     
 			PreviewInstance->SetAnimationAsset(PreviewAsset); 
-
+			
+			// Reset to previous animation asset's root motion playback time to prevent this from influencing the new animation asset previewing during root motion consumption.
+			ConsumeRootMotionPreviousPlaybackTime = 0.0f;
+			
 			// Reset mesh relative transform. It could be offset from the origin if we are processing root motion
 			if (!GetRelativeTransform().Equals(FTransform::Identity))
 			{
