Cooker: Add PumpRequests and WaitForAsync profile data to the end-of-cook logs.
#rnx

diff --git a/Engine/Source/Editor/UnrealEd/Private/CookOnTheFlyServer.cpp b/Engine/Source/Editor/UnrealEd/Private/CookOnTheFlyServer.cpp
--- a/Engine/Source/Editor/UnrealEd/Private/CookOnTheFlyServer.cpp
+++ b/Engine/Source/Editor/UnrealEd/Private/CookOnTheFlyServer.cpp
@@ -2104,7 +2104,7 @@ void UCookOnTheFlyServer::InitializePollables()
 void UCookOnTheFlyServer::WaitForAsync(UE::Cook::FTickStackData& StackData)
 {
 	// Sleep until the next time that DecideNextCookAction will find work to do, up to a maximum of WaitForAsyncSleepSeconds
-	UE_SCOPED_COOKTIMER(WaitForAsync);
+	UE_SCOPED_HIERARCHICAL_COOKTIMER(WaitForAsync);
 	double CurrentTime = FPlatformTime::Seconds();
 	double SleepDuration = WaitForAsyncSleepSeconds;
 	SleepDuration = FMath::Min(SleepDuration, StackData.Timer.GetEndTimeSeconds() - CurrentTime);
@@ -2357,7 +2357,7 @@ bool UCookOnTheFlyServer::TryCreateRequestCluster(UE::Cook::FPackageData& Packag
 
 void UCookOnTheFlyServer::PumpRequests(UE::Cook::FTickStackData& StackData, int32& OutNumPushed)
 {
-	UE_SCOPED_COOKTIMER(PumpRequests);
+	UE_SCOPED_HIERARCHICAL_COOKTIMER(PumpRequests);
 	using namespace UE::Cook;
 
 	OutNumPushed = 0;
