Fix broken editor quick launch for Lyra.
It was aborting early without a specific error message, but it was caused by a ELauncherProfileValidationErrors::FallbackBuildTargetIsRequired validation error.

#rnx

diff --git a/Engine/Source/Editor/UnrealEd/Private/PlayLevelViaLauncher.cpp b/Engine/Source/Editor/UnrealEd/Private/PlayLevelViaLauncher.cpp
--- a/Engine/Source/Editor/UnrealEd/Private/PlayLevelViaLauncher.cpp
+++ b/Engine/Source/Editor/UnrealEd/Private/PlayLevelViaLauncher.cpp
@@ -284,6 +284,7 @@ void UEditorEngine::StartPlayUsingLauncherSession(FRequestPlaySessionParams& InR
 		if (TargetInfo != nullptr)
 		{
 			LauncherProfile->SetBuildTarget(TargetInfo->Name);
+			LauncherProfile->SetBuildTargetSpecified(true);
 		}
 	}
 	LauncherProfile->SetCookMode(CurrentLauncherCookMode);
