Fix IProjectExternalContentInterface::RemoveExternalContent documentation

diff --git a/Engine/Source/Editor/UnrealEd/Public/IProjectExternalContentInterface.h b/Engine/Source/Editor/UnrealEd/Public/IProjectExternalContentInterface.h
--- a/Engine/Source/Editor/UnrealEd/Public/IProjectExternalContentInterface.h
+++ b/Engine/Source/Editor/UnrealEd/Public/IProjectExternalContentInterface.h
@@ -50,7 +50,7 @@ public:
 	/**
 	 * Removes a reference to external content from the project and unloads the external content
 	 * @param ExternalContentId External content identifier (verse path, link code, etc)
-	 * @param Params See FRemoveExternalContentParams
+	 * @param CompleteCallback See FRemoveExternalContentComplete
 	 */
 	virtual void RemoveExternalContent(const FString& ExternalContentId, FRemoveExternalContentComplete CompleteCallback = FRemoveExternalContentComplete()) = 0;
 };
