World Partition - HLOD: Added extra logging to track which HLOD packages are modified during a build

diff --git a/Engine/Source/Editor/UnrealEd/Private/WorldPartition/WorldPartitionHLODsBuilder.cpp b/Engine/Source/Editor/UnrealEd/Private/WorldPartition/WorldPartitionHLODsBuilder.cpp
--- a/Engine/Source/Editor/UnrealEd/Private/WorldPartition/WorldPartitionHLODsBuilder.cpp
+++ b/Engine/Source/Editor/UnrealEd/Private/WorldPartition/WorldPartitionHLODsBuilder.cpp
@@ -491,6 +491,8 @@ bool UWorldPartitionHLODsBuilder::BuildHLODActors()
 		UPackage* ActorPackage = HLODActor->GetPackage();
 		if (ActorPackage->IsDirty())
 		{
+			UE_LOG(LogWorldPartitionHLODsBuilder, Display, TEXT("[%d / %d] HLOD actor %s was modified, saving..."), CurrentActor + 1, HLODActorsToBuild.Num(), *HLODActor->GetActorLabel());
+
 			bool bSaved = SourceControlHelper->Save(ActorPackage);
 			if (!bSaved)
 			{
