Fix invalid removal of empty vertex declaration for certain vertex shaders during PSO cache build phase

#localization none
#tests run build step locally and verified that no more valid PSOs are being removed anymore

diff --git a/Engine/Source/Editor/UnrealEd/Private/Commandlets/ShaderPipelineCacheToolsCommandlet.cpp b/Engine/Source/Editor/UnrealEd/Private/Commandlets/ShaderPipelineCacheToolsCommandlet.cpp
--- a/Engine/Source/Editor/UnrealEd/Private/Commandlets/ShaderPipelineCacheToolsCommandlet.cpp
+++ b/Engine/Source/Editor/UnrealEd/Private/Commandlets/ShaderPipelineCacheToolsCommandlet.cpp
@@ -1632,7 +1632,9 @@ void FilterInvalidPSOs(TSet<FPipelineCacheFileFormatPSO>& InOutPSOs, const TMult
 		TEXT("FRenderSkyAtmosphereVS"),
 		TEXT("TPageTableUpdateVS<true>"),
 		TEXT("TPageTableUpdateVS<false>"),
-		TEXT("FShaderDrawDebugVS")
+		TEXT("FShaderDrawDebugVS"),
+		TEXT("FHWRasterizeVS"),
+		TEXT("FRenderRealTimeReflectionHeightFogVS")
 	};
 
 	TSet<FName> VShadersUsableWithEmptyVertexDecl;
