Deprecate FBlueprintEditorUtils::RegenerateBlueprintClass, a long unused entrypoint for blueprint 'on load' compilation

diff --git a/Engine/Source/Editor/UnrealEd/Public/Kismet2/BlueprintEditorUtils.h b/Engine/Source/Editor/UnrealEd/Public/Kismet2/BlueprintEditorUtils.h
--- a/Engine/Source/Editor/UnrealEd/Public/Kismet2/BlueprintEditorUtils.h
+++ b/Engine/Source/Editor/UnrealEd/Public/Kismet2/BlueprintEditorUtils.h
@@ -212,6 +212,7 @@ public:
 	/**
 	 * Regenerates the class at class load time, and refreshes the blueprint
 	 */
+	UE_DEPRECATED(5.2, "Replace with FKismetEditorUtilities::CompileBlueprint")
 	static UClass* RegenerateBlueprintClass(UBlueprint* Blueprint, UClass* ClassToRegenerate, UObject* PreviousCDO);
 	
 	/**
