MPCook: Close down the CookWorker cooksession at the start of the flush rather than waiting until receiving abortacknowledge from the director, so that shader files are released before the director tries to read them.
#rnx

diff --git a/Engine/Source/Editor/UnrealEd/Private/CookOnTheFlyServer.cpp b/Engine/Source/Editor/UnrealEd/Private/CookOnTheFlyServer.cpp
--- a/Engine/Source/Editor/UnrealEd/Private/CookOnTheFlyServer.cpp
+++ b/Engine/Source/Editor/UnrealEd/Private/CookOnTheFlyServer.cpp
@@ -9550,6 +9550,12 @@ void UCookOnTheFlyServer::LogCookWorkerStats()
 
 void UCookOnTheFlyServer::CookAsCookWorkerFinished()
 {
+	if (CookWorkerClient->HasRunFinished())
+	{
+		return;
+	}
+	CookWorkerClient->SetHasRunFinished(true);
+
 	FString LibraryName = GetProjectShaderLibraryName();
 	FString ActualLibraryName = GenerateShaderCodeLibraryName(LibraryName, IsCookFlagSet(ECookInitializationFlags::IterateSharedBuild));
 	FShaderLibraryCooker::EndCookingLibrary(ActualLibraryName);
@@ -9560,6 +9566,7 @@ void UCookOnTheFlyServer::CookAsCookWorkerFinished()
 	{
 		GShaderCompilingManager->SkipShaderCompilation(false);
 	}
+	LogCookWorkerStats();
 }
 
 void UCookOnTheFlyServer::ShutdownCookAsCookWorker()
diff --git a/Engine/Source/Editor/UnrealEd/Private/Cooker/CookWorkerClient.cpp b/Engine/Source/Editor/UnrealEd/Private/Cooker/CookWorkerClient.cpp
--- a/Engine/Source/Editor/UnrealEd/Private/Cooker/CookWorkerClient.cpp
+++ b/Engine/Source/Editor/UnrealEd/Private/Cooker/CookWorkerClient.cpp
@@ -477,9 +477,9 @@ void FCookWorkerClient::PumpDisconnect(FTickStackData& StackData)
 		{
 		case EConnectStatus::FlushAndAbortFirst:
 		{
-			COTFS.LogCookWorkerStats();
 			TickCollectors(StackData, true /* bFlush */);
 			// Add code here for any waiting we need to do for the local CookOnTheFlyServer to gracefully shutdown
+			COTFS.CookAsCookWorkerFinished();
 			SendMessage(FAbortWorkerMessage(FAbortWorkerMessage::EType::Abort));
 			SendToState(EConnectStatus::WaitForAbortAcknowledge);
 			break;
diff --git a/Engine/Source/Editor/UnrealEd/Private/Cooker/CookWorkerClient.h b/Engine/Source/Editor/UnrealEd/Private/Cooker/CookWorkerClient.h
--- a/Engine/Source/Editor/UnrealEd/Private/Cooker/CookWorkerClient.h
+++ b/Engine/Source/Editor/UnrealEd/Private/Cooker/CookWorkerClient.h
@@ -47,6 +47,8 @@ public:
 	FCookOnTheFlyOptions&& ConsumeCookOnTheFlyOptions();
 	/** Mark that initialization is complete and we can free the memory for initialization settings. */
 	void DoneWithInitialSettings();
+	bool HasRunFinished() const { return bHasRunFinished; }
+	void SetHasRunFinished(bool Value) { bHasRunFinished = Value; }
 
 	/** Queue a message to the server that the Package was cook-suppressed. Will be sent during Tick. */
 	void ReportDemoteToIdle(const FPackageData& PackageData, ESuppressCookReason Reason);
@@ -120,6 +122,7 @@ private:
 	double NextTickCollectorsTimeSeconds = 0.;
 	EConnectStatus ConnectStatus = EConnectStatus::Uninitialized;
 	ECookMode::Type DirectorCookMode = ECookMode::CookByTheBook;
+	bool bHasRunFinished = false;
 };
 
 }
\ No newline at end of file
