MPCook: Add diagnostics for how long CookWorkers are idle.
#rnx

diff --git a/Engine/Source/Editor/UnrealEd/Classes/CookOnTheSide/CookOnTheFlyServer.h b/Engine/Source/Editor/UnrealEd/Classes/CookOnTheSide/CookOnTheFlyServer.h
--- a/Engine/Source/Editor/UnrealEd/Classes/CookOnTheSide/CookOnTheFlyServer.h
+++ b/Engine/Source/Editor/UnrealEd/Classes/CookOnTheSide/CookOnTheFlyServer.h
@@ -460,8 +460,8 @@ private:
 	};
 	/** Inspect all tasks the scheduler could do and return which one it should do. */
 	ECookAction DecideNextCookAction(UE::Cook::FTickStackData& StackData);
-	/** Is the local CookOnTheFlyServer a CookDirector and it has finished all locally processed Requests? */
-	bool IsMultiprocessLocalWorkerIdle() const;
+	/** How many packages are assigned to the CookOnTheFlyServer for saving/loading? Returns 0 if the current process is not a CookDirector. */
+	int32 NumMultiprocessLocalWorkerAssignments() const;
 	/** Execute any existing external callbacks and push any existing external cook requests into new RequestClusters. */
 	void PumpExternalRequests(const UE::Cook::FCookerTimer& CookerTimer);
 	/** Send the PackageData back to request state to create a request cluster, if it has not yet been explored. */
diff --git a/Engine/Source/Editor/UnrealEd/Private/CookOnTheFlyServer.cpp b/Engine/Source/Editor/UnrealEd/Private/CookOnTheFlyServer.cpp
--- a/Engine/Source/Editor/UnrealEd/Private/CookOnTheFlyServer.cpp
+++ b/Engine/Source/Editor/UnrealEd/Private/CookOnTheFlyServer.cpp
@@ -1657,10 +1657,15 @@ void UCookOnTheFlyServer::UpdateDisplay(UE::Cook::FTickStackData& StackData, boo
 		}
 		if (!IsCookOnTheFlyMode() || bCookOnTheFlyShouldDisplay)
 		{
-			UE_CLOG(!(StackData.TickFlags & ECookTickFlags::HideProgressDisplay) && (GCookProgressDisplay != (int32)ECookProgressDisplayMode::Nothing),
-				LogCook, Display,
-				TEXT("Cook Diagnostics: OpenFileHandles=%d, VirtualMemory=%dMiB"),
-				OpenFileHandles, FPlatformMemory::GetStats().UsedVirtual / 1024 / 1024);
+			if (!(StackData.TickFlags & ECookTickFlags::HideProgressDisplay) && (GCookProgressDisplay != (int32)ECookProgressDisplayMode::Nothing))
+			{
+				UE_LOG(LogCook, Display, TEXT("Cook Diagnostics: OpenFileHandles=%d, VirtualMemory=%dMiB"),
+					OpenFileHandles, FPlatformMemory::GetStats().UsedVirtual / 1024 / 1024);
+				if (CookDirector)
+				{
+					CookDirector->UpdateDisplayDiagnostics();
+				}
+			}
 		}
 		if (bCookOnTheFlyShouldDisplay)
 		{
@@ -2284,18 +2289,18 @@ UCookOnTheFlyServer::ECookAction UCookOnTheFlyServer::DecideNextCookAction(UE::C
 	return ECookAction::Done;
 }
 
-bool UCookOnTheFlyServer::IsMultiprocessLocalWorkerIdle() const
+int32 UCookOnTheFlyServer::NumMultiprocessLocalWorkerAssignments() const
 {
 	if (!CookDirector.IsValid())
 	{
-		return false;
+		return 0;
 	}
 	UE::Cook::FPackageDataMonitor& Monitor = PackageDatas->GetMonitor();
-	return !WorkerRequests->HasExternalRequests() &&
-		PackageDatas->GetRequestQueue().IsEmpty() &&
-		PackageDatas->GetLoadPrepareQueue().IsEmpty() &&
-		PackageDatas->GetLoadReadyQueue().IsEmpty() &&
-		PackageDatas->GetSaveQueue().IsEmpty();
+	return WorkerRequests->GetNumExternalRequests() +
+		PackageDatas->GetRequestQueue().Num() +
+		PackageDatas->GetLoadPrepareQueue().Num() +
+		PackageDatas->GetLoadReadyQueue().Num() +
+		PackageDatas->GetSaveQueue().Num();
 }
 
 void UCookOnTheFlyServer::PumpExternalRequests(const UE::Cook::FCookerTimer& CookerTimer)
diff --git a/Engine/Source/Editor/UnrealEd/Private/Cooker/CookDirector.cpp b/Engine/Source/Editor/UnrealEd/Private/Cooker/CookDirector.cpp
--- a/Engine/Source/Editor/UnrealEd/Private/Cooker/CookDirector.cpp
+++ b/Engine/Source/Editor/UnrealEd/Private/Cooker/CookDirector.cpp
@@ -12,6 +12,7 @@
 #include "CoreGlobals.h"
 #include "LoadBalanceCookBurden.h"
 #include "HAL/PlatformMisc.h"
+#include "HAL/PlatformTime.h"
 #include "HAL/RunnableThread.h"
 #include "Math/NumericLimits.h"
 #include "Misc/CommandLine.h"
@@ -29,6 +30,24 @@ extern CORE_API int32 GNumForegroundWorkers; // TaskGraph.cpp
 namespace UE::Cook
 {
 
+/** Profile data for each CookWorker that needs to be collected on the Director. */
+struct FCookWorkerProfileData
+{
+	float IdleTimeSeconds = 0.f;
+	bool bIsIdle = true;
+	void UpdateIdle(bool bInIsIdle, float DeltaTime)
+	{
+		if (bInIsIdle)
+		{
+			if (bIsIdle)
+			{
+				IdleTimeSeconds += DeltaTime;
+			}
+		}
+		bIsIdle = bInIsIdle;
+	}
+};
+
 FCookDirector::FCookDirector(UCookOnTheFlyServer& InCOTFS, int32 CookProcessCount)
 	: RunnableShunt(*this) 
 	, COTFS(InCOTFS)
@@ -37,6 +56,7 @@ FCookDirector::FCookDirector(UCookOnTheFlyServer& InCOTFS, int32 CookProcessCoun
 	WorkersStalledStartTimeSeconds = MAX_flt;
 	WorkersStalledWarnTimeSeconds = MAX_flt;
 	ShutdownEvent->Reset();
+	LocalWorkerProfileData = MakeUnique<FCookWorkerProfileData>();
 
 	bool bConfigValid;
 	ParseConfig(CookProcessCount, bConfigValid);
@@ -59,6 +79,9 @@ FCookDirector::FCookDirector(UCookOnTheFlyServer& InCOTFS, int32 CookProcessCoun
 		RequestedCookWorkerCount+1, RequestedCookWorkerCount, RequestedCookWorkerCount > 1 ? TEXT("CookWorkers") : TEXT("CookWorker"));
 
 	Register(new FLogMessagesMessageHandler());
+	LastTickTimeSeconds = FPlatformTime::Seconds();
+
+	FCookStatsManager::CookStatsCallbacks.AddRaw(this, &FCookDirector::LogCookStats);
 }
 
 bool FCookDirector::IsMultiprocessAvailable() const
@@ -123,6 +146,7 @@ void FCookDirector::ParseConfig(int32 CookProcessCount, bool& bOutValid)
 FCookDirector::~FCookDirector()
 {
 	StopCommunicationThread();
+	FCookStatsManager::CookStatsCallbacks.RemoveAll(this);
 
 	TSet<FPackageData*> AbortedAssignments;
 	for (TPair<int32, TRefCountPtr<FCookWorkerServer>>& Pair : RemoteWorkers)
@@ -136,6 +160,7 @@ FCookDirector::~FCookDirector()
 		PackageData->SendToState(UE::Cook::EPackageState::Request, ESendFlags::QueueAddAndRemove);
 	}
 	RemoteWorkers.Empty();
+	RemoteWorkerProfileDatas.Empty();
 	PendingConnections.Empty();
 	Sockets::CloseSocket(WorkerConnectSocket);
 }
@@ -303,19 +328,28 @@ void FCookDirector::RemoveFromWorker(FPackageData& PackageData)
 
 void FCookDirector::TickFromSchedulerThread()
 {
+	double CurrentTime = FPlatformTime::Seconds();
 	if (!CommunicationThread)
 	{
 		TickCommunication(ECookDirectorThread::SchedulerThread);
 	}
 
+	bool bLocalWorkerIdle = COTFS.NumMultiprocessLocalWorkerAssignments() == 0;
+	float DeltaTime = static_cast<float>(CurrentTime - LastTickTimeSeconds);
+	LastTickTimeSeconds = CurrentTime;
+	LocalWorkerProfileData->UpdateIdle(bLocalWorkerIdle, DeltaTime);
+
 	TArray<TRefCountPtr<FCookWorkerServer>, TInlineAllocator<16>> WorkersWithMessage;
 	{
 		FScopeLock CommunicationScopeLock(&CommunicationLock);
 		for (TPair<int32, TRefCountPtr<FCookWorkerServer>>& Pair : RemoteWorkers)
 		{
-			if (Pair.Value->HasMessages())
+			FCookWorkerServer* RemoteWorker = Pair.Value.GetReference();
+			FCookWorkerProfileData& ProfileData = RemoteWorkerProfileDatas[RemoteWorker->GetProfileId()];
+			ProfileData.UpdateIdle(RemoteWorker->NumAssignments() == 0, DeltaTime);
+			if (RemoteWorker->HasMessages())
 			{
-				WorkersWithMessage.Add(Pair.Value);
+				WorkersWithMessage.Add(RemoteWorker);
 			}
 		}
 		for (TPair<FCookWorkerServer*, TRefCountPtr<FCookWorkerServer>>& Pair : ShuttingDownWorkers)
@@ -327,15 +361,27 @@ void FCookDirector::TickFromSchedulerThread()
 		}
 
 	}
-	bool bIsStalled = COTFS.IsMultiprocessLocalWorkerIdle() && !COTFS.PackageDatas->GetAssignedToWorkerSet().IsEmpty();
+	bool bIsStalled = bLocalWorkerIdle && !COTFS.PackageDatas->GetAssignedToWorkerSet().IsEmpty() && WorkersWithMessage.IsEmpty();
 	for (TRefCountPtr<FCookWorkerServer>& Worker : WorkersWithMessage)
 	{
 		Worker->HandleReceiveMessages(ECookDirectorThread::SchedulerThread);
-		bIsStalled = false;
 	}
 	WorkersWithMessage.Empty();
 
 	SetWorkersStalled(bIsStalled);
+	LastTickTimeSeconds = CurrentTime;
+}
+
+void FCookDirector::UpdateDisplayDiagnostics()
+{
+	UE_LOG(LogCook, Display, TEXT("\tLocalWorker: %d packages remain."), COTFS.NumMultiprocessLocalWorkerAssignments());
+	FScopeLock CommunicationScopeLock(&CommunicationLock);
+	for (TPair<int32, TRefCountPtr<FCookWorkerServer>>& Pair : RemoteWorkers)
+	{
+		FCookWorkerServer* RemoteWorker = Pair.Value;
+		UE_LOG(LogCook, Display, TEXT("\tCookWorker %d: %d packages remain."),
+			RemoteWorker->GetProfileId(), RemoteWorker->NumAssignments());
+	}
 }
 
 void FCookDirector::TickCommunication(ECookDirectorThread TickThread)
@@ -398,7 +444,7 @@ void FCookDirector::PumpCookComplete(bool& bCompleted)
 				// the last message sent about completed saves. These discovered packages can cause the cook to fall
 				// back to incomplete. Don't send CookComplete messages to the CookWorkers until all CookWorkers have
 				// reported they are idle and have no further messages to send.
-				if (RemoteWorker.HasAssignments())
+				if (RemoteWorker.NumAssignments() > 0)
 				{
 					bAllIdle = false;
 					break;
@@ -621,7 +667,9 @@ void FCookDirector::InitializeWorkers()
 	RemoteWorkers.Reserve(RequestedCookWorkerCount);
 	for (int32 RemoteIndex = 0; RemoteIndex < RequestedCookWorkerCount; ++RemoteIndex)
 	{
-		RemoteWorkers.Add(RemoteIndex, new FCookWorkerServer(*this, FWorkerId::FromRemoteIndex(RemoteIndex)));
+		int32 ProfileId = RemoteWorkerProfileDatas.Num();
+		RemoteWorkerProfileDatas.Emplace();
+		RemoteWorkers.Add(RemoteIndex, new FCookWorkerServer(*this, ProfileId, FWorkerId::FromRemoteIndex(RemoteIndex)));
 	}
 	bWorkersActive = true;
 
@@ -668,7 +716,9 @@ void FCookDirector::RecreateWorkers()
 		{
 			RemoteIndex = RemoteWorkers.Num();
 		}
-		RemoteWorkers.Add(RemoteIndex, new FCookWorkerServer(*this, FWorkerId::FromRemoteIndex(RemoteIndex)));
+		int32 ProfileId = RemoteWorkerProfileDatas.Num();
+		RemoteWorkerProfileDatas.Emplace();
+		RemoteWorkers.Add(RemoteIndex, new FCookWorkerServer(*this, ProfileId, FWorkerId::FromRemoteIndex(RemoteIndex)));
 		bWorkersActive = true;
 	}
 
@@ -862,7 +912,7 @@ void FCookDirector::TickWorkerShutdowns(ECookDirectorThread TickThread)
 	CompletedWorkers.Empty();
 }
 
-FString FCookDirector::GetWorkerCommandLine(FWorkerId WorkerId)
+FString FCookDirector::GetWorkerCommandLine(FWorkerId WorkerId, int32 ProfileId)
 {
 	FString CommandLine = FCommandLine::Get();
 
@@ -874,17 +924,14 @@ FString FCookDirector::GetWorkerCommandLine(FWorkerId WorkerId)
 			if (Token.StartsWith(TEXT("-run=")) ||
 				Token == TEXT("-CookOnTheFly") ||
 				Token == TEXT("-CookWorker") ||
-				Token == TEXT("-CookMultiProcess") ||
-				Token == TEXT("-CookSingleProcess") ||
 				Token.StartsWith(TEXT("-TargetPlatform")) ||
 				Token.StartsWith(TEXT("-CookCultures")) ||
-				Token.StartsWith(TEXT("-CookDirectorCount=")) ||
 				Token.StartsWith(TEXT("-CookDirectorHost=")) ||
 				Token.StartsWith(TEXT("-MultiprocessId=")) ||
+				Token.StartsWith(TEXT("-CookProfileId=")) ||
 				Token.StartsWith(TEXT("-ShowCookWorker")) ||
 				Token.StartsWith(TEXT("-CoreLimit")) ||
-				Token.StartsWith(TEXT("-PhysicalCoreLimit")) ||
-				Token.StartsWith(TEXT("-CoreLimitHyperThreads"))
+				Token.StartsWith(TEXT("-PhysicalCoreLimit"))
 				)
 			{
 				return;
@@ -900,6 +947,7 @@ FString FCookDirector::GetWorkerCommandLine(FWorkerId WorkerId)
 	check(!WorkerConnectAuthority.IsEmpty()); // This should have been constructed in TryCreateWorkerConnectSocket before any CookWorkerServers could exist to call GetWorkerCommandLine
 	Tokens.Add(FString::Printf(TEXT("-CookDirectorHost=%s"), *WorkerConnectAuthority));
 	Tokens.Add(FString::Printf(TEXT("-MultiprocessId=%d"), WorkerId.GetRemoteIndex() + 1));
+	Tokens.Add(FString::Printf(TEXT("-CookProfileId=%d"), ProfileId));
 	if (CoreLimit > 0)
 	{
 		Tokens.Add(FString::Printf(TEXT("-PhysicalCoreLimit=%d"), CoreLimit));
@@ -1011,14 +1059,30 @@ const FInitialConfigMessage& FCookDirector::GetInitialConfigMessage()
 	return *InitialConfigMessage;
 }
 
-FCookDirector::FLaunchInfo FCookDirector::GetLaunchInfo(FWorkerId WorkerId)
+FCookDirector::FLaunchInfo FCookDirector::GetLaunchInfo(FWorkerId WorkerId, int32 ProfileId)
 {
 	FLaunchInfo Info;
 	Info.ShowWorkerOption = GetShowWorkerOption();
 	Info.CommandletExecutable = CommandletExecutablePath;
-	Info.WorkerCommandLine = GetWorkerCommandLine(WorkerId);
+	Info.WorkerCommandLine = GetWorkerCommandLine(WorkerId, ProfileId);
 	return Info;
 }
 
+void FCookDirector::LogCookStats(FCookStatsManager::AddStatFuncRef AddStat)
+{
+	auto IdleTimeToString = [](float IdleTime)
+	{
+		return FString::Printf(TEXT("%.1fs"), IdleTime);
+	};
+	TArray<FCookStatsManager::StringKeyValue> Stats;
+	Stats.Emplace(TEXT("LocalWorker IdleTime"), IdleTimeToString(LocalWorkerProfileData->IdleTimeSeconds));
+	for (int32 ProfileId = 0; ProfileId < RemoteWorkerProfileDatas.Num(); ++ProfileId)
+	{
+		FCookWorkerProfileData& ProfileData = RemoteWorkerProfileDatas[ProfileId];
+		Stats.Emplace(FString::Printf(TEXT("CookWorker %d IdleTime"), ProfileId),
+			IdleTimeToString(ProfileData.IdleTimeSeconds));
+	}
+	AddStat(TEXT("CookDirector"), Stats);
+}
 
 }
diff --git a/Engine/Source/Editor/UnrealEd/Private/Cooker/CookDirector.h b/Engine/Source/Editor/UnrealEd/Private/Cooker/CookDirector.h
--- a/Engine/Source/Editor/UnrealEd/Private/Cooker/CookDirector.h
+++ b/Engine/Source/Editor/UnrealEd/Private/Cooker/CookDirector.h
@@ -15,6 +15,7 @@
 #include "Memory/SharedBuffer.h"
 #include "Misc/Guid.h"
 #include "Misc/ScopeLock.h"
+#include "ProfilingDebugging/CookStats.h"
 #include "Templates/RefCounting.h"
 #include "Templates/UniquePtr.h"
 
@@ -24,6 +25,7 @@ class FCbObject;
 class FCbWriter;
 class FRunnableThread;
 class UCookOnTheFlyServer;
+namespace UE::Cook { struct FCookWorkerProfileData; }
 namespace UE::Cook { class FCookWorkerServer; }
 namespace UE::Cook { class IMPCollector; }
 namespace UE::Cook { struct FInitialConfigMessage; }
@@ -68,6 +70,8 @@ public:
 	void RemoveFromWorker(FPackageData& PackageData);
 	/** Periodic tick function. Sends/Receives messages to CookWorkers. */
 	void TickFromSchedulerThread();
+	/** Periodic display function, called from CookOnTheFlyServer.UpdateDisplay. */
+	void UpdateDisplayDiagnostics();
 	/** Called when the COTFS Server has detected all packages are complete. Tells the CookWorkers to flush messages and exit. */
 	void PumpCookComplete(bool& bOutCompleted);
 	/** Called when a session ends. The Director blocks on shutdown of all CookWorkers and returns state to before session started. */
@@ -94,7 +98,7 @@ public:
 		FString CommandletExecutable;
 		FString WorkerCommandLine;
 	};
-	FLaunchInfo GetLaunchInfo(FWorkerId WorkerId);
+	FLaunchInfo GetLaunchInfo(FWorkerId WorkerId, int32 ProfileId);
 
 	/** The message CookWorkerServer sends to the remote process once it is ready to connect. */
 	const FInitialConfigMessage& GetInitialConfigMessage();
@@ -162,7 +166,7 @@ private:
 	/** Tick helper: tick any workers that are shutting down. */
 	void TickWorkerShutdowns(ECookDirectorThread TickThread);
 	/** Get the commandline to launch a worker process with. */
-	FString GetWorkerCommandLine(FWorkerId WorkerId);
+	FString GetWorkerCommandLine(FWorkerId WorkerId, int32 ProfileId);
 	/** Calls the configured LoadBalanceAlgorithm. Input Requests have been sorted by leaf to root load order. */
 	void LoadBalance(TConstArrayView<TRefCountPtr<FCookWorkerServer>> Workers, TArrayView<FPackageData*> Requests,
 		TMap<FPackageData*, TArray<FPackageData*>>&& RequestGraph, TArray<FWorkerId>& OutAssignments);
@@ -174,6 +178,8 @@ private:
 	 * Send warning when it goes on too long.
 	 */
 	void SetWorkersStalled(bool bInWorkersStalled);
+	/** Callback for CookStats system to log our stats. */
+	void LogCookStats(FCookStatsManager::AddStatFuncRef AddStat);
 
 private:
 	// Synchronization primitives that can be used from any thread
@@ -183,10 +189,13 @@ private:
 	// Data only accessible from the SchedulerThread
 	FRunnableShunt RunnableShunt;
 	FRunnableThread* CommunicationThread = nullptr;
+	TArray<FCookWorkerProfileData> RemoteWorkerProfileDatas;
 	TArray<FPendingConnection> PendingConnections;
+	TUniquePtr<FCookWorkerProfileData> LocalWorkerProfileData;
 	UCookOnTheFlyServer& COTFS;
 	double WorkersStalledStartTimeSeconds = 0.;
 	double WorkersStalledWarnTimeSeconds = 0.;
+	double LastTickTimeSeconds = 0.;
 	bool bWorkersInitialized = false;
 	bool bHasReducedMachineResources = false;
 	bool bIsFirstAssignment = true;
@@ -239,4 +248,4 @@ public:
 	static FGuid MessageType;
 };
 
-}
\ No newline at end of file
+} // namespace UE::Cook
\ No newline at end of file
diff --git a/Engine/Source/Editor/UnrealEd/Private/Cooker/CookMPCollector.h b/Engine/Source/Editor/UnrealEd/Private/Cooker/CookMPCollector.h
--- a/Engine/Source/Editor/UnrealEd/Private/Cooker/CookMPCollector.h
+++ b/Engine/Source/Editor/UnrealEd/Private/Cooker/CookMPCollector.h
@@ -53,10 +53,12 @@ public:
 		uint8 PlatformToInt(const ITargetPlatform* Platform) const;
 		const ITargetPlatform* IntToPlatform(uint8 PlatformAsInt) const;
 		FWorkerId GetWorkerId() const { return WorkerId; }
+		int32 GetProfileId() const { return ProfileId; }
 
 	private:
 		TConstArrayView<const ITargetPlatform*> Platforms;
 		FWorkerId WorkerId;
+		int32 ProfileId;
 
 		friend class FCookWorkerServer;
 	};
diff --git a/Engine/Source/Editor/UnrealEd/Private/Cooker/CookWorkerClient.cpp b/Engine/Source/Editor/UnrealEd/Private/Cooker/CookWorkerClient.cpp
--- a/Engine/Source/Editor/UnrealEd/Private/Cooker/CookWorkerClient.cpp
+++ b/Engine/Source/Editor/UnrealEd/Private/Cooker/CookWorkerClient.cpp
@@ -37,7 +37,7 @@ FCookWorkerClient::~FCookWorkerClient()
 	if (ConnectStatus == EConnectStatus::Connected ||
 		(EConnectStatus::FlushAndAbortFirst <= ConnectStatus && ConnectStatus <= EConnectStatus::FlushAndAbortLast))
 	{
-		UE_LOG(LogCook, Warning, TEXT("CookWorkerServer %d was destroyed before it finished Disconnect. The CookDirector may be missing some information."));
+		UE_LOG(LogCook, Warning, TEXT("CookWorker was destroyed before it finished Disconnect. The CookDirector may be missing some information."));
 	}
 	Sockets::CloseSocket(ServerSocket);
 }
diff --git a/Engine/Source/Editor/UnrealEd/Private/Cooker/CookWorkerServer.cpp b/Engine/Source/Editor/UnrealEd/Private/Cooker/CookWorkerServer.cpp
--- a/Engine/Source/Editor/UnrealEd/Private/Cooker/CookWorkerServer.cpp
+++ b/Engine/Source/Editor/UnrealEd/Private/Cooker/CookWorkerServer.cpp
@@ -21,9 +21,10 @@
 namespace UE::Cook
 {
 
-FCookWorkerServer::FCookWorkerServer(FCookDirector& InDirector, FWorkerId InWorkerId)
+FCookWorkerServer::FCookWorkerServer(FCookDirector& InDirector, int32 InProfileId, FWorkerId InWorkerId)
 	: Director(InDirector)
 	, COTFS(InDirector.COTFS)
+	, ProfileId(InProfileId)
 	, WorkerId(InWorkerId)
 {
 }
@@ -38,7 +39,7 @@ FCookWorkerServer::~FCookWorkerServer()
 	if (ConnectStatus == EConnectStatus::Connected || ConnectStatus == EConnectStatus::PumpingCookComplete || ConnectStatus == EConnectStatus::WaitForDisconnect)
 	{
 		UE_LOG(LogCook, Error, TEXT("CookWorkerServer %d was destroyed before it finished Disconnect. The remote process may linger and may interfere with writes of future packages."),
-			WorkerId.GetRemoteIndex());
+			ProfileId);
 	}
 	DetachFromRemoteProcess();
 }
@@ -179,10 +180,10 @@ bool FCookWorkerServer::IsShutdownComplete() const
 	return ConnectStatus == EConnectStatus::LostConnection;
 }
 
-bool FCookWorkerServer::HasAssignments() const
+int32 FCookWorkerServer::NumAssignments() const
 {
 	FScopeLock CommunicationScopeLock(&CommunicationLock);
-	return !PackagesToAssign.IsEmpty() || !PendingPackages.IsEmpty();
+	return PackagesToAssign.Num() + PendingPackages.Num();
 }
 
 bool FCookWorkerServer::HasMessages() const
@@ -203,7 +204,7 @@ bool FCookWorkerServer::TryHandleConnectMessage(FWorkerConnectMessage& Message,
 	Socket = InSocket;
 
 	SendToState(EConnectStatus::Connected);
-	UE_LOG(LogCook, Display, TEXT("CookWorker %d connected after %.3fs."), WorkerId.GetRemoteIndex(),
+	UE_LOG(LogCook, Display, TEXT("CookWorker %d connected after %.3fs."), ProfileId,
 		static_cast<float>(FPlatformTime::Seconds() - ConnectStartTimeSeconds));
 	for (UE::CompactBinaryTCP::FMarshalledMessage& OtherMessage : OtherPacketMessages)
 	{
@@ -255,7 +256,7 @@ void FCookWorkerServer::TickCommunication(ECookDirectorThread TickThread)
 					return; // Try again later
 				}
 				UE_LOG(LogCook, Error, TEXT("CookWorker process of CookWorkerServer %d failed to finalize its cook within %.0f seconds; we will tell it to shutdown."),
-					WorkerId.GetRemoteIndex(), WaitForPumpCompleteTimeout);
+					ProfileId, WaitForPumpCompleteTimeout);
 				SendMessage(FAbortWorkerMessage(FAbortWorkerMessage::EType::Abort));
 				SendToState(EConnectStatus::WaitForDisconnect);
 			}
@@ -298,7 +299,7 @@ void FCookWorkerServer::SignalCookComplete(ECookDirectorThread TickThread)
 
 void FCookWorkerServer::LaunchProcess()
 {
-	FCookDirector::FLaunchInfo LaunchInfo = Director.GetLaunchInfo(WorkerId);
+	FCookDirector::FLaunchInfo LaunchInfo = Director.GetLaunchInfo(WorkerId, ProfileId);
 	bool bShowCookWorkers = LaunchInfo.ShowWorkerOption == FCookDirector::EShowWorker::SeparateWindows;
 
 	CookWorkerHandle = FPlatformProcess::CreateProc(*LaunchInfo.CommandletExecutable, *LaunchInfo.WorkerCommandLine,
@@ -307,15 +308,15 @@ void FCookWorkerServer::LaunchProcess()
 		nullptr /* PipeWriteChild */);
 	if (CookWorkerHandle.IsValid())
 	{
-		UE_LOG(LogCook, Display, TEXT("CookWorkerServer %d launched CookWorker as PID %u with commandline \"%s\"."),
-			WorkerId.GetRemoteIndex(), CookWorkerProcessId, *LaunchInfo.WorkerCommandLine);
+		UE_LOG(LogCook, Display, TEXT("CookWorkerServer %d launched CookWorker as WorkerId %d and PID %u with commandline \"%s\"."),
+			ProfileId, WorkerId.GetRemoteIndex(), CookWorkerProcessId, *LaunchInfo.WorkerCommandLine);
 		SendToState(EConnectStatus::WaitForConnect);
 	}
 	else
 	{
 		// GetLastError information was logged by CreateProc
 		UE_LOG(LogCook, Error, TEXT("CookWorkerServer %d failed to create CookWorker process. Assigned packages will be returned to the director."),
-			WorkerId.GetRemoteIndex());
+			ProfileId);
 		SendToState(EConnectStatus::LostConnection);
 	}
 }
@@ -333,7 +334,7 @@ void FCookWorkerServer::TickWaitForConnect()
 		if (!FPlatformProcess::IsProcRunning(CookWorkerHandle))
 		{
 			UE_LOG(LogCook, Error, TEXT("CookWorkerServer %d process terminated before connecting. Assigned packages will be returned to the director."),
-				WorkerId.GetRemoteIndex());
+				ProfileId);
 			SendToState(EConnectStatus::LostConnection);
 			return;
 		}
@@ -343,7 +344,7 @@ void FCookWorkerServer::TickWaitForConnect()
 	if (CurrentTime - ConnectStartTimeSeconds > WaitForConnectTimeout && !IsCookIgnoreTimeouts())
 	{
 		UE_LOG(LogCook, Error, TEXT("CookWorkerServer %d process failed to connect within %.0f seconds. Assigned packages will be returned to the director."),
-			WorkerId.GetRemoteIndex(), WaitForConnectTimeout);
+			ProfileId, WaitForConnectTimeout);
 		ShutdownRemoteProcess();
 		SendToState(EConnectStatus::LostConnection);
 		return;
@@ -375,7 +376,7 @@ void FCookWorkerServer::TickWaitForDisconnect()
 	{
 		UE_CLOG(!bTerminateImmediately, LogCook, Warning,
 			TEXT("CookWorker process of CookWorkerServer %d failed to disconnect within %.0f seconds; we will terminate it."),
-			WorkerId.GetRemoteIndex(), WaitForDisconnectTimeout);
+			ProfileId, WaitForDisconnectTimeout);
 		ShutdownRemoteProcess();
 		SendToState(EConnectStatus::LostConnection);
 	}
@@ -387,7 +388,7 @@ void FCookWorkerServer::PumpSendMessages()
 	if (Status == UE::CompactBinaryTCP::Failed)
 	{
 		UE_LOG(LogCook, Error, TEXT("CookWorkerServer %d failed to write to socket, we will shutdown the remote process. Assigned packages will be returned to the director."),
-			WorkerId.GetRemoteIndex());
+			ProfileId);
 		SendToState(EConnectStatus::WaitForDisconnect);
 		bTerminateImmediately = true;
 	}
@@ -421,7 +422,7 @@ void FCookWorkerServer::PumpReceiveMessages()
 	if (SocketStatus != EConnectionStatus::Okay && SocketStatus != EConnectionStatus::Incomplete)
 	{
 		UE_LOG(LogCook, Error, TEXT("CookWorkerServer %d failed to read from socket, we will shutdown the remote process. Assigned packages will be returned to the director."),
-			WorkerId.GetRemoteIndex());
+			ProfileId);
 		SendToState(EConnectStatus::WaitForDisconnect);
 		bTerminateImmediately = true;
 		return;
@@ -450,7 +451,7 @@ void FCookWorkerServer::HandleReceiveMessagesInternal()
 			UE::CompactBinaryTCP::FMarshalledMessage Message = ReceiveMessages.PopFrontValue();
 			UE_CLOG(ConnectStatus != EConnectStatus::PumpingCookComplete && ConnectStatus != EConnectStatus::WaitForDisconnect,
 				LogCook, Error, TEXT("CookWorkerServer %d remote process shut down unexpectedly. Assigned packages will be returned to the director."),
-				WorkerId.GetRemoteIndex());
+				ProfileId);
 
 			SendMessage(FAbortWorkerMessage(FAbortWorkerMessage::AbortAcknowledge));
 			SendToState(EConnectStatus::WaitForDisconnect);
@@ -500,6 +501,7 @@ void FCookWorkerServer::HandleReceiveMessagesInternal()
 				IMPCollector::FServerContext Context;
 				Context.Platforms = OrderedSessionPlatforms;
 				Context.WorkerId = WorkerId;
+				Context.ProfileId = ProfileId;
 				(*Collector)->ReceiveMessage(Context, Message.Object);
 			}
 			else
@@ -557,13 +559,13 @@ void FCookWorkerServer::RecordResults(FPackageResultsMessage& Message)
 		if (!PackageData)
 		{
 			UE_LOG(LogCook, Warning, TEXT("CookWorkerServer %d received FPackageResultsMessage for invalid package %s. Ignoring it."),
-				WorkerId.GetRemoteIndex() , *Result.PackageName.ToString());
+				ProfileId, *Result.PackageName.ToString());
 			continue;
 		}
 		if (PendingPackages.Remove(PackageData) != 1)
 		{
 			UE_LOG(LogCook, Warning, TEXT("CookWorkerServer %d received FPackageResultsMessage for package %s which is not a pending package. Ignoring it."),
-				WorkerId.GetRemoteIndex() , *Result.PackageName.ToString());
+				ProfileId, *Result.PackageName.ToString());
 			continue;
 		}
 		PackageData->SetWorkerAssignment(FWorkerId::Invalid());
@@ -576,7 +578,7 @@ void FCookWorkerServer::RecordResults(FPackageResultsMessage& Message)
 			if (Result.Platforms.Num() != NumPlatforms)
 			{
 				UE_LOG(LogCook, Warning, TEXT("CookWorkerServer %d received FPackageResultsMessage for package %s with an invalid number of platform results: expected %d, actual %d. Ignoring it."),
-					WorkerId.GetRemoteIndex(), *Result.PackageName.ToString(), NumPlatforms, Result.Platforms.Num());
+					ProfileId, *Result.PackageName.ToString(), NumPlatforms, Result.Platforms.Num());
 				continue;
 			}
 			for (int32 PlatformIndex = 0; PlatformIndex < NumPlatforms; ++PlatformIndex)
@@ -953,7 +955,7 @@ void FLogMessagesMessageHandler::ReceiveMessage(FServerContext& Context, FCbObje
 	for (FReplicatedLogData& LogData : Messages)
 	{
 		GLog->CategorizedLogf(LogData.Category, LogData.Verbosity, TEXT("[CookWorker %d]: %s"),
-			Context.GetWorkerId().GetRemoteIndex(), *LogData.Message);
+			Context.GetProfileId(), *LogData.Message);
 	}
 }
 
diff --git a/Engine/Source/Editor/UnrealEd/Private/Cooker/CookWorkerServer.h b/Engine/Source/Editor/UnrealEd/Private/Cooker/CookWorkerServer.h
--- a/Engine/Source/Editor/UnrealEd/Private/Cooker/CookWorkerServer.h
+++ b/Engine/Source/Editor/UnrealEd/Private/Cooker/CookWorkerServer.h
@@ -36,9 +36,10 @@ namespace UE::Cook
 class FCookWorkerServer : public FThreadSafeRefCountedObject
 {
 public:
-	FCookWorkerServer(FCookDirector& InDirector, FWorkerId InWorkerId);
+	FCookWorkerServer(FCookDirector& InDirector, int32 InProfileId, FWorkerId InWorkerId);
 	~FCookWorkerServer();
 
+	int32 GetProfileId() const { return ProfileId; }
 	FWorkerId GetWorkerId() const { return WorkerId; }
 
 	/** Add the given assignments for the CookWorker. They will be sent during Tick */
@@ -71,8 +72,8 @@ public:
 	bool IsFlushingBeforeShutdown() const;
 	/** Is this not yet or no longer connected to a remote Client? */
 	bool IsShutdownComplete() const;
-	/** Does this Server have any package assignments the remmote CookWorker is supposed to save but hasn't yet? */
-	bool HasAssignments() const;
+	/** How many package assignments is the remote CookWorker supposed to save but hasn't yet? */
+	int32 NumAssignments() const;
 	/** Does this Server have any ReceivedMessages that need to be processed by the Scheduler thread? */
 	bool HasMessages() const;
 
@@ -160,6 +161,7 @@ private:
 	FProcHandle CookWorkerHandle;
 	FTickState TickState;
 	uint32 CookWorkerProcessId = 0;
+	int32 ProfileId = 0;
 	double ConnectStartTimeSeconds = 0.;
 	double ConnectTestStartTimeSeconds = 0.;
 	FWorkerId WorkerId = FWorkerId::Invalid();
