Add support for ShaderPipelineCache tool command let to filter out PSOs which could possibly use shaders with rely on specific RHI features which might not be supported on all systems and IHVs causing PSO compilation failure during startup

diff --git a/Engine/Source/Editor/UnrealEd/Private/Commandlets/ShaderPipelineCacheToolsCommandlet.cpp b/Engine/Source/Editor/UnrealEd/Private/Commandlets/ShaderPipelineCacheToolsCommandlet.cpp
--- a/Engine/Source/Editor/UnrealEd/Private/Commandlets/ShaderPipelineCacheToolsCommandlet.cpp
+++ b/Engine/Source/Editor/UnrealEd/Private/Commandlets/ShaderPipelineCacheToolsCommandlet.cpp
@@ -1656,6 +1656,18 @@ void FilterInvalidPSOs(TSet<FPipelineCacheFileFormatPSO>& InOutPSOs, const TMult
 		VFactoriesUsableWithEmptyVertexDecl.Add(FName(VFType));
 	}
 
+	// list of Shaders which are using RHI features which are not available on all systems and could fail to compile
+	const TCHAR* ShadersUsingPossibleUnsupportedRHIFeatures_Table[] =
+	{
+		TEXT("FTSRRejectShadingCS")
+	};
+	
+	TSet<FName> ShadersUsingPossibleUnsupportedRHIFeatures;
+	for (const TCHAR* ShaderType : ShadersUsingPossibleUnsupportedRHIFeatures_Table)
+	{
+		ShadersUsingPossibleUnsupportedRHIFeatures.Add(FName(ShaderType));
+	}
+
 	// This may be too strict, but we cannot know the VS signature.
 	auto IsInputLayoutCompatible = [](const FVertexDeclarationElementList& A, const FVertexDeclarationElementList& B, TMap<TTuple<EVertexElementType, EVertexElementType>, int32>& MismatchStats) -> bool
 	{
@@ -1812,6 +1824,7 @@ void FilterInvalidPSOs(TSet<FPipelineCacheFileFormatPSO>& InOutPSOs, const TMult
 	TSet<FStableShaderKeyAndValue> PossiblyIncorrectUsageWithEmptyDeclaration;
 	int32 NumPSOsFilteredDueToEmptyDecls = 0;
 	int32 NumPSOsFilteredDueToInconsistentDecls = 0;
+	int32 NumPSOsFilteredDueToUsingPossibleUnsupportedRHIFeatures = 0;
 	int32 NumPSOsOriginal = InOutPSOs.Num();
 
 	for (const FPipelineCacheFileFormatPSO& CurPSO : InOutPSOs)
@@ -1895,71 +1908,110 @@ void FilterInvalidPSOs(TSet<FPipelineCacheFileFormatPSO>& InOutPSOs, const TMult
 
 	FName UnknownVFType(TEXT("null"));
 
+	// Filter the PSOs using possible unsupported RHI features
+	auto ContainsShaderWithPossibleUnsupportedRHIFeatures = [InverseMap, ShadersUsingPossibleUnsupportedRHIFeatures](const FSHAHash& ShaderHash) -> bool
+	{
+		if (ShaderHash != FSHAHash())
+		{
+			const TArray<FStableShaderKeyAndValue>* Shaders = InverseMap.Find(ShaderHash);
+			if (Shaders != nullptr)
+			{
+				for (const FStableShaderKeyAndValue& Shader : *Shaders)
+				{
+					if (ShadersUsingPossibleUnsupportedRHIFeatures.Contains(Shader.ShaderType))
+					{
+						UE_LOG(LogShaderPipelineCacheTools, Warning, TEXT("Filtering out PSO using shader with possible unsupported RHI feature:\n %s"), *Shader.ToString());
+						return true;
+					}
+				}
+			}
+		}
+
+		return false;
+	};
+
 	// filter the PSOs
 	TSet<FPipelineCacheFileFormatPSO> RetainedPSOs;
 	for (const FPipelineCacheFileFormatPSO& CurPSO : InOutPSOs)
 	{
-		if (CurPSO.Type != FPipelineCacheFileFormatPSO::DescriptorType::Graphics)
+		switch (CurPSO.Type)
 		{
-			RetainedPSOs.Add(CurPSO);
-			continue;
-		}
-
-		if (CurPSO.GraphicsDesc.MeshShader != FSHAHash())
+		case FPipelineCacheFileFormatPSO::DescriptorType::Compute:
 		{
-			RetainedPSOs.Add(CurPSO);
-			continue;
-		}
+			if (ContainsShaderWithPossibleUnsupportedRHIFeatures(CurPSO.ComputeDesc.ComputeShader))
+			{
+				++NumPSOsFilteredDueToUsingPossibleUnsupportedRHIFeatures;
+				continue;
+			}
 
-		if (SuspiciousVertexShaders.Contains(CurPSO.GraphicsDesc.VertexShader))
-		{
-			++NumPSOsFilteredDueToInconsistentDecls;
-			continue;
+			break;
 		}
-
-		// check if the vertex shader is known to be used with an empty declaration - this is the largest source of driver crashes
-		if (CurPSO.GraphicsDesc.VertexDescriptor.Num() == 0)
+		case FPipelineCacheFileFormatPSO::DescriptorType::Graphics:
 		{
-			// check against the list
-			const TArray<FStableShaderKeyAndValue>* OriginalShaders = InverseMap.Find(CurPSO.GraphicsDesc.VertexShader);
-			if (OriginalShaders == nullptr)
+			if (CurPSO.GraphicsDesc.MeshShader != FSHAHash())
+			{
+				RetainedPSOs.Add(CurPSO);
+				continue;
+			}
+
+			if (SuspiciousVertexShaders.Contains(CurPSO.GraphicsDesc.VertexShader))
 			{
-				UE_LOG(LogShaderPipelineCacheTools, Warning, TEXT("PSO with an empty vertex declaration and unknown VS %s encountered, filtering out"), *CurPSO.GraphicsDesc.VertexShader.ToString());
-				++NumPSOsFilteredDueToEmptyDecls;
+				++NumPSOsFilteredDueToInconsistentDecls;
 				continue;
 			}
 
-			// all shader classes need to be usabe with empty declarations for this to pass
-			bool bAllShadersAllowed = true;
-			for (const FStableShaderKeyAndValue& OriginalShader : *OriginalShaders)
+			// check if the vertex shader is known to be used with an empty declaration - this is the largest source of driver crashes
+			if (CurPSO.GraphicsDesc.VertexDescriptor.Num() == 0)
 			{
-				if (!VShadersUsableWithEmptyVertexDecl.Contains(OriginalShader.ShaderType))
+				// check against the list
+				const TArray<FStableShaderKeyAndValue>* OriginalShaders = InverseMap.Find(CurPSO.GraphicsDesc.VertexShader);
+				if (OriginalShaders == nullptr)
 				{
-					// if this shader has a vertex factory type associated, check if VF is known to have empty decl
-					if (OriginalShader.VFType != UnknownVFType)
+					UE_LOG(LogShaderPipelineCacheTools, Warning, TEXT("PSO with an empty vertex declaration and unknown VS %s encountered, filtering out"), *CurPSO.GraphicsDesc.VertexShader.ToString());
+					++NumPSOsFilteredDueToEmptyDecls;
+					continue;
+				}
+
+				// all shader classes need to be usabe with empty declarations for this to pass
+				bool bAllShadersAllowed = true;
+				for (const FStableShaderKeyAndValue& OriginalShader : *OriginalShaders)
+				{
+					if (!VShadersUsableWithEmptyVertexDecl.Contains(OriginalShader.ShaderType))
 					{
-						if (VFactoriesUsableWithEmptyVertexDecl.Contains(OriginalShader.VFType))
+						// if this shader has a vertex factory type associated, check if VF is known to have empty decl
+						if (OriginalShader.VFType != UnknownVFType)
 						{
-							// allow, vertex factory can have an empty declaration
-							continue;
-						}
+							if (VFactoriesUsableWithEmptyVertexDecl.Contains(OriginalShader.VFType))
+							{
+								// allow, vertex factory can have an empty declaration
+								continue;
+							}
 
-						// found an incompatible (possibly, but we will err on the side of caution) usage. Log it
-						PossiblyIncorrectUsageWithEmptyDeclaration.Add(OriginalShader);
+							// found an incompatible (possibly, but we will err on the side of caution) usage. Log it
+							PossiblyIncorrectUsageWithEmptyDeclaration.Add(OriginalShader);
+						}
+						bAllShadersAllowed = false;
+						break;
 					}
-					bAllShadersAllowed = false;
-					break;
 				}
-			}
 
-			if (!bAllShadersAllowed)
-			{
-				// skip this PSO
-				++NumPSOsFilteredDueToEmptyDecls;
-				continue;
+				if (!bAllShadersAllowed)
+				{
+					// skip this PSO
+					++NumPSOsFilteredDueToEmptyDecls;
+					continue;
+				}
 			}
+
+			break;
+		}
+		case FPipelineCacheFileFormatPSO::DescriptorType::RayTracing:
+		{
+			break;
+		}
 		}
 
+		// still used
 		RetainedPSOs.Add(CurPSO);
 	}
 
@@ -1984,6 +2036,7 @@ void FilterInvalidPSOs(TSet<FPipelineCacheFileFormatPSO>& InOutPSOs, const TMult
 	UE_LOG(LogShaderPipelineCacheTools, Display, TEXT("Before sanitization: .................................................................... %6d PSOs"), NumPSOsOriginal);
 	UE_LOG(LogShaderPipelineCacheTools, Display, TEXT("Filtered out due to inconsistent vertex declaration for the same vertex shader:.......... %6d PSOs"), NumPSOsFilteredDueToInconsistentDecls);
 	UE_LOG(LogShaderPipelineCacheTools, Display, TEXT("Filtered out due to VS being possibly incompatible with an empty vertex declaration:..... %6d PSOs"), NumPSOsFilteredDueToEmptyDecls);
+	UE_LOG(LogShaderPipelineCacheTools, Display, TEXT("Filtered out due to using possible unsupported RHI features:............................. %6d PSOs"), NumPSOsFilteredDueToUsingPossibleUnsupportedRHIFeatures);
 	UE_LOG(LogShaderPipelineCacheTools, Display, TEXT("-----"));
 	UE_LOG(LogShaderPipelineCacheTools, Display, TEXT("Number of PSOs after sanity checks:...................................................... %6d PSOs"), InOutPSOs.Num());
 }
diff --git a/Engine/Source/Runtime/D3D12RHI/Private/Windows/WindowsD3D12PipelineState.cpp b/Engine/Source/Runtime/D3D12RHI/Private/Windows/WindowsD3D12PipelineState.cpp
--- a/Engine/Source/Runtime/D3D12RHI/Private/Windows/WindowsD3D12PipelineState.cpp
+++ b/Engine/Source/Runtime/D3D12RHI/Private/Windows/WindowsD3D12PipelineState.cpp
@@ -798,15 +798,14 @@ static void CreatePipelineStateWrapper(ID3D12PipelineState** PSO, FD3D12Adapter*
 		HRESULT hr = CreatePipelineStateFromStream(*PSO, pDevice2, &StreamDesc, static_cast<ID3D12PipelineLibrary1*>(CreationArgs->Library), Name);	// Static cast to ID3D12PipelineLibrary1 since we already checked for ID3D12Device2.
 		if (FAILED(hr))
 		{
+			// Always dump the graphics PSO state - internal driver compiler error can still return DXGI_ERROR_DEVICE_REMOVED
+			DumpGraphicsPSO(CreationArgs->Desc.Desc, Name);
+
 			// First check if D3D device removed, hung or out of memory and handle that separately 
 			if (hr == DXGI_ERROR_DEVICE_REMOVED || hr == DXGI_ERROR_DEVICE_HUNG || hr == E_OUTOFMEMORY)
 			{
 				VERIFYD3D12RESULT_EX(hr, pDevice2);
 			}
-			else
-			{
-				DumpGraphicsPSO(CreationArgs->Desc.Desc, Name);
-			}
 		}
 	}
 	else
@@ -817,15 +816,14 @@ static void CreatePipelineStateWrapper(ID3D12PipelineState** PSO, FD3D12Adapter*
 		HRESULT hr = CreatePipelineStateFromStream(*PSO, pDevice2, &StreamDesc, static_cast<ID3D12PipelineLibrary1*>(CreationArgs->Library), Name);	// Static cast to ID3D12PipelineLibrary1 since we already checked for ID3D12Device2.
 		if (FAILED(hr))
 		{
+			// Always dump the graphics PSO state - internal driver compiler error can still return DXGI_ERROR_DEVICE_REMOVED
+			DumpGraphicsPSO(CreationArgs->Desc.Desc, Name);
+
 			// First check if D3D device removed, hung or out of memory and handle that separately 
 			if (hr == DXGI_ERROR_DEVICE_REMOVED || hr == DXGI_ERROR_DEVICE_HUNG || hr == E_OUTOFMEMORY)
 			{
 				VERIFYD3D12RESULT_EX(hr, pDevice2);
 			}
-			else
-			{
-				DumpGraphicsPSO(CreationArgs->Desc.Desc, Name);
-			}
 		}
 	}
 }
