Add a new option 'SkipVirtualized' to the ResavePackages commandlet which will cause packages that already contain virtualized data to be skipped.

#rnx

- This option is at its most useful when trying to convert a project from non-virtualized, to virtualized as it adds an easy way to identify and re-save all packages that have not yet been virtualized.

diff --git a/Engine/Source/Editor/UnrealEd/Classes/Commandlets/ResavePackagesCommandlet.h b/Engine/Source/Editor/UnrealEd/Classes/Commandlets/ResavePackagesCommandlet.h
--- a/Engine/Source/Editor/UnrealEd/Classes/Commandlets/ResavePackagesCommandlet.h
+++ b/Engine/Source/Editor/UnrealEd/Classes/Commandlets/ResavePackagesCommandlet.h
@@ -105,6 +105,9 @@ protected:
 	/** Only save packages containing bulkdata payloads that have been virtualized **/
 	bool bOnlyVirtualized;
 
+	/** Only save packages containing bulkdata payloads that have not been virtualized **/
+	bool bSkipVirtualized;
+	
 	/** Only save packages containing FPayloadTrailers **/
 	bool bOnlyPayloadTrailers;
 
diff --git a/Engine/Source/Editor/UnrealEd/Private/Commandlets/ContentCommandlets.cpp b/Engine/Source/Editor/UnrealEd/Private/Commandlets/ContentCommandlets.cpp
--- a/Engine/Source/Editor/UnrealEd/Private/Commandlets/ContentCommandlets.cpp
+++ b/Engine/Source/Editor/UnrealEd/Private/Commandlets/ContentCommandlets.cpp
@@ -1126,8 +1126,10 @@ int32 UResavePackagesCommandlet::Main( const FString& Params )
 	bOnlyUnversioned = Switches.Contains(TEXT("OnlyUnversioned"));
 	/** whether we should only save packages saved by licensees */
 	bOnlyLicenseed = Switches.Contains(TEXT("OnlyLicenseed"));
-	/** whether we should only save packages containing virtualized bulkdata payloads */
+	/** whether we should only save packages containing virtualized payloads */
 	bOnlyVirtualized = Switches.Contains(TEXT("OnlyVirtualized"));
+	/** whether we should skip saving packages that already contain virtualized payloads */
+	bSkipVirtualized = Switches.Contains(TEXT("SkipVirtualized"));
 	/** whether we should only save packages containing FPayloadTrailers */
 	bOnlyPayloadTrailers = Switches.Contains(TEXT("OnlyPayloadTrailers"));
 	/** whether we should skip saving packages that already have a payload trailer */
@@ -1520,7 +1522,6 @@ void UResavePackagesCommandlet::PerformPreloadOperations( FLinkerLoad* PackageLi
 		}
 	}
 
-	// Check if the package contains virtualized bulkdata payloads
 	if (bOnlyVirtualized)
 	{
 		const UE::FPackageTrailer* Trailer = PackageLinker->GetPackageTrailer();
@@ -1531,6 +1532,16 @@ void UResavePackagesCommandlet::PerformPreloadOperations( FLinkerLoad* PackageLi
 		}
 	}
 
+	if (bSkipVirtualized)
+	{
+		const UE::FPackageTrailer* Trailer = PackageLinker->GetPackageTrailer();
+		if (Trailer != nullptr && Trailer->GetNumPayloads(UE::EPayloadStorageType::Virtualized) > 0)
+		{
+			bSavePackage = false;
+			return;
+		}
+	}
+
 	if (bOnlyPayloadTrailers && PackageLinker->GetPackageTrailer() == nullptr)
 	{
 		bSavePackage = false;
