Changing FSimpleElementTexture2DPreviewPS to use shader parameter structs.

diff --git a/Engine/Shaders/Private/SimpleElementTexture2DPreviewPixelShader.usf b/Engine/Shaders/Private/SimpleElementTexture2DPreviewPixelShader.usf
--- a/Engine/Shaders/Private/SimpleElementTexture2DPreviewPixelShader.usf
+++ b/Engine/Shaders/Private/SimpleElementTexture2DPreviewPixelShader.usf
@@ -25,20 +25,20 @@
 
 #define WRITE_TO_GBUFFER (FEATURE_LEVEL >= FEATURE_LEVEL_SM4 && !FORWARD_SHADING)
 
-#if TEXTURE_ARRAY
-	Texture2DArray InTexture;
-#else
-	Texture2D InTexture;
-#endif
-
 SamplerState InTextureSampler;
 
 #if SAMPLE_VIRTUAL_TEXTURE
-Texture2D<uint4> InPageTableTexture0;
-Texture2D<uint4> InPageTableTexture1;
-uint4 VTPackedPageTableUniform[2];
-uint4 VTPackedUniform;
-#endif // SAMPLE_VIRTUAL_TEXTURE
+	Texture2D<uint4> InPageTableTexture0;
+	Texture2D<uint4> InPageTableTexture1;
+	Texture2D InPhysicalTexture;
+
+	uint4 VTPackedPageTableUniform[2];
+	uint4 VTPackedUniform;
+#elif TEXTURE_ARRAY
+	Texture2DArray InTextureArray;
+#else
+	Texture2D InTexture;
+#endif
 
 half4 TextureComponentReplicate;
 half4 TextureComponentReplicateAlpha;
@@ -83,18 +83,18 @@ void Main(
 
 	if( MipLevel >= 0.0f )
 	{
-		Sample = Texture2DArraySampleLevel(InTexture, InTextureSampler, SampleCoordinates, MipLevel);
+		Sample = Texture2DArraySampleLevel(InTextureArray, InTextureSampler, SampleCoordinates, MipLevel);
 	}
 	else
 	{
-		Sample = Texture2DArraySample(InTexture, InTextureSampler, SampleCoordinates);
+		Sample = Texture2DArraySample(InTextureArray, InTextureSampler, SampleCoordinates);
 	}
 #else //TEXTURE_ARAY
 	if( MipLevel >= 0.0f )
 	{
 #if SAMPLE_VIRTUAL_TEXTURE
 		VTPageTableResult PageTableResult = TextureLoadVirtualPageTableLevel(InPageTableTexture0, InPageTableTexture1, VTPageTableUniform_Unpack(VTPackedPageTableUniform[0], VTPackedPageTableUniform[1]), TextureCoordinate, VTADDRESSMODE_WRAP, VTADDRESSMODE_WRAP, MipLevel);
-		Sample = TextureVirtualSampleLevel( InTexture, InTextureSampler, PageTableResult, LayerIndex, VTUniform_Unpack(VTPackedUniform));
+		Sample = TextureVirtualSampleLevel( InPhysicalTexture, InTextureSampler, PageTableResult, LayerIndex, VTUniform_Unpack(VTPackedUniform));
 #else
 		Sample = Texture2DSampleLevel(InTexture, InTextureSampler,TextureCoordinate,MipLevel);
 #endif
@@ -103,7 +103,7 @@ void Main(
 	{
 #if SAMPLE_VIRTUAL_TEXTURE
 		VTPageTableResult PageTableResult = TextureLoadVirtualPageTable(InPageTableTexture0, InPageTableTexture1, VTPageTableUniform_Unpack(VTPackedPageTableUniform[0], VTPackedPageTableUniform[1]), TextureCoordinate, VTADDRESSMODE_WRAP, VTADDRESSMODE_WRAP, 0.0f, InPosition.xy);
-		Sample = TextureVirtualSample( InTexture, InTextureSampler, PageTableResult, LayerIndex, VTUniform_Unpack(VTPackedUniform));
+		Sample = TextureVirtualSample( InPhysicalTexture, InTextureSampler, PageTableResult, LayerIndex, VTUniform_Unpack(VTPackedUniform));
 #else
 		Sample = Texture2DSample(InTexture, InTextureSampler,TextureCoordinate);
 #endif
diff --git a/Engine/Source/Editor/UnrealEd/Private/Texture2DPreview.cpp b/Engine/Source/Editor/UnrealEd/Private/Texture2DPreview.cpp
--- a/Engine/Source/Editor/UnrealEd/Private/Texture2DPreview.cpp
+++ b/Engine/Source/Editor/UnrealEd/Private/Texture2DPreview.cpp
@@ -15,128 +15,66 @@
 #include "RenderCore.h"
 #include "VirtualTexturing.h"
 #include "Engine/Texture2DArray.h"
+#include "ShaderParameterStruct.h"
+
 /*------------------------------------------------------------------------------
 	Batched element shaders for previewing 2d textures.
 ------------------------------------------------------------------------------*/
-/**
- * Simple pixel shader for previewing 2d textures at a specified mip level
- */
-namespace
-{
-	class FTexture2DPreviewVirtualTexture : SHADER_PERMUTATION_BOOL("SAMPLE_VIRTUAL_TEXTURE");
-	class FTexture2DPreviewTexture2DArray : SHADER_PERMUTATION_BOOL("TEXTURE_ARRAY");
-}
-
 
 class FSimpleElementTexture2DPreviewPS : public FGlobalShader
 {
-	DECLARE_SHADER_TYPE(FSimpleElementTexture2DPreviewPS, Global);
 public:
-	using FPermutationDomain = TShaderPermutationDomain<FTexture2DPreviewVirtualTexture, FTexture2DPreviewTexture2DArray>;
+	DECLARE_GLOBAL_SHADER(FSimpleElementTexture2DPreviewPS);
+	SHADER_USE_PARAMETER_STRUCT(FSimpleElementTexture2DPreviewPS, FGlobalShader);
 
-	FSimpleElementTexture2DPreviewPS(const ShaderMetaType::CompiledShaderInitializerType& Initializer) :
-		FGlobalShader(Initializer)
-	{
-		InTexture.Bind(Initializer.ParameterMap,TEXT("InTexture"), SPF_Mandatory);
-		InTextureSampler.Bind(Initializer.ParameterMap,TEXT("InTextureSampler"));
-		InPageTableTexture0.Bind(Initializer.ParameterMap, TEXT("InPageTableTexture0"));
-		InPageTableTexture1.Bind(Initializer.ParameterMap, TEXT("InPageTableTexture1"));
-		VTPackedPageTableUniform.Bind(Initializer.ParameterMap, TEXT("VTPackedPageTableUniform"));
-		VTPackedUniform.Bind(Initializer.ParameterMap, TEXT("VTPackedUniform"));
-		TextureComponentReplicate.Bind(Initializer.ParameterMap,TEXT("TextureComponentReplicate"));
-		TextureComponentReplicateAlpha.Bind(Initializer.ParameterMap,TEXT("TextureComponentReplicateAlpha"));
-		ColorWeights.Bind(Initializer.ParameterMap,TEXT("ColorWeights"));
-		PackedParameters.Bind(Initializer.ParameterMap,TEXT("PackedParams"));
-		NumSlices.Bind(Initializer.ParameterMap, TEXT("NumSlices"));
-		SliceIndex.Bind(Initializer.ParameterMap, TEXT("SliceIndex"));
-	}
-	FSimpleElementTexture2DPreviewPS() {}
+	class FVirtualTextureDim : SHADER_PERMUTATION_BOOL("SAMPLE_VIRTUAL_TEXTURE");
+	class FTexture2DArrayDim : SHADER_PERMUTATION_BOOL("TEXTURE_ARRAY");
+
+	using FPermutationDomain = TShaderPermutationDomain<FVirtualTextureDim, FTexture2DArrayDim>;
 
 	/** Should the shader be cached? Always. */
 	static bool ShouldCompilePermutation(const FGlobalShaderPermutationParameters& Parameters)
 	{
-		FPermutationDomain PermutationVector(Parameters.PermutationId);
-		if (IsFeatureLevelSupported(Parameters.Platform, ERHIFeatureLevel::SM5) == false)
+		if (!IsFeatureLevelSupported(Parameters.Platform, ERHIFeatureLevel::SM5))
 		{
 			return false;
 		}
-		return IsFeatureLevelSupported(Parameters.Platform, ERHIFeatureLevel::SM5) && !IsConsolePlatform(Parameters.Platform);
-	}
 
-	void SetParameters(FRHICommandList& InRHICmdList, const FTexture* InTextureValue, const FMatrix44f& InColorWeightsValue, float InGammaValue, float InMipLevel, float InLayerIndex, float InSliceIndex, bool bInIsNormalMap, bool bInIsSingleVTPhysicalSpace, bool bInIsVirtualTexture, bool bInIsTextureArray, bool bInUsePointSampling)
-	{
-		FRHIPixelShader* ShaderRHI = InRHICmdList.GetBoundPixelShader();
-		if (bInIsVirtualTexture)
+		if (IsConsolePlatform(Parameters.Platform))
 		{
-			FVirtualTexture2DResource* VirtualTextureValue = (FVirtualTexture2DResource*)InTextureValue;
-			IAllocatedVirtualTexture* AllocatedVT = VirtualTextureValue->AcquireAllocatedVT();
-	
-			FRHIShaderResourceView* PhysicalView = AllocatedVT->GetPhysicalTextureSRV((uint32)InLayerIndex, InTextureValue->bSRGB);
-			SetSRVParameter(InRHICmdList, ShaderRHI, InTexture, PhysicalView);
-
-			FRHISamplerState* SamplerState = bInUsePointSampling ? TStaticSamplerState<SF_Point>::GetRHI() : VirtualTextureValue->SamplerStateRHI.GetReference();
-			SetSamplerParameter(InRHICmdList, ShaderRHI, InTextureSampler, SamplerState);
-
-			SetTextureParameter(InRHICmdList, ShaderRHI, InPageTableTexture0, AllocatedVT->GetPageTableTexture(0u));
-			if (AllocatedVT->GetNumPageTableTextures() > 1u)
-			{
-				SetTextureParameter(InRHICmdList, ShaderRHI, InPageTableTexture1, AllocatedVT->GetPageTableTexture(1u));
-			}
-			else
-			{
-				SetTextureParameter(InRHICmdList, ShaderRHI, InPageTableTexture1, GBlackTexture->TextureRHI);
-			}
-
-			FUintVector4 PageTableUniform[2];
-			FUintVector4 Uniform;
-
-			AllocatedVT->GetPackedPageTableUniform(PageTableUniform);
-			AllocatedVT->GetPackedUniform(&Uniform, (uint32)InLayerIndex);
-
-			SetShaderValueArray(InRHICmdList, ShaderRHI, VTPackedPageTableUniform, PageTableUniform, UE_ARRAY_COUNT(PageTableUniform));
-			SetShaderValue(InRHICmdList, ShaderRHI, VTPackedUniform, Uniform);
+			return false;
 		}
-		else
-		{
-			SetTextureParameter(InRHICmdList, ShaderRHI, InPageTableTexture0, GBlackTexture->TextureRHI);
-			SetTextureParameter(InRHICmdList, ShaderRHI, InPageTableTexture1, GBlackTexture->TextureRHI);
 
-			if (InTextureValue != nullptr)
-			{
-				FRHISamplerState* SamplerState = bInUsePointSampling ? TStaticSamplerState<SF_Point>::GetRHI() : InTextureValue->SamplerStateRHI.GetReference();
-				SetTextureParameter(InRHICmdList, ShaderRHI, InTexture, InTextureSampler, SamplerState, InTextureValue->TextureRHI);
-			}
-		}
-		
-		SetShaderValue(InRHICmdList, ShaderRHI,ColorWeights,InColorWeightsValue);
-		FVector4f PackedParametersValue(InGammaValue, InMipLevel, bInIsNormalMap ? 1.0f : -1.0f, bInIsSingleVTPhysicalSpace ? 0 : InLayerIndex);
-		SetShaderValue(InRHICmdList, ShaderRHI, PackedParameters, PackedParametersValue);
+		FPermutationDomain PermutationVector(Parameters.PermutationId);
 
-		// Store slice count and selected slice index for texture array
-		if (bInIsTextureArray)
+		// Both SAMPLE_VIRTUAL_TEXTURE and TEXTURE_ARRAY can't be set at the same time
+		if (PermutationVector.Get<FVirtualTextureDim>() && PermutationVector.Get<FTexture2DArrayDim>())
 		{
-			const float NumSlicesData = InTextureValue != nullptr ? static_cast<float>(InTextureValue->GetSizeZ()) : 1.0f;
-			SetShaderValue(InRHICmdList, ShaderRHI, NumSlices, NumSlicesData);
-			SetShaderValue(InRHICmdList, ShaderRHI, SliceIndex, InSliceIndex);
+			return false;
 		}
 
-		SetShaderValue(InRHICmdList, ShaderRHI, TextureComponentReplicate, (InTextureValue && InTextureValue->bGreyScaleFormat) ? FLinearColor(1,0,0,0) : FLinearColor(0,0,0,0));
-		SetShaderValue(InRHICmdList, ShaderRHI, TextureComponentReplicateAlpha, (InTextureValue && InTextureValue->bGreyScaleFormat) ? FLinearColor(1,0,0,0) : FLinearColor(0,0,0,1));
+		return true;
 	}
 
-private:
-	LAYOUT_FIELD(FShaderResourceParameter, InTexture);			// if VT, this used as the physical texture
-	LAYOUT_FIELD(FShaderResourceParameter, InTextureSampler);	// if VT, this is the physical sampler
-	LAYOUT_FIELD(FShaderResourceParameter, InPageTableTexture0);
-	LAYOUT_FIELD(FShaderResourceParameter, InPageTableTexture1);
-	LAYOUT_FIELD(FShaderParameter, VTPackedPageTableUniform);
-	LAYOUT_FIELD(FShaderParameter, VTPackedUniform);
-	LAYOUT_FIELD(FShaderParameter, TextureComponentReplicate);
-	LAYOUT_FIELD(FShaderParameter, TextureComponentReplicateAlpha);
-	LAYOUT_FIELD(FShaderParameter, ColorWeights);
-	LAYOUT_FIELD(FShaderParameter, PackedParameters);
-	LAYOUT_FIELD(FShaderParameter, NumSlices);
-	LAYOUT_FIELD(FShaderParameter, SliceIndex);
+	BEGIN_SHADER_PARAMETER_STRUCT(FParameters, )
+		SHADER_PARAMETER_ARRAY(FUintVector4, VTPackedPageTableUniform, [2])
+		SHADER_PARAMETER(FUintVector4, VTPackedUniform)
+
+		SHADER_PARAMETER_EX(FVector4f, TextureComponentReplicate, EShaderPrecisionModifier::Half)
+		SHADER_PARAMETER_EX(FVector4f, TextureComponentReplicateAlpha, EShaderPrecisionModifier::Half)
+		SHADER_PARAMETER(FMatrix44f, ColorWeights)
+		SHADER_PARAMETER(FVector4f, PackedParams)
+
+		SHADER_PARAMETER(float, NumSlices)
+		SHADER_PARAMETER(float, SliceIndex)
+
+		SHADER_PARAMETER_SRV(Texture2D, InPhysicalTexture)
+		SHADER_PARAMETER_TEXTURE(Texture2D, InTexture)
+		SHADER_PARAMETER_TEXTURE(Texture2DArray, InTextureArray)
+		SHADER_PARAMETER_SAMPLER(SamplerState, InTextureSampler)
+		SHADER_PARAMETER_TEXTURE(Texture2D, InPageTableTexture0)
+		SHADER_PARAMETER_TEXTURE(Texture2D, InPageTableTexture1)
+	END_SHADER_PARAMETER_STRUCT()
 };
 
 IMPLEMENT_GLOBAL_SHADER(FSimpleElementTexture2DPreviewPS, "/Engine/Private/SimpleElementTexture2DPreviewPixelShader.usf", "Main", SF_Pixel);
@@ -154,10 +92,65 @@ void FBatchedElementTexture2DPreviewParameters::BindShaders(
 	TShaderMapRef<FSimpleElementVS> VertexShader(GetGlobalShaderMap(InFeatureLevel));
 
 	FSimpleElementTexture2DPreviewPS::FPermutationDomain PermutationVector;
-	PermutationVector.Set<FTexture2DPreviewVirtualTexture>(bIsVirtualTexture);	
-	PermutationVector.Set<FTexture2DPreviewTexture2DArray>(bIsTextureArray);
+	PermutationVector.Set<FSimpleElementTexture2DPreviewPS::FVirtualTextureDim>(bIsVirtualTexture);
+	PermutationVector.Set<FSimpleElementTexture2DPreviewPS::FTexture2DArrayDim>(bIsTextureArray);
 	TShaderMapRef<FSimpleElementTexture2DPreviewPS> PixelShader(GetGlobalShaderMap(InFeatureLevel), PermutationVector);
 
+	FSimpleElementTexture2DPreviewPS::FParameters Parameters{};
+	{
+		if (bIsVirtualTexture)
+		{
+			FVirtualTexture2DResource* VirtualTextureValue = (FVirtualTexture2DResource*)Texture;
+			IAllocatedVirtualTexture* AllocatedVT = VirtualTextureValue->AcquireAllocatedVT();
+
+			Parameters.InPhysicalTexture = AllocatedVT->GetPhysicalTextureSRV((uint32)LayerIndex, Texture->bSRGB);
+			Parameters.InTextureSampler = VirtualTextureValue->SamplerStateRHI.GetReference();
+
+			Parameters.InPageTableTexture0 = AllocatedVT->GetPageTableTexture(0u);
+			Parameters.InPageTableTexture1 = AllocatedVT->GetNumPageTableTextures() > 1u ? AllocatedVT->GetPageTableTexture(1u) : GBlackTexture->TextureRHI.GetReference();
+
+			FUintVector4 VTPackedPageTableUniform[2];
+			FUintVector4 VTPackedUniform;
+
+			AllocatedVT->GetPackedPageTableUniform(VTPackedPageTableUniform);
+			AllocatedVT->GetPackedUniform(&VTPackedUniform, (uint32)LayerIndex);
+
+			Parameters.VTPackedPageTableUniform[0] = VTPackedPageTableUniform[0];
+			Parameters.VTPackedPageTableUniform[1] = VTPackedPageTableUniform[1];
+			Parameters.VTPackedUniform = VTPackedUniform;
+		}
+		else if (Texture != nullptr)
+		{
+			if (bIsTextureArray)
+			{
+				Parameters.InTextureArray = Texture->TextureRHI;
+			}
+			else
+			{
+				Parameters.InTexture = Texture->TextureRHI;
+			}
+			Parameters.InTextureSampler = Texture->SamplerStateRHI.GetReference();
+		}
+
+		if (bUsePointSampling)
+		{
+			Parameters.InTextureSampler = TStaticSamplerState<SF_Point>::GetRHI();
+		}
+
+		Parameters.ColorWeights = FMatrix44f(ColorWeights);
+		Parameters.PackedParams = FVector4f(InGamma, MipLevel, bIsNormalMap ? 1.0f : -1.0f, bIsSingleVTPhysicalSpace ? 0 : LayerIndex);
+
+		// Store slice count and selected slice index for texture array
+		if (bIsTextureArray)
+		{
+			Parameters.NumSlices = Texture != nullptr ? static_cast<float>(Texture->GetSizeZ()) : 1.0f;
+			Parameters.SliceIndex = SliceIndex;
+		}
+
+		Parameters.TextureComponentReplicate = (Texture && Texture->bGreyScaleFormat) ? FLinearColor(1, 0, 0, 0) : FLinearColor(0, 0, 0, 0);
+		Parameters.TextureComponentReplicateAlpha = (Texture && Texture->bGreyScaleFormat) ? FLinearColor(1, 0, 0, 0) : FLinearColor(0, 0, 0, 1);
+	}
+
 	GraphicsPSOInit.BoundShaderState.VertexDeclarationRHI = GSimpleElementVertexDeclaration.VertexDeclarationRHI;
 	GraphicsPSOInit.BoundShaderState.VertexShaderRHI = VertexShader.GetVertexShader();
 	GraphicsPSOInit.BoundShaderState.PixelShaderRHI = PixelShader.GetPixelShader();
@@ -172,5 +165,6 @@ void FBatchedElementTexture2DPreviewParameters::BindShaders(
 	SetGraphicsPipelineState(RHICmdList, GraphicsPSOInit, 0);
 
 	VertexShader->SetParameters(RHICmdList, InTransform);
-	PixelShader->SetParameters(RHICmdList, Texture, FMatrix44f(ColorWeights), InGamma, MipLevel, LayerIndex, SliceIndex, bIsNormalMap, bIsSingleVTPhysicalSpace, bIsVirtualTexture, bIsTextureArray, bUsePointSampling);
+
+	SetShaderParameters(RHICmdList, PixelShader, PixelShader.GetPixelShader(), Parameters);
 }
