Remove logging
#rnx

diff --git a/Engine/Source/Editor/UnrealEd/Private/Kismet2/KismetReinstanceUtilities.cpp b/Engine/Source/Editor/UnrealEd/Private/Kismet2/KismetReinstanceUtilities.cpp
--- a/Engine/Source/Editor/UnrealEd/Private/Kismet2/KismetReinstanceUtilities.cpp
+++ b/Engine/Source/Editor/UnrealEd/Private/Kismet2/KismetReinstanceUtilities.cpp
@@ -57,8 +57,6 @@ static FAutoConsoleVariableRef CVarUseLegacyAnimInstanceReinstancingBehavior(
 	TEXT("Use the legacy re-instancing behavior for anim instances where the instance is destroyed and re-created.")
 );
 
-extern ENGINE_API UObject* GLogReinstancerReferenceReplacementObj;
-
 struct FReplaceReferenceHelper
 {
 	static void IncludeCDO(UClass* OldClass, UClass* NewClass, TMap<UObject*, UObject*> &OldToNewInstanceMap, TArray<UObject*> &SourceObjects, UObject* OriginalCDO)
@@ -96,21 +94,6 @@ struct FReplaceReferenceHelper
 			return;
 		}
 
-		if (GLogReinstancerReferenceReplacementObj)
-		{
-			FString Classes;
-			for (UObject* Obj : SourceObjects)
-			{
-				Classes.Append(Obj->GetClass()->GetName() + " ");
-			}
-			UE_LOG(LogBlueprint, Warning, TEXT("%s %s %s %s Replacing References to: %s"), 
-				*(GLogReinstancerReferenceReplacementObj->GetOutermost()->GetName()),
-				GLogReinstancerReferenceReplacementObj->GetOuter() ? *(GLogReinstancerReferenceReplacementObj->GetOuter()->GetName()) : TEXT("nullouter"),
-				*(GLogReinstancerReferenceReplacementObj->GetPathName()),
-				*(GLogReinstancerReferenceReplacementObj->GetClass()->GetName()), 
-				*Classes);
-		}
-
 		// Remember what values were in UActorChannel::Actor so we can restore them later (this should only affect reinstancing during PIE)
 		// We need the old actor channel to tear down cleanly without affecting the new actor
 		TMap<UActorChannel*, AActor*> ActorChannelActorRestorationMap;
@@ -193,72 +176,7 @@ struct FReplaceReferenceHelper
 						const TMap<FSoftObjectPath, UObject*>& WeakReferencesMap;
 					};
 
-					if (GLogReinstancerReferenceReplacementObj)
-					{
-						// inlining the whole of UGCObjectReferencer::AddReferencedObjects because I want to track down
-						// a likely memory stomp, but don't want to slow down or destabilize other clients of UGCObjectReferencer:
-						if (UGCObjectReferencer* GCObjectReferencer = Cast<UGCObjectReferencer>(Obj))
-						{
-							ReferenceReplace ReplaceAr(Obj, OldToNewInstanceMap, ReinstancedObjectsWeakReferenceMap, EArchiveReplaceObjectFlags::DelayStart);
-							class FReplaceObjectRefCollector : public FReferenceCollector
-							{
-								FArchive& Ar;
-								bool bAllowReferenceElimination;
-							public:
-								FReplaceObjectRefCollector(FArchive& InAr)
-									: Ar(InAr)
-									, bAllowReferenceElimination(true)
-								{
-								}
-								virtual bool IsIgnoringArchetypeRef() const override
-								{
-									return Ar.IsIgnoringArchetypeRef();
-								}
-								virtual bool IsIgnoringTransient() const override
-								{
-									return false;
-								}
-								virtual void AllowEliminatingReferences(bool bAllow) override
-								{
-									bAllowReferenceElimination = bAllow;
-								}
-								virtual void HandleObjectReference(UObject*& InObject, const UObject* InReferencingObject, const FProperty* InReferencingProperty) override
-								{
-									if (bAllowReferenceElimination)
-									{
-										FProperty* NewSerializedProperty = const_cast<FProperty*>(InReferencingProperty);
-										FSerializedPropertyScope SerializedPropertyScope(Ar, NewSerializedProperty ? NewSerializedProperty : Ar.GetSerializedProperty());
-										Ar << InObject;
-									}
-								}
-							} ReplaceRefCollector(ReplaceAr);
-
-							checkSlow(!GCObjectReferencer->bIsAddingReferencedObjects);
-							GCObjectReferencer->bIsAddingReferencedObjects = true;
-							for ( int32 I = 0; I < GCObjectReferencer->ReferencedObjects.Num(); ++I )
-							{
-								FGCObject* Object = GCObjectReferencer->ReferencedObjects[I];
-								// GetReferencerName seems to be as good as i can do, but the results are poor so in an effort
-								// to provide clarity i'm also outputing the FGCObject index:
-								UE_LOG(LogBlueprint, Warning, TEXT("Replacing References within: %d %s"), I, *(Object->GetReferencerName()));
-								check(Object);
-								GCObjectReferencer->CurrentlySerializingObject = Object;
-								Object->AddReferencedObjects(ReplaceRefCollector);
-							}
-							GCObjectReferencer->CurrentlySerializingObject = nullptr;
-							UObject::AddReferencedObjects(GCObjectReferencer, ReplaceRefCollector);
-							GCObjectReferencer->bIsAddingReferencedObjects = false;
-						}
-						else
-						{
-							UE_LOG(LogBlueprint, Warning, TEXT("Replacing References within: %s"), *(Obj->GetName()));
-							ReferenceReplace ReplaceAr(Obj, OldToNewInstanceMap, ReinstancedObjectsWeakReferenceMap);
-						}
-					}
-					else
-					{
-						ReferenceReplace ReplaceAr(Obj, OldToNewInstanceMap, ReinstancedObjectsWeakReferenceMap);
-					}
+					ReferenceReplace ReplaceAr(Obj, OldToNewInstanceMap, ReinstancedObjectsWeakReferenceMap);
 				}
 			}
 		}
diff --git a/Engine/Source/Runtime/Engine/Private/SkeletalMeshComponentPhysics.cpp b/Engine/Source/Runtime/Engine/Private/SkeletalMeshComponentPhysics.cpp
--- a/Engine/Source/Runtime/Engine/Private/SkeletalMeshComponentPhysics.cpp
+++ b/Engine/Source/Runtime/Engine/Private/SkeletalMeshComponentPhysics.cpp
@@ -2689,10 +2689,6 @@ void USkeletalMeshComponent::AddClothingBounds(FBoxSphereBounds& InOutBounds, co
 	}
 }
 
-#if WITH_EDITOR
-ENGINE_API UObject* GLogReinstancerReferenceReplacementObj = nullptr;
-#endif
-
 void USkeletalMeshComponent::RecreateClothingActors()
 {
 	CSV_SCOPED_TIMING_STAT(Animation, ClothInit);
@@ -2705,10 +2701,6 @@ void USkeletalMeshComponent::RecreateClothingActors()
 		return;
 	}
 
-#if WITH_EDITOR
-	TGuardValue<UObject*> EnableReinstancingLog(GLogReinstancerReferenceReplacementObj, this);
-#endif
-
 	if(CVarEnableClothPhysics.GetValueOnGameThread() && (SkelMesh->GetMeshClothingAssets().Num() > 0))
 	{
 		UClass* SimFactoryClass = *ClothingSimulationFactory;
