#UE Lower verbosity on message about PSOs with possible unsupported RHI features

diff --git a/Engine/Source/Editor/UnrealEd/Private/Commandlets/ShaderPipelineCacheToolsCommandlet.cpp b/Engine/Source/Editor/UnrealEd/Private/Commandlets/ShaderPipelineCacheToolsCommandlet.cpp
--- a/Engine/Source/Editor/UnrealEd/Private/Commandlets/ShaderPipelineCacheToolsCommandlet.cpp
+++ b/Engine/Source/Editor/UnrealEd/Private/Commandlets/ShaderPipelineCacheToolsCommandlet.cpp
@@ -1920,7 +1920,7 @@ void FilterInvalidPSOs(TSet<FPipelineCacheFileFormatPSO>& InOutPSOs, const TMult
 				{
 					if (ShadersUsingPossibleUnsupportedRHIFeatures.Contains(Shader.ShaderType))
 					{
-						UE_LOG(LogShaderPipelineCacheTools, Warning, TEXT("Filtering out PSO using shader with possible unsupported RHI feature:\n %s"), *Shader.ToString());
+						UE_LOG(LogShaderPipelineCacheTools, Display, TEXT("Filtering out PSO using shader with possible unsupported RHI feature:\n %s"), *Shader.ToString());
 						return true;
 					}
 				}
