MPCook: Add implementation subclasses of IMPCollector that use a CompactBinaryTCP::IMessage subclass to serialize the message, and pass the result to a callback. This reduces the boilerplate required for CookDirectory and CookWorkerClient to pass messages to each other.
Add IMPCollector::ClientReceiveMessage so that a collector running on the client can receive messages from the server.
Change CompactBinaryTCP::IMessage to take an FCbObjectView instead of FCbObject&& to match IMPCollector. Taking an FCbObjectView increases copy costs if keeping an FCbObject is necessary, but it prevents the possibility of keeping a long message referenced when holding on to a single FCbObject within that message.
Change FCookWorkerClient to also register its IMPCollectors to handle incoming messages. Change FCookDirector::MessageHandlers variable name to Collectors to prepare for allowing those Collectors to be ticked as well as handle messages, but ticking is not yet necessary so not yet implemented.
Add RetractionRequest and RetractionResults message boilerplate; implementation of them is coming in a future changelist.

#rnx

diff --git a/Engine/Source/Editor/UnrealEd/Private/Commandlets/AssetRegistryGenerator.cpp b/Engine/Source/Editor/UnrealEd/Private/Commandlets/AssetRegistryGenerator.cpp
--- a/Engine/Source/Editor/UnrealEd/Private/Commandlets/AssetRegistryGenerator.cpp
+++ b/Engine/Source/Editor/UnrealEd/Private/Commandlets/AssetRegistryGenerator.cpp
@@ -2662,7 +2662,7 @@ void FAssetRegistryPackageMessage::Write(FCbWriter& Writer, const FPackageData&
 	Writer << "S" << DiskSize;
 }
 
-bool FAssetRegistryPackageMessage::TryRead(FCbObject&& Object, FPackageData& PackageData, const ITargetPlatform* TargetPlatform)
+bool FAssetRegistryPackageMessage::TryRead(FCbObjectView Object, FPackageData& PackageData, const ITargetPlatform* TargetPlatform)
 {
 	FCbFieldView AssetDatasField = Object["A"];
 	FCbArrayView AssetDatasArray = AssetDatasField.AsArrayView();
diff --git a/Engine/Source/Editor/UnrealEd/Private/CookOnTheFlyServer.cpp b/Engine/Source/Editor/UnrealEd/Private/CookOnTheFlyServer.cpp
--- a/Engine/Source/Editor/UnrealEd/Private/CookOnTheFlyServer.cpp
+++ b/Engine/Source/Editor/UnrealEd/Private/CookOnTheFlyServer.cpp
@@ -7971,14 +7971,14 @@ public:
 	virtual FGuid GetMessageType() const override { return MessageType; }
 	virtual const TCHAR* GetDebugName() const override { return TEXT("FShaderLibraryCollector"); }
 
-	virtual void ClientTick(FClientContext& Context) override;
-	virtual void ReceiveMessage(FServerContext& Context, FCbObjectView Message) override;
+	virtual void ClientTick(FMPCollectorClientTickContext& Context) override;
+	virtual void ServerReceiveMessage(FMPCollectorServerMessageContext& Context, FCbObjectView Message) override;
 
 	static FGuid MessageType;
 };
 FGuid FShaderLibraryCollector::MessageType(TEXT("4DF3B36BBA2F4E04A846E894E24EB2C4"));
 
-void FShaderLibraryCollector::ClientTick(FClientContext& Context)
+void FShaderLibraryCollector::ClientTick(FMPCollectorClientTickContext& Context)
 {
 	FCbWriter Writer;
 	bool bHasData;
@@ -7992,7 +7992,7 @@ void FShaderLibraryCollector::ClientTick(FClientContext& Context)
 	}
 }
 
-void FShaderLibraryCollector::ReceiveMessage(FServerContext& Context, FCbObjectView Message)
+void FShaderLibraryCollector::ServerReceiveMessage(FMPCollectorServerMessageContext& Context, FCbObjectView Message)
 {
 	bool bSuccessful = FShaderLibraryCooker::AppendFromCompactBinary(Message["S"]);
 	UE_CLOG(!bSuccessful, LogCook, Error,
@@ -9955,14 +9955,14 @@ public:
 	virtual FGuid GetMessageType() const override { return MessageType; }
 	virtual const TCHAR* GetDebugName() const override { return TEXT("FEDLMPCollector"); }
 
-	virtual void ClientTick(FClientContext& Context) override;
-	virtual void ReceiveMessage(FServerContext& Context, FCbObjectView Message) override;
+	virtual void ClientTick(FMPCollectorClientTickContext& Context) override;
+	virtual void ServerReceiveMessage(FMPCollectorServerMessageContext& Context, FCbObjectView Message) override;
 
 	static FGuid MessageType;
 };
 FGuid FEDLMPCollector::MessageType(TEXT("0164FD08F6884F6A82D2D00F8F70B182"));
 
-void FEDLMPCollector::ClientTick(FClientContext& Context)
+void FEDLMPCollector::ClientTick(FMPCollectorClientTickContext& Context)
 {
 	FCbWriter Writer;
 	bool bHasData;
@@ -9973,7 +9973,7 @@ void FEDLMPCollector::ClientTick(FClientContext& Context)
 	}
 }
 
-void FEDLMPCollector::ReceiveMessage(FServerContext& Context, FCbObjectView Message)
+void FEDLMPCollector::ServerReceiveMessage(FMPCollectorServerMessageContext& Context, FCbObjectView Message)
 {
 	UE::SavePackageUtilities::EDLCookInfoAppendFromCompactBinary(Message.AsFieldView());
 }
diff --git a/Engine/Source/Editor/UnrealEd/Private/Cooker/CompactBinaryTCP.h b/Engine/Source/Editor/UnrealEd/Private/Cooker/CompactBinaryTCP.h
--- a/Engine/Source/Editor/UnrealEd/Private/Cooker/CompactBinaryTCP.h
+++ b/Engine/Source/Editor/UnrealEd/Private/Cooker/CompactBinaryTCP.h
@@ -46,7 +46,7 @@ public:
 	/** Marshall the message to a CompactBinaryObject. */
 	virtual void Write(FCbWriter& Writer) const = 0;
 	/** Unmarshall the message from a CompactBinaryObject. */
-	virtual bool TryRead(FCbObject&& Object) = 0;
+	virtual bool TryRead(FCbObjectView Object) = 0;
 	/** Return the Guid that identifies the message to the remote connection. */
 	virtual FGuid GetMessageType() const = 0;
 };
diff --git a/Engine/Source/Editor/UnrealEd/Private/Cooker/CookDirector.cpp b/Engine/Source/Editor/UnrealEd/Private/Cooker/CookDirector.cpp
--- a/Engine/Source/Editor/UnrealEd/Private/Cooker/CookDirector.cpp
+++ b/Engine/Source/Editor/UnrealEd/Private/Cooker/CookDirector.cpp
@@ -79,6 +79,11 @@ FCookDirector::FCookDirector(UCookOnTheFlyServer& InCOTFS, int32 CookProcessCoun
 		RequestedCookWorkerCount+1, RequestedCookWorkerCount, RequestedCookWorkerCount > 1 ? TEXT("CookWorkers") : TEXT("CookWorker"));
 
 	Register(new FLogMessagesMessageHandler());
+	Register(new IMPCollectorCbServerMessage<FRetractionResultsMessage>([this]
+	(FMPCollectorServerMessageContext& Context, bool bReadSuccessful, FRetractionResultsMessage&& Message)
+		{
+			HandleRetractionMessage(Context, bReadSuccessful, MoveTemp(Message));
+		}, TEXT("HandleRetractionMessage")));
 	LastTickTimeSeconds = FPlatformTime::Seconds();
 
 	FCookStatsManager::CookStatsCallbacks.AddRaw(this, &FCookDirector::LogCookStats);
@@ -531,7 +536,7 @@ void FCookDirector::ShutdownCookSession()
 
 void FCookDirector::Register(IMPCollector* Collector)
 {
-	TRefCountPtr<IMPCollector>& Existing = MessageHandlers.FindOrAdd(Collector->GetMessageType());
+	TRefCountPtr<IMPCollector>& Existing = Collectors.FindOrAdd(Collector->GetMessageType());
 	if (Existing)
 	{
 		UE_LOG(LogCook, Error, TEXT("Duplicate IMPCollectors registered. Guid: %s, Existing: %s, Registering: %s. Keeping the Existing."),
@@ -544,12 +549,12 @@ void FCookDirector::Register(IMPCollector* Collector)
 void FCookDirector::Unregister(IMPCollector* Collector)
 {
 	TRefCountPtr<IMPCollector> Existing;
-	MessageHandlers.RemoveAndCopyValue(Collector->GetMessageType(), Existing);
+	Collectors.RemoveAndCopyValue(Collector->GetMessageType(), Existing);
 	if (Existing && Existing.GetReference() != Collector)
 	{
 		UE_LOG(LogCook, Error, TEXT("Duplicate IMPCollector during Unregister. Guid: %s, Existing: %s, Unregistering: %s. Ignoring the Unregister."),
 			*Collector->GetMessageType().ToString(), Existing->GetDebugName(), Collector->GetDebugName());
-		MessageHandlers.Add(Collector->GetMessageType(), MoveTemp(Existing));
+		Collectors.Add(Collector->GetMessageType(), MoveTemp(Existing));
 	}
 }
 
@@ -605,7 +610,7 @@ void FWorkerConnectMessage::Write(FCbWriter& Writer) const
 	Writer << "RemoteIndex" << RemoteIndex;
 }
 
-bool FWorkerConnectMessage::TryRead(FCbObject&& Object)
+bool FWorkerConnectMessage::TryRead(FCbObjectView Object)
 {
 	RemoteIndex = Object["RemoteIndex"].AsInt32(-1);
 	return RemoteIndex >= 0;
@@ -1085,4 +1090,31 @@ void FCookDirector::LogCookStats(FCookStatsManager::AddStatFuncRef AddStat)
 	AddStat(TEXT("CookDirector"), Stats);
 }
 
+void FCookDirector::HandleRetractionMessage(FMPCollectorServerMessageContext& Context, bool bReadSuccessful,
+	FRetractionResultsMessage&& Message)
+{
+}
+
+void FRetractionRequestMessage::Write(FCbWriter& Writer) const
+{
+	Writer << "RequestedCount" << RequestedCount;
+}
+bool FRetractionRequestMessage::TryRead(FCbObjectView Object)
+{
+	return LoadFromCompactBinary(Object["RequestedCount"], RequestedCount);
+}
+
+FGuid FRetractionRequestMessage::MessageType(TEXT("7109E168E8A8405BA65F9E1E82571D1A"));
+
+void FRetractionResultsMessage::Write(FCbWriter& Writer) const
+{
+	Writer << "ReturnedPackages" << ReturnedPackages;
+}
+bool FRetractionResultsMessage::TryRead(FCbObjectView Object)
+{
+	return LoadFromCompactBinary(Object["ReturnedPackages"], ReturnedPackages);
+}
+
+FGuid FRetractionResultsMessage::MessageType(TEXT("CBFB840A4FB94903A757C490514A4B86"));
+
 }
diff --git a/Engine/Source/Editor/UnrealEd/Private/Cooker/CookDirector.h b/Engine/Source/Editor/UnrealEd/Private/Cooker/CookDirector.h
--- a/Engine/Source/Editor/UnrealEd/Private/Cooker/CookDirector.h
+++ b/Engine/Source/Editor/UnrealEd/Private/Cooker/CookDirector.h
@@ -25,11 +25,13 @@ class FCbObject;
 class FCbWriter;
 class FRunnableThread;
 class UCookOnTheFlyServer;
-namespace UE::Cook { struct FCookWorkerProfileData; }
 namespace UE::Cook { class FCookWorkerServer; }
+namespace UE::Cook { class FMPCollectorServerMessageContext; }
 namespace UE::Cook { class IMPCollector; }
+namespace UE::Cook { struct FCookWorkerProfileData; }
 namespace UE::Cook { struct FInitialConfigMessage; }
 namespace UE::Cook { struct FPackageData; }
+namespace UE::Cook { struct FRetractionResultsMessage; }
 namespace UE::Cook { struct FWorkerId; }
 
 namespace UE::Cook
@@ -180,6 +182,9 @@ private:
 	void SetWorkersStalled(bool bInWorkersStalled);
 	/** Callback for CookStats system to log our stats. */
 	void LogCookStats(FCookStatsManager::AddStatFuncRef AddStat);
+	void HandleRetractionMessage(FMPCollectorServerMessageContext& Context, bool bReadSuccessful,
+		FRetractionResultsMessage&& Message);
+
 
 private:
 	// Synchronization primitives that can be used from any thread
@@ -205,7 +210,7 @@ private:
 
 	// Data that is read-only while the CommunicationThread is active and is readable from any thread
 	FBeginCookContextForWorker BeginCookContext;
-	TMap<FGuid, TRefCountPtr<IMPCollector>> MessageHandlers;
+	TMap<FGuid, TRefCountPtr<IMPCollector>> Collectors;
 	TUniquePtr<FInitialConfigMessage> InitialConfigMessage;
 	FString WorkerConnectAuthority;
 	FString CommandletExecutablePath;
@@ -240,7 +245,7 @@ struct FWorkerConnectMessage : public UE::CompactBinaryTCP::IMessage
 {
 public:
 	virtual void Write(FCbWriter& Writer) const override;
-	virtual bool TryRead(FCbObject&& Object) override;
+	virtual bool TryRead(FCbObjectView Object) override;
 	virtual FGuid GetMessageType() const override { return MessageType; }
 
 public:
@@ -248,4 +253,34 @@ public:
 	static FGuid MessageType;
 };
 
+/**
+ * Message sent from CookDirector to a CookWorker to cancel some of its assigned packages and return them
+ * dispatch to idle workers.
+ */
+struct FRetractionRequestMessage : public UE::CompactBinaryTCP::IMessage
+{
+	virtual void Write(FCbWriter& Writer) const override;
+	virtual bool TryRead(FCbObjectView Object) override;
+	virtual FGuid GetMessageType() const override { return MessageType; }
+
+public:
+	int32 RequestedCount = 0;
+	static FGuid MessageType;
+};
+
+/**
+ * Message sent from CookWorker to CookDirector identifying which assigned packages it chose to satisfy a
+ * a RetractionRequest.
+ */
+struct FRetractionResultsMessage : public UE::CompactBinaryTCP::IMessage
+{
+	virtual void Write(FCbWriter& Writer) const override;
+	virtual bool TryRead(FCbObjectView Object) override;
+	virtual FGuid GetMessageType() const override { return MessageType; }
+
+public:
+	TArray<FName> ReturnedPackages;
+	static FGuid MessageType;
+};
+
 } // namespace UE::Cook
\ No newline at end of file
diff --git a/Engine/Source/Editor/UnrealEd/Private/Cooker/CookMPCollector.cpp b/Engine/Source/Editor/UnrealEd/Private/Cooker/CookMPCollector.cpp
--- a/Engine/Source/Editor/UnrealEd/Private/Cooker/CookMPCollector.cpp
+++ b/Engine/Source/Editor/UnrealEd/Private/Cooker/CookMPCollector.cpp
@@ -26,34 +26,39 @@ static const ITargetPlatform* IntToPlatformImpl(uint8 PlatformAsInt, TConstArray
 	return nullptr;
 }
 
-void IMPCollector::FClientContext::AddMessage(FCbObject Object)
+void FMPCollectorClientTickContext::AddMessage(FCbObject Object)
 {
 	Messages.Add(Object);
 }
 
-uint8 IMPCollector::FClientContext::PlatformToInt(const ITargetPlatform* Platform) const
+uint8 FMPCollectorClientTickContext::PlatformToInt(const ITargetPlatform* Platform) const
 {
 	return PlatformToIntImpl(Platform, Platforms);
 }
 
-const ITargetPlatform* IMPCollector::FClientContext::IntToPlatform(uint8 PlatformAsInt) const
+const ITargetPlatform* FMPCollectorClientTickContext::IntToPlatform(uint8 PlatformAsInt) const
 {
 	return IntToPlatformImpl(PlatformAsInt, Platforms);
 }
 
-uint8 IMPCollector::FServerContext::PlatformToInt(const ITargetPlatform* Platform) const
+uint8 FMPCollectorClientMessageContext::PlatformToInt(const ITargetPlatform* Platform) const
 {
 	return PlatformToIntImpl(Platform, Platforms);
 }
 
-const ITargetPlatform* IMPCollector::FServerContext::IntToPlatform(uint8 PlatformAsInt) const
+const ITargetPlatform* FMPCollectorClientMessageContext::IntToPlatform(uint8 PlatformAsInt) const
 {
 	return IntToPlatformImpl(PlatformAsInt, Platforms);
 }
 
-void IMPCollector::ReceiveMessage(FServerContext& Context, FCbObjectView Message)
+uint8 FMPCollectorServerMessageContext::PlatformToInt(const ITargetPlatform* Platform) const
 {
+	return PlatformToIntImpl(Platform, Platforms);
+}
 
+const ITargetPlatform* FMPCollectorServerMessageContext::IntToPlatform(uint8 PlatformAsInt) const
+{
+	return IntToPlatformImpl(PlatformAsInt, Platforms);
 }
 
 }
diff --git a/Engine/Source/Editor/UnrealEd/Private/Cooker/CookMPCollector.h b/Engine/Source/Editor/UnrealEd/Private/Cooker/CookMPCollector.h
--- a/Engine/Source/Editor/UnrealEd/Private/Cooker/CookMPCollector.h
+++ b/Engine/Source/Editor/UnrealEd/Private/Cooker/CookMPCollector.h
@@ -4,9 +4,11 @@
 
 #include "Containers/Array.h"
 #include "Containers/ArrayView.h"
+#include "Cooker/CompactBinaryTCP.h"
 #include "Cooker/CookTypes.h"
 #include "HAL/Platform.h"
 #include "Misc/Guid.h"
+#include "Templates/Function.h"
 #include "Templates/RefCounting.h"
 
 class FCbObject;
@@ -14,11 +16,63 @@ class FCbObjectView;
 class FCbWriter;
 class ITargetPlatform;
 
+namespace UE::Cook { class FCookWorkerClient; }
 namespace UE::Cook { class FCookWorkerServer; }
 
 namespace UE::Cook
 {
 
+class FMPCollectorClientTickContext
+{
+public:
+	TConstArrayView<const ITargetPlatform*> GetPlatforms() const { return Platforms; }
+	bool IsFlush()  const { return bFlush; }
+
+	void AddMessage(FCbObject Object);
+
+	uint8 PlatformToInt(const ITargetPlatform* Platform) const;
+	const ITargetPlatform* IntToPlatform(uint8 PlatformAsInt) const;
+
+private:
+	TConstArrayView<const ITargetPlatform*> Platforms;
+	TArray<FCbObject> Messages;
+	bool bFlush = false;
+
+	friend class FCookWorkerClient;
+};
+
+class FMPCollectorClientMessageContext
+{
+public:
+	TConstArrayView<const ITargetPlatform*> GetPlatforms() { return Platforms; }
+
+	uint8 PlatformToInt(const ITargetPlatform* Platform) const;
+	const ITargetPlatform* IntToPlatform(uint8 PlatformAsInt) const;
+
+private:
+	TConstArrayView<const ITargetPlatform*> Platforms;
+
+	friend class FCookWorkerClient;
+};
+
+class FMPCollectorServerMessageContext
+{
+public:
+	TConstArrayView<const ITargetPlatform*> GetPlatforms() { return Platforms; }
+
+	uint8 PlatformToInt(const ITargetPlatform* Platform) const;
+	const ITargetPlatform* IntToPlatform(uint8 PlatformAsInt) const;
+	FWorkerId GetWorkerId() const { return WorkerId; }
+	int32 GetProfileId() const { return ProfileId; }
+
+private:
+	TConstArrayView<const ITargetPlatform*> Platforms;
+	FWorkerId WorkerId;
+	int32 ProfileId;
+
+	friend class FCookWorkerServer;
+};
+
 /**
  * Interface used during cooking to send data collected from save/load on a remote CookWorker
  * to the Director for aggregation into files saves at the end of the cook. 
@@ -26,50 +80,102 @@ namespace UE::Cook
 class IMPCollector : public FRefCountBase
 {
 public:
-	class FClientContext
-	{
-	public:
-		TConstArrayView<const ITargetPlatform*> GetPlatforms() const { return Platforms; }
-		bool IsFlush()  const { return bFlush; }
+	virtual ~IMPCollector() {}
 
-		void AddMessage(FCbObject Object);
+	virtual FGuid GetMessageType() const = 0;
+	virtual const TCHAR* GetDebugName() const = 0;
 
-		uint8 PlatformToInt(const ITargetPlatform* Platform) const;
-		const ITargetPlatform* IntToPlatform(uint8 PlatformAsInt) const;
+	virtual void ClientTick(FMPCollectorClientTickContext& Context) {}
+	virtual void ClientReceiveMessage(FMPCollectorClientMessageContext& Context, FCbObjectView Message) {}
+	virtual void ServerReceiveMessage(FMPCollectorServerMessageContext& Context, FCbObjectView Message) {}
+};
 
-	private:
-		TConstArrayView<const ITargetPlatform*> Platforms;
-		TArray<FCbObject> Messages;
-		bool bFlush = false;
+/** A subinterface of IMPCollector that uses a UE::CompactBinaryTCP::IMessage subclass to serialize the message. */
+template <typename CompactBinaryTCPMessageType>
+class IMPCollectorCbMessage : public IMPCollector
+{
+public:
+	virtual void ClientReceiveMessage(FMPCollectorClientMessageContext& Context, bool bReadSuccessful,
+		CompactBinaryTCPMessageType&& Message) {}
+	virtual void ServerReceiveMessage(FMPCollectorServerMessageContext& Context, bool bReadSuccessful,
+		CompactBinaryTCPMessageType&& Message) {}
 
-		friend class FCookWorkerClient;
-	};
+	virtual FGuid GetMessageType() const override
+	{
+		return CompactBinaryTCPMessageType::MessageType;
+	}
 
-	class FServerContext
+	virtual void ClientReceiveMessage(FMPCollectorClientMessageContext& Context, FCbObjectView Message) override
 	{
-	public:
-		TConstArrayView<const ITargetPlatform*> GetPlatforms() { return Platforms; }
+		CompactBinaryTCPMessageType CBMessage;
+		bool bReadSuccessful = CBMessage.TryRead(Message);
+		ClientReceiveMessage(Context, bReadSuccessful, MoveTemp(CBMessage));
+	}
 
-		uint8 PlatformToInt(const ITargetPlatform* Platform) const;
-		const ITargetPlatform* IntToPlatform(uint8 PlatformAsInt) const;
-		FWorkerId GetWorkerId() const { return WorkerId; }
-		int32 GetProfileId() const { return ProfileId; }
+	virtual void ServerReceiveMessage(FMPCollectorServerMessageContext& Context, FCbObjectView Message) override
+	{
+		CompactBinaryTCPMessageType CBMessage;
+		bool bReadSuccessful = CBMessage.TryRead(Message);
+		ServerReceiveMessage(Context, bReadSuccessful, MoveTemp(CBMessage));
+	}
+};
 
-	private:
-		TConstArrayView<const ITargetPlatform*> Platforms;
-		FWorkerId WorkerId;
-		int32 ProfileId;
+/**
+ * An implementation of IMPCollector that uses a UE::CompactBinaryTCP::IMessage subclass to serialize the message,
+ * and directs messages received on the client to the given callback.
+ */
+template <typename CompactBinaryTCPMessageType>
+class IMPCollectorCbClientMessage : public IMPCollectorCbMessage<CompactBinaryTCPMessageType>
+{
+public:
+	IMPCollectorCbClientMessage(TUniqueFunction<void(FMPCollectorClientMessageContext& Context, bool bReadSuccessful,
+		CompactBinaryTCPMessageType&& Message)>&& InCallback, const TCHAR* InDebugName)
+		: Callback(MoveTemp(InCallback))
+		, DebugName(InDebugName)
+	{
+		check(Callback);
+	}
 
-		friend class FCookWorkerServer;
-	};
+	virtual const TCHAR* GetDebugName() const override { return DebugName; }
 
-	virtual ~IMPCollector() {}
+	virtual void ClientReceiveMessage(FMPCollectorClientMessageContext& Context, bool bReadSuccessful,
+		CompactBinaryTCPMessageType&& Message) override
+	{
+		Callback(Context, bReadSuccessful, MoveTemp(Message));
+	}
 
-	virtual FGuid GetMessageType() const = 0;
-	virtual const TCHAR* GetDebugName() const = 0;
+private:
+	TUniqueFunction<void(FMPCollectorClientMessageContext& Context, bool bReadSuccessful, CompactBinaryTCPMessageType&& Message)> Callback;
+	const TCHAR* DebugName;
+};
+
+/**
+ * An implementation of IMPCollector that uses a UE::CompactBinaryTCP::IMessage subclass to serialize the message,
+ * and directs messages received on the server to the given callback.
+ */
+template <typename CompactBinaryTCPMessageType>
+class IMPCollectorCbServerMessage : public IMPCollectorCbMessage<CompactBinaryTCPMessageType>
+{
+public:
+	IMPCollectorCbServerMessage(TUniqueFunction<void(FMPCollectorServerMessageContext& Context, bool bReadSuccessful,
+		CompactBinaryTCPMessageType&& Message)>&& InCallback, const TCHAR* InDebugName)
+		: Callback(MoveTemp(InCallback))
+		, DebugName(InDebugName)
+	{
+		check(Callback);
+	}
+
+	virtual const TCHAR* GetDebugName() const override { return DebugName; }
+
+	virtual void ServerReceiveMessage(FMPCollectorServerMessageContext& Context, bool bReadSuccessful,
+		CompactBinaryTCPMessageType&& Message) override
+	{
+		Callback(Context, bReadSuccessful, MoveTemp(Message));
+	}
 
-	virtual void ClientTick(FClientContext& Context) {}
-	virtual void ReceiveMessage(FServerContext& Context, FCbObjectView Message);
+private:
+	TUniqueFunction<void(FMPCollectorServerMessageContext& Context, bool bReadSuccessful, CompactBinaryTCPMessageType&& Message)> Callback;
+	const TCHAR* DebugName;
 };
 
 }
\ No newline at end of file
diff --git a/Engine/Source/Editor/UnrealEd/Private/Cooker/CookWorkerClient.cpp b/Engine/Source/Editor/UnrealEd/Private/Cooker/CookWorkerClient.cpp
--- a/Engine/Source/Editor/UnrealEd/Private/Cooker/CookWorkerClient.cpp
+++ b/Engine/Source/Editor/UnrealEd/Private/Cooker/CookWorkerClient.cpp
@@ -30,6 +30,11 @@ FCookWorkerClient::FCookWorkerClient(UCookOnTheFlyServer& InCOTFS)
 	FLogMessagesMessageHandler* Handler = new FLogMessagesMessageHandler();
 	Handler->InitializeClient();
 	Register(Handler);
+	Register(new IMPCollectorCbClientMessage<FRetractionRequestMessage>([this]
+	(FMPCollectorClientMessageContext& Context, bool bReadSuccessful, FRetractionRequestMessage&& Message)
+		{
+			HandleRetractionMessage(Context, bReadSuccessful, MoveTemp(Message));
+		}, TEXT("HandleRetractionMessage")));
 }
 
 FCookWorkerClient::~FCookWorkerClient()
@@ -354,7 +359,7 @@ void FCookWorkerClient::PollReceiveConfigMessage()
 	}
 	check(!InitialConfigMessage);
 	InitialConfigMessage = MakeUnique<FInitialConfigMessage>();
-	if (!InitialConfigMessage->TryRead(MoveTemp(Messages[0].Object)))
+	if (!InitialConfigMessage->TryRead(Messages[0].Object))
 	{
 		UE_LOG(LogCook, Warning, TEXT("CookWorker initialization failure: Director sent an invalid InitialConfigMessage."));
 		SendToState(EConnectStatus::LostConnection);
@@ -436,7 +441,7 @@ void FCookWorkerClient::HandleReceiveMessages(TArray<UE::CompactBinaryTCP::FMars
 			if (Message.MessageType == FAbortWorkerMessage::MessageType)
 			{
 				FAbortWorkerMessage AbortMessage;
-				AbortMessage.TryRead(MoveTemp(Message.Object));
+				AbortMessage.TryRead(Message.Object);
 				if (AbortMessage.Type == FAbortWorkerMessage::EType::CookComplete)
 				{
 					UE_LOG(LogCook, Display, TEXT("CookWorkerClient received CookComplete message from Director. Flushing messages and shutting down."));
@@ -456,7 +461,7 @@ void FCookWorkerClient::HandleReceiveMessages(TArray<UE::CompactBinaryTCP::FMars
 			else if (Message.MessageType == FAssignPackagesMessage::MessageType)
 			{
 				FAssignPackagesMessage AssignPackagesMessage;
-				if (!AssignPackagesMessage.TryRead(MoveTemp(Message.Object)))
+				if (!AssignPackagesMessage.TryRead(Message.Object))
 				{
 					LogInvalidMessage(TEXT("FAssignPackagesMessage"));
 				}
@@ -465,6 +470,22 @@ void FCookWorkerClient::HandleReceiveMessages(TArray<UE::CompactBinaryTCP::FMars
 					AssignPackages(AssignPackagesMessage);
 				}
 			}
+			else
+			{
+				TRefCountPtr<IMPCollector>* Collector = Collectors.Find(Message.MessageType);
+				if (Collector)
+				{
+					check(*Collector);
+					FMPCollectorClientMessageContext Context;
+					Context.Platforms = OrderedSessionPlatforms;
+					(*Collector)->ClientReceiveMessage(Context, Message.Object);
+				}
+				else
+				{
+					UE_LOG(LogCook, Error, TEXT("CookWorkerClient received message of unknown type %s from CookDirector. Ignoring it."),
+						*Message.MessageType.ToString());
+				}
+			}
 		}
 	}
 }
@@ -579,12 +600,26 @@ void FCookWorkerClient::AssignPackages(FAssignPackagesMessage& Message)
 
 void FCookWorkerClient::Register(IMPCollector* Collector)
 {
-	CollectorsToTick.AddUnique(Collector);
+	TRefCountPtr<IMPCollector>& Existing = Collectors.FindOrAdd(Collector->GetMessageType());
+	if (Existing)
+	{
+		UE_LOG(LogCook, Error, TEXT("Duplicate IMPCollectors registered. Guid: %s, Existing: %s, Registering: %s. Keeping the Existing."),
+			*Collector->GetMessageType().ToString(), Existing->GetDebugName(), Collector->GetDebugName());
+		return;
+	}
+	Existing = Collector;
 }
 
 void FCookWorkerClient::Unregister(IMPCollector* Collector)
 {
-	CollectorsToTick.RemoveSwap(Collector);
+	TRefCountPtr<IMPCollector> Existing;
+	Collectors.RemoveAndCopyValue(Collector->GetMessageType(), Existing);
+	if (Existing && Existing.GetReference() != Collector)
+	{
+		UE_LOG(LogCook, Error, TEXT("Duplicate IMPCollector during Unregister. Guid: %s, Existing: %s, Unregistering: %s. Ignoring the Unregister."),
+			*Collector->GetMessageType().ToString(), Existing->GetDebugName(), Collector->GetDebugName());
+		Collectors.Add(Collector->GetMessageType(), MoveTemp(Existing));
+	}
 }
 
 void FCookWorkerClient::TickCollectors(FTickStackData& StackData, bool bFlush)
@@ -594,15 +629,16 @@ void FCookWorkerClient::TickCollectors(FTickStackData& StackData, bool bFlush)
 		return;
 	}
 
-	if (!CollectorsToTick.IsEmpty())
+	if (!Collectors.IsEmpty())
 	{
-		IMPCollector::FClientContext Context;
+		FMPCollectorClientTickContext Context;
 		Context.Platforms = OrderedSessionPlatforms;
 		Context.bFlush = bFlush;
 		TArray<UE::CompactBinaryTCP::FMarshalledMessage> MarshalledMessages;
 
-		for (IMPCollector* Collector : CollectorsToTick)
+		for (const TPair<FGuid, TRefCountPtr<IMPCollector>>& Pair: Collectors)
 		{
+			IMPCollector* Collector = Pair.Value.GetReference();
 			Collector->ClientTick(Context);
 			if (!Context.Messages.IsEmpty())
 			{
@@ -625,4 +661,10 @@ void FCookWorkerClient::TickCollectors(FTickStackData& StackData, bool bFlush)
 	NextTickCollectorsTimeSeconds = FPlatformTime::Seconds() + TickCollectorsPeriodSeconds;
 }
 
-}
\ No newline at end of file
+void FCookWorkerClient::HandleRetractionMessage(FMPCollectorClientMessageContext& Context, bool bReadSuccessful,
+	FRetractionRequestMessage&& Message)
+{
+
+}
+
+} // namespace UE::Cook
\ No newline at end of file
diff --git a/Engine/Source/Editor/UnrealEd/Private/Cooker/CookWorkerClient.h b/Engine/Source/Editor/UnrealEd/Private/Cooker/CookWorkerClient.h
--- a/Engine/Source/Editor/UnrealEd/Private/Cooker/CookWorkerClient.h
+++ b/Engine/Source/Editor/UnrealEd/Private/Cooker/CookWorkerClient.h
@@ -7,13 +7,16 @@
 #include "CookOnTheSide/CookOnTheFlyServer.h"
 #include "CookTypes.h"
 #include "IPAddress.h"
+#include "Misc/Guid.h"
 
+namespace UE::Cook { class FMPCollectorClientMessageContext; }
 namespace UE::Cook { class IMPCollector; }
 namespace UE::Cook { struct FAssignPackagesMessage; }
 namespace UE::Cook { struct FDirectorConnectionInfo; }
 namespace UE::Cook { struct FDiscoveredPackage; }
 namespace UE::Cook { struct FInitialConfigMessage; }
 namespace UE::Cook { struct FPackageRemoteResult; }
+namespace UE::Cook { struct FRetractionRequestMessage; }
 
 namespace UE::Cook
 {
@@ -104,7 +107,8 @@ private:
 	void AssignPackages(FAssignPackagesMessage& Message);
 	/** Tick the registered collectors. */
 	void TickCollectors(FTickStackData& StackData, bool bFlush);
-
+	void HandleRetractionMessage(FMPCollectorClientMessageContext& Context, bool bReadSuccessful,
+		FRetractionRequestMessage&& Message);
 
 private:
 	TSharedPtr<FInternetAddr> DirectorAddr;
@@ -112,7 +116,7 @@ private:
 	TArray<ITargetPlatform*> OrderedSessionPlatforms;
 	TArray<FPackageRemoteResult> PendingResults;
 	TArray<FDiscoveredPackage> PendingDiscoveredPackages;
-	TArray<TRefCountPtr<IMPCollector>> CollectorsToTick;
+	TMap<FGuid, TRefCountPtr<IMPCollector>> Collectors;
 	UE::CompactBinaryTCP::FSendBuffer SendBuffer;
 	UE::CompactBinaryTCP::FReceiveBuffer ReceiveBuffer;
 	FString DirectorURI;
diff --git a/Engine/Source/Editor/UnrealEd/Private/Cooker/CookWorkerServer.cpp b/Engine/Source/Editor/UnrealEd/Private/Cooker/CookWorkerServer.cpp
--- a/Engine/Source/Editor/UnrealEd/Private/Cooker/CookWorkerServer.cpp
+++ b/Engine/Source/Editor/UnrealEd/Private/Cooker/CookWorkerServer.cpp
@@ -468,7 +468,7 @@ void FCookWorkerServer::HandleReceiveMessagesInternal()
 		if (Message.MessageType == FPackageResultsMessage::MessageType)
 		{
 			FPackageResultsMessage ResultsMessage;
-			if (!ResultsMessage.TryRead(MoveTemp(Message.Object)))
+			if (!ResultsMessage.TryRead(Message.Object))
 			{
 				LogInvalidMessage(TEXT("FPackageResultsMessage"));
 			}
@@ -480,7 +480,7 @@ void FCookWorkerServer::HandleReceiveMessagesInternal()
 		else if (Message.MessageType == FDiscoveredPackagesMessage::MessageType)
 		{
 			FDiscoveredPackagesMessage DiscoveredMessage;
-			if (!DiscoveredMessage.TryRead(MoveTemp(Message.Object)))
+			if (!DiscoveredMessage.TryRead(Message.Object))
 			{
 				LogInvalidMessage(TEXT("FDiscoveredPackagesMessage"));
 			}
@@ -494,15 +494,15 @@ void FCookWorkerServer::HandleReceiveMessagesInternal()
 		}
 		else
 		{
-			TRefCountPtr<IMPCollector>* Collector = Director.MessageHandlers.Find(Message.MessageType);
+			TRefCountPtr<IMPCollector>* Collector = Director.Collectors.Find(Message.MessageType);
 			if (Collector)
 			{
 				check(*Collector);
-				IMPCollector::FServerContext Context;
+				FMPCollectorServerMessageContext Context;
 				Context.Platforms = OrderedSessionPlatforms;
 				Context.WorkerId = WorkerId;
 				Context.ProfileId = ProfileId;
-				(*Collector)->ReceiveMessage(Context, Message.Object);
+				(*Collector)->ServerReceiveMessage(Context, Message.Object);
 			}
 			else
 			{
@@ -522,7 +522,7 @@ void FCookWorkerServer::HandleReceivedPackagePlatformMessages(FPackageData& Pack
 		if (Message.MessageType == FAssetRegistryPackageMessage::MessageType)
 		{
 			FAssetRegistryPackageMessage ARMessage;
-			if (!ARMessage.TryRead(MoveTemp(Message.Object), PackageData, TargetPlatform))
+			if (!ARMessage.TryRead(Message.Object, PackageData, TargetPlatform))
 			{
 				LogInvalidMessage(TEXT("FAssetRegistryPackageMessage"));
 			}
@@ -662,7 +662,7 @@ void FAssignPackagesMessage::Write(FCbWriter& Writer) const
 	Writer << "P" << PackageDatas;
 }
 
-bool FAssignPackagesMessage::TryRead(FCbObject&& Object)
+bool FAssignPackagesMessage::TryRead(FCbObjectView Object)
 {
 	return LoadFromCompactBinary(Object["P"], PackageDatas);
 }
@@ -722,7 +722,7 @@ void FAbortPackagesMessage::Write(FCbWriter& Writer) const
 	Writer << "PackageNames" <<  PackageNames;
 }
 
-bool FAbortPackagesMessage::TryRead(FCbObject&& Object)
+bool FAbortPackagesMessage::TryRead(FCbObjectView Object)
 {
 	return LoadFromCompactBinary(Object["PackageNames"], PackageNames);
 }
@@ -739,7 +739,7 @@ void FAbortWorkerMessage::Write(FCbWriter& Writer) const
 	Writer << "Type" << (uint8)Type;
 }
 
-bool FAbortWorkerMessage::TryRead(FCbObject&& Object)
+bool FAbortWorkerMessage::TryRead(FCbObjectView Object)
 {
 	Type = static_cast<EType>(Object["Type"].AsUInt8((uint8)EType::Abort));
 	return true;
@@ -783,7 +783,7 @@ void FInitialConfigMessage::Write(FCbWriter& Writer) const
 	Writer << "CookOnTheFlyOptions" << CookOnTheFlyOptions;
 }
 
-bool FInitialConfigMessage::TryRead(FCbObject&& Object)
+bool FInitialConfigMessage::TryRead(FCbObjectView Object)
 {
 	bool bOk = true;
 	int32 LocalCookMode;
@@ -873,7 +873,7 @@ void FDiscoveredPackagesMessage::Write(FCbWriter& Writer) const
 	Writer << "Packages" << Packages;
 }
 
-bool FDiscoveredPackagesMessage::TryRead(FCbObject&& Object)
+bool FDiscoveredPackagesMessage::TryRead(FCbObjectView Object)
 {
 	return LoadFromCompactBinary(Object["Packages"], Packages);
 }
@@ -926,7 +926,7 @@ void FLogMessagesMessageHandler::InitializeClient()
 	bRegistered = true;
 }
 
-void FLogMessagesMessageHandler::ClientTick(FClientContext& Context)
+void FLogMessagesMessageHandler::ClientTick(FMPCollectorClientTickContext& Context)
 {
 	{
 		FScopeLock QueueScopeLock(&QueueLock);
@@ -943,7 +943,7 @@ void FLogMessagesMessageHandler::ClientTick(FClientContext& Context)
 	}
 }
 
-void FLogMessagesMessageHandler::ReceiveMessage(FServerContext& Context, FCbObjectView InMessage)
+void FLogMessagesMessageHandler::ServerReceiveMessage(FMPCollectorServerMessageContext& Context, FCbObjectView InMessage)
 {
 	TArray<FReplicatedLogData> Messages;
 	if (!LoadFromCompactBinary(InMessage["Messages"], Messages))
diff --git a/Engine/Source/Editor/UnrealEd/Private/Cooker/CookWorkerServer.h b/Engine/Source/Editor/UnrealEd/Private/Cooker/CookWorkerServer.h
--- a/Engine/Source/Editor/UnrealEd/Private/Cooker/CookWorkerServer.h
+++ b/Engine/Source/Editor/UnrealEd/Private/Cooker/CookWorkerServer.h
@@ -188,7 +188,7 @@ public:
 	FAssignPackagesMessage(TArray<FAssignPackageData>&& InPackageDatas);
 
 	virtual void Write(FCbWriter& Writer) const override;
-	virtual bool TryRead(FCbObject&& Object) override;
+	virtual bool TryRead(FCbObjectView Object) override;
 	virtual FGuid GetMessageType() const override { return MessageType; }
 
 public:
@@ -204,7 +204,7 @@ public:
 	FAbortPackagesMessage(TArray<FName>&& InPackageNames);
 
 	virtual void Write(FCbWriter& Writer) const override;
-	virtual bool TryRead(FCbObject&& Object) override;
+	virtual bool TryRead(FCbObjectView Object) override;
 	virtual FGuid GetMessageType() const override { return MessageType; }
 
 public:
@@ -228,7 +228,7 @@ public:
 	};
 	FAbortWorkerMessage(EType InType = EType::Abort);
 	virtual void Write(FCbWriter& Writer) const override;
-	virtual bool TryRead(FCbObject&& Object) override;
+	virtual bool TryRead(FCbObjectView Object) override;
 	virtual FGuid GetMessageType() const override { return MessageType; }
 
 public:
@@ -241,7 +241,7 @@ struct FInitialConfigMessage : public UE::CompactBinaryTCP::IMessage
 {
 public:
 	virtual void Write(FCbWriter& Writer) const override;
-	virtual bool TryRead(FCbObject&& Object) override;
+	virtual bool TryRead(FCbObjectView Object) override;
 	virtual FGuid GetMessageType() const override { return MessageType; }
 
 	void ReadFromLocal(const UCookOnTheFlyServer& COTFS, const TArray<ITargetPlatform*>& InOrderedSessionPlatforms,
@@ -290,7 +290,7 @@ struct FDiscoveredPackagesMessage : public UE::CompactBinaryTCP::IMessage
 {
 public:
 	virtual void Write(FCbWriter& Writer) const override;
-	virtual bool TryRead(FCbObject&& Object) override;
+	virtual bool TryRead(FCbObjectView Object) override;
 	virtual FGuid GetMessageType() const override { return MessageType; }
 
 public:
@@ -321,8 +321,8 @@ public:
 	// IMPCollector
 	virtual FGuid GetMessageType() const override { return MessageType; }
 	virtual const TCHAR* GetDebugName() const override { return TEXT("FLogMessagesMessageHandler"); }
-	virtual void ClientTick(FClientContext& Context) override;
-	virtual void ReceiveMessage(FServerContext& Context, FCbObjectView Message) override;
+	virtual void ClientTick(FMPCollectorClientTickContext& Context) override;
+	virtual void ServerReceiveMessage(FMPCollectorServerMessageContext& Context, FCbObjectView Message) override;
 
 	// FOutputDevice
 	virtual void Serialize(const TCHAR* V, ELogVerbosity::Type Verbosity, const FName& Category) override;
diff --git a/Engine/Source/Editor/UnrealEd/Private/Cooker/PackageResultsMessage.cpp b/Engine/Source/Editor/UnrealEd/Private/Cooker/PackageResultsMessage.cpp
--- a/Engine/Source/Editor/UnrealEd/Private/Cooker/PackageResultsMessage.cpp
+++ b/Engine/Source/Editor/UnrealEd/Private/Cooker/PackageResultsMessage.cpp
@@ -69,12 +69,12 @@ void FPackageResultsMessage::Write(FCbWriter& Writer) const
 	Writer.EndArray();
 }
 
-bool FPackageResultsMessage::TryRead(FCbObject&& Object)
+bool FPackageResultsMessage::TryRead(FCbObjectView Object)
 {
 	Results.Reset();
-	for (FCbField ResultField : Object["R"])
+	for (FCbFieldView ResultField : Object["R"])
 	{
-		FCbObject ResultObject = ResultField.AsObject();
+		FCbObjectView ResultObject = ResultField.AsObjectView();
 		FPackageRemoteResult& Result = Results.Emplace_GetRef();
 		LoadFromCompactBinary(ResultObject["N"], Result.PackageName);
 		if (Result.PackageName.IsNone())
@@ -89,19 +89,19 @@ bool FPackageResultsMessage::TryRead(FCbObject&& Object)
 		Result.SuppressCookReason = static_cast<ESuppressCookReason>(LocalSuppressCookReason);
 		Result.bReferencedOnlyByEditorOnlyData = ResultObject["E"].AsBool();
 		Result.Platforms.Reset();
-		for (FCbField PlatformField : ResultObject["P"])
+		for (FCbFieldView PlatformField : ResultObject["P"])
 		{
 			FPackageRemoteResult::FPlatformResult& PlatformResult = Result.Platforms.Emplace_GetRef();
-			FCbObject PlatformObject = PlatformField.AsObject();
+			FCbObjectView PlatformObject = PlatformField.AsObjectView();
 
 			PlatformResult.bSuccessful = PlatformObject["S"].AsBool();
 			PlatformResult.PackageGuid = PlatformObject["G"].AsUuid();
-			PlatformResult.TargetDomainDependencies = PlatformObject["D"].AsObject();
+			PlatformResult.TargetDomainDependencies = FCbObject::Clone(PlatformObject["D"].AsObjectView());
 
 			// We read a nonhomogenous array of length 2N. 2N+0 is the Message type, 2N+1 is the message object.
-			FCbArray MessagesArray = PlatformObject["M"].AsArray();
+			FCbArrayView MessagesArray = PlatformObject["M"].AsArrayView();
 			PlatformResult.Messages.Reserve(MessagesArray.Num() / 2);
-			FCbFieldIterator MessageField = MessagesArray.CreateIterator();
+			FCbFieldViewIterator MessageField = MessagesArray.CreateViewIterator();
 			while (MessageField)
 			{
 				UE::CompactBinaryTCP::FMarshalledMessage& Message = PlatformResult.Messages.Emplace_GetRef();
@@ -111,7 +111,7 @@ bool FPackageResultsMessage::TryRead(FCbObject&& Object)
 					return false;
 				}
 				++MessageField;
-				Message.Object = MessageField->AsObject();
+				Message.Object = FCbObject::Clone(MessageField->AsObjectView());
 				if (MessageField.HasError())
 				{
 					return false;
diff --git a/Engine/Source/Editor/UnrealEd/Private/Cooker/PackageResultsMessage.h b/Engine/Source/Editor/UnrealEd/Private/Cooker/PackageResultsMessage.h
--- a/Engine/Source/Editor/UnrealEd/Private/Cooker/PackageResultsMessage.h
+++ b/Engine/Source/Editor/UnrealEd/Private/Cooker/PackageResultsMessage.h
@@ -26,7 +26,7 @@ public:
 	/** Marshall the message to a CompactBinaryObject. */
 	virtual void Write(FCbWriter& Writer, const FPackageData& PackageData, const ITargetPlatform* TargetPlatform) const = 0;
 	/** Unmarshall the message from a CompactBinaryObject. */
-	virtual bool TryRead(FCbObject&& Object, FPackageData& PackageData, const ITargetPlatform* TargetPlatform) = 0;
+	virtual bool TryRead(FCbObjectView Object, FPackageData& PackageData, const ITargetPlatform* TargetPlatform) = 0;
 	/** Return the Guid that identifies the message to the remote connection. */
 	virtual FGuid GetMessageType() const = 0;
 };
@@ -65,7 +65,7 @@ struct FPackageResultsMessage : public UE::CompactBinaryTCP::IMessage
 {
 public:
 	virtual void Write(FCbWriter& Writer) const override;
-	virtual bool TryRead(FCbObject&& Object) override;
+	virtual bool TryRead(FCbObjectView Object) override;
 	virtual FGuid GetMessageType() const override { return MessageType; }
 
 public:
diff --git a/Engine/Source/Editor/UnrealEd/Public/Commandlets/AssetRegistryGenerator.h b/Engine/Source/Editor/UnrealEd/Public/Commandlets/AssetRegistryGenerator.h
--- a/Engine/Source/Editor/UnrealEd/Public/Commandlets/AssetRegistryGenerator.h
+++ b/Engine/Source/Editor/UnrealEd/Public/Commandlets/AssetRegistryGenerator.h
@@ -474,7 +474,7 @@ class FAssetRegistryPackageMessage : public IPackageMessage
 {
 public:
 	virtual void Write(FCbWriter& Writer, const FPackageData& PackageData, const ITargetPlatform* TargetPlatform) const override;
-	virtual bool TryRead(FCbObject&& Object, FPackageData& PackageData, const ITargetPlatform* TargetPlatform) override;
+	virtual bool TryRead(FCbObjectView Object, FPackageData& PackageData, const ITargetPlatform* TargetPlatform) override;
 	virtual FGuid GetMessageType() const override { return MessageType; }
 
 	TArray<FAssetData> AssetDatas;
