MPCook: CookWorkers need to initialize SessionStartupObjects so that garbage collection diagnostics exclude those objects from the results.
#rnx

diff --git a/Engine/Source/Editor/UnrealEd/Private/CookOnTheFlyServer.cpp b/Engine/Source/Editor/UnrealEd/Private/CookOnTheFlyServer.cpp
--- a/Engine/Source/Editor/UnrealEd/Private/CookOnTheFlyServer.cpp
+++ b/Engine/Source/Editor/UnrealEd/Private/CookOnTheFlyServer.cpp
@@ -9607,6 +9607,14 @@ FBeginCookContext UCookOnTheFlyServer::CreateCookWorkerContext()
 		// PlatformContext.PlatformData is currently null and is set in SelectSessionPlatforms
 	}
 
+	TArray<FWeakObjectPtr>& SessionStartupObjects = CookByTheBookOptions->SessionStartupObjects;
+	SessionStartupObjects.Reset();
+	for (FThreadSafeObjectIterator Iter; Iter; ++Iter)
+	{
+		SessionStartupObjects.Emplace(*Iter);
+	}
+	SessionStartupObjects.Shrink();
+
 	return BeginContext;
 }
 
