Canceling an edit on an existing theme now restores the correct previous theme colors. Duplicated themes start with a blank name instead of a generated name.

diff --git a/Engine/Source/Editor/UnrealEd/Private/Settings/EditorStyleSettingsCustomization.cpp b/Engine/Source/Editor/UnrealEd/Private/Settings/EditorStyleSettingsCustomization.cpp
--- a/Engine/Source/Editor/UnrealEd/Private/Settings/EditorStyleSettingsCustomization.cpp
+++ b/Engine/Source/Editor/UnrealEd/Private/Settings/EditorStyleSettingsCustomization.cpp
@@ -172,6 +172,13 @@ private:
 	{
 		FText OutErrorMessage;
 		const TArray<FStyleTheme> ThemeOptions = USlateThemeManager::Get().GetThemes();
+		
+		if (ThemeName.IsEmpty())
+		{
+			OutErrorMessage = LOCTEXT("ThemeNameEmpty", "The theme name cannot be empty.");
+			EditableThemeName->SetError(OutErrorMessage);
+			return false;
+		}
 
 		for (const FStyleTheme& Theme : ThemeOptions)
 		{
@@ -653,6 +660,8 @@ FReply FEditorStyleSettingsCustomization::OnDuplicateAndEditThemeClicked()
 
 	FGuid NewThemeId = USlateThemeManager::Get().DuplicateActiveTheme();
 	USlateThemeManager::Get().ApplyTheme(NewThemeId);
+	// Set the new theme name to empty FText, to avoid a generated name collision or needing to delete a template name
+	USlateThemeManager::Get().SetCurrentThemeDisplayName(FText::GetEmpty());
 	CurrentActiveThemeDisplayName = USlateThemeManager::Get().GetCurrentTheme().DisplayName.ToString();
 	OriginalThemeName = USlateThemeManager::Get().GetCurrentTheme().DisplayName.ToString(); 
 
@@ -665,11 +674,13 @@ FReply FEditorStyleSettingsCustomization::OnDuplicateAndEditThemeClicked()
 
 FReply FEditorStyleSettingsCustomization::OnEditThemeClicked()
 {
+	FGuid CurrentlyActiveTheme = USlateThemeManager::Get().GetCurrentTheme().Id;
 
 	CurrentActiveThemeDisplayName = USlateThemeManager::Get().GetCurrentTheme().DisplayName.ToString();
 	OriginalThemeName = CurrentActiveThemeDisplayName; 
-
-	OpenThemeEditorWindow(FOnThemeEditorClosed::CreateStatic(&OnThemeEditorClosed, TWeakPtr<FEditorStyleSettingsCustomization>(SharedThis(this)), FGuid(), FGuid()));
+	
+	// There is no new theme created, so just pass in the current active theme ID
+	OpenThemeEditorWindow(FOnThemeEditorClosed::CreateStatic(&OnThemeEditorClosed, TWeakPtr<FEditorStyleSettingsCustomization>(SharedThis(this)), FGuid(), CurrentlyActiveTheme));
 
 	return FReply::Handled();
 }
