EditorDomain: Add determinism test, fix indeterminism in UPackage.PersistentGuid, and bump the EditorDomain global salt.
#rnx

diff --git a/Engine/Source/Editor/UnrealEd/Private/EditorDomain/EditorDomain.cpp b/Engine/Source/Editor/UnrealEd/Private/EditorDomain/EditorDomain.cpp
--- a/Engine/Source/Editor/UnrealEd/Private/EditorDomain/EditorDomain.cpp
+++ b/Engine/Source/Editor/UnrealEd/Private/EditorDomain/EditorDomain.cpp
@@ -116,7 +116,8 @@ FEditorDomain::FEditorDomain()
 	// without needing to call ScanPathsSynchronous
 	AssetRegistry->SearchAllAssets(false /* bSynchronousSearch */);
 
-	bEditorDomainReadEnabled = !FParse::Param(FCommandLine::Get(), TEXT("noeditordomainread"));
+	bEditorDomainReadEnabled = !FParse::Param(FCommandLine::Get(), TEXT("noeditordomainread"))
+	 && !FParse::Param(FCommandLine::Get(), TEXT("testeditordomaindeterminism"));
 
 	ELoadingPhase::Type CurrentPhase = IPluginManager::Get().GetLastCompletedLoadingPhase();
 	if (CurrentPhase == ELoadingPhase::None || CurrentPhase < ELoadingPhase::PostEngineInit)
diff --git a/Engine/Source/Editor/UnrealEd/Private/EditorDomain/EditorDomainUtils.cpp b/Engine/Source/Editor/UnrealEd/Private/EditorDomain/EditorDomainUtils.cpp
--- a/Engine/Source/Editor/UnrealEd/Private/EditorDomain/EditorDomainUtils.cpp
+++ b/Engine/Source/Editor/UnrealEd/Private/EditorDomain/EditorDomainUtils.cpp
@@ -21,12 +21,14 @@
 #include "HAL/IConsoleManager.h"
 #include "Hash/Blake3.h"
 #include "Memory/SharedBuffer.h"
+#include "Misc/CommandLine.h"
 #include "Misc/CoreDelegates.h"
 #include "Misc/DelayedAutoRegister.h"
 #include "Misc/FileHelper.h"
 #include "Misc/PackageAccessTracking.h"
 #include "Misc/PackageAccessTrackingOps.h"
 #include "Misc/PackagePath.h"
+#include "Misc/Parse.h"
 #include "Misc/ScopeRWLock.h"
 #include "Misc/StringBuilder.h"
 #include "Serialization/BulkDataRegistry.h"
@@ -201,7 +203,7 @@ FBlake3Hash GGlobalConstructClassesHash;
 int64 GMaxBulkDataSize = -1;
 
 // Change to a new guid when EditorDomain needs to be invalidated
-const TCHAR* EditorDomainVersion = TEXT("0CE4958BE8484993A0399BF2CB52EF4E");
+const TCHAR* EditorDomainVersion = TEXT("6D2972C7108F4366A13B7886DEDE68D5");
 
 // Identifier of the CacheBuckets for EditorDomain tables
 const TCHAR* EditorDomainPackageBucketName = TEXT("EditorDomainPackage");
@@ -1723,6 +1725,9 @@ static FGuid IoHashToGuid(const FIoHash& Hash)
 	return FGuid(HashInts[0], HashInts[1], HashInts[2], HashInts[3]);
 }
 
+void ValidateEditorDomainDeterminism(FName PackageName, bool bStorageResultValid,
+	FEditorDomainPackageWriter* PackageWriter, UE::DerivedData::FCacheGetResponse&& GetResponse);
+
 bool TrySavePackage(UPackage* Package)
 {
 	using namespace UE::DerivedData;
@@ -2017,6 +2022,24 @@ bool TrySavePackage(UPackage* Package)
 	}
 
 	bool bStorageResultValid = StorageResult == ESaveStorageResult::Valid;
+	static bool bTestEditorDomainDeterminism = FParse::Param(FCommandLine::Get(), TEXT("testeditordomaindeterminism"));;
+	if (bTestEditorDomainDeterminism)
+	{
+		FPackagePath PackagePath;
+		FRequestOwner GetOwner(EPriority::Blocking);
+		TOptional<FCacheGetResponse> GetResponse;
+		verify(FPackagePath::TryFromPackageName(PackageName, PackagePath));
+		RequestEditorDomainPackage(PackagePath, PackageDigest.Hash, ECachePolicy::None, GetOwner,
+			[&GetResponse](FCacheGetResponse&& Response) { GetResponse.Emplace(MoveTemp(Response)); });
+		GetOwner.Wait();
+		if (!GetResponse || GetResponse->Status != EStatus::Ok)
+		{
+			return true;
+		}
+		ValidateEditorDomainDeterminism(PackageName, bStorageResultValid, PackageWriter, MoveTemp(*GetResponse));
+		return true;
+	}
+
 	TCbWriter<16> MetaData;
 	MetaData.BeginObject();
 	uint64 FileSize = PackageWriter->GetFileSize();
@@ -2077,6 +2100,71 @@ bool TrySavePackage(UPackage* Package)
 	return true;
 }
 
+void ValidateEditorDomainDeterminism(FName PackageName, bool bStorageResultValid,
+	FEditorDomainPackageWriter* PackageWriter, UE::DerivedData::FCacheGetResponse&& GetResponse)
+{
+	using namespace UE::DerivedData;
+
+	FCacheRecord& Record = GetResponse.Record;
+
+	const FCbObject& MetaData = Record.GetMeta();
+	bool bPreviousStorageResultValid = MetaData["Valid"].AsBool(false);
+	auto LogHeaderWarning = [PackageName]()
+	{
+		UE_LOG(LogEditorDomain, Warning, TEXT("Indeterminism detected in EditorDomain save of %s."), *PackageName.ToString());
+	};
+
+	if (bPreviousStorageResultValid != bStorageResultValid)
+	{
+		LogHeaderWarning();
+		UE_LOG(LogEditorDomain, Display, TEXT("\tStorageResultValid not equal. Previous: %s, Current: %s"),
+			*LexToString(bPreviousStorageResultValid), *LexToString(bStorageResultValid));
+		return;
+	}
+	if (!bStorageResultValid)
+	{
+		// Nothing further to test; we ignore the PackageWriter data and write an empty record in this case.
+		return;
+	}
+
+	TConstArrayView<FValueWithId> PreviousValues = Record.GetValues();
+	TConstArrayView<FEditorDomainPackageWriter::FAttachment> CurrentValues = PackageWriter->GetAttachments();
+	for (int32 SegmentIndex = 0; SegmentIndex < FMath::Max(PreviousValues.Num(), CurrentValues.Num()); ++SegmentIndex)
+	{
+		if (PreviousValues.Num() <= SegmentIndex || CurrentValues.Num() <= SegmentIndex)
+		{
+			LogHeaderWarning();
+			UE_LOG(LogEditorDomain, Display, TEXT("\tNumber of segments differ. Previous: %d, Current: %d"),
+				PreviousValues.Num(), CurrentValues.Num());
+			return;
+		}
+		FSharedBuffer PreviousValue = PreviousValues[SegmentIndex].GetData().Decompress();
+		FSharedBuffer CurrentValue = CurrentValues[SegmentIndex].Buffer;
+		const uint8* PreviousBytes = reinterpret_cast<const uint8*>(PreviousValue.GetData());
+		const uint8* CurrentBytes = reinterpret_cast<const uint8*>(CurrentValue.GetData());
+		uint64 MinNumBytes = FMath::Min(PreviousValue.GetSize(), CurrentValue.GetSize());
+		for (uint64 ByteIndex = 0; ByteIndex < MinNumBytes; ByteIndex++)
+		{
+			if (PreviousBytes[ByteIndex] != CurrentBytes[ByteIndex])
+			{
+				LogHeaderWarning();
+				UE_LOG(LogEditorDomain, Display, TEXT("\tBytes differ in Segment %d, Offset %" UINT64_FMT ". Previous: %d, Current: %d"),
+					SegmentIndex, ByteIndex, PreviousBytes[ByteIndex], CurrentBytes[ByteIndex]);
+				return;
+			}
+		}
+		if (MinNumBytes < PreviousValue.GetSize() || MinNumBytes < CurrentValue.GetSize())
+		{
+			LogHeaderWarning();
+			UE_LOG(LogEditorDomain, Display, TEXT("\tSegment sizes differ in Segment %d. Previous: %" UINT64_FMT ", Current: %" UINT64_FMT "."),
+				SegmentIndex, PreviousValue.GetSize(), CurrentValue.GetSize());
+			return;
+		}
+	}
+
+	// All bytes identical
+}
+
 void GetBulkDataList(FName PackageName, UE::DerivedData::IRequestOwner& Owner, TUniqueFunction<void(FSharedBuffer Buffer)>&& Callback)
 {
 	IPackageDigestCache* DigestCache = IPackageDigestCache::Get();
diff --git a/Engine/Source/Runtime/CoreUObject/Private/UObject/SavePackage2.cpp b/Engine/Source/Runtime/CoreUObject/Private/UObject/SavePackage2.cpp
--- a/Engine/Source/Runtime/CoreUObject/Private/UObject/SavePackage2.cpp
+++ b/Engine/Source/Runtime/CoreUObject/Private/UObject/SavePackage2.cpp
@@ -249,7 +249,7 @@ void PreSavePackage(FSaveContext& SaveContext)
 	// if the in memory package filename is different the filename we are saving it to,
 	// regenerate a new persistent id for it.
 	UPackage* Package = SaveContext.GetPackage();
-	if (!SaveContext.IsCooking() && !SaveContext.IsFromAutoSave() && !Package->GetLoadedPath().IsEmpty() && Package->GetLoadedPath() != SaveContext.GetTargetPackagePath())
+	if (!SaveContext.IsProceduralSave() && !SaveContext.IsFromAutoSave() && !Package->GetLoadedPath().IsEmpty() && Package->GetLoadedPath() != SaveContext.GetTargetPackagePath())
 	{
 		Package->SetPersistentGuid(FGuid::NewGuid());
 	}
