Editor: Add Max number of auto saves setting to preferences and use it to specify the man number of auto saves

diff --git a/Engine/Source/Editor/UnrealEd/Classes/Settings/EditorLoadingSavingSettings.h b/Engine/Source/Editor/UnrealEd/Classes/Settings/EditorLoadingSavingSettings.h
--- a/Engine/Source/Editor/UnrealEd/Classes/Settings/EditorLoadingSavingSettings.h
+++ b/Engine/Source/Editor/UnrealEd/Classes/Settings/EditorLoadingSavingSettings.h
@@ -162,6 +162,10 @@ public:
 	UPROPERTY(EditAnywhere, config, Category=AutoSave, meta=(DisplayName="Warning in seconds", ClampMin = "0", UIMin = "0", UIMax = "20"))
 	int32 AutoSaveWarningInSeconds;
 
+	/** How many auto save files to keep around*/
+	UPROPERTY(EditAnywhere, config, Category = AutoSave, meta = (DisplayName = "Maximum number of AutoSaves", ClampMin = "1", UIMin = "1", UIMax = "100"))
+	int32 AutoSaveMaxBackups = 10;
+
 public:
 
 	/** Whether to automatically checkout on asset modification */
diff --git a/Engine/Source/Editor/UnrealEd/Private/PackageAutoSaver.cpp b/Engine/Source/Editor/UnrealEd/Private/PackageAutoSaver.cpp
--- a/Engine/Source/Editor/UnrealEd/Private/PackageAutoSaver.cpp
+++ b/Engine/Source/Editor/UnrealEd/Private/PackageAutoSaver.cpp
@@ -192,8 +192,9 @@ void FPackageAutoSaver::AttemptAutoSave()
 			const FString AutoSaveDir = AutoSaveUtils::GetAutoSaveDir();
 			IFileManager::Get().MakeDirectory(*AutoSaveDir, true);
 
+			const int32 AutoSaveMaxBackups = LoadingSavingSettings->AutoSaveMaxBackups > 0 ? LoadingSavingSettings->AutoSaveMaxBackups : 10;
 			// Auto-save maps and/or content packages based on user settings.
-			const int32 NewAutoSaveIndex = (AutoSaveIndex + 1) % 10;
+			const int32 NewAutoSaveIndex = (AutoSaveIndex + 1) % AutoSaveMaxBackups;
 
 			bool bLevelSaved = false;
 			auto MapsSaveResults = EAutosaveContentPackagesResult::NothingToDo;
