Fix rare crash when moving directories to new locations in the content browser.

- This is very similar to the issue fixed in CL 19319379 which concerned the renaming of single packages, except this use case involves the mass renaming of many packages.
- Due to being submitted in a hot fix we cannot expose the utilities added in that changelist to PackageTools so we must create a copy of the code. Once this has merged back to main, the original code can be exposed for 5.2 and the duplication removed.
- The problem is that the call to RenameObject can cause ResetLoader to be called while async compilation is ongoing.
- We are using the same sort of utility to wait on async compilation as in PackageTools to attempt to limit how often we block here.
- As with the PackageTools fix, this code can be removed when the editor bulkdata thread safety work is completed.

diff --git a/Engine/Source/Editor/UnrealEd/Private/ObjectTools.cpp b/Engine/Source/Editor/UnrealEd/Private/ObjectTools.cpp
--- a/Engine/Source/Editor/UnrealEd/Private/ObjectTools.cpp
+++ b/Engine/Source/Editor/UnrealEd/Private/ObjectTools.cpp
@@ -158,6 +158,53 @@ void ReloadEditorWorldForReferenceReplacementIfNecessary(TArray< TWeakObjectPtr<
 	}
 }
 
+// Note that the contents of this namespace are a copy of code from PackageTools.cpp submitted as part of a 
+// release hot fix. In UE 5.2 we should merge the two sets of code and place them in a common utility file.
+namespace
+{
+
+/**
+ * Utility function that checks each UObject inside of the given UPackage to see if it is waiting
+ * on async compilation.
+ *
+ * @return true if the package contains at least one UObject that has compilation work running, otherwise false.
+ */
+static bool IsPackageCompiling(const UPackage* Package)
+{
+	bool bIsCompiling = false;
+	ForEachObjectWithPackage(Package, [&bIsCompiling](const UObject* Object)
+		{
+			const IInterface_AsyncCompilation* AsyncCompilationIF = Cast<IInterface_AsyncCompilation>(Object);
+			if (AsyncCompilationIF != nullptr && AsyncCompilationIF->IsCompiling())
+			{
+				bIsCompiling = true;
+				return false;
+			}
+			else
+			{
+				return true;
+			}
+		});
+
+	return bIsCompiling;
+}
+
+/**
+ * Utility function that checks the provided package to see if it contains
+ * any assets that currently have async compilation work running.
+ * If there are assets that are waiting on async compilation work then we
+ * wait on all currently outstanding work to finish before returning.
+ */
+static void FlushAsyncCompilation(const UPackage* Package)
+{
+	if (IsPackageCompiling(Package))
+	{
+		FAssetCompilingManager::Get().FinishAllCompilation();
+	}
+}
+
+} // namespace
+
 namespace ObjectTools
 {
 	/** Returns true if the specified object can be displayed in a content browser */
@@ -4133,6 +4180,9 @@ namespace ObjectTools
 			// Look for a thumbnail for this asset before we rename it
 			FObjectThumbnail* Thumbnail = ThumbnailTools::GetThumbnailForObject(Object);
 
+			//
+			FlushAsyncCompilation(Object->GetPackage());
+
 			FString OldObjectFullName = Object->GetFullName();
 			FString OldObjectPathName = Object->GetPathName();
 			GEditor->RenameObject( Object, NewPackage, *ObjName, bLeaveRedirector ? REN_None : REN_DontCreateRedirectors );
