Show selected refraction mode next to the Refraction pin of the root node.

diff --git a/Engine/Source/Editor/UnrealEd/Private/MaterialGraphNode_Root.cpp b/Engine/Source/Editor/UnrealEd/Private/MaterialGraphNode_Root.cpp
--- a/Engine/Source/Editor/UnrealEd/Private/MaterialGraphNode_Root.cpp
+++ b/Engine/Source/Editor/UnrealEd/Private/MaterialGraphNode_Root.cpp
@@ -111,7 +111,36 @@ void UMaterialGraphNode_Root::CreateInputPins()
 		const FMaterialInputInfo& MaterialInput = MaterialGraph->MaterialInputs[Index];
 		EMaterialProperty Property = MaterialInput.GetProperty();
 	
-		UEdGraphPin* InputPin = CreatePin(EGPD_Input, UMaterialGraphSchema::PC_MaterialInput, *FString::Printf(TEXT("%d"), (int32)Property), *MaterialInput.GetName().ToString());
+		FName MaterialInputName;
+		if (Property == MP_Refraction)
+		{
+			FString RefractionMethodStr;
+			switch (MaterialGraph->Material->RefractionMethod)
+			{
+			case ERefractionMode::RM_None:
+				RefractionMethodStr = TEXT("Disabled");
+				break;
+			case ERefractionMode::RM_IndexOfRefraction:
+				RefractionMethodStr = TEXT("Index Of Refraction");
+				break;
+			case ERefractionMode::RM_PixelNormalOffset:
+				RefractionMethodStr = TEXT("Pixel Normal Offset");
+				break;
+			case ERefractionMode::RM_2DOffset:
+				RefractionMethodStr = TEXT("2D Offset");
+				break;
+			default:
+				RefractionMethodStr = TEXT("UNKOWN ERefractionMode");
+			}
+			
+			MaterialInputName = *FString::Printf(TEXT("%s (%s)"), *MaterialInput.GetName().ToString(), *RefractionMethodStr);
+		}
+		else
+		{
+			MaterialInputName = *MaterialInput.GetName().ToString();
+		}
+
+		UEdGraphPin* InputPin = CreatePin(EGPD_Input, UMaterialGraphSchema::PC_MaterialInput, *FString::Printf(TEXT("%d"), (int32)Property), MaterialInputName);
 		InputPin->SourceIndex = Index;
 	}
 
