Make sure skeletalmesh thumbnail renderer do not stall the main thread.

#rnx

diff --git a/Engine/Source/Editor/UnrealEd/Private/ThumbnailRendering/SkeletalMeshThumbnailRenderer.cpp b/Engine/Source/Editor/UnrealEd/Private/ThumbnailRendering/SkeletalMeshThumbnailRenderer.cpp
--- a/Engine/Source/Editor/UnrealEd/Private/ThumbnailRendering/SkeletalMeshThumbnailRenderer.cpp
+++ b/Engine/Source/Editor/UnrealEd/Private/ThumbnailRendering/SkeletalMeshThumbnailRenderer.cpp
@@ -52,7 +52,7 @@ bool USkeletalMeshThumbnailRenderer::CanVisualizeAsset(UObject* Object)
 		return false;
 	}
 	USkeletalMesh* SkeletalMesh = Cast<USkeletalMesh>(Object);
-	const bool bValidRenderData = SkeletalMesh && SkeletalMesh->GetResourceForRendering();
+	const bool bValidRenderData = SkeletalMesh && !SkeletalMesh->IsCompiling() && SkeletalMesh->GetResourceForRendering();
 	return bValidRenderData;
 }
 
