Fixes to the code that handles closing and reopening open editors when any of their assets is force-reloaded.

This fixes potential crashes due to editors not being correctly detected and closed when dependant assets are reloaded.

diff --git a/Engine/Source/Editor/UnrealEd/Private/Subsystems/AssetEditorSubsystem.cpp b/Engine/Source/Editor/UnrealEd/Private/Subsystems/AssetEditorSubsystem.cpp
--- a/Engine/Source/Editor/UnrealEd/Private/Subsystems/AssetEditorSubsystem.cpp
+++ b/Engine/Source/Editor/UnrealEd/Private/Subsystems/AssetEditorSubsystem.cpp
@@ -969,30 +969,27 @@ void UAssetEditorSubsystem::HandlePackageReloaded(const EPackageReloadPhase InPa
 		for (TPair<FAssetEntry, IAssetEditorInstance*>& AssetEditorPair : OpenedAssets)
 		{
 			UObject* NewAsset = nullptr;
-			if (InPackageReloadedEvent->GetRepointedObject(AssetEditorPair.Key.ObjectPtr.Get(), NewAsset))
+			if (AssetEditorPair.Key.RawPtr && InPackageReloadedEvent->GetRepointedObject(AssetEditorPair.Key.RawPtr, NewAsset))
 			{
 				if (NewAsset)
 				{
 					PendingAssetsToOpen.AddUnique(NewAsset);
 				}
 
-				UObject* OldAsset = AssetEditorPair.Key.RawPtr; // Not validating the asset here since we'd want to close editors for garbage collected assets
+				// Not validating the asset here since we'd want to close editors for garbage collected assets
+				UObject* OldAsset = AssetEditorPair.Key.RawPtr;
 				ObjectsToClose.AddUnique(OldAsset);
 
-				if (AssetEditorPair.Key.ObjectPtr.IsValid())
+				// Gather other assets referencing reloaded asset and mark their editors to be closed too.
+				TArray<FReferencerInformation> AssetInternalReferencers, AssetExternalReferencers;
+				AssetEditorPair.Key.RawPtr->RetrieveReferencers(&AssetInternalReferencers, &AssetExternalReferencers);
+				for (const FReferencerInformation& Ref : AssetExternalReferencers)
 				{
-					// The asset being reloaded might have other assets that depend on it. Find the list of
-					// external referencers to this asset and mark them to be closed and reopened as well.
-					TArray<FReferencerInformation> AssetInternalReferencers, AssetExternalReferencers;
-					AssetEditorPair.Key.ObjectPtr.Get()->RetrieveReferencers(&AssetInternalReferencers, &AssetExternalReferencers);
-					for (const FReferencerInformation& Ref : AssetExternalReferencers)
-					{
-						ObjectsToClose.AddUnique(Ref.Referencer);
+					ObjectsToClose.AddUnique(Ref.Referencer);
 
-						if (!FindEditorsForAssetAndSubObjects(Ref.Referencer).IsEmpty())
-						{
-							PendingAssetsToOpen.AddUnique(Ref.Referencer);
-						}
+					if (!FindEditorsForAssetAndSubObjects(Ref.Referencer).IsEmpty())
+					{
+						PendingAssetsToOpen.AddUnique(Ref.Referencer);
 					}
 				}
 			}
