nop comment only

diff --git a/Engine/Source/Editor/UnrealEd/Private/Factories/EditorFactories.cpp b/Engine/Source/Editor/UnrealEd/Private/Factories/EditorFactories.cpp
--- a/Engine/Source/Editor/UnrealEd/Private/Factories/EditorFactories.cpp
+++ b/Engine/Source/Editor/UnrealEd/Private/Factories/EditorFactories.cpp
@@ -4141,9 +4141,6 @@ UObject* UTextureFactory::FactoryCreateBinary
 				Texture->CompressionSettings = TC_EditorIcon; // "UserInterface2D"
 			}
 		}
-
-		// @@CB move this to ApplyDefaultsForNewlyImportedTextures so it can be shared with Interchange, not code-duped
-		// also set LODGroup to UI ?
 	}
 			
 	static const auto CVarVirtualTexturesEnabled = IConsoleManager::Get().FindTConsoleVariableDataInt(TEXT("r.VirtualTextures"));
