Fix CIS StressTestUnity job

- The submit 23370951 required that some functionality be copied from PackageTools.cpp to ObjectTools.cpp to avoid changing the header files which is causing issues with the StressTestUnity job.
- Moved the copied code to UE::Private namespace for now, then on UE5 Main the duplicate code paths will be removed and replaced with a common utility.

diff --git a/Engine/Source/Editor/UnrealEd/Private/ObjectTools.cpp b/Engine/Source/Editor/UnrealEd/Private/ObjectTools.cpp
--- a/Engine/Source/Editor/UnrealEd/Private/ObjectTools.cpp
+++ b/Engine/Source/Editor/UnrealEd/Private/ObjectTools.cpp
@@ -160,7 +160,7 @@ void ReloadEditorWorldForReferenceReplacementIfNecessary(TArray< TWeakObjectPtr<
 
 // Note that the contents of this namespace are a copy of code from PackageTools.cpp submitted as part of a 
 // release hot fix. In UE 5.2 we should merge the two sets of code and place them in a common utility file.
-namespace
+namespace UE::Private
 {
 
 /**
@@ -197,13 +197,13 @@ static bool IsPackageCompiling(const UPackage* Package)
  */
 static void FlushAsyncCompilation(const UPackage* Package)
 {
-	if (IsPackageCompiling(Package))
+	if (UE::Private::IsPackageCompiling(Package))
 	{
 		FAssetCompilingManager::Get().FinishAllCompilation();
 	}
 }
 
-} // namespace
+} // namespace UE::Private
 
 namespace ObjectTools
 {
@@ -4181,7 +4181,7 @@ namespace ObjectTools
 			FObjectThumbnail* Thumbnail = ThumbnailTools::GetThumbnailForObject(Object);
 
 			//
-			FlushAsyncCompilation(Object->GetPackage());
+			UE::Private::FlushAsyncCompilation(Object->GetPackage());
 
 			FString OldObjectFullName = Object->GetFullName();
 			FString OldObjectPathName = Object->GetPathName();
