Renaming a package with PKG_NewlyCreated wasn't propagated to the new package.

diff --git a/Engine/Source/Editor/UnrealEd/Private/ObjectTools.cpp b/Engine/Source/Editor/UnrealEd/Private/ObjectTools.cpp
--- a/Engine/Source/Editor/UnrealEd/Private/ObjectTools.cpp
+++ b/Engine/Source/Editor/UnrealEd/Private/ObjectTools.cpp
@@ -4053,6 +4053,11 @@ namespace ObjectTools
 							NewPackage->SetPackageFlags(PKG_DisallowExport);
 						}
 
+						if (OldPackage->HasAnyPackageFlags(PKG_NewlyCreated))
+						{
+							NewPackage->SetPackageFlags(PKG_NewlyCreated);
+						}
+
 						NewPackage->SetIsExternallyReferenceable(OldPackage->IsExternallyReferenceable());
 
 						// When renaming a World Composition map, make sure to properly initialize WorldTileInfo
