Add generic finish compilation in asset compilers and use it to improve performance of UnloadPackages by partially flushing only what's required instead of everything
  - This can improve map load time under cold conditions when plugins are unloaded by BP scripts
  - This can improve time to delete or rename packages while the editor is still processing async compilations in the background

#rnx

diff --git a/Engine/Source/Editor/UnrealEd/Private/PackageTools.cpp b/Engine/Source/Editor/UnrealEd/Private/PackageTools.cpp
--- a/Engine/Source/Editor/UnrealEd/Private/PackageTools.cpp
+++ b/Engine/Source/Editor/UnrealEd/Private/PackageTools.cpp
@@ -54,6 +54,8 @@
 #include "UObject/UObjectGlobals.h"
 #include "UObject/UObjectIterator.h"
 
+#include "AssetCompilingManager.h"
+
 #include "ShaderCompiler.h"
 #include "DistanceFieldAtlas.h"
 #include "MeshCardRepresentation.h"
@@ -74,57 +76,35 @@ FDelegateHandle UPackageTools::ReachabilityCallbackHandle;
 namespace
 {
 
-/** 
- * Utility function that checks each UObject inside of the given UPackage to see if it is waiting 
- * on async compilation.
- * 
- * @return true if the package contains at least one UObject that has compilation work running, otherwise false.
- */
-static bool IsPackageCompiling(const UPackage* Package)
-{
-	bool bIsCompiling = false;
-	ForEachObjectWithPackage(Package, [&bIsCompiling](const UObject* Object)
-	{
-		const IInterface_AsyncCompilation* AsyncCompilationIF = Cast<IInterface_AsyncCompilation>(Object);
-		if (AsyncCompilationIF != nullptr && AsyncCompilationIF->IsCompiling())
-		{
-			bIsCompiling = true;
-			return false;
-		}
-		else
-		{
-			return true;
-		}
-	});
-
-	return bIsCompiling;
-}
-
 /**
- * Utility function that checks all of the provided packages to see if any
- * of them contain assets that currently have async compilation work running.
- * If there are assets that are waiting on async compilation work then we 
- * wait on all currently outstanding work to finish before returning.
+ * Utility function that gathers all async compilable objects from given packages
+ * and flush them to make sure there is no remaining async work trying to load data
+ * from said packages.
  */
 static void FlushAsyncCompilation(const TSet<UPackage*>& PackagesToUnload)
 {
-	bool bHasAsyncCompilationWork = false;
+	TArray<UObject*> ObjectsToFinish;
+
 	for (const UPackage* Package : PackagesToUnload)
 	{
-		if (IsPackageCompiling(Package))
+		ForEachObjectWithPackage(Package, [&ObjectsToFinish](UObject* Object)
 		{
-			bHasAsyncCompilationWork = true;
-			break;
-		}
+			if (const IInterface_AsyncCompilation* AsyncCompilationIF = Cast<IInterface_AsyncCompilation>(Object))
+			{
+				ObjectsToFinish.Add(Object);
+			}
+
+			return true;
+		});
 	}
 
-	if (bHasAsyncCompilationWork)
+	if (ObjectsToFinish.Num())
 	{
-		FAssetCompilingManager::Get().FinishAllCompilation();
+		FAssetCompilingManager::Get().FinishCompilationForObjects(ObjectsToFinish);
 	}
 }
 
-}
+} // anonymous namespace
 
 UPackageTools::UPackageTools(const FObjectInitializer& ObjectInitializer)
 	: Super(ObjectInitializer)
diff --git a/Engine/Source/Runtime/Engine/Private/AssetCompilingManager.cpp b/Engine/Source/Runtime/Engine/Private/AssetCompilingManager.cpp
--- a/Engine/Source/Runtime/Engine/Private/AssetCompilingManager.cpp
+++ b/Engine/Source/Runtime/Engine/Private/AssetCompilingManager.cpp
@@ -421,6 +421,47 @@ void FAssetCompilingManager::FinishAllCompilation()
 	UpdateNumRemainingAssets();
 }
 
+/**
+ * Finish compilation of the requested objects.
+ */
+void FAssetCompilingManager::FinishCompilationForObjects(TArrayView<UObject* const> InObjects)
+{
+#if WITH_EDITOR
+	if (InObjects.Num())
+	{
+		TRACE_CPUPROFILER_EVENT_SCOPE(FAssetCompilingManager::FinishCompilationForObjects);
+
+		for (IAssetCompilingManager* AssetCompilingManager : AssetCompilingManagers)
+		{
+			AssetCompilingManager->FinishCompilationForObjects(InObjects);
+		}
+
+		bool bAreObjectsStillCompiling = false;
+		for (const UObject* Object : InObjects)
+		{
+			if (const IInterface_AsyncCompilation* AsyncCompilationInterface = Cast<IInterface_AsyncCompilation>(Object))
+			{
+				if (AsyncCompilationInterface->IsCompiling())
+				{
+					bAreObjectsStillCompiling = true;
+
+					ensureMsgf(false,
+						TEXT("A finish compilation has been requested on %s but it is still being compiled, this might indicate a missing implementation of IAssetCompilingManager::FinishCompilationForObjects"),
+						*Object->GetFullName());
+
+					break;
+				}
+			}
+		}
+
+		if (bAreObjectsStillCompiling)
+		{
+			FinishAllCompilation();
+		}
+	}
+#endif
+}
+
 /**
  * Cancel any pending work and blocks until it is safe to shut down.
  */
diff --git a/Engine/Source/Runtime/Engine/Private/DistanceFieldAtlas.cpp b/Engine/Source/Runtime/Engine/Private/DistanceFieldAtlas.cpp
--- a/Engine/Source/Runtime/Engine/Private/DistanceFieldAtlas.cpp
+++ b/Engine/Source/Runtime/Engine/Private/DistanceFieldAtlas.cpp
@@ -372,6 +372,32 @@ void FDistanceFieldAsyncQueue::FinishAllCompilation()
 	BlockUntilAllBuildsComplete(); 
 }
 
+void FDistanceFieldAsyncQueue::FinishCompilationForObjects(TArrayView<UObject* const> InObjects)
+{
+	TRACE_CPUPROFILER_EVENT_SCOPE(FDistanceFieldAsyncQueue::FinishCompilationForObjects);
+
+	TSet<UStaticMesh*> StaticMeshes;
+	for (UObject* Object : InObjects)
+	{
+		if (UStaticMesh* StaticMesh = Cast<UStaticMesh>(Object))
+		{
+			StaticMeshes.Add(StaticMesh);
+		}
+		else if (UStaticMeshComponent* StaticMeshComponent = Cast<UStaticMeshComponent>(Object))
+		{
+			if (StaticMeshComponent->GetStaticMesh())
+			{
+				StaticMeshes.Add(StaticMeshComponent->GetStaticMesh());
+			}
+		}
+	}
+
+	for (UStaticMesh* StaticMesh : StaticMeshes)
+	{
+		BlockUntilBuildComplete(StaticMesh, false);
+	}
+}
+
 bool FDistanceFieldAsyncQueue::IsTaskInvalid(FAsyncDistanceFieldTask* Task) const
 {
 	return (Task->StaticMesh && Task->StaticMesh->IsUnreachable()) || (Task->GenerateSource && Task->GenerateSource->IsUnreachable());
@@ -464,6 +490,8 @@ void FDistanceFieldAsyncQueue::StartBackgroundTask(FAsyncDistanceFieldTask* Task
 
 void FDistanceFieldAsyncQueue::ProcessPendingTasks()
 {
+	TRACE_CPUPROFILER_EVENT_SCOPE(FDistanceFieldAsyncQueue::ProcessPendingTasks);
+
 	FScopeLock Lock(&CriticalSection);
 
 	for (auto It = PendingTasks.CreateIterator(); It; ++It)
diff --git a/Engine/Source/Runtime/Engine/Private/MeshCardRepresentation.cpp b/Engine/Source/Runtime/Engine/Private/MeshCardRepresentation.cpp
--- a/Engine/Source/Runtime/Engine/Private/MeshCardRepresentation.cpp
+++ b/Engine/Source/Runtime/Engine/Private/MeshCardRepresentation.cpp
@@ -399,6 +399,8 @@ void FCardRepresentationAsyncQueue::StartBackgroundTask(FAsyncCardRepresentation
 
 void FCardRepresentationAsyncQueue::ProcessPendingTasks()
 {
+	TRACE_CPUPROFILER_EVENT_SCOPE(FCardRepresentationAsyncQueue::ProcessPendingTasks);
+
 	FScopeLock Lock(&CriticalSection);
 
 	for (auto It = PendingTasks.CreateIterator(); It; ++It)
@@ -501,8 +503,36 @@ void FCardRepresentationAsyncQueue::RescheduleBackgroundTask(FAsyncCardRepresent
 	}
 }
 
+void FCardRepresentationAsyncQueue::FinishCompilationForObjects(TArrayView<UObject* const> InObjects)
+{
+	TRACE_CPUPROFILER_EVENT_SCOPE(FCardRepresentationAsyncQueue::FinishCompilationForObjects);
+
+	TSet<UStaticMesh*> StaticMeshes;
+	for (UObject* Object : InObjects)
+	{
+		if (UStaticMesh* StaticMesh = Cast<UStaticMesh>(Object))
+		{
+			StaticMeshes.Add(StaticMesh);
+		}
+		else if (UStaticMeshComponent* StaticMeshComponent = Cast<UStaticMeshComponent>(Object))
+		{
+			if (StaticMeshComponent->GetStaticMesh())
+			{
+				StaticMeshes.Add(StaticMeshComponent->GetStaticMesh());
+			}
+		}
+	}
+
+	for (UStaticMesh* StaticMesh : StaticMeshes)
+	{
+		BlockUntilBuildComplete(StaticMesh, false);
+	}
+}
+
 void FCardRepresentationAsyncQueue::BlockUntilBuildComplete(UStaticMesh* InStaticMesh, bool bWarnIfBlocked)
 {
+	TRACE_CPUPROFILER_EVENT_SCOPE(FCardRepresentationAsyncQueue::BlockUntilBuildComplete);
+
 	// We will track the wait time here, but only the cycles used.
 	// This function is called whether or not an async task is pending, 
 	// so we have to look elsewhere to properly count how many resources have actually finished building.
diff --git a/Engine/Source/Runtime/Engine/Private/SkinnedAssetCompiler.cpp b/Engine/Source/Runtime/Engine/Private/SkinnedAssetCompiler.cpp
--- a/Engine/Source/Runtime/Engine/Private/SkinnedAssetCompiler.cpp
+++ b/Engine/Source/Runtime/Engine/Private/SkinnedAssetCompiler.cpp
@@ -384,9 +384,35 @@ void FSkinnedAssetCompilingManager::FinishAllCompilation()
 	}
 }
 
+void FSkinnedAssetCompilingManager::FinishCompilationForObjects(TArrayView<UObject* const> InObjects)
+{
+	TRACE_CPUPROFILER_EVENT_SCOPE(FSkinnedAssetCompilingManager::FinishCompilationForObjects);
+
+	TSet<USkinnedAsset*> SkinnedAssets;
+	for (UObject* Object : InObjects)
+	{
+		if (USkinnedAsset* SkinnedAsset = Cast<USkinnedAsset>(Object))
+		{
+			SkinnedAssets.Add(SkinnedAsset);
+		}
+		else if (USkinnedMeshComponent* SkinnedMeshComponent = Cast<USkinnedMeshComponent>(Object))
+		{
+			if (SkinnedMeshComponent->GetSkinnedAsset())
+			{
+				SkinnedAssets.Add(SkinnedMeshComponent->GetSkinnedAsset());
+			}
+		}
+	}
+
+	if (SkinnedAssets.Num())
+	{
+		FinishCompilation(SkinnedAssets.Array());
+	}
+}
+
 void FSkinnedAssetCompilingManager::Reschedule()
 {
-	
+
 }
 
 void FSkinnedAssetCompilingManager::ProcessSkinnedAssets(bool bLimitExecutionTime, int32 MinBatchSize)
diff --git a/Engine/Source/Runtime/Engine/Private/SoundWaveCompiler.cpp b/Engine/Source/Runtime/Engine/Private/SoundWaveCompiler.cpp
--- a/Engine/Source/Runtime/Engine/Private/SoundWaveCompiler.cpp
+++ b/Engine/Source/Runtime/Engine/Private/SoundWaveCompiler.cpp
@@ -311,6 +311,25 @@ void FSoundWaveCompilingManager::FinishAllCompilation()
 	}
 }
 
+void FSoundWaveCompilingManager::FinishCompilationForObjects(TArrayView<UObject* const> InObjects)
+{
+	TRACE_CPUPROFILER_EVENT_SCOPE(FSoundWaveCompilingManager::FinishCompilationForObjects);
+
+	TSet<USoundWave*> SoundWaves;
+	for (UObject* Object : InObjects)
+	{
+		if (USoundWave* SoundWave = Cast<USoundWave>(Object))
+		{
+			SoundWaves.Add(SoundWave);
+		}
+	}
+
+	if (SoundWaves.Num())
+	{
+		FinishCompilation(SoundWaves.Array());
+	}
+}
+
 void FSoundWaveCompilingManager::ProcessSoundWaves(bool bLimitExecutionTime, int32 MaximumPriority)
 {
 	using namespace SoundWaveCompilingManagerImpl;
diff --git a/Engine/Source/Runtime/Engine/Private/StaticMeshCompiler.cpp b/Engine/Source/Runtime/Engine/Private/StaticMeshCompiler.cpp
--- a/Engine/Source/Runtime/Engine/Private/StaticMeshCompiler.cpp
+++ b/Engine/Source/Runtime/Engine/Private/StaticMeshCompiler.cpp
@@ -588,6 +588,32 @@ void FStaticMeshCompilingManager::FinishAllCompilation()
 	}
 }
 
+void FStaticMeshCompilingManager::FinishCompilationForObjects(TArrayView<UObject* const> InObjects)
+{
+	TRACE_CPUPROFILER_EVENT_SCOPE(FStaticMeshCompilingManager::FinishCompilationForObjects);
+
+	TSet<UStaticMesh*> StaticMeshes;
+	for (UObject* Object : InObjects)
+	{
+		if (UStaticMesh* StaticMesh = Cast<UStaticMesh>(Object))
+		{
+			StaticMeshes.Add(StaticMesh);
+		}
+		else if (UStaticMeshComponent* StaticMeshComponent = Cast<UStaticMeshComponent>(Object))
+		{
+			if (StaticMeshComponent->GetStaticMesh())
+			{
+				StaticMeshes.Add(StaticMeshComponent->GetStaticMesh());
+			}
+		}
+	}
+
+	if (StaticMeshes.Num())
+	{
+		FinishCompilation(StaticMeshes.Array());
+	}
+}
+
 void FStaticMeshCompilingManager::Reschedule()
 {
 	if (RegisteredStaticMesh.Num() > 1)
diff --git a/Engine/Source/Runtime/Engine/Private/TextureCompiler.cpp b/Engine/Source/Runtime/Engine/Private/TextureCompiler.cpp
--- a/Engine/Source/Runtime/Engine/Private/TextureCompiler.cpp
+++ b/Engine/Source/Runtime/Engine/Private/TextureCompiler.cpp
@@ -275,6 +275,25 @@ void FTextureCompilingManager::AddTextures(TArrayView<UTexture* const> InTexture
 	TRACE_COUNTER_SET(QueuedTextureCompilation, GetNumRemainingTextures());
 }
 
+void FTextureCompilingManager::FinishCompilationForObjects(TArrayView<UObject* const> InObjects)
+{
+	TRACE_CPUPROFILER_EVENT_SCOPE(FTextureCompilingManager::FinishCompilationForObjects);
+
+	TSet<UTexture*> Textures;
+	for (UObject* Object : InObjects)
+	{
+		if (UTexture* Texture = Cast<UTexture>(Object))
+		{
+			Textures.Add(Texture);
+		}
+	}
+
+	if (Textures.Num())
+	{
+		FinishCompilation(Textures.Array());
+	}
+}
+
 void FTextureCompilingManager::FinishCompilation(TArrayView<UTexture* const> InTextures)
 {
 	TRACE_CPUPROFILER_EVENT_SCOPE(FTextureCompilingManager::FinishCompilation);
diff --git a/Engine/Source/Runtime/Engine/Public/AssetCompilingManager.h b/Engine/Source/Runtime/Engine/Public/AssetCompilingManager.h
--- a/Engine/Source/Runtime/Engine/Public/AssetCompilingManager.h
+++ b/Engine/Source/Runtime/Engine/Public/AssetCompilingManager.h
@@ -52,6 +52,11 @@ struct IAssetCompilingManager
 	 */
 	ENGINE_API virtual int32 GetNumRemainingAssets() const = 0;
 
+	/**
+	 * Blocks until completion of the requested objects.
+	 */
+	ENGINE_API virtual void FinishCompilationForObjects(TArrayView<UObject* const> InObjects) {} /* Optional for backward compatibility */
+
 	/**
 	 * Blocks until completion of all assets.
 	 */
@@ -106,6 +111,11 @@ public:
 	 */
 	ENGINE_API void FinishAllCompilation();
 
+	/**
+	 * Finish compilation of the requested objects.
+	 */
+	ENGINE_API void FinishCompilationForObjects(TArrayView<UObject* const> InObjects);
+
 	/**
 	 * Cancel any pending work and blocks until it is safe to shut down.
 	 */
diff --git a/Engine/Source/Runtime/Engine/Public/DistanceFieldAtlas.h b/Engine/Source/Runtime/Engine/Public/DistanceFieldAtlas.h
--- a/Engine/Source/Runtime/Engine/Public/DistanceFieldAtlas.h
+++ b/Engine/Source/Runtime/Engine/Public/DistanceFieldAtlas.h
@@ -412,6 +412,9 @@ private:
 	/** Cancel or finish any work for the given task. */
 	void CancelAndDeleteTask(const TSet<FAsyncDistanceFieldTask*>& Tasks);
 
+	/** Handle generic finish compilation */
+	void FinishCompilationForObjects(TArrayView<UObject* const> InObjects) override;
+
 	/** Return whether the task has become invalid and should be cancelled (i.e. reference unreachable objects) */
 	bool IsTaskInvalid(FAsyncDistanceFieldTask* Task) const;
 
diff --git a/Engine/Source/Runtime/Engine/Public/MeshCardRepresentation.h b/Engine/Source/Runtime/Engine/Public/MeshCardRepresentation.h
--- a/Engine/Source/Runtime/Engine/Public/MeshCardRepresentation.h
+++ b/Engine/Source/Runtime/Engine/Public/MeshCardRepresentation.h
@@ -394,6 +394,9 @@ private:
 	/** Cancel or finish any work for the given task. */
 	void CancelAndDeleteTask(const TSet<FAsyncCardRepresentationTask*>& Tasks);
 
+	/** Handle generic finish compilation */
+	void FinishCompilationForObjects(TArrayView<UObject* const> InObjects) override;
+
 	/** Return whether the task has become invalid and should be canceled (i.e. reference unreachable objects) */
 	bool IsTaskInvalid(FAsyncCardRepresentationTask* Task) const;
 
diff --git a/Engine/Source/Runtime/Engine/Public/SkinnedAssetCompiler.h b/Engine/Source/Runtime/Engine/Public/SkinnedAssetCompiler.h
--- a/Engine/Source/Runtime/Engine/Public/SkinnedAssetCompiler.h
+++ b/Engine/Source/Runtime/Engine/Public/SkinnedAssetCompiler.h
@@ -82,6 +82,8 @@ private:
 	int32 GetNumRemainingAssets() const override;
 	void ProcessAsyncTasks(bool bLimitExecutionTime = false) override;
 
+	void FinishCompilationForObjects(TArrayView<UObject* const> InObjects) override;
+
 	friend class FAssetCompilingManager;
 	
 	bool bHasShutdown = false;
diff --git a/Engine/Source/Runtime/Engine/Public/SoundWaveCompiler.h b/Engine/Source/Runtime/Engine/Public/SoundWaveCompiler.h
--- a/Engine/Source/Runtime/Engine/Public/SoundWaveCompiler.h
+++ b/Engine/Source/Runtime/Engine/Public/SoundWaveCompiler.h
@@ -87,6 +87,9 @@ private:
 	TArrayView<FName> GetDependentTypeNames() const override;
 	int32 GetNumRemainingAssets() const override;
 
+	/** Handle generic finish compilation */
+	void FinishCompilationForObjects(TArrayView<UObject* const> InObjects) override;
+
 	void UpdateCompilationNotification();
 	void PostCompilation(TArrayView<USoundWave* const> InCompiledSoundWaves);
 	void PostCompilation(USoundWave* SoundWave);
diff --git a/Engine/Source/Runtime/Engine/Public/StaticMeshCompiler.h b/Engine/Source/Runtime/Engine/Public/StaticMeshCompiler.h
--- a/Engine/Source/Runtime/Engine/Public/StaticMeshCompiler.h
+++ b/Engine/Source/Runtime/Engine/Public/StaticMeshCompiler.h
@@ -82,6 +82,9 @@ private:
 	/** Called once per frame, fetches completed tasks and applies them to the scene. */
 	void ProcessAsyncTasks(bool bLimitExecutionTime = false) override;
 
+	/** Handle generic finish compilation */
+	void FinishCompilationForObjects(TArrayView<UObject* const> InObjects) override;
+
 	bool bHasShutdown = false;
 	TSet<TWeakObjectPtr<UStaticMesh>> RegisteredStaticMesh;
 	FAsyncCompilationNotification Notification;
diff --git a/Engine/Source/Runtime/Engine/Public/TextureCompiler.h b/Engine/Source/Runtime/Engine/Public/TextureCompiler.h
--- a/Engine/Source/Runtime/Engine/Public/TextureCompiler.h
+++ b/Engine/Source/Runtime/Engine/Public/TextureCompiler.h
@@ -83,6 +83,9 @@ private:
 
 	FTextureCompilingManager();
 
+	/** Handle generic finish compilation */
+	void FinishCompilationForObjects(TArrayView<UObject* const> InObjects) override;
+
 	FName GetAssetTypeName() const override;
 	FTextFormat GetAssetNameFormat() const override;
 	TArrayView<FName> GetDependentTypeNames() const override;
