MPCook: Implement retraction from busy workers and reassignment to idle workers.
#rnx

diff --git a/Engine/Source/Editor/UnrealEd/Classes/CookOnTheSide/CookOnTheFlyServer.h b/Engine/Source/Editor/UnrealEd/Classes/CookOnTheSide/CookOnTheFlyServer.h
--- a/Engine/Source/Editor/UnrealEd/Classes/CookOnTheSide/CookOnTheFlyServer.h
+++ b/Engine/Source/Editor/UnrealEd/Classes/CookOnTheSide/CookOnTheFlyServer.h
@@ -1044,6 +1044,8 @@ private:
 	/** Construct the on-stack data for StartCook functions, based on the CookWorker's configuration */
 	FBeginCookContext CreateCookWorkerContext();
 	void CookAsCookWorkerFinished();
+	/** Return the best packages to retract and give to other workers -the ones with least effort spent so far. */
+	void GetPackagesToRetract(int32 NumToRetract, TArray<FName>& OutRetractionPackages);
 
 	//////////////////////////////////////////////////////////////////////////
 	// general functions
diff --git a/Engine/Source/Editor/UnrealEd/Private/CookOnTheFlyServer.cpp b/Engine/Source/Editor/UnrealEd/Private/CookOnTheFlyServer.cpp
--- a/Engine/Source/Editor/UnrealEd/Private/CookOnTheFlyServer.cpp
+++ b/Engine/Source/Editor/UnrealEd/Private/CookOnTheFlyServer.cpp
@@ -9578,6 +9578,72 @@ void UCookOnTheFlyServer::CookAsCookWorkerFinished()
 	LogCookWorkerStats();
 }
 
+void UCookOnTheFlyServer::GetPackagesToRetract(int32 NumToRetract, TArray<FName>& OutRetractionPackages)
+{
+	using namespace UE::Cook;
+	OutRetractionPackages.Reset(NumToRetract);
+	if (OutRetractionPackages.Num() >= NumToRetract)
+	{
+		return;
+	}
+
+	FRequestQueue& RequestQueue = PackageDatas->GetRequestQueue();
+	TArray<FPackageData*> PoppedPackages;
+	if (!RequestQueue.IsReadyRequestsEmpty())
+	{
+		while (!RequestQueue.IsReadyRequestsEmpty())
+		{
+			PoppedPackages.Add(RequestQueue.PopReadyRequest());
+		}
+		for (FPackageData* PackageData : PoppedPackages)
+		{
+			if (OutRetractionPackages.Num() < NumToRetract)
+			{
+				OutRetractionPackages.Add(PackageData->GetPackageName());
+			}
+			RequestQueue.AddReadyRequest(PackageData);
+		}
+	}
+	if (OutRetractionPackages.Num() >= NumToRetract)
+	{
+		return;
+	}
+
+	FLoadPrepareQueue& LoadPrepareQueue = PackageDatas->GetLoadPrepareQueue();
+	for (FPackageData* PackageData : LoadPrepareQueue.EntryQueue)
+	{
+		OutRetractionPackages.Add(PackageData->GetPackageName());
+		if (OutRetractionPackages.Num() >= NumToRetract)
+		{
+			return;
+		}
+	}
+	for (FPackageData* PackageData : LoadPrepareQueue.PreloadingQueue)
+	{
+		OutRetractionPackages.Add(PackageData->GetPackageName());
+		if (OutRetractionPackages.Num() >= NumToRetract)
+		{
+			return;
+		}
+	}
+	for (FPackageData* PackageData : PackageDatas->GetLoadReadyQueue())
+	{
+		OutRetractionPackages.Add(PackageData->GetPackageName());
+		if (OutRetractionPackages.Num() >= NumToRetract)
+		{
+			return;
+		}
+	}
+	for (FPackageData* PackageData : PackageDatas->GetSaveQueue())
+	{
+		OutRetractionPackages.Add(PackageData->GetPackageName());
+		if (OutRetractionPackages.Num() >= NumToRetract)
+		{
+			return;
+		}
+	}
+}
+
 void UCookOnTheFlyServer::ShutdownCookAsCookWorker()
 {
 	if (IsDirectorCookByTheBook())
diff --git a/Engine/Source/Editor/UnrealEd/Private/Cooker/CookDirector.cpp b/Engine/Source/Editor/UnrealEd/Private/Cooker/CookDirector.cpp
--- a/Engine/Source/Editor/UnrealEd/Private/Cooker/CookDirector.cpp
+++ b/Engine/Source/Editor/UnrealEd/Private/Cooker/CookDirector.cpp
@@ -30,6 +30,8 @@ extern CORE_API int32 GNumForegroundWorkers; // TaskGraph.cpp
 namespace UE::Cook
 {
 
+constexpr int32 RetractionMinimumNumAssignments = 100;
+
 /** Profile data for each CookWorker that needs to be collected on the Director. */
 struct FCookWorkerProfileData
 {
@@ -48,6 +50,41 @@ struct FCookWorkerProfileData
 	}
 };
 
+/**
+ * A class that has an instance active while we need to handle retraction of assigned results from a CookWorker.
+ * Keeps track of the expected message coming back from the remote worker, prevents repeatedly sending messages, gives
+ * a warning if the remote worker does not respond.
+ */
+class FCookDirector::FRetractionHandler
+{
+public:
+	FRetractionHandler(FCookDirector& InDirector);
+	/** Initialize to search idle and busy workers to send a RetractionRequestMessage. */
+	void Initialize();
+	/** Initialize to handle an unexpected RetractionResultsMessage. */
+	void InitializeForResultsMessage(const FWorkerId& FromWorker);
+	void TickFromSchedulerThread(bool bAnyIdle, bool& bOutComplete);
+	/** Hook called by the director when a retraction message comes in. */
+	void HandleRetractionMessage(const FWorkerId& FromWorker, TConstArrayView<FName> Packages);
+
+private:
+	/**
+	 * Pick workers to give the retracted packages to, and assign those packages to the worker
+	 * in the local and remote state.
+	 */
+	void ReassignPackages(const FWorkerId& WorkerId, TConstArrayView<FPackageData*> Packages);
+	/** Pick workers to give the retracted packages to. */
+	TArray<FWorkerId> CalculateWorkersToSplitOver(int32 NumPackages, const FWorkerId& FromWorker,
+		TConstArrayView<TRefCountPtr<FCookWorkerServer>> LocalRemoteWorkers);
+
+	FCookDirector& Director;
+	FWorkerId ExpectedWorker;
+	TMap<FWorkerId, TArray<FName>> PackagesToRetract;
+	FWorkerId WorkerWithResults;
+	double MessageSentTimeSeconds = 0.;
+	double LastWarnTimeSeconds = 0.;
+};
+
 FCookDirector::FCookDirector(UCookOnTheFlyServer& InCOTFS, int32 CookProcessCount)
 	: RunnableShunt(*this) 
 	, COTFS(InCOTFS)
@@ -225,49 +262,99 @@ void FCookDirector::StartCook(const FBeginCookContext& InBeginContext)
 	BeginCookContext.Set(InBeginContext);
 }
 
-void FCookDirector::AssignRequests(TArrayView<UE::Cook::FPackageData*> Requests, TArray<FWorkerId>& OutAssignments,
+void FCookDirector::AssignRequests(TArrayView<FPackageData*> Requests, TArray<FWorkerId>& OutAssignments,
 	TMap<FPackageData*, TArray<FPackageData*>>&& RequestGraph)
 {
 	ActivateMachineResourceReduction();
 
-	TArray<TRefCountPtr<FCookWorkerServer>> SortedWorkers;
-	int32 MaxRemoteIndex = -1;
-	FScopeLock CommunicationScopeLock(&CommunicationLock);
+	TArray<FWorkerId> WorkerIds;
+	TArray<TRefCountPtr<FCookWorkerServer>> LocalRemoteWorkers;
 	{
+		FScopeLock CommunicationScopeLock(&CommunicationLock);
 		InitializeWorkers();
+	}
+	LocalRemoteWorkers = CopyRemoteWorkers();;
 
-		// Convert the Map of RemoteWorkers into an array sorted by WorkerIndex
-		SortedWorkers.Reserve(RemoteWorkers.Num());
-		for (TPair<int32, TRefCountPtr<FCookWorkerServer>>& Pair : RemoteWorkers)
-		{
-			SortedWorkers.Add(Pair.Value);
-			MaxRemoteIndex = FMath::Max(Pair.Key, MaxRemoteIndex);
-		}
+	WorkerIds.Reserve(LocalRemoteWorkers.Num() + 1);
+	WorkerIds.Add(FWorkerId::Local());
+	for (const TRefCountPtr<FCookWorkerServer>& RemoteWorker : LocalRemoteWorkers)
+	{
+		WorkerIds.Add(RemoteWorker->GetWorkerId());
 	}
-	if (SortedWorkers.IsEmpty())
+
+	AssignRequests(MoveTemp(WorkerIds), LocalRemoteWorkers, Requests, OutAssignments, MoveTemp(RequestGraph));
+}
+
+void FCookDirector::AssignRequests(TArray<FWorkerId>&& InWorkers, TArray<TRefCountPtr<FCookWorkerServer>>& InRemoteWorkers,
+	TArrayView<FPackageData*> Requests, TArray<FWorkerId>& OutAssignments, TMap<FPackageData*, TArray<FPackageData*>>&& RequestGraph)
+{
+	check(InWorkers.Num() > 0);
+	if (InWorkers.Num() <= 1)
 	{
+		FWorkerId WorkerId = InWorkers[0];
 		OutAssignments.SetNum(Requests.Num());
-		for (FWorkerId& Assignment : OutAssignments)
+		TArray<UE::Cook::FPackageData*> RemovedRequests;
+		for (int32 RequestIndex = 0; RequestIndex < Requests.Num(); ++RequestIndex)
 		{
-			Assignment = FWorkerId::Local();
+			FPackageData* Request = Requests[RequestIndex];
+			FWorkerId& Assignment = OutAssignments[RequestIndex];
+			FWorkerId WorkerIdConstraint = Request->GetWorkerAssignmentConstraint();
+			if (WorkerIdConstraint.IsValid() && WorkerIdConstraint != WorkerId)
+			{
+				UE_LOG(LogCook, Error, TEXT("Package %s can only be cooked by a now-disconnected CookWorker. The package can not be cooked."),
+					*Request->GetPackageName().ToString());
+				Assignment = FWorkerId::Invalid();
+				RemovedRequests.Add(Request);
+			}
+			else
+			{
+				Assignment = WorkerId;
+			}
+		}
+		if (!WorkerId.IsLocal())
+		{
+			TRefCountPtr<FCookWorkerServer>* RemoteWorker = InRemoteWorkers.FindByPredicate(
+				[&WorkerId](const TRefCountPtr<FCookWorkerServer>& X) { return X->GetWorkerId() == WorkerId; });
+			check(RemoteWorker);
+			TArrayView<FPackageData*> RequestsToSend = Requests;
+			TArray<FPackageData*> RequestBuffer;
+			if (!RemovedRequests.IsEmpty())
+			{
+				TSet<FPackageData*> RequestSet;
+				for (FPackageData* Request : Requests)
+				{
+					RequestSet.Add(Request);
+				}
+				for (FPackageData* Remove : RemovedRequests)
+				{
+					RequestSet.Remove(Remove);
+				}
+				RequestBuffer = RequestSet.Array();
+				RequestsToSend = RequestBuffer;
+			}
+			(*RemoteWorker)->AppendAssignments(RequestsToSend, ECookDirectorThread::SchedulerThread);
 		}
 		return;
 	}
-	check(MaxRemoteIndex >= 0);
-	SortedWorkers.Sort([](const TRefCountPtr<FCookWorkerServer>& A, const TRefCountPtr<FCookWorkerServer>& B)
-		{ return A->GetWorkerId() < B->GetWorkerId(); });
+
+	InWorkers.Sort();
 
 	// Call the LoadBalancing algorithm to split the requests among the LocalWorker and RemoteWorkers
-	LoadBalance(SortedWorkers, Requests, MoveTemp(RequestGraph), OutAssignments);
+	LoadBalance(InWorkers, Requests, MoveTemp(RequestGraph), OutAssignments);
+
+	int32 MaxRemoteIndex = InWorkers.Last().IsLocal() ? -1 : InWorkers.Last().GetRemoteIndex();
 
 	// Split the output array of WorkerId assignments into a batch for each of the RemoteWorkers 
 	TArray<TArray<FPackageData*>> RemoteBatches; // Indexed by WorkerId.GetRemoteIndex()
 	TArray<bool> RemoteIndexIsValid; // Indexed by WorkerId.GetRemoteIndex()
 	RemoteBatches.SetNum(MaxRemoteIndex+1);
 	RemoteIndexIsValid.Init(false, MaxRemoteIndex+1);
-	for (FCookWorkerServer* Worker : SortedWorkers)
+	for (FWorkerId& WorkerId : InWorkers)
 	{
-		RemoteIndexIsValid[Worker->GetWorkerId().GetRemoteIndex()] = true;
+		if (!WorkerId.IsLocal())
+		{
+			RemoteIndexIsValid[WorkerId.GetRemoteIndex()] = true;
+		}
 	}
 
 	for (int32 RequestIndex = 0; RequestIndex < Requests.Num(); ++RequestIndex)
@@ -279,6 +366,7 @@ void FCookDirector::AssignRequests(TArrayView<UE::Cook::FPackageData*> Requests,
 		FWorkerId WorkerIdConstraint = Requests[RequestIndex]->GetWorkerAssignmentConstraint();
 		if (WorkerIdConstraint.IsValid())
 		{
+			check(InWorkers.Contains(WorkerIdConstraint));
 			OutAssignments[RequestIndex] = WorkerIdConstraint;
 			WorkerId = WorkerIdConstraint;
 		}
@@ -289,15 +377,15 @@ void FCookDirector::AssignRequests(TArrayView<UE::Cook::FPackageData*> Requests,
 			check(RemoteIndex < RemoteBatches.Num());
 			if (!RemoteIndexIsValid[RemoteIndex])
 			{
-				UE_LOG(LogCook, Error, TEXT("Package %s can only be cooked by CookWorkerServer %d, but this worker has disconnected. The package can not be cooked."),
-					*Requests[RequestIndex]->GetPackageName().ToString(), RemoteIndex);
+				UE_LOG(LogCook, Error, TEXT("Package %s can only be cooked by a now-disconnected CookWorker. The package can not be cooked."),
+					*Requests[RequestIndex]->GetPackageName().ToString());
 				OutAssignments[RequestIndex] = FWorkerId::Invalid();
 				continue;
 			}
 			TArray<FPackageData*>& RemoteBatch = RemoteBatches[RemoteIndex];
 			if (RemoteBatch.Num() == 0)
 			{
-				RemoteBatch.Reserve(2 * Requests.Num() / (SortedWorkers.Num() + 1));
+				RemoteBatch.Reserve(2 * Requests.Num() / (InWorkers.Num()));
 			}
 			RemoteBatch.Add(Requests[RequestIndex]);
 		}
@@ -307,14 +395,32 @@ void FCookDirector::AssignRequests(TArrayView<UE::Cook::FPackageData*> Requests,
 
 	// Assign each batch to the FCookWorkerServer in RemoteWorkers;
 	// the CookWorkerServer's tick will handle sending the message to the remote process
-	for (const TRefCountPtr<FCookWorkerServer>& RemoteWorker : SortedWorkers)
+	for (FWorkerId& WorkerId : InWorkers)
 	{
-		RemoteWorker->AppendAssignments(RemoteBatches[RemoteWorker->GetWorkerId().GetRemoteIndex()], ECookDirectorThread::SchedulerThread);
+		if (!WorkerId.IsLocal())
+		{
+			TRefCountPtr<FCookWorkerServer>* RemoteWorker = InRemoteWorkers.FindByPredicate(
+				[&WorkerId](const TRefCountPtr<FCookWorkerServer>& X) { return X->GetWorkerId() == WorkerId; });
+			check(RemoteWorker);
+			(*RemoteWorker)->AppendAssignments(RemoteBatches[WorkerId.GetRemoteIndex()], ECookDirectorThread::SchedulerThread);
+		}
 	}
 
 	bIsFirstAssignment = false;
 }
 
+TArray<TRefCountPtr<FCookWorkerServer>> FCookDirector::CopyRemoteWorkers() const
+{
+	TArray<TRefCountPtr<FCookWorkerServer>> LocalRemoteWorkers;
+	FScopeLock CommunicationScopeLock(&CommunicationLock);
+	LocalRemoteWorkers.Reset(RemoteWorkers.Num());
+	for (const TPair<int32, TRefCountPtr<FCookWorkerServer>>& Pair : RemoteWorkers)
+	{
+		LocalRemoteWorkers.Add(Pair.Value);
+	}
+	return LocalRemoteWorkers;
+}
+
 void FCookDirector::RemoveFromWorker(FPackageData& PackageData)
 {
 	TArray<FCookWorkerServer*> Workers;
@@ -339,19 +445,29 @@ void FCookDirector::TickFromSchedulerThread()
 		TickCommunication(ECookDirectorThread::SchedulerThread);
 	}
 
-	bool bLocalWorkerIdle = COTFS.NumMultiprocessLocalWorkerAssignments() == 0;
+	int32 BusiestNumAssignments = COTFS.NumMultiprocessLocalWorkerAssignments();
+	bool bLocalWorkerIdle = BusiestNumAssignments == 0;
 	float DeltaTime = static_cast<float>(CurrentTime - LastTickTimeSeconds);
+	bool bAnyIdle = bLocalWorkerIdle;
 	LastTickTimeSeconds = CurrentTime;
 	LocalWorkerProfileData->UpdateIdle(bLocalWorkerIdle, DeltaTime);
 
 	TArray<TRefCountPtr<FCookWorkerServer>, TInlineAllocator<16>> WorkersWithMessage;
+	bool bAllWorkersConnected = false;
+	if (bWorkersInitialized)
 	{
 		FScopeLock CommunicationScopeLock(&CommunicationLock);
+		bAllWorkersConnected = true;
 		for (TPair<int32, TRefCountPtr<FCookWorkerServer>>& Pair : RemoteWorkers)
 		{
 			FCookWorkerServer* RemoteWorker = Pair.Value.GetReference();
 			FCookWorkerProfileData& ProfileData = RemoteWorkerProfileDatas[RemoteWorker->GetProfileId()];
-			ProfileData.UpdateIdle(RemoteWorker->NumAssignments() == 0, DeltaTime);
+			bAllWorkersConnected &= RemoteWorker->IsConnected();
+			int32 NumAssignments = RemoteWorker->NumAssignments();
+			BusiestNumAssignments = FMath::Max(NumAssignments, BusiestNumAssignments);
+			bool bWorkerIdle = NumAssignments == 0;
+			bAnyIdle |= bWorkerIdle;
+			ProfileData.UpdateIdle(bWorkerIdle, DeltaTime);
 			if (RemoteWorker->HasMessages())
 			{
 				WorkersWithMessage.Add(RemoteWorker);
@@ -364,8 +480,12 @@ void FCookDirector::TickFromSchedulerThread()
 				WorkersWithMessage.Add(Pair.Value);
 			}
 		}
-
 	}
+	if (bAllWorkersConnected && (RetractionHandler.IsValid() || bAnyIdle))
+	{
+		TickRetractionFromSchedulerThread(bAnyIdle, BusiestNumAssignments);
+	}
+
 	bool bIsStalled = bLocalWorkerIdle && !COTFS.PackageDatas->GetAssignedToWorkerSet().IsEmpty() && WorkersWithMessage.IsEmpty();
 	for (TRefCountPtr<FCookWorkerServer>& Worker : WorkersWithMessage)
 	{
@@ -377,18 +497,72 @@ void FCookDirector::TickFromSchedulerThread()
 	LastTickTimeSeconds = CurrentTime;
 }
 
-void FCookDirector::UpdateDisplayDiagnostics()
+void FCookDirector::UpdateDisplayDiagnostics() const
+{
+	DisplayRemainingPackages();
+}
+
+void FCookDirector::DisplayRemainingPackages() const
 {
-	UE_LOG(LogCook, Display, TEXT("\tLocalWorker: %d packages remain."), COTFS.NumMultiprocessLocalWorkerAssignments());
+	constexpr int32 DisplayWidth = 16;
+	UE_LOG(LogCook, Display, TEXT("\t%s: %d packages remain."), *GetDisplayName(FWorkerId::Local(), DisplayWidth),
+		COTFS.NumMultiprocessLocalWorkerAssignments());
 	FScopeLock CommunicationScopeLock(&CommunicationLock);
-	for (TPair<int32, TRefCountPtr<FCookWorkerServer>>& Pair : RemoteWorkers)
+	for (const TPair<int32, TRefCountPtr<FCookWorkerServer>>& Pair : RemoteWorkers)
 	{
 		FCookWorkerServer* RemoteWorker = Pair.Value;
-		UE_LOG(LogCook, Display, TEXT("\tCookWorker %d: %d packages remain."),
-			RemoteWorker->GetProfileId(), RemoteWorker->NumAssignments());
+		UE_LOG(LogCook, Display, TEXT("\t%s: %d packages remain."),
+			*GetDisplayName(*RemoteWorker, DisplayWidth), RemoteWorker->NumAssignments());
 	}
 }
 
+FString FCookDirector::GetDisplayName(const FWorkerId& WorkerId, int32 PreferredWidth) const
+{
+	FString Result;
+	if (WorkerId.IsLocal())
+	{
+		Result = TEXTVIEW("Local");
+	}
+	else
+	{
+		const TRefCountPtr<FCookWorkerServer>* RemoteWorker = nullptr;
+		{
+			FScopeLock CommunicationScopeLock(&CommunicationLock);
+			RemoteWorker = RemoteWorkers.Find(WorkerId.GetRemoteIndex());
+			if (!RemoteWorker)
+			{
+				for (const TPair<FCookWorkerServer*, TRefCountPtr<FCookWorkerServer>>& Pair : ShuttingDownWorkers)
+				{
+					if (Pair.Value && Pair.Value->GetWorkerId() == WorkerId)
+					{
+						RemoteWorker = &Pair.Value;
+						break;
+					}
+				}
+			}
+		}
+		if (RemoteWorker)
+		{
+			Result = FString::Printf(TEXT("%d"), (*RemoteWorker)->GetProfileId());
+		}
+		else
+		{
+			Result = FString::Printf(TEXT("Unknown (WorkerId %d)"), WorkerId.GetRemoteIndex());
+		}
+	}
+	constexpr FStringView Prefix(TEXTVIEW("CookWorker "));
+	Result = FString(Prefix) + Result.LeftPad(PreferredWidth-Prefix.Len());
+	return Result;
+}
+
+FString FCookDirector::GetDisplayName(const FCookWorkerServer& RemoteWorker, int32 PreferredWidth) const
+{
+	FString Result = FString::Printf(TEXT("%d"), RemoteWorker.GetProfileId());
+	constexpr FStringView Prefix(TEXTVIEW("CookWorker "));
+	Result = FString(Prefix) + Result.LeftPad(PreferredWidth - Prefix.Len());
+	return Result;
+}
+
 void FCookDirector::TickCommunication(ECookDirectorThread TickThread)
 {
 	bool bHasShutdownWorkers = false;
@@ -936,7 +1110,8 @@ FString FCookDirector::GetWorkerCommandLine(FWorkerId WorkerId, int32 ProfileId)
 				Token.StartsWith(TEXT("-CookProfileId=")) ||
 				Token.StartsWith(TEXT("-ShowCookWorker")) ||
 				Token.StartsWith(TEXT("-CoreLimit")) ||
-				Token.StartsWith(TEXT("-PhysicalCoreLimit"))
+				Token.StartsWith(TEXT("-PhysicalCoreLimit")) ||
+				Token.StartsWith(TEXT("-CookProcessCount="))
 				)
 			{
 				return;
@@ -984,29 +1159,21 @@ bool FDirectorConnectionInfo::TryParseCommandLine()
 	return true;
 }
 
-void FCookDirector::LoadBalance(TConstArrayView<TRefCountPtr<FCookWorkerServer>> SortedWorkers, TArrayView<FPackageData*> Requests,
+void FCookDirector::LoadBalance(TConstArrayView<FWorkerId> SortedWorkers, TArrayView<FPackageData*> Requests,
 	TMap<FPackageData*, TArray<FPackageData*>>&& RequestGraph, TArray<FWorkerId>& OutAssignments)
 {
-	TArray<FWorkerId> AllWorkers;
-	int32 NumAllWorkers = SortedWorkers.Num() + 1;
-	AllWorkers.Reserve(NumAllWorkers);
-	AllWorkers.Add(FWorkerId::Local());
-	for (FCookWorkerServer* Worker : SortedWorkers)
-	{
-		AllWorkers.Add(Worker->GetWorkerId());
-	}
 	OutAssignments.Reset(Requests.Num());
 	bool bLogResults = bIsFirstAssignment;
 
 	switch (LoadBalanceAlgorithm)
 	{
 	case ELoadBalanceAlgorithm::Striped:
-		return LoadBalanceStriped(AllWorkers, Requests, MoveTemp(RequestGraph), OutAssignments, bLogResults);
+		return LoadBalanceStriped(SortedWorkers, Requests, MoveTemp(RequestGraph), OutAssignments, bLogResults);
 	case ELoadBalanceAlgorithm::CookBurden:
-		return LoadBalanceCookBurden(AllWorkers, Requests, MoveTemp(RequestGraph), OutAssignments, bLogResults);
+		return LoadBalanceCookBurden(SortedWorkers, Requests, MoveTemp(RequestGraph), OutAssignments, bLogResults);
 	}
 	checkNoEntry();
-	return LoadBalanceCookBurden(AllWorkers, Requests, MoveTemp(RequestGraph), OutAssignments, bLogResults);
+	return LoadBalanceCookBurden(SortedWorkers, Requests, MoveTemp(RequestGraph), OutAssignments, bLogResults);
 }
 
 void FCookDirector::AbortWorker(FWorkerId WorkerId, ECookDirectorThread TickThread)
@@ -1023,7 +1190,7 @@ void FCookDirector::AbortWorker(FWorkerId WorkerId, ECookDirectorThread TickThre
 		}
 	}
 	TSet<FPackageData*> PackagesToReassign;
-	RemoteWorker->AbortAssignments(PackagesToReassign, TickThread);
+	RemoteWorker->AbortAllAssignments(PackagesToReassign, TickThread);
 	if (!RemoteWorker->IsShuttingDown())
 	{
 		RemoteWorker->AbortWorker(PackagesToReassign, TickThread);
@@ -1093,6 +1260,346 @@ void FCookDirector::LogCookStats(FCookStatsManager::AddStatFuncRef AddStat)
 void FCookDirector::HandleRetractionMessage(FMPCollectorServerMessageContext& Context, bool bReadSuccessful,
 	FRetractionResultsMessage&& Message)
 {
+	if (!bReadSuccessful)
+	{
+		UE_LOG(LogCook, Error, TEXT("Corrupt RetractionResultsMessage received from CookWorker %d. It will be ignored and packages may fail to cook."),
+			Context.GetProfileId());
+		return;
+	}
+
+	if (!RetractionHandler)
+	{
+		UE_LOG(LogCook, Warning, TEXT("Retractionmessage received from CookWorker %d when we were not expecting one."), Context.GetProfileId());
+		RetractionHandler = MakeUnique<FRetractionHandler>(*this);
+		RetractionHandler->InitializeForResultsMessage(Context.GetWorkerId());
+	}
+	RetractionHandler->HandleRetractionMessage(Context.GetWorkerId(), Message.ReturnedPackages);
+}
+
+void FCookDirector::TickRetractionFromSchedulerThread(bool bAnyIdle, int32 BusiestNumAssignments)
+{
+	if (RetractionHandler)
+	{
+		bool bComplete;
+		RetractionHandler->TickFromSchedulerThread(bAnyIdle, bComplete);
+		if (bComplete)
+		{
+			RetractionHandler.Reset();
+		}
+	}
+	else if (bAnyIdle && BusiestNumAssignments > RetractionMinimumNumAssignments)
+	{
+		RetractionHandler = MakeUnique<FRetractionHandler>(*this);
+		RetractionHandler->Initialize();
+	}
+}
+
+FCookDirector::FRetractionHandler::FRetractionHandler(FCookDirector& InDirector)
+	: Director(InDirector)
+{
+}
+
+void FCookDirector::FRetractionHandler::Initialize()
+{
+	FWorkerId BusiestWorker;
+	TArray<FWorkerId> IdleWorkers;
+	int32 BusiestNumAssignments = 0;
+
+	{
+		int32 NumAssignments = Director.COTFS.NumMultiprocessLocalWorkerAssignments();
+		BusiestWorker = FWorkerId::Local();
+		BusiestNumAssignments = NumAssignments;
+		if (NumAssignments == 0)
+		{
+			IdleWorkers.Add(FWorkerId::Local());
+		}
+	}
+
+	TArray<TRefCountPtr<FCookWorkerServer>> LocalRemoteWorkers = Director.CopyRemoteWorkers();
+	for (const TRefCountPtr<FCookWorkerServer>& RemoteWorker : LocalRemoteWorkers)
+	{
+		int32 NumAssignments = RemoteWorker->NumAssignments();
+		if (NumAssignments > BusiestNumAssignments)
+		{
+			BusiestWorker = RemoteWorker->GetWorkerId();
+			BusiestNumAssignments = NumAssignments;
+		}
+		if (NumAssignments == 0)
+		{
+			IdleWorkers.Add(RemoteWorker->GetWorkerId());
+		}
+	}
+	if (IdleWorkers.IsEmpty() || BusiestNumAssignments < RetractionMinimumNumAssignments)
+	{
+		// Worker loads changed after the point where we decided to initialize the RetractionHandler,
+		// and retraction is no longer needed. Early exit now and this will be deleted later in Tick.
+		return;
+	}
+
+	// Plan to divide the assignments evenly between all idle workers and the one busiest worker. This means
+	// retracting all but 1/(N+1) packages from the busiest worker.
+	int32 NumAssignmentsToRetract = (BusiestNumAssignments * IdleWorkers.Num()) / (IdleWorkers.Num() + 1);
+	TStringBuilder<256> IdleWorkerListText;
+	for (FWorkerId& WorkerId : IdleWorkers)
+	{
+		IdleWorkerListText << Director.GetDisplayName(WorkerId) << TEXT(", ");
+	}
+	IdleWorkerListText.RemoveSuffix(2);
+	UE_LOG(LogCook, Display, TEXT("Idle CookWorkers: { %s }. Retracting %d packages from %s to distribute to the idle CookWorkers."),
+		*IdleWorkerListText, NumAssignmentsToRetract, *Director.GetDisplayName(BusiestWorker));
+	Director.DisplayRemainingPackages();
+
+	if (BusiestWorker.IsLocal())
+	{
+		ExpectedWorker = FWorkerId::Local();
+		WorkerWithResults = ExpectedWorker;
+		TArray<FName> LocalPackagesToRetract;
+		Director.COTFS.GetPackagesToRetract(NumAssignmentsToRetract, LocalPackagesToRetract);
+		PackagesToRetract.FindOrAdd(ExpectedWorker).Append(MoveTemp(LocalPackagesToRetract));
+	}
+	else
+	{
+		TRefCountPtr<FCookWorkerServer>* RemoteWorker = LocalRemoteWorkers.FindByPredicate(
+			[&BusiestWorker](const TRefCountPtr<FCookWorkerServer>& X) { return X->GetWorkerId() == BusiestWorker; });
+		check(RemoteWorker);
+		FRetractionRequestMessage Message;
+		Message.RequestedCount = NumAssignmentsToRetract;
+		(*RemoteWorker)->SendMessage(Message, ECookDirectorThread::SchedulerThread);
+		ExpectedWorker = BusiestWorker;
+		MessageSentTimeSeconds = FPlatformTime::Seconds();
+		LastWarnTimeSeconds = MessageSentTimeSeconds;
+	}
+}
+
+void FCookDirector::FRetractionHandler::InitializeForResultsMessage(const FWorkerId& FromWorker)
+{
+	ExpectedWorker = FromWorker;
+}
+
+void FCookDirector::FRetractionHandler::TickFromSchedulerThread(bool bAnyIdle, bool& bOutComplete)
+{
+	if (ExpectedWorker.IsInvalid())
+	{
+		// We decided to cancel
+		checkf(PackagesToRetract.IsEmpty(), TEXT("We should not have any packages when we cancelled."));
+		bOutComplete = true;
+		return;
+	}
+	if (WorkerWithResults.IsInvalid())
+	{
+		double CurrentTime = FPlatformTime::Seconds();
+		constexpr float WarnDuration = 60.f;
+
+		if (static_cast<float>(CurrentTime - LastWarnTimeSeconds) < WarnDuration)
+		{
+			bOutComplete = false;
+			return;
+		}
+		check(ExpectedWorker.IsRemote());
+		{
+			FScopeLock CommunicationScopeLock(&Director.CommunicationLock);
+			TRefCountPtr<FCookWorkerServer>* RemoteWorkerPtr = Director.RemoteWorkers.Find(ExpectedWorker.GetRemoteIndex());
+			if (!RemoteWorkerPtr)
+			{
+				// The CookWorker aborted and we already reassigned all of its packages; stop waiting for a retraction message from it.
+				check(PackagesToRetract.IsEmpty()); // Otherwise WorkerWithResults would have been set
+				ExpectedWorker = FWorkerId::Invalid();
+				bOutComplete = true;
+				return;
+			}
+		}
+		bOutComplete = false;
+		UE_CLOG(!IsCookIgnoreTimeouts(), LogCook, Warning, TEXT("%s has not responded to a RetractionRequest message for %.1f seconds. Continuing to wait..."),
+			*Director.GetDisplayName(ExpectedWorker), static_cast<float>(CurrentTime - MessageSentTimeSeconds));
+		LastWarnTimeSeconds = CurrentTime;
+		return;
+	}
+
+	// Convert names to packagedatas and collect results from all CookWorkers who sent a message.
+	TArray<FPackageData*> PackageDatasToReassign;
+	for (const TPair<FWorkerId, TArray<FName>>& Pair : PackagesToRetract)
+	{
+		TRefCountPtr<FCookWorkerServer> RemoteWorker;
+		if (Pair.Key.IsRemote())
+		{
+			FScopeLock CommunicationScopeLock(&Director.CommunicationLock);
+			TRefCountPtr<FCookWorkerServer>* FoundRemoteWorker = Director.RemoteWorkers.Find(Pair.Key.GetRemoteIndex());
+			if (FoundRemoteWorker)
+			{
+				RemoteWorker = *FoundRemoteWorker;
+			}
+		}
+		TArray<FPackageData*> WorkerPackageDatas;
+		WorkerPackageDatas.Reserve(Pair.Value.Num());
+		for (FName PackageName : Pair.Value)
+		{
+			FPackageData* PackageData = Director.COTFS.PackageDatas->FindPackageDataByPackageName(PackageName);
+			if (PackageData)
+			{
+				WorkerPackageDatas.Add(PackageData);
+			}
+		}
+		if (RemoteWorker)
+		{
+			// The worker(s) that sent the retraction message aborted all of the packages, so mark locally that they have been aborted
+			RemoteWorker->AbortAssignments(WorkerPackageDatas, ECookDirectorThread::SchedulerThread,
+				ENotifyRemote::LocalOnly);
+		}
+		PackageDatasToReassign.Append(MoveTemp(WorkerPackageDatas));
+	}
+
+	// Reassign the packages
+	ReassignPackages(WorkerWithResults, PackageDatasToReassign);
+
+	// Mark that we are no longer waiting
+	ExpectedWorker = FWorkerId::Invalid();
+	WorkerWithResults = FWorkerId::Invalid();
+	PackagesToRetract.Empty();
+	bOutComplete = true;
+}
+
+void FCookDirector::FRetractionHandler::HandleRetractionMessage(const FWorkerId& FromWorker, TConstArrayView<FName> Packages)
+{
+	UE_CLOG(WorkerWithResults.IsValid(), LogCook, Error,
+		TEXT("Unexpectedly received RetractionResults message from multiple CookWorkers. Merging the results."));
+	WorkerWithResults = FromWorker;
+	PackagesToRetract.FindOrAdd(FromWorker).Append(Packages);
+}
+
+void FCookDirector::FRetractionHandler::ReassignPackages(const FWorkerId& FromWorker, TConstArrayView<FPackageData*> Packages)
+{
+	TArray<TRefCountPtr<FCookWorkerServer>> LocalRemoteWorkers = Director.CopyRemoteWorkers();
+
+	TArray<FWorkerId> WorkersRequiredByConstraint;
+	TArray<FPackageData*> AssignmentPackages;
+	for (FPackageData* PackageData : Packages)
+	{
+		EPackageState State = PackageData->GetState();
+		if (State == EPackageState::Idle)
+		{
+			continue;
+		}
+		FWorkerId WorkerConstraint = PackageData->GetWorkerAssignmentConstraint();
+		if (WorkerConstraint.IsValid())
+		{
+			if (!WorkerConstraint.IsLocal() && !LocalRemoteWorkers.FindByPredicate(
+				[&WorkerConstraint](const TRefCountPtr<FCookWorkerServer>& X) { return X->GetWorkerId() == WorkerConstraint; }))
+			{
+				continue;
+			}
+			WorkersRequiredByConstraint.AddUnique(WorkerConstraint);
+		}
+
+		AssignmentPackages.Add(PackageData);
+		PackageData->SendToState(EPackageState::Request, ESendFlags::QueueRemove);
+	}
+	if (AssignmentPackages.IsEmpty())
+	{
+		UE_LOG(LogCook, Display, TEXT("Retraction results message received from %s; no packages were available for retraction."),
+			*Director.GetDisplayName(FromWorker));
+		Director.DisplayRemainingPackages();
+		return;
+	}
+
+	TArray<FWorkerId> WorkersToSplitOver = CalculateWorkersToSplitOver(AssignmentPackages.Num(), FromWorker, LocalRemoteWorkers);
+	if (WorkersToSplitOver.IsEmpty())
+	{
+		// Send the packages back to the Director for reassignment
+		FPackageDataSet& UnclusteredRequests = Director.COTFS.PackageDatas->GetRequestQueue().GetUnclusteredRequests();
+		for (FPackageData* PackageData : AssignmentPackages)
+		{
+			UnclusteredRequests.Add(PackageData);
+		}
+		// MPCOOKTODO: Add a method to PumpRequests long enough to assign the packages
+		UE_LOG(LogCook, Display, TEXT("%d packages retracted from %s. No workers are currently idle so the packages were assigned evenly to all CookWorkers."),
+			AssignmentPackages.Num(), *Director.GetDisplayName(FromWorker));
+		Director.DisplayRemainingPackages();
+		return;
+	}
+	for (const FWorkerId& WorkerId : WorkersRequiredByConstraint)
+	{
+		WorkersToSplitOver.AddUnique(WorkerId);
+	}
+
+	TStringBuilder<256> WorkerListText;
+	for (const FWorkerId& WorkerId : WorkersToSplitOver)
+	{
+		WorkerListText << Director.GetDisplayName(WorkerId) << TEXT(", ");
+	}
+	WorkerListText.RemoveSuffix(2);
+	UE_LOG(LogCook, Display, TEXT("%d packages retracted from %s and distributed to idle workers { %s }."),
+		AssignmentPackages.Num(), *Director.GetDisplayName(FromWorker), *WorkerListText);
+
+	TMap<FPackageData*, TArray<FPackageData*>> RequestGraph;
+	TArray<FWorkerId> Assignments;
+	Director.AssignRequests(MoveTemp(WorkersToSplitOver), LocalRemoteWorkers, AssignmentPackages, Assignments, MoveTemp(RequestGraph));
+	FRequestQueue& RequestQueue = Director.COTFS.PackageDatas->GetRequestQueue();
+	for (int32 Index = 0; Index < AssignmentPackages.Num(); ++Index)
+	{
+		FPackageData* PackageData = AssignmentPackages[Index];
+		FWorkerId Assignment = Assignments[Index];
+		if (Assignment.IsInvalid())
+		{
+			Director.COTFS.DemoteToIdle(*PackageData, ESendFlags::QueueAdd, ESuppressCookReason::MultiprocessAssignmentError);
+		}
+		else if (Assignment.IsLocal())
+		{
+			RequestQueue.AddReadyRequest(PackageData);
+		}
+		else
+		{
+			PackageData->SetWorkerAssignment(Assignment);
+			PackageData->SendToState(EPackageState::AssignedToWorker, ESendFlags::QueueAdd);
+		}
+	}
+	Director.DisplayRemainingPackages();
+}
+
+TArray<FWorkerId> FCookDirector::FRetractionHandler::CalculateWorkersToSplitOver(int32 NumPackages, const FWorkerId& FromWorker,
+	TConstArrayView<TRefCountPtr<FCookWorkerServer>> LocalRemoteWorkers)
+{
+	TArray<TPair<FWorkerId, int32>> WorkerNumPackages;
+	if (FromWorker != FWorkerId::Local())
+	{
+		WorkerNumPackages.Emplace(FWorkerId::Local(), Director.COTFS.NumMultiprocessLocalWorkerAssignments());
+	}
+	for (const TRefCountPtr<FCookWorkerServer>& RemoteWorker : LocalRemoteWorkers)
+	{
+		if (FromWorker != RemoteWorker->GetWorkerId())
+		{
+			WorkerNumPackages.Emplace(RemoteWorker->GetWorkerId(), RemoteWorker->NumAssignments());
+		}
+	}
+	if (WorkerNumPackages.Num() == 0)
+	{
+		return TArray<FWorkerId>();
+	}
+	WorkerNumPackages.Sort([](const TPair<FWorkerId, int32>& A, const TPair<FWorkerId, int32>& B) { return A.Value < B.Value; });
+
+	// Consider splitting the packages amonst the 1 lowest, 2 lowest, ... n lowest (not including the FromWorker)
+	// Pick the value to split over based on whichever split group results in the lowest post split maximum
+	// So splitting 500 over 0,1000,1000,1000 -> would give them all to the first, but splitting 500 over 0, 100, 1000, 1000 would
+	// split them amongst the first two.
+	int32 BestNumToSplitOver = 0;
+	int32 BestPostSplitValue = 0;
+	for (int32 NumToSplitOver = 1; NumToSplitOver <= WorkerNumPackages.Num(); ++NumToSplitOver)
+	{
+		int32 PostSplitValue = WorkerNumPackages[NumToSplitOver - 1].Value + NumPackages / NumToSplitOver;
+		if (BestNumToSplitOver == 0 || PostSplitValue < BestPostSplitValue)
+		{
+			BestNumToSplitOver = NumToSplitOver;
+			BestPostSplitValue = PostSplitValue;
+		}
+	}
+	check(BestNumToSplitOver > 0);
+	TArray<FWorkerId> Results;
+	Results.Reserve(BestNumToSplitOver);
+	for (const TPair<FWorkerId, int32>& Pair :
+		TArrayView<TPair<FWorkerId, int32>>(WorkerNumPackages).Left(BestNumToSplitOver))
+	{
+		Results.Add(Pair.Key);
+	}
+	return Results;
 }
 
 void FRetractionRequestMessage::Write(FCbWriter& Writer) const
diff --git a/Engine/Source/Editor/UnrealEd/Private/Cooker/CookDirector.h b/Engine/Source/Editor/UnrealEd/Private/Cooker/CookDirector.h
--- a/Engine/Source/Editor/UnrealEd/Private/Cooker/CookDirector.h
+++ b/Engine/Source/Editor/UnrealEd/Private/Cooker/CookDirector.h
@@ -66,14 +66,14 @@ public:
 	 * Assign the given requests out to CookWorkers (or keep on local COTFS), return the list of assignments.
 	 * Input requests have been sorted by leaf to root load order.
 	 */
-	void AssignRequests(TArrayView<UE::Cook::FPackageData*> Requests, TArray<FWorkerId>& OutAssignments,
+	void AssignRequests(TArrayView<FPackageData*> Requests, TArray<FWorkerId>& OutAssignments,
 		TMap<FPackageData*, TArray<FPackageData*>>&& RequestGraph);
 	/** Notify the CookWorker that owns the cook of the package that the Director wants to take it back. */
 	void RemoveFromWorker(FPackageData& PackageData);
 	/** Periodic tick function. Sends/Receives messages to CookWorkers. */
 	void TickFromSchedulerThread();
 	/** Periodic display function, called from CookOnTheFlyServer.UpdateDisplay. */
-	void UpdateDisplayDiagnostics();
+	void UpdateDisplayDiagnostics() const;
 	/** Called when the COTFS Server has detected all packages are complete. Tells the CookWorkers to flush messages and exit. */
 	void PumpCookComplete(bool& bOutCompleted);
 	/** Called when a session ends. The Director blocks on shutdown of all CookWorkers and returns state to before session started. */
@@ -135,6 +135,7 @@ private:
 		virtual void Stop() override;
 		FCookDirector& Director;
 	};
+	class FRetractionHandler;
 
 private:
 	/** Helper for constructor parsing. */
@@ -170,7 +171,7 @@ private:
 	/** Get the commandline to launch a worker process with. */
 	FString GetWorkerCommandLine(FWorkerId WorkerId, int32 ProfileId);
 	/** Calls the configured LoadBalanceAlgorithm. Input Requests have been sorted by leaf to root load order. */
-	void LoadBalance(TConstArrayView<TRefCountPtr<FCookWorkerServer>> Workers, TArrayView<FPackageData*> Requests,
+	void LoadBalance(TConstArrayView<FWorkerId> SortedWorkers, TArrayView<FPackageData*> Requests,
 		TMap<FPackageData*, TArray<FPackageData*>>&& RequestGraph, TArray<FWorkerId>& OutAssignments);
 
 	/** Move the given worker from active workers to the list of workers shutting down. */
@@ -182,13 +183,21 @@ private:
 	void SetWorkersStalled(bool bInWorkersStalled);
 	/** Callback for CookStats system to log our stats. */
 	void LogCookStats(FCookStatsManager::AddStatFuncRef AddStat);
+	void TickRetractionFromSchedulerThread(bool bAnyIdle, int32 BusiestNumAssignments);
 	void HandleRetractionMessage(FMPCollectorServerMessageContext& Context, bool bReadSuccessful,
 		FRetractionResultsMessage&& Message);
+	void AssignRequests(TArray<FWorkerId>&& InWorkers, TArray<TRefCountPtr<FCookWorkerServer>>& InRemoteWorkers, 
+		TArrayView<FPackageData*> Requests, TArray<FWorkerId>& OutAssignments,
+		TMap<FPackageData*, TArray<FPackageData*>>&& RequestGraph);
 
+	TArray<TRefCountPtr<FCookWorkerServer>> CopyRemoteWorkers() const;
+	void DisplayRemainingPackages() const;
+	FString GetDisplayName(const FWorkerId& WorkerId, int32 PreferredWidth = -1) const;
+	FString GetDisplayName(const FCookWorkerServer& RemoteWorker, int32 PreferredWidth=-1) const;
 
 private:
 	// Synchronization primitives that can be used from any thread
-	FCriticalSection CommunicationLock;
+	mutable FCriticalSection CommunicationLock;
 	FEventRef ShutdownEvent {EEventMode::ManualReset};
 
 	// Data only accessible from the SchedulerThread
@@ -226,6 +235,7 @@ private:
 	// Data shared between SchedulerThread and CommunicationThread that can only be accessed inside CommunicationLock
 	TMap<int32, TRefCountPtr<FCookWorkerServer>> RemoteWorkers;
 	TMap<FCookWorkerServer*, TRefCountPtr<FCookWorkerServer>> ShuttingDownWorkers;
+	TUniquePtr<FRetractionHandler> RetractionHandler;
 	bool bWorkersActive = false;
 
 	friend class UE::Cook::FCookWorkerServer;
diff --git a/Engine/Source/Editor/UnrealEd/Private/Cooker/CookTypes.h b/Engine/Source/Editor/UnrealEd/Private/Cooker/CookTypes.h
--- a/Engine/Source/Editor/UnrealEd/Private/Cooker/CookTypes.h
+++ b/Engine/Source/Editor/UnrealEd/Private/Cooker/CookTypes.h
@@ -145,6 +145,7 @@ namespace UE::Cook
 		OnlyEditorOnly,
 		CookCanceled,
 		MultiprocessAssignmentError,
+		RetractedByCookDirector,
 	};
 
 	/** The type of callback for External Requests that needs to be executed within the Scheduler's lock. */
@@ -338,7 +339,10 @@ namespace UE::Cook
 		bool IsRemote() const { return Id != InvalidId && Id != LocalId; }
 		uint8 GetRemoteIndex() const { check(IsRemote()); return Id - 1U; }
 		uint8 GetLocalOrRemoteIndex() const { check(IsValid()); return Id; }
+		bool operator==(const FWorkerId& Other) const { return Id == Other.Id; }
+		bool operator!=(const FWorkerId& Other) const { return Id != Other.Id; }
 		bool operator<(const FWorkerId& Other) const { return Id < Other.Id; }
+		inline friend int32 GetTypeHash(const FWorkerId& WorkerId) { return WorkerId.Id; }
 
 	private:
 		constexpr explicit FWorkerId(uint8 InId) : Id(InId) {}
diff --git a/Engine/Source/Editor/UnrealEd/Private/Cooker/CookWorkerClient.cpp b/Engine/Source/Editor/UnrealEd/Private/Cooker/CookWorkerClient.cpp
--- a/Engine/Source/Editor/UnrealEd/Private/Cooker/CookWorkerClient.cpp
+++ b/Engine/Source/Editor/UnrealEd/Private/Cooker/CookWorkerClient.cpp
@@ -35,6 +35,11 @@ FCookWorkerClient::FCookWorkerClient(UCookOnTheFlyServer& InCOTFS)
 		{
 			HandleRetractionMessage(Context, bReadSuccessful, MoveTemp(Message));
 		}, TEXT("HandleRetractionMessage")));
+	Register(new IMPCollectorCbClientMessage<FAbortPackagesMessage>([this]
+	(FMPCollectorClientMessageContext& Context, bool bReadSuccessful, FAbortPackagesMessage&& Message)
+		{
+			HandleAbortPackagesMessage(Context, bReadSuccessful, MoveTemp(Message));
+		}, TEXT("HandleAbortPackagesMessage")));
 }
 
 FCookWorkerClient::~FCookWorkerClient()
@@ -133,6 +138,10 @@ void FCookWorkerClient::DoneWithInitialSettings()
 
 void FCookWorkerClient::ReportDemoteToIdle(const FPackageData& PackageData, ESuppressCookReason Reason)
 {
+	if (Reason == ESuppressCookReason::RetractedByCookDirector)
+	{
+		return;
+	}
 	FPackageRemoteResult& Result = PendingResults.Emplace_GetRef();
 	Result.PackageName = PackageData.GetPackageName();
 	Result.SuppressCookReason = Reason;
@@ -661,10 +670,44 @@ void FCookWorkerClient::TickCollectors(FTickStackData& StackData, bool bFlush)
 	NextTickCollectorsTimeSeconds = FPlatformTime::Seconds() + TickCollectorsPeriodSeconds;
 }
 
+void FCookWorkerClient::HandleAbortPackagesMessage(FMPCollectorClientMessageContext& Context, bool bReadSuccessful,
+	FAbortPackagesMessage&& Message)
+{
+	if (!bReadSuccessful)
+	{
+		LogInvalidMessage(TEXT("AbortPackagesMessage"));
+	}
+	for (FName PackageName : Message.PackageNames)
+	{
+		FPackageData* PackageData = COTFS.PackageDatas->FindPackageDataByPackageName(PackageName);
+		if (PackageData)
+		{
+			COTFS.DemoteToIdle(*PackageData, ESendFlags::QueueAddAndRemove, ESuppressCookReason::RetractedByCookDirector);
+		}
+	}
+}
+
 void FCookWorkerClient::HandleRetractionMessage(FMPCollectorClientMessageContext& Context, bool bReadSuccessful,
 	FRetractionRequestMessage&& Message)
 {
+	if (!bReadSuccessful)
+	{
+		LogInvalidMessage(TEXT("RetractionRequestMessage"));
+		return;
+	}
+
+	TArray<FName> PackageNames;
+	COTFS.GetPackagesToRetract(Message.RequestedCount, PackageNames);
+	for (FName PackageName : PackageNames)
+	{
+		FPackageData* PackageData = COTFS.PackageDatas->FindPackageDataByPackageName(PackageName);
+		check(PackageData);
+		COTFS.DemoteToIdle(*PackageData, ESendFlags::QueueAddAndRemove, ESuppressCookReason::RetractedByCookDirector);
+	}
 
+	FRetractionResultsMessage ResultsMessage;
+	ResultsMessage.ReturnedPackages = MoveTemp(PackageNames);
+	SendMessage(ResultsMessage);
 }
 
 } // namespace UE::Cook
\ No newline at end of file
diff --git a/Engine/Source/Editor/UnrealEd/Private/Cooker/CookWorkerClient.h b/Engine/Source/Editor/UnrealEd/Private/Cooker/CookWorkerClient.h
--- a/Engine/Source/Editor/UnrealEd/Private/Cooker/CookWorkerClient.h
+++ b/Engine/Source/Editor/UnrealEd/Private/Cooker/CookWorkerClient.h
@@ -11,6 +11,7 @@
 
 namespace UE::Cook { class FMPCollectorClientMessageContext; }
 namespace UE::Cook { class IMPCollector; }
+namespace UE::Cook { struct FAbortPackagesMessage; }
 namespace UE::Cook { struct FAssignPackagesMessage; }
 namespace UE::Cook { struct FDirectorConnectionInfo; }
 namespace UE::Cook { struct FDiscoveredPackage; }
@@ -107,6 +108,8 @@ private:
 	void AssignPackages(FAssignPackagesMessage& Message);
 	/** Tick the registered collectors. */
 	void TickCollectors(FTickStackData& StackData, bool bFlush);
+	void HandleAbortPackagesMessage(FMPCollectorClientMessageContext& Context, bool bReadSuccessful,
+		FAbortPackagesMessage&& Message);
 	void HandleRetractionMessage(FMPCollectorClientMessageContext& Context, bool bReadSuccessful,
 		FRetractionRequestMessage&& Message);
 
diff --git a/Engine/Source/Editor/UnrealEd/Private/Cooker/CookWorkerServer.cpp b/Engine/Source/Editor/UnrealEd/Private/Cooker/CookWorkerServer.cpp
--- a/Engine/Source/Editor/UnrealEd/Private/Cooker/CookWorkerServer.cpp
+++ b/Engine/Source/Editor/UnrealEd/Private/Cooker/CookWorkerServer.cpp
@@ -70,13 +70,13 @@ void FCookWorkerServer::AppendAssignments(TArrayView<FPackageData*> Assignments,
 	PackagesToAssign.Append(Assignments);
 }
 
-void FCookWorkerServer::AbortAssignments(TSet<FPackageData*>& OutPendingPackages, ECookDirectorThread TickThread)
+void FCookWorkerServer::AbortAllAssignments(TSet<FPackageData*>& OutPendingPackages, ECookDirectorThread TickThread)
 {
 	FCommunicationScopeLock ScopeLock(this, TickThread, ETickAction::Queue);
-	AbortAssignmentsInLock(OutPendingPackages);
+	AbortAllAssignmentsInLock(OutPendingPackages);
 }
 
-void FCookWorkerServer::AbortAssignmentsInLock(TSet<FPackageData*>& OutPendingPackages)
+void FCookWorkerServer::AbortAllAssignmentsInLock(TSet<FPackageData*>& OutPendingPackages)
 {
 	if (PendingPackages.Num())
 	{
@@ -88,7 +88,7 @@ void FCookWorkerServer::AbortAssignmentsInLock(TSet<FPackageData*>& OutPendingPa
 			{
 				PackageNames.Add(PackageData->GetPackageName());
 			}
-			SendMessage(FAbortPackagesMessage(MoveTemp(PackageNames)));
+			SendMessageInLock(FAbortPackagesMessage(MoveTemp(PackageNames)));
 		}
 		OutPendingPackages.Append(MoveTemp(PendingPackages));
 		PendingPackages.Empty();
@@ -97,28 +97,43 @@ void FCookWorkerServer::AbortAssignmentsInLock(TSet<FPackageData*>& OutPendingPa
 	PackagesToAssign.Empty();
 }
 
-void FCookWorkerServer::AbortAssignment(FPackageData& PackageData, ECookDirectorThread TickThread)
+void FCookWorkerServer::AbortAssignment(FPackageData& PackageData, ECookDirectorThread TickThread,
+	ENotifyRemote NotifyRemote)
+{
+	TArray<FPackageData*> PackageDatas;
+	AbortAssignments(PackageDatas, TickThread, NotifyRemote);
+}
+
+void FCookWorkerServer::AbortAssignments(TConstArrayView<FPackageData*> PackageDatas, ECookDirectorThread TickThread,
+	ENotifyRemote NotifyRemote)
 {
 	FCommunicationScopeLock ScopeLock(this, TickThread, ETickAction::Queue);
 
-	if (PendingPackages.Remove(&PackageData))
+	TArray<FName> PackageNamesToMessage;
+	bool bSignalRemote = ConnectStatus == EConnectStatus::Connected && NotifyRemote == ENotifyRemote::NotifyRemote;
+	for (FPackageData* PackageData : PackageDatas)
 	{
-		if (ConnectStatus == EConnectStatus::Connected)
+		if (PendingPackages.Remove(PackageData))
 		{
-			TArray<FName> PackageNames;
-			PackageNames.Add(PackageData.GetPackageName());
-			SendMessage(FAbortPackagesMessage(MoveTemp(PackageNames)));
+			if (bSignalRemote)
+			{
+				PackageNamesToMessage.Add(PackageData->GetPackageName());
+			}
 		}
-	}
 
-	PackagesToAssign.Remove(&PackageData);
+		PackagesToAssign.Remove(PackageData);
+	}
+	if (!PackageNamesToMessage.IsEmpty())
+	{
+		SendMessageInLock(FAbortPackagesMessage(MoveTemp(PackageNamesToMessage)));
+	}
 }
 
 void FCookWorkerServer::AbortWorker(TSet<FPackageData*>& OutPendingPackages, ECookDirectorThread TickThread)
 {
 	FCommunicationScopeLock ScopeLock(this, TickThread, ETickAction::Tick);
 
-	AbortAssignmentsInLock(OutPendingPackages);
+	AbortAllAssignmentsInLock(OutPendingPackages);
 	switch (ConnectStatus)
 	{
 	case EConnectStatus::Uninitialized: // Fall through
@@ -128,7 +143,7 @@ void FCookWorkerServer::AbortWorker(TSet<FPackageData*>& OutPendingPackages, ECo
 	case EConnectStatus::Connected: // Fall through
 	case EConnectStatus::PumpingCookComplete:
 	{
-		SendMessage(FAbortWorkerMessage(FAbortWorkerMessage::EType::Abort));
+		SendMessageInLock(FAbortWorkerMessage(FAbortWorkerMessage::EType::Abort));
 		SendToState(EConnectStatus::WaitForDisconnect);
 		break;
 	}
@@ -162,6 +177,12 @@ void FCookWorkerServer::SendToState(EConnectStatus TargetStatus)
 	ConnectStatus = TargetStatus;
 }
 
+bool FCookWorkerServer::IsConnected() const
+{
+	FScopeLock CommunicationScopeLock(&CommunicationLock);
+	return ConnectStatus == EConnectStatus::Connected;
+}
+
 bool FCookWorkerServer::IsShuttingDown() const
 {
 	FScopeLock CommunicationScopeLock(&CommunicationLock);
@@ -213,7 +234,7 @@ bool FCookWorkerServer::TryHandleConnectMessage(FWorkerConnectMessage& Message,
 	HandleReceiveMessagesInternal();
 	const FInitialConfigMessage& InitialConfigMessage = Director.GetInitialConfigMessage();
 	OrderedSessionPlatforms = InitialConfigMessage.GetOrderedSessionPlatforms();
-	SendMessage(InitialConfigMessage);
+	SendMessageInLock(InitialConfigMessage);
 	return true;
 }
 
@@ -257,7 +278,7 @@ void FCookWorkerServer::TickCommunication(ECookDirectorThread TickThread)
 				}
 				UE_LOG(LogCook, Error, TEXT("CookWorker process of CookWorkerServer %d failed to finalize its cook within %.0f seconds; we will tell it to shutdown."),
 					ProfileId, WaitForPumpCompleteTimeout);
-				SendMessage(FAbortWorkerMessage(FAbortWorkerMessage::EType::Abort));
+				SendMessageInLock(FAbortWorkerMessage(FAbortWorkerMessage::EType::Abort));
 				SendToState(EConnectStatus::WaitForDisconnect);
 			}
 			break;
@@ -289,7 +310,7 @@ void FCookWorkerServer::SignalCookComplete(ECookDirectorThread TickThread)
 		SendToState(EConnectStatus::LostConnection);
 		break;
 	case EConnectStatus::Connected:
-		SendMessage(FAbortWorkerMessage(FAbortWorkerMessage::EType::CookComplete));
+		SendMessageInLock(FAbortWorkerMessage(FAbortWorkerMessage::EType::CookComplete));
 		SendToState(EConnectStatus::PumpingCookComplete);
 		break;
 	default:
@@ -411,7 +432,7 @@ void FCookWorkerServer::SendPendingPackages()
 	}
 	PendingPackages.Append(PackagesToAssign);
 	PackagesToAssign.Empty();
-	SendMessage(FAssignPackagesMessage(MoveTemp(AssignDatas)));
+	SendMessageInLock(FAssignPackagesMessage(MoveTemp(AssignDatas)));
 }
 
 void FCookWorkerServer::PumpReceiveMessages()
@@ -453,7 +474,7 @@ void FCookWorkerServer::HandleReceiveMessagesInternal()
 				LogCook, Error, TEXT("CookWorkerServer %d remote process shut down unexpectedly. Assigned packages will be returned to the director."),
 				ProfileId);
 
-			SendMessage(FAbortWorkerMessage(FAbortWorkerMessage::AbortAcknowledge));
+			SendMessageInLock(FAbortWorkerMessage(FAbortWorkerMessage::AbortAcknowledge));
 			SendToState(EConnectStatus::WaitForDisconnect);
 			ReceiveMessages.Reset();
 			break;
@@ -536,7 +557,13 @@ void FCookWorkerServer::HandleReceivedPackagePlatformMessages(FPackageData& Pack
 	}
 }
 
-void FCookWorkerServer::SendMessage(const UE::CompactBinaryTCP::IMessage& Message)
+void FCookWorkerServer::SendMessage(const UE::CompactBinaryTCP::IMessage& Message, ECookDirectorThread TickThread)
+{
+	FCommunicationScopeLock ScopeLock(this, TickThread, ETickAction::Tick);
+	SendMessageInLock(Message);
+}
+
+void FCookWorkerServer::SendMessageInLock(const UE::CompactBinaryTCP::IMessage& Message)
 {
 	if (TickState.TickAction == ETickAction::Tick)
 	{
diff --git a/Engine/Source/Editor/UnrealEd/Private/Cooker/CookWorkerServer.h b/Engine/Source/Editor/UnrealEd/Private/Cooker/CookWorkerServer.h
--- a/Engine/Source/Editor/UnrealEd/Private/Cooker/CookWorkerServer.h
+++ b/Engine/Source/Editor/UnrealEd/Private/Cooker/CookWorkerServer.h
@@ -32,6 +32,12 @@ namespace UE::Cook { struct FWorkerConnectMessage; }
 namespace UE::Cook
 {
 
+enum class ENotifyRemote
+{
+	NotifyRemote,
+	LocalOnly,
+};
+
 /** Class in a Director process that communicates over a Socket with FCookWorkerClient in a CookWorker process. */
 class FCookWorkerServer : public FThreadSafeRefCountedObject
 {
@@ -44,18 +50,25 @@ public:
 
 	/** Add the given assignments for the CookWorker. They will be sent during Tick */
 	void AppendAssignments(TArrayView<FPackageData*> Assignments, ECookDirectorThread TickThread);
-	/** Remove assignment of the package from local state and from the connected Client. */
-	void AbortAssignment(FPackageData& PackageData, ECookDirectorThread TickThread);
+	/** Remove assignment of the package from local state and optionally from the connected Client. */
+	void AbortAssignment(FPackageData& PackageData, ECookDirectorThread TickThread,
+		ENotifyRemote NotifyRemote = ENotifyRemote::NotifyRemote);
+	void AbortAssignments(TConstArrayView<FPackageData*> PackageData, ECookDirectorThread TickThread,
+		ENotifyRemote = ENotifyRemote::NotifyRemote);
+
 	/**
-	 * Remove assignment of the all assigned packages from local state and from the connected Client.
-	 * Report all packages that were unassigned.
+	 * Remove assignment of all assigned packages from local state and from the connected Client.
+	 * Report all packages that were removed.
 	 */
-	void AbortAssignments(TSet<FPackageData*>& OutPendingPackages, ECookDirectorThread TickThread);
-	/** AbortAssignments and tell the connected Client to gracefully terminate. Report all packages that were unassigned. */
+	void AbortAllAssignments(TSet<FPackageData*>& OutPendingPackages, ECookDirectorThread TickThread);
+	/** AbortAllAssignments and tell the connected Client to gracefully terminate. Report all packages that were unassigned. */
 	void AbortWorker(TSet<FPackageData*>& OutPendingPackages, ECookDirectorThread TickThread);
 	/** Take over the Socket for a CookWorker that has just connected. */
 	bool TryHandleConnectMessage(FWorkerConnectMessage& Message, FSocket* InSocket, TArray<UE::CompactBinaryTCP::FMarshalledMessage>&& OtherPacketMessages, ECookDirectorThread TickThread);
 
+	/** Send the message immediately to the Socket. If cannot complete immediately, it will be finished during Tick. */
+	void SendMessage(const UE::CompactBinaryTCP::IMessage& Message, ECookDirectorThread TickThread);
+
 	/** Periodic Tick function to send and receive messages to the Client. */
 	void TickCommunication(ECookDirectorThread TickThread);
 	/** Called when the COTFS Server has detected all packages are complete. Tell the CookWorker to flush messages and exit. */
@@ -66,6 +79,8 @@ public:
 	 */
 	void HandleReceiveMessages(ECookDirectorThread TickThread);
 
+	/** Is this done connecting and not yet shutting down? */
+	bool IsConnected() const;
 	/** Is this either shutting down or completed shutdown of its remote Client? */
 	bool IsShuttingDown() const;
 	/** Is this executing the portion of graceful shutdown where it waits for the CookWorker to transfer remaining messages? */
@@ -126,10 +141,10 @@ private:
 	void SendPendingPackages();
 	/** Helper for Tick, pump receive messages from a connected Client. */
 	void PumpReceiveMessages();
-	/** The main implementation of AbortAssignments, only callable from inside the lock. */
-	void AbortAssignmentsInLock(TSet<FPackageData*>& OutPendingPackages);
+	/** The main implementation of AbortAllAssignments, only callable from inside the lock. */
+	void AbortAllAssignmentsInLock(TSet<FPackageData*>& OutPendingPackages);
 	/** Send the message immediately to the Socket. If cannot complete immediately, it will be finished during Tick. */
-	void SendMessage(const UE::CompactBinaryTCP::IMessage& Message);
+	void SendMessageInLock(const UE::CompactBinaryTCP::IMessage& Message);
 	/** Send this into the given state. Update any state-dependent variables. */
 	void SendToState(EConnectStatus TargetStatus);
 	/** Close the connection and connection resources to the remote process. Does not kill the process. */
