Stagger multiple platform CachePlatformDatas for textures when doing a DDC fill with shared linear texture encoding enabled to allow it to be effective.

diff --git a/Engine/Source/Editor/UnrealEd/Classes/Commandlets/DerivedDataCacheCommandlet.h b/Engine/Source/Editor/UnrealEd/Classes/Commandlets/DerivedDataCacheCommandlet.h
--- a/Engine/Source/Editor/UnrealEd/Classes/Commandlets/DerivedDataCacheCommandlet.h
+++ b/Engine/Source/Editor/UnrealEd/Classes/Commandlets/DerivedDataCacheCommandlet.h
@@ -43,6 +43,13 @@ class UDerivedDataCacheCommandlet : public UCommandlet
 		 * expensive on some objects
 		 */
 		double LastTimeTested = 0.;
+
+		/**
+		 * If true, we wait to launch additional platforms until the first platform completes.
+		 * This is specifically so that shared linear texture encoding can reuse the base texture encode - otherwise
+		 * both platforms try to fetch and build the base texture at the same time and benefits are limited.
+		 */
+		bool bFirstPlatformIsSolo = false;
 	};
 	TMap<UObject*, FCachingData> CachingObjects;
 	TSet<FName>    ProcessedPackages;
@@ -50,6 +57,7 @@ class UDerivedDataCacheCommandlet : public UCommandlet
 	TArray<const ITargetPlatform*> Platforms;
 	double FinishCacheTime = 0.0;
 	double BeginCacheTime = 0.0;
+	bool bSharedLinearTextureEncodingEnabled = false;
 
 	class FPackageListener;
 	TUniquePtr<FPackageListener> PackageListener;
diff --git a/Engine/Source/Editor/UnrealEd/Private/Commandlets/DerivedDataCacheCommandlet.cpp b/Engine/Source/Editor/UnrealEd/Private/Commandlets/DerivedDataCacheCommandlet.cpp
--- a/Engine/Source/Editor/UnrealEd/Private/Commandlets/DerivedDataCacheCommandlet.cpp
+++ b/Engine/Source/Editor/UnrealEd/Private/Commandlets/DerivedDataCacheCommandlet.cpp
@@ -245,12 +245,32 @@ void UDerivedDataCacheCommandlet::CacheLoadedPackages(UPackage* CurrentPackage,
 				UE_TRACK_REFERENCING_PACKAGE_SCOPED(NewPackage, PackageAccessTrackingOps::NAME_CookerBuildObject);
 				for (UObject* Object : ObjectsWithOuter)
 				{
+					FCachingData& CachingData = CachingObjects.FindOrAdd(Object);
+
+					// For texture ddc fills, we want to stagger the platforms so that the base texture is only encoded
+					// once with shared linear texture encoding. If we kick everything off at the same time then all worker
+					// tasks ask for the linear encoding at the same time. If we queue the other platforms after the first one
+					// is done, then the subsequent platforms can retrieve that encoded texture and avoid a lot of work.
+					//
+					// This only actually helps if some of the platforms you're filling have tiling - otherwise
+					// you're delaying the other platforms for no reason. When I did some measuring this had no visible effect
+					// on ddc fill times - there's enough work to fill up the space, so rather than add a complicated system
+					// for determining if the platform will utilize it, we just always do it.
+					if (bSharedLinearTextureEncodingEnabled && Platforms.Num() > 1 && Object->IsA<UTexture>())
+					{
+						CachingData.bFirstPlatformIsSolo = true;
+					}
+
 					for (auto Platform : Platforms)
 					{
 						Object->BeginCacheForCookedPlatformData(Platform);
+						if (CachingData.bFirstPlatformIsSolo)
+						{
+							// We will start the other platforms later.
+							break;
+						}
 					}
 
-					FCachingData& CachingData = CachingObjects.FindOrAdd(Object);
 					CachingData.PlatformIsComplete.Init(false, Platforms.Num());
 				}
 			}
@@ -305,6 +325,20 @@ bool UDerivedDataCacheCommandlet::ProcessCachingObjects()
 							{
 								CachingData.PlatformIsComplete[PlatformIndex] = true;
 								bHadActivity = true;
+
+								if (CachingData.bFirstPlatformIsSolo)
+								{
+									bIsFinished = false;
+									CachingData.bFirstPlatformIsSolo = false;
+
+									// Start the remaining platforms now. We only launched PlatformIndex == 0.
+									check(PlatformIndex == 0);
+									check(Platforms.Num() > 1); // we shouldn't be here for only 1 platform
+									for (int32 AddingPlatformIndex = 1; AddingPlatformIndex < Platforms.Num(); AddingPlatformIndex++)
+									{
+										Object->BeginCacheForCookedPlatformData(Platforms[AddingPlatformIndex]);
+									}
+								}
 							}
 							else
 							{
@@ -430,6 +464,12 @@ int32 UDerivedDataCacheCommandlet::Main( const FString& Params )
 	FParse::Value(*Params, TEXT("SubsetTarget="), SubsetTarget);
 	bool bDoSubset = SubsetMod > 0 && SubsetTarget < SubsetMod;
 
+	static auto CVarSharedLinearTextureEncoding = IConsoleManager::Get().FindTConsoleVariableDataInt(TEXT("r.SharedLinearTextureEncoding"));
+	if (CVarSharedLinearTextureEncoding->GetValueOnAnyThread())
+	{
+		bSharedLinearTextureEncodingEnabled = true;
+	}
+
 	double FindProcessedPackagesTime = 0.0;
 	double GCTime = 0.0;
 	FinishCacheTime = 0.;
