MPCook: Do not allow retraction of packages constrained to a CookWorker.
#rnx

diff --git a/Engine/Source/Editor/UnrealEd/Private/CookOnTheFlyServer.cpp b/Engine/Source/Editor/UnrealEd/Private/CookOnTheFlyServer.cpp
--- a/Engine/Source/Editor/UnrealEd/Private/CookOnTheFlyServer.cpp
+++ b/Engine/Source/Editor/UnrealEd/Private/CookOnTheFlyServer.cpp
@@ -9525,7 +9525,6 @@ void UCookOnTheFlyServer::StartCookAsCookWorker()
 	{
 		BeginCookStartShaderCodeLibrary(BeginContext); // start shader code library cooking asynchronously; we block on it later
 	}
-	// MPCOOKTODO: Remove PlatformAssetRegistries, funnel it through API on WorkerRequests
 	RefreshPlatformAssetRegistries(BeginContext.TargetPlatforms); // Required by BeginCookSandbox stage
 
 	// Clear in-memory CookedPackages, or preserve them and cook iteratively in-process. We do not modify the CookedPackages on disk, because
@@ -9587,6 +9586,25 @@ void UCookOnTheFlyServer::GetPackagesToRetract(int32 NumToRetract, TArray<FName>
 		return;
 	}
 
+	auto AddPackageIfPossibleAndReportDone = [&OutRetractionPackages, NumToRetract](FPackageData* PackageData)
+	{
+		if (OutRetractionPackages.Num() >= NumToRetract)
+		{
+			return true;
+		}
+
+		if (PackageData->GetWorkerAssignmentConstraint().IsValid())
+		{
+			// Don't send back Packages that are constrained to this worker. Doing so will just
+			// cause the CookDirector to send it back to us, and this can cause the cooker to crash
+			// on WorldPartition packages, if we abort them and then try to restart them later.
+			return false;
+		}
+
+		OutRetractionPackages.Add(PackageData->GetPackageName());
+		return OutRetractionPackages.Num() >= NumToRetract;
+	};
+
 	FRequestQueue& RequestQueue = PackageDatas->GetRequestQueue();
 	TArray<FPackageData*> PoppedPackages;
 	if (!RequestQueue.IsReadyRequestsEmpty())
@@ -9597,11 +9615,8 @@ void UCookOnTheFlyServer::GetPackagesToRetract(int32 NumToRetract, TArray<FName>
 		}
 		for (FPackageData* PackageData : PoppedPackages)
 		{
-			if (OutRetractionPackages.Num() < NumToRetract)
-			{
-				OutRetractionPackages.Add(PackageData->GetPackageName());
-			}
 			RequestQueue.AddReadyRequest(PackageData);
+			AddPackageIfPossibleAndReportDone(PackageData);
 		}
 	}
 	if (OutRetractionPackages.Num() >= NumToRetract)
@@ -9612,32 +9627,28 @@ void UCookOnTheFlyServer::GetPackagesToRetract(int32 NumToRetract, TArray<FName>
 	FLoadPrepareQueue& LoadPrepareQueue = PackageDatas->GetLoadPrepareQueue();
 	for (FPackageData* PackageData : LoadPrepareQueue.EntryQueue)
 	{
-		OutRetractionPackages.Add(PackageData->GetPackageName());
-		if (OutRetractionPackages.Num() >= NumToRetract)
+		if (AddPackageIfPossibleAndReportDone(PackageData))
 		{
 			return;
 		}
 	}
 	for (FPackageData* PackageData : LoadPrepareQueue.PreloadingQueue)
 	{
-		OutRetractionPackages.Add(PackageData->GetPackageName());
-		if (OutRetractionPackages.Num() >= NumToRetract)
+		if (AddPackageIfPossibleAndReportDone(PackageData))
 		{
 			return;
 		}
 	}
 	for (FPackageData* PackageData : PackageDatas->GetLoadReadyQueue())
 	{
-		OutRetractionPackages.Add(PackageData->GetPackageName());
-		if (OutRetractionPackages.Num() >= NumToRetract)
+		if (AddPackageIfPossibleAndReportDone(PackageData))
 		{
 			return;
 		}
 	}
 	for (FPackageData* PackageData : PackageDatas->GetSaveQueue())
 	{
-		OutRetractionPackages.Add(PackageData->GetPackageName());
-		if (OutRetractionPackages.Num() >= NumToRetract)
+		if (AddPackageIfPossibleAndReportDone(PackageData))
 		{
 			return;
 		}
