MPCook: Fix bug with packages not being removed from CookWorker when retracted from local worker.
#rnx

diff --git a/Engine/Source/Editor/UnrealEd/Private/Cooker/CookDirector.cpp b/Engine/Source/Editor/UnrealEd/Private/Cooker/CookDirector.cpp
--- a/Engine/Source/Editor/UnrealEd/Private/Cooker/CookDirector.cpp
+++ b/Engine/Source/Editor/UnrealEd/Private/Cooker/CookDirector.cpp
@@ -198,7 +198,7 @@ FCookDirector::~FCookDirector()
 	for (FPackageData* PackageData : AbortedAssignments)
 	{
 		check(PackageData->IsInProgress()); // Packages that were assigned to workers should be in the AssignedToWorker state
-		PackageData->SetWorkerAssignment(FWorkerId::Invalid());
+		PackageData->SetWorkerAssignment(FWorkerId::Invalid(), ESendFlags::QueueNone);
 		PackageData->SendToState(UE::Cook::EPackageState::Request, ESendFlags::QueueAddAndRemove);
 	}
 	RemoteWorkers.Empty();
@@ -349,7 +349,7 @@ void FCookDirector::AssignRequests(TArray<FWorkerId>&& InWorkers, TArray<TRefCou
 	TArray<bool> RemoteIndexIsValid; // Indexed by WorkerId.GetRemoteIndex()
 	RemoteBatches.SetNum(MaxRemoteIndex+1);
 	RemoteIndexIsValid.Init(false, MaxRemoteIndex+1);
-	for (FWorkerId& WorkerId : InWorkers)
+	for (FWorkerId WorkerId : InWorkers)
 	{
 		if (!WorkerId.IsLocal())
 		{
@@ -359,7 +359,7 @@ void FCookDirector::AssignRequests(TArray<FWorkerId>&& InWorkers, TArray<TRefCou
 
 	for (int32 RequestIndex = 0; RequestIndex < Requests.Num(); ++RequestIndex)
 	{
-		FWorkerId WorkerId = OutAssignments[RequestIndex];
+		FWorkerId& WorkerId = OutAssignments[RequestIndex];
 		// Override the loadbalancer's assignment if the Package has a WorkerAssignmentConstraint
 		// This allows us to guarantee that generated packages will be cooked on the worker that cooked
 		// their generator package
@@ -367,7 +367,6 @@ void FCookDirector::AssignRequests(TArray<FWorkerId>&& InWorkers, TArray<TRefCou
 		if (WorkerIdConstraint.IsValid())
 		{
 			check(InWorkers.Contains(WorkerIdConstraint));
-			OutAssignments[RequestIndex] = WorkerIdConstraint;
 			WorkerId = WorkerIdConstraint;
 		}
 
@@ -379,7 +378,7 @@ void FCookDirector::AssignRequests(TArray<FWorkerId>&& InWorkers, TArray<TRefCou
 			{
 				UE_LOG(LogCook, Error, TEXT("Package %s can only be cooked by a now-disconnected CookWorker. The package can not be cooked."),
 					*Requests[RequestIndex]->GetPackageName().ToString());
-				OutAssignments[RequestIndex] = FWorkerId::Invalid();
+				WorkerId = FWorkerId::Invalid();
 				continue;
 			}
 			TArray<FPackageData*>& RemoteBatch = RemoteBatches[RemoteIndex];
@@ -395,7 +394,7 @@ void FCookDirector::AssignRequests(TArray<FWorkerId>&& InWorkers, TArray<TRefCou
 
 	// Assign each batch to the FCookWorkerServer in RemoteWorkers;
 	// the CookWorkerServer's tick will handle sending the message to the remote process
-	for (FWorkerId& WorkerId : InWorkers)
+	for (FWorkerId WorkerId : InWorkers)
 	{
 		if (!WorkerId.IsLocal())
 		{
@@ -423,18 +422,24 @@ TArray<TRefCountPtr<FCookWorkerServer>> FCookDirector::CopyRemoteWorkers() const
 
 void FCookDirector::RemoveFromWorker(FPackageData& PackageData)
 {
-	TArray<FCookWorkerServer*> Workers;
+	FWorkerId WorkerId = PackageData.GetWorkerAssignment();
+	if (!WorkerId.IsRemote())
+	{
+		return;
+	}
+
+	TRefCountPtr<FCookWorkerServer> OwningWorker;
 	{
 		FScopeLock CommunicationScopeLock(&CommunicationLock);
-		for (TPair<int32, TRefCountPtr<FCookWorkerServer>>& Pair : RemoteWorkers)
+		TRefCountPtr<FCookWorkerServer>* RemoteWorkerPtr = RemoteWorkers.Find(WorkerId.GetRemoteIndex());
+		if (!RemoteWorkerPtr)
 		{
-			Workers.Add(Pair.Value);
+			return;
 		}
+		OwningWorker = *RemoteWorkerPtr;
 	}
-	for (FCookWorkerServer* Worker : Workers)
-	{
-		Worker->AbortAssignment(PackageData, ECookDirectorThread::SchedulerThread);
-	}
+
+	OwningWorker->AbortAssignment(PackageData, ECookDirectorThread::SchedulerThread);
 }
 
 void FCookDirector::TickFromSchedulerThread()
@@ -1278,19 +1283,24 @@ void FCookDirector::HandleRetractionMessage(FMPCollectorServerMessageContext& Co
 
 void FCookDirector::TickRetractionFromSchedulerThread(bool bAnyIdle, int32 BusiestNumAssignments)
 {
-	if (RetractionHandler)
+	if (!RetractionHandler)
 	{
-		bool bComplete;
-		RetractionHandler->TickFromSchedulerThread(bAnyIdle, bComplete);
-		if (bComplete)
+		if (bAnyIdle && BusiestNumAssignments > RetractionMinimumNumAssignments)
 		{
-			RetractionHandler.Reset();
+			RetractionHandler = MakeUnique<FRetractionHandler>(*this);
+			RetractionHandler->Initialize();
+		}
+		else
+		{
+			return;
 		}
 	}
-	else if (bAnyIdle && BusiestNumAssignments > RetractionMinimumNumAssignments)
+
+	bool bComplete;
+	RetractionHandler->TickFromSchedulerThread(bAnyIdle, bComplete);
+	if (bComplete)
 	{
-		RetractionHandler = MakeUnique<FRetractionHandler>(*this);
-		RetractionHandler->Initialize();
+		RetractionHandler.Reset();
 	}
 }
 
@@ -1548,8 +1558,8 @@ void FCookDirector::FRetractionHandler::ReassignPackages(const FWorkerId& FromWo
 		}
 		else
 		{
-			PackageData->SetWorkerAssignment(Assignment);
 			PackageData->SendToState(EPackageState::AssignedToWorker, ESendFlags::QueueAdd);
+			PackageData->SetWorkerAssignment(Assignment);
 		}
 	}
 	Director.DisplayRemainingPackages();
diff --git a/Engine/Source/Editor/UnrealEd/Private/Cooker/CookPackageData.cpp b/Engine/Source/Editor/UnrealEd/Private/Cooker/CookPackageData.cpp
--- a/Engine/Source/Editor/UnrealEd/Private/Cooker/CookPackageData.cpp
+++ b/Engine/Source/Editor/UnrealEd/Private/Cooker/CookPackageData.cpp
@@ -727,14 +727,34 @@ void FPackageData::OnEnterAssignedToWorker()
 {
 }
 
-void FPackageData::OnExitAssignedToWorker()
+void FPackageData::SetWorkerAssignment(FWorkerId InWorkerAssignment, ESendFlags SendFlags)
 {
-	if (GetWorkerAssignment().IsValid())
+	if (WorkerAssignment.IsValid())
+	{
+		checkf(InWorkerAssignment.IsInvalid(), TEXT("Package %s is being assigned to worker %d while it is already assigned to worker %d."),
+			*GetPackageName().ToString(), WorkerAssignment.GetRemoteIndex(), WorkerAssignment.GetRemoteIndex());
+		if (EnumHasAnyFlags(SendFlags, ESendFlags::QueueRemove))
+		{
+			PackageDatas.GetCookOnTheFlyServer().NotifyRemovedFromWorker(*this);
+		}
+		WorkerAssignment = FWorkerId::Invalid();
+	}
+	else
 	{
-		PackageDatas.GetCookOnTheFlyServer().NotifyRemovedFromWorker(*this);
+		if (InWorkerAssignment.IsValid())
+		{
+			checkf(GetState() == EPackageState::AssignedToWorker, TEXT("Package %s is being assigned to worker %d while in a state other than AssignedToWorker. This is invalid."),
+				*GetPackageName().ToString(), GetWorkerAssignment().GetRemoteIndex());
+		}
+		WorkerAssignment = InWorkerAssignment;
 	}
 }
 
+void FPackageData::OnExitAssignedToWorker()
+{
+	SetWorkerAssignment(FWorkerId::Invalid());
+}
+
 void FPackageData::OnEnterLoadPrepare()
 {
 }
diff --git a/Engine/Source/Editor/UnrealEd/Private/Cooker/CookPackageData.h b/Engine/Source/Editor/UnrealEd/Private/Cooker/CookPackageData.h
--- a/Engine/Source/Editor/UnrealEd/Private/Cooker/CookPackageData.h
+++ b/Engine/Source/Editor/UnrealEd/Private/Cooker/CookPackageData.h
@@ -484,8 +484,11 @@ public:
 
 	/** For MultiProcessCooks, Get the id of the worker this Package is assigned to; InvalidId means owned by local. */
 	FWorkerId GetWorkerAssignment() const { return WorkerAssignment; }
-	/** Set the id of the worker this Package is assigned to. */
-	void SetWorkerAssignment(FWorkerId InWorkerAssignment) { WorkerAssignment = InWorkerAssignment; }
+	/**
+	 * Set the id of the worker this Package is assigned to. If value changes from Valid to Invalid and SendFlags
+	 * includes QueueRemove, also calls NotifyRemovedFromWorker.
+	 */
+	void SetWorkerAssignment(FWorkerId InWorkerAssignment, ESendFlags SendFlags = ESendFlags::QueueAddAndRemove);
 	/** Get the workerid that is the only worker allowed to cook this package; InvalidId means no constraint. */
 	FWorkerId GetWorkerAssignmentConstraint() const { return WorkerAssignmentConstraint; }
 	/** Set the workerid that is the only worker allowed to cook this package; default is InvalidId; */
diff --git a/Engine/Source/Editor/UnrealEd/Private/Cooker/CookWorkerServer.cpp b/Engine/Source/Editor/UnrealEd/Private/Cooker/CookWorkerServer.cpp
--- a/Engine/Source/Editor/UnrealEd/Private/Cooker/CookWorkerServer.cpp
+++ b/Engine/Source/Editor/UnrealEd/Private/Cooker/CookWorkerServer.cpp
@@ -100,7 +100,8 @@ void FCookWorkerServer::AbortAllAssignmentsInLock(TSet<FPackageData*>& OutPendin
 void FCookWorkerServer::AbortAssignment(FPackageData& PackageData, ECookDirectorThread TickThread,
 	ENotifyRemote NotifyRemote)
 {
-	TArray<FPackageData*> PackageDatas;
+	TArray<FPackageData*, TInlineAllocator<1>> PackageDatas;
+	PackageDatas.Add(&PackageData);
 	AbortAssignments(PackageDatas, TickThread, NotifyRemote);
 }
 
@@ -595,7 +596,7 @@ void FCookWorkerServer::RecordResults(FPackageResultsMessage& Message)
 				ProfileId, *Result.PackageName.ToString());
 			continue;
 		}
-		PackageData->SetWorkerAssignment(FWorkerId::Invalid());
+		PackageData->SetWorkerAssignment(FWorkerId::Invalid(), ESendFlags::QueueNone);
 
 		// MPCOOKTODO: Refactor FSaveCookedPackageContext::FinishPlatform and ::FinishPackage so we can call them from here
 		// to reduce duplication
