[shaders] changes to CompileShadersTestbedCommandlet
-  add an "all" option to the compile shaders commandlet which builds all materials & material instances found in a project
- fix potential issue causing uploads of incomplete shader maps to DDC (resulting in non-determinism warnings) - the "IsCachedCookedPlatformDataLoaded" function is actually what registers a material dependency with the associated shader map, so we must call this at least once for all materials/material instances before any calls to ProcessAsyncResults to avoid the possibility of shader maps thinking they are complete due to some dependencies finishing before all have been registered.
- minor cleanup to -collection path to simplify filter creation
- minor updates to log messages for the sake of clarity

diff --git a/Engine/Source/Editor/UnrealEd/Private/Commandlets/CompileShadersTestBedCommandlet.cpp b/Engine/Source/Editor/UnrealEd/Private/Commandlets/CompileShadersTestBedCommandlet.cpp
--- a/Engine/Source/Editor/UnrealEd/Private/Commandlets/CompileShadersTestBedCommandlet.cpp
+++ b/Engine/Source/Editor/UnrealEd/Private/Commandlets/CompileShadersTestBedCommandlet.cpp
@@ -38,7 +38,8 @@ int32 UCompileShadersTestBedCommandlet::Main(const FString& Params)
 		UE_LOG(LogCompileShadersTestBedCommandlet, Log, TEXT("CompileShadersTestBed"));
 		UE_LOG(LogCompileShadersTestBedCommandlet, Log, TEXT("This commandlet compiles global and default material shaders.  Used to profile and test shader compilation."));
 		UE_LOG(LogCompileShadersTestBedCommandlet, Log, TEXT(" Optional: -collection=<name>                (You can also specify a collection of assets to narrow down the results e.g. if you maintain a collection that represents the actually used in-game assets)."));
-		UE_LOG(LogCompileShadersTestBedCommandlet, Log, TEXT(" Optional: -materials=<path1>+<path2>        (You can also specify a list of material asset paths separated by a '+' to narrow down the results."));
+		UE_LOG(LogCompileShadersTestBedCommandlet, Log, TEXT(" Optional: -materials=<path1>+<path2>        (You can also specify a list of material asset paths separated by a '+' to narrow down the results.)"));
+		UE_LOG(LogCompileShadersTestBedCommandlet, Log, TEXT(" Optional: -all                              (You can specify -all to compile all global/default shaders as well as shaders for all materials/material instances found in a project.)"))
 		return 0;
 	}
 
@@ -49,56 +50,62 @@ int32 UCompileShadersTestBedCommandlet::Main(const FString& Params)
 	TArray<FAssetData> MaterialList;
 
 	FARFilter Filter;
-
-	// Parse collection
-	FString CollectionName;
-	if (FParse::Value(*Params, TEXT("collection="), CollectionName, true))
+	Filter.bRecursiveClasses = true;
+	// Find all materials/material instances if -all is specified
+	if (Switches.Contains(TEXT("all")))
+	{
+		Filter.ClassPaths.Add(UMaterial::StaticClass()->GetClassPathName());
+		Filter.ClassPaths.Add(UMaterialInstance::StaticClass()->GetClassPathName());
+		AssetRegistry.GetAssets(Filter, MaterialList);
+		UE_LOG(LogCompileShadersTestBedCommandlet, Log, TEXT("Found %d materials/material instances in project."), MaterialList.Num());
+	}
+	else // otherwise parse -collection and -materials arguments
 	{
-		if (!CollectionName.IsEmpty())
+		// Parse collection
+		FString CollectionName;
+		if (FParse::Value(*Params, TEXT("collection="), CollectionName, true))
 		{
-			// Get the list of materials from a collection
-			Filter.PackagePaths.Add(FName(TEXT("/Game")));
-			Filter.bRecursivePaths = true;
-			Filter.ClassPaths.Add(UMaterial::StaticClass()->GetClassPathName());
-
-			FCollectionManagerModule& CollectionManagerModule = FCollectionManagerModule::GetModule();
-			CollectionManagerModule.Get().GetObjectsInCollection(FName(*CollectionName), ECollectionShareType::CST_All, Filter.SoftObjectPaths, ECollectionRecursionFlags::SelfAndChildren);
-
-			AssetRegistry.GetAssets(Filter, MaterialList);
+			if (!CollectionName.IsEmpty())
+			{
+				// Get the list of materials from a collection
+				Filter.PackagePaths.Add(FName(TEXT("/Game")));
+				Filter.bRecursivePaths = true;
+				Filter.ClassPaths.Add(UMaterial::StaticClass()->GetClassPathName());
+				Filter.ClassPaths.Add(UMaterialInstance::StaticClass()->GetClassPathName());
 
-			Filter.ClassPaths.Empty();
-			Filter.ClassPaths.Add(UMaterialInstance::StaticClass()->GetClassPathName());
-			Filter.ClassPaths.Add(UMaterialInstanceConstant::StaticClass()->GetClassPathName());
+				FCollectionManagerModule& CollectionManagerModule = FCollectionManagerModule::GetModule();
+				CollectionManagerModule.Get().GetObjectsInCollection(FName(*CollectionName), ECollectionShareType::CST_All, Filter.SoftObjectPaths, ECollectionRecursionFlags::SelfAndChildren);
 
-			AssetRegistry.GetAssets(Filter, MaterialList);
+				AssetRegistry.GetAssets(Filter, MaterialList);
+				UE_LOG(LogCompileShadersTestBedCommandlet, Log, TEXT("Found %d materials/material instances from collection %s."), MaterialList.Num(), *CollectionName);
+			}
+		}
+		// Process -materials= switches separated by a '+'
+		TArray<FString> CmdLineMaterialEntries;
+		const TCHAR* MaterialsSwitchName = TEXT("Materials");
+		if (const FString* MaterialsSwitches = ParamVals.Find(MaterialsSwitchName))
+		{
+			MaterialsSwitches->ParseIntoArray(CmdLineMaterialEntries, TEXT("+"));
 		}
-	}
-
-	UE_LOG(LogCompileShadersTestBedCommandlet, Log, TEXT("Found %d/%d Materials."), MaterialList.Num(), Filter.SoftObjectPaths.Num());
-
-	// Process -materials= switches separated by a '+'
-	TArray<FString> CmdLineMaterialEntries;
-	const TCHAR* MaterialsSwitchName = TEXT("Materials");
-	if (const FString* MaterialsSwitches = ParamVals.Find(MaterialsSwitchName))
-	{
-		MaterialsSwitches->ParseIntoArray(CmdLineMaterialEntries, TEXT("+"));
-	}
 
-	if (CmdLineMaterialEntries.Num())
-	{
-		// re-use the filter and only filter based on the passed in objects.
-		Filter.ClassPaths.Empty();
-		Filter.SoftObjectPaths.Empty();
-		for (const FString& MaterialPathString : CmdLineMaterialEntries)
+		if (CmdLineMaterialEntries.Num())
 		{
-			const FSoftObjectPath MaterialPath(MaterialPathString);
-			if (!Filter.SoftObjectPaths.Contains(MaterialPath))
+			// re-use the filter and only filter based on the passed in objects.
+			Filter.ClassPaths.Empty();
+			Filter.SoftObjectPaths.Empty();
+			int MaterialsNumBefore = MaterialList.Num();
+			for (const FString& MaterialPathString : CmdLineMaterialEntries)
 			{
-				Filter.SoftObjectPaths.Add(MaterialPath);
+				const FSoftObjectPath MaterialPath(MaterialPathString);
+				if (!Filter.SoftObjectPaths.Contains(MaterialPath))
+				{
+					Filter.SoftObjectPaths.Add(MaterialPath);
+				}
 			}
-		}
 
-		AssetRegistry.GetAssets(Filter, MaterialList);
+			AssetRegistry.GetAssets(Filter, MaterialList);
+			UE_LOG(LogCompileShadersTestBedCommandlet, Log, TEXT("Found %d/%d requested materials/material instances."), MaterialList.Num() - MaterialsNumBefore, Filter.SoftObjectPaths.Num());
+		}
 	}
 
 	static constexpr bool bLimitExecutationTime = false;
@@ -163,7 +170,14 @@ int32 UCompileShadersTestBedCommandlet::Main(const FString& Params)
 				{
 					UE_LOG(LogCompileShadersTestBedCommandlet, Display, TEXT("BeginCache for %s"), *MaterialInterface->GetFullName());
 					MaterialInterface->BeginCacheForCookedPlatformData(Platform);
-					MaterialsToCompile.Add(MaterialInterface);
+					// need to call this once for all objects before any calls to ProcessAsyncResults as otherwise we'll potentially upload
+					// incremental/incomplete shadermaps to DDC (as this function actually triggers compilation, some compiles for a particular
+					// material may finish before we've even started others - if we call ProcessAsyncResults in that case the associated shader
+					// maps will think they are "finished" due to having no outstanding dependencies).
+					if (!MaterialInterface->IsCachedCookedPlatformDataLoaded(Platform))
+					{
+						MaterialsToCompile.Add(MaterialInterface);
+					}
 				}
 			}
 		}
