DebugSkelMeshComponent: Clear preview instance correctly

diff --git a/Engine/Plugins/Animation/ControlRig/Source/ControlRig/Private/Sequencer/ControlRigLayerInstanceProxy.h b/Engine/Plugins/Animation/ControlRig/Source/ControlRig/Private/Sequencer/ControlRigLayerInstanceProxy.h
--- a/Engine/Plugins/Animation/ControlRig/Source/ControlRig/Private/Sequencer/ControlRigLayerInstanceProxy.h
+++ b/Engine/Plugins/Animation/ControlRig/Source/ControlRig/Private/Sequencer/ControlRigLayerInstanceProxy.h
@@ -17,6 +17,7 @@ struct FAnimNode_ControlRigInputPose : public FAnimNode_Base
 
 	FAnimNode_ControlRigInputPose()
 		: InputProxy(nullptr)
+		, InputAnimInstance(nullptr)
 	{
 	}
 
diff --git a/Engine/Source/Editor/UnrealEd/Classes/Animation/DebugSkelMeshComponent.h b/Engine/Source/Editor/UnrealEd/Classes/Animation/DebugSkelMeshComponent.h
--- a/Engine/Source/Editor/UnrealEd/Classes/Animation/DebugSkelMeshComponent.h
+++ b/Engine/Source/Editor/UnrealEd/Classes/Animation/DebugSkelMeshComponent.h
@@ -366,6 +366,7 @@ class UNREALED_API UDebugSkelMeshComponent : public USkeletalMeshComponent
 	virtual void InitAnim(bool bForceReinit) override;
 	virtual bool IsWindEnabled() const override { return true; }
 	virtual void SetAnimClass(class UClass* NewClass) override;
+	virtual void OnClearAnimScriptInstance() override;
 	//~ End SkeletalMeshComponent Interface
 
 	// return true if currently preview animation asset is on
diff --git a/Engine/Source/Editor/UnrealEd/Private/Animation/DebugSkelMeshComponent.cpp b/Engine/Source/Editor/UnrealEd/Private/Animation/DebugSkelMeshComponent.cpp
--- a/Engine/Source/Editor/UnrealEd/Private/Animation/DebugSkelMeshComponent.cpp
+++ b/Engine/Source/Editor/UnrealEd/Private/Animation/DebugSkelMeshComponent.cpp
@@ -564,6 +564,14 @@ void UDebugSkelMeshComponent::SetAnimClass(class UClass* NewClass)
 	UE_LOG(LogAnimation, Warning, TEXT("Attempting to destroy an animation preview actor, skipping."));
 }
 
+void UDebugSkelMeshComponent::OnClearAnimScriptInstance()
+{
+	// call to super not strictly necessary (since it is empty)
+	Super::OnClearAnimScriptInstance();
+	
+	PreviewInstance = nullptr;
+}
+
 void UDebugSkelMeshComponent::EnablePreview(bool bEnable, UAnimationAsset* PreviewAsset)
 {
 	if (PreviewInstance)
diff --git a/Engine/Source/Runtime/Engine/Classes/Components/SkeletalMeshComponent.h b/Engine/Source/Runtime/Engine/Classes/Components/SkeletalMeshComponent.h
--- a/Engine/Source/Runtime/Engine/Classes/Components/SkeletalMeshComponent.h
+++ b/Engine/Source/Runtime/Engine/Classes/Components/SkeletalMeshComponent.h
@@ -2318,6 +2318,8 @@ protected:
 
 private:
 
+	virtual void OnClearAnimScriptInstance() {};
+
 	friend struct FSkeletalMeshComponentEndPhysicsTickFunction;
 
 	/** Update systems after physics sim is done */
diff --git a/Engine/Source/Runtime/Engine/Private/Components/SkeletalMeshComponent.cpp b/Engine/Source/Runtime/Engine/Private/Components/SkeletalMeshComponent.cpp
--- a/Engine/Source/Runtime/Engine/Private/Components/SkeletalMeshComponent.cpp
+++ b/Engine/Source/Runtime/Engine/Private/Components/SkeletalMeshComponent.cpp
@@ -1012,6 +1012,7 @@ void USkeletalMeshComponent::ClearAnimScriptInstance()
 	AnimScriptInstance = nullptr;
 	ResetLinkedAnimInstances();
 	ClearCachedAnimProperties();
+	OnClearAnimScriptInstance();
 }
 
 void USkeletalMeshComponent::ClearCachedAnimProperties()
