EditorDomain: Change global EditorDomain version to be different between UE5Main and other streams so that we don't load incompatible EditorDomain packages due to lack of handling of redirectors in the package digest.
#rnx

diff --git a/Engine/Source/Editor/UnrealEd/Private/EditorDomain/EditorDomainUtils.cpp b/Engine/Source/Editor/UnrealEd/Private/EditorDomain/EditorDomainUtils.cpp
--- a/Engine/Source/Editor/UnrealEd/Private/EditorDomain/EditorDomainUtils.cpp
+++ b/Engine/Source/Editor/UnrealEd/Private/EditorDomain/EditorDomainUtils.cpp
@@ -203,7 +203,7 @@ FBlake3Hash GGlobalConstructClassesHash;
 int64 GMaxBulkDataSize = -1;
 
 // Change to a new guid when EditorDomain needs to be invalidated
-const TCHAR* EditorDomainVersion = TEXT("74B1CEDEF6774CC89F3A45BEA18B9B8B");
+const TCHAR* EditorDomainVersion = TEXT("382517D23ED840A48501AE70D95B3CE2");
 
 // Identifier of the CacheBuckets for EditorDomain tables
 const TCHAR* EditorDomainPackageBucketName = TEXT("EditorDomainPackage");
