Blueprint Compile: Fix ReplaceInstancesOfClass_Inner to handle Actors with nonnull GetLevel() but a nullptr GetWorld().
#rnx

diff --git a/Engine/Source/Editor/UnrealEd/Private/Kismet2/KismetReinstanceUtilities.cpp b/Engine/Source/Editor/UnrealEd/Private/Kismet2/KismetReinstanceUtilities.cpp
--- a/Engine/Source/Editor/UnrealEd/Private/Kismet2/KismetReinstanceUtilities.cpp
+++ b/Engine/Source/Editor/UnrealEd/Private/Kismet2/KismetReinstanceUtilities.cpp
@@ -2382,7 +2382,7 @@ void FBlueprintCompileReinstancer::ReplaceInstancesOfClass_Inner(const TMap<UCla
 					if (OldActor != nullptr)
 					{
 						UObject* NewUObject = nullptr;
-						if (OldActor->GetLevel())
+						if (OldActor->GetLevel() && OldActor->GetWorld())
 						{
 							ReplaceActorHelper(OldActor, OldClass, NewUObject, NewClass, OldToNewInstanceMap, InOldToNewClassMap, ReinstancedObjectsWeakReferenceMap, ActorAttachmentData, ReplacementActors, bPreserveRootComponent, bSelectionChanged);
 						}
