Cooker: Fix pure virtual function call when exiting the editor. It is caused by TAllocatorFixedSizeFreeList being used by global FHierarchicalTimerInfo being destructed before FHierarchicalTimerInfo.
#rnx

diff --git a/Engine/Source/Editor/UnrealEd/Private/CookOnTheFlyServer.cpp b/Engine/Source/Editor/UnrealEd/Private/CookOnTheFlyServer.cpp
--- a/Engine/Source/Editor/UnrealEd/Private/CookOnTheFlyServer.cpp
+++ b/Engine/Source/Editor/UnrealEd/Private/CookOnTheFlyServer.cpp
@@ -461,6 +461,11 @@ UCookOnTheFlyServer::~UCookOnTheFlyServer()
 	FCoreUObjectDelegates::GetPreGarbageCollectDelegate().RemoveAll(this);
 	FCoreUObjectDelegates::GetPostGarbageCollect().RemoveAll(this);
 	GetTargetPlatformManager()->GetOnTargetPlatformsInvalidatedDelegate().RemoveAll(this);
+
+	if (HasAnyFlags(RF_ClassDefaultObject))
+	{
+		ClearHierarchyTimers();
+	}
 }
 
 // This tick only happens in the editor.  The cook commandlet directly calls tick on the side.
@@ -1317,12 +1322,11 @@ uint32 UCookOnTheFlyServer::TickCookWorker()
 
 void UCookOnTheFlyServer::TickMainCookLoop(UE::Cook::FTickStackData& StackData)
 {
-	UE_SCOPED_HIERARCHICAL_COOKTIMER(TickMainCookLoop);
 	if (!IsInSession())
 	{
 		return;
 	}
-
+	UE_SCOPED_HIERARCHICAL_COOKTIMER(TickMainCookLoop);
 	bool bContinueTick = true;
 	while (bContinueTick && (!IsEngineExitRequested() || (IsCookByTheBookMode() && !IsCookingInEditor())))
 	{
