BulkDataRegistry: Second attempt to remove repeated misses for BulkDataPayloadIds.

1) Change the key for BulkDataPayloadIds to include only the PackageGuid and BulkDataId to remove spurious invalidations when EditorDomain PackageDigest changes.

2) Load the BulkData and issue a Put every time a Get misses.

#rnx

diff --git a/Engine/Source/Editor/UnrealEd/Private/EditorDomain/BulkDataRegistryImpl.cpp b/Engine/Source/Editor/UnrealEd/Private/EditorDomain/BulkDataRegistryImpl.cpp
--- a/Engine/Source/Editor/UnrealEd/Private/EditorDomain/BulkDataRegistryImpl.cpp
+++ b/Engine/Source/Editor/UnrealEd/Private/EditorDomain/BulkDataRegistryImpl.cpp
@@ -327,37 +327,48 @@ void FBulkDataRegistryImpl::UpdatePlaceholderPayloadId(const UE::Serialization::
 	}
 	const FGuid& Key = BulkData.GetIdentifier();
 
-	FWriteScopeLock RegistryScopeLock(RegistryLock);
-	if (!bActive)
-	{
-		return;
-	}
-	FRegisteredBulk* Existing = Registry.Find(Key);
-	if (!Existing)
+	TOptional<UE::Serialization::FEditorBulkData> WritePayloadIdBulkData;
+	FName WritePayloadIdPackageName;
 	{
-		return;
-	}
-	if (!Existing->bRegistered ||
-		!Existing->BulkData.GetIdentifier().IsValid() ||
-		!Existing->BulkData.HasPlaceholderPayloadId() ||
-		!BulkData.LocationMatches(Existing->BulkData))
-	{
-		return;
-	}
+		FWriteScopeLock RegistryScopeLock(RegistryLock);
+		if (!bActive)
+		{
+			return;
+		}
+		FRegisteredBulk* Existing = Registry.Find(Key);
+		if (!Existing)
+		{
+			return;
+		}
+		if (!Existing->bRegistered ||
+			!Existing->BulkData.GetIdentifier().IsValid() ||
+			!Existing->BulkData.HasPlaceholderPayloadId() ||
+			!BulkData.LocationMatches(Existing->BulkData))
+		{
+			return;
+		}
 
-	Existing->BulkData = BulkData.CopyTornOff();
-	Existing->bAllowedToWritePayloadIdToCache = Existing->bAllowedToWritePayloadIdToCache && Existing->BulkData.CanSaveForRegistry();
-	if (Existing->bAllowedToWritePayloadIdToCache)
-	{
-		// If there is a Get still pending the PayloadId may already be in the cache. Do not Put it until we find out whether
-		// it exists. The response lambda in ReadPayloadIdsFromCache will issue the WritePayloadIdToCache if it is a miss.
-		if (!PendingPayloadIds.Find(Key))
+		Existing->BulkData = BulkData.CopyTornOff();
+		Existing->bAllowedToWritePayloadIdToCache = Existing->bAllowedToWritePayloadIdToCache && Existing->BulkData.CanSaveForRegistry();
+		if (Existing->bAllowedToWritePayloadIdToCache)
 		{
-			WritePayloadIdToCache(Existing->PackageName, Existing->BulkData);
-			// Set that it is no longer allowed to write PayloadIdToCache since we have already written it.
-			Existing->bAllowedToWritePayloadIdToCache = false;
+			// If there is a Get still pending the PayloadId may already be in the cache. Do not Put it until we find out whether
+			// it exists. The response lambda in ReadPayloadIdsFromCache will issue the WritePayloadIdToCache if it is a miss.
+			if (!PendingPayloadIds.Find(Key))
+			{
+				// Set that it is no longer allowed to write PayloadIdToCache since we are writing it now
+				Existing->bAllowedToWritePayloadIdToCache = false;
+
+				// Save the BulkData for writing to the cache outside of the lock
+				WritePayloadIdBulkData.Emplace(Existing->BulkData);
+				WritePayloadIdPackageName = Existing->PackageName;
+			}
 		}
 	}
+	if (WritePayloadIdBulkData)
+	{
+		FBulkDataRegistryImpl::WritePayloadIdToCache(WritePayloadIdPackageName, *WritePayloadIdBulkData);
+	}
 }
 
 TFuture<UE::BulkDataRegistry::FMetaData> FBulkDataRegistryImpl::GetMeta(const FGuid& BulkDataId)
@@ -397,7 +408,7 @@ TFuture<UE::BulkDataRegistry::FMetaData> FBulkDataRegistryImpl::GetMeta(const FG
 		FUpdatingPayload& UpdatingPayload = UpdatingPayloads.FindOrAdd(BulkDataId);
 		if (!UpdatingPayload.AsyncTask)
 		{
-			UpdatingPayload.AsyncTask = new FAutoDeleteAsyncTask<FUpdatePayloadWorker>(this, BulkData);
+			UpdatingPayload.AsyncTask = new FAutoDeleteAsyncTask<FUpdatePayloadWorker>(this, BulkData, true /* bKeepTempLoadedPayload */);
 			UpdatingPayload.AsyncTask->StartBackgroundTask();
 		}
 		TPromise<UE::BulkDataRegistry::FMetaData> Promise;
@@ -464,7 +475,7 @@ TFuture<UE::BulkDataRegistry::FData> FBulkDataRegistryImpl::GetData(const FGuid&
 			FUpdatingPayload& UpdatingPayload = UpdatingPayloads.FindOrAdd(BulkDataId);
 			if (!UpdatingPayload.AsyncTask)
 			{
-				UpdatingPayload.AsyncTask = new FAutoDeleteAsyncTask<FUpdatePayloadWorker>(this, Existing->BulkData);
+				UpdatingPayload.AsyncTask = new FAutoDeleteAsyncTask<FUpdatePayloadWorker>(this, Existing->BulkData, false /* bKeepTempLoadedPayload */);
 			}
 			TPromise<UE::BulkDataRegistry::FData> Promise;
 			TFuture<UE::BulkDataRegistry::FData> Future = Promise.GetFuture();
@@ -576,7 +587,7 @@ void FBulkDataRegistryImpl::PruneTempLoadedPayloads()
 	}
 }
 
-void FBulkDataRegistryImpl::WritePayloadIdToCache(FName PackageName, const UE::Serialization::FEditorBulkData& BulkData) const
+void FBulkDataRegistryImpl::WritePayloadIdToCache(FName PackageName, const UE::Serialization::FEditorBulkData& BulkData)
 {
 	check(!PackageName.IsNone());
 	TArray<uint8> Bytes;
@@ -617,53 +628,77 @@ void FBulkDataRegistryImpl::ReadPayloadIdsFromCache(FName PackageName, TArray<TR
 				}
 			}
 
-			FWriteScopeLock RegistryScopeLock(RegistryLock);
-			if (!bActive)
-			{
-				return;
-			}
-			TRefCountPtr<FPendingPayloadId> ExistingPending;
-			if (!PendingPayloadIds.RemoveAndCopyValue(BulkDataId, ExistingPending))
+			TOptional<UE::Serialization::FEditorBulkData> WritePayloadIdBulkData;
+			FName WritePayloadIdPackageName;
 			{
-				return;
-			}
-			check(ExistingPending->GetBulkDataId() == BulkDataId);
-			if (ExistingPending != NewPending)
-			{
-				// We removed ExistingPending because we thought it was equal to NewPending, but it's not, so put it back
-				PendingPayloadIds.Add(BulkDataId, MoveTemp(ExistingPending));
-				return;
-			}
+				FWriteScopeLock RegistryScopeLock(RegistryLock);
+				if (!bActive)
+				{
+					return;
+				}
+				TRefCountPtr<FPendingPayloadId> ExistingPending;
+				if (!PendingPayloadIds.RemoveAndCopyValue(BulkDataId, ExistingPending))
+				{
+					return;
+				}
+				check(ExistingPending->GetBulkDataId() == BulkDataId);
+				if (ExistingPending != NewPending)
+				{
+					// We removed ExistingPending because we thought it was equal to NewPending, but it's not, so put it back
+					PendingPayloadIds.Add(BulkDataId, MoveTemp(ExistingPending));
+					return;
+				}
 
-			FRegisteredBulk* ExistingRegisteredBulk = Registry.Find(BulkDataId);
-			if (!ExistingRegisteredBulk)
-			{
-				return;
-			}
+				FRegisteredBulk* ExistingRegisteredBulk = Registry.Find(BulkDataId);
+				if (!ExistingRegisteredBulk)
+				{
+					return;
+				}
 
-			UE::Serialization::FEditorBulkData& ExistingBulkData = ExistingRegisteredBulk->BulkData;
-			if (CachedBulkData.GetIdentifier().IsValid())
-			{
-				check(ExistingBulkData.GetIdentifier() == BulkDataId);
-				if (ExistingBulkData.HasPlaceholderPayloadId() && CachedBulkData.LocationMatches(ExistingBulkData))
+				UE::Serialization::FEditorBulkData& ExistingBulkData = ExistingRegisteredBulk->BulkData;
+				if (CachedBulkData.GetIdentifier().IsValid())
 				{
-					ExistingBulkData = CachedBulkData;
-					// No longer allowed to write the payloadId because we have found it already exists
-					ExistingRegisteredBulk->bAllowedToWritePayloadIdToCache = false;
+					check(ExistingBulkData.GetIdentifier() == BulkDataId);
+					if (ExistingBulkData.HasPlaceholderPayloadId() && CachedBulkData.LocationMatches(ExistingBulkData))
+					{
+						ExistingBulkData = CachedBulkData;
+						// No longer allowed to write the payloadId because we have found it already exists
+						ExistingRegisteredBulk->bAllowedToWritePayloadIdToCache = false;
+					}
 				}
-			}
-			else
-			{
-				if (!ExistingBulkData.HasPlaceholderPayloadId() && ExistingRegisteredBulk->bAllowedToWritePayloadIdToCache)
+				else if (ExistingRegisteredBulk->bAllowedToWritePayloadIdToCache)
 				{
-					// In between the point where we started the cache query and we received this result,
-					// The FEditorBulkData has updated its PayloadId and informed us by calling UpdatePlaceholderPayloadId
-					// Put this locally computed PayloadId into the cache, since the cache is missing it.
-					WritePayloadIdToCache(ExistingRegisteredBulk->PackageName, ExistingBulkData);
-					// Set that it is no longer allowed to write PayloadIdToCache since we have already written it.
-					ExistingRegisteredBulk->bAllowedToWritePayloadIdToCache = false;
+					// We had a missing Get, so calculate the value locally and Put the results
+					if (!ExistingBulkData.HasPlaceholderPayloadId())
+					{
+						// In between the point where we started the cache query and we received this result,
+						// The FEditorBulkData has updated its PayloadId and informed us by calling UpdatePlaceholderPayloadId
+						// Put this locally computed PayloadId into the cache, since the cache is missing it.
+
+						// Set that it is no longer allowed to write PayloadIdToCache since we are writing it now
+						ExistingRegisteredBulk->bAllowedToWritePayloadIdToCache = false;
+
+						// Save the BulkData for writing to the cache outside of the lock
+						WritePayloadIdBulkData.Emplace(ExistingBulkData);
+						WritePayloadIdPackageName = ExistingRegisteredBulk->PackageName;
+					}
+					else
+					{
+						// Create an FUpdatingPayload for the BulkData. Its DoWork function loads the data, sets the
+						// PayloadId on the BulkData in this->Registry, and calls WritePayloadIdToCache.
+						FUpdatingPayload& UpdatingPayload = UpdatingPayloads.FindOrAdd(BulkDataId);
+						if (!UpdatingPayload.AsyncTask)
+						{
+							UpdatingPayload.AsyncTask = new FAutoDeleteAsyncTask<FUpdatePayloadWorker>(this, ExistingBulkData, false/* bKeepTempLoadedPayload */);
+							UpdatingPayload.AsyncTask->StartBackgroundTask();
+						}
+					}
 				}
 			}
+			if (WritePayloadIdBulkData)
+			{
+				FBulkDataRegistryImpl::WritePayloadIdToCache(WritePayloadIdPackageName, *WritePayloadIdBulkData);
+			}
 		});
 	}
 
@@ -880,9 +915,10 @@ void FPendingPackage::WriteCache()
 }
 
 FUpdatePayloadWorker::FUpdatePayloadWorker(FBulkDataRegistryImpl* InBulkDataRegistry,
-	const UE::Serialization::FEditorBulkData& InSourceBulk)
+	const UE::Serialization::FEditorBulkData& InSourceBulk, bool bInKeepTempLoadedPayload)
 	: BulkData(InSourceBulk)
 	, BulkDataRegistry(InBulkDataRegistry)
+	, bKeepTempLoadedPayload(bInKeepTempLoadedPayload)
 {
 	SharedDataLock = InBulkDataRegistry->SharedDataLock;
 }
@@ -897,6 +933,8 @@ void FUpdatePayloadWorker::DoWork()
 		BulkData.UpdatePayloadId();
 		Buffer = BulkData.GetCompressedPayload().Get();
 
+		TOptional<UE::Serialization::FEditorBulkData> WritePayloadIdBulkData;
+		FName WritePayloadIdPackageName;
 		{
 			FReadScopeLock SharedDataScopeLock(SharedDataLock->ActiveLock);
 			if (!SharedDataLock->bActive)
@@ -938,19 +976,34 @@ void FUpdatePayloadWorker::DoWork()
 			// Store the new PayloadId in the Registry's entry for the BulkData; new MetaData requests will no longer need to wait for it
 			RegisteredBulk->BulkData = BulkData;
 
-			// Mark that the next GetData call should remove the temporary payload
-			RegisteredBulk->bHasTempPayload = true;
-			BulkDataRegistry->AddTempLoadedPayload(BulkData.GetIdentifier(), BulkData.GetPayloadSize());
-			BulkDataRegistry->PruneTempLoadedPayloads();
+			// The BulkData also has the payload still loaded; keep it or remove it
+			if (bKeepTempLoadedPayload)
+			{
+				// Keep the payload, and mark that the next GetData call should remove it
+				RegisteredBulk->bHasTempPayload = true;
+				BulkDataRegistry->AddTempLoadedPayload(BulkData.GetIdentifier(), BulkData.GetPayloadSize());
+				BulkDataRegistry->PruneTempLoadedPayloads();
+			}
+			else
+			{
+				RegisteredBulk->BulkData.UnloadData();
+			}
 
 			if (RegisteredBulk->bAllowedToWritePayloadIdToCache)
 			{
-				BulkDataRegistry->WritePayloadIdToCache(RegisteredBulk->PackageName, BulkData);
-				// Set that it is no longer allowed to write PayloadIdToCache since we have already written it.
+				// Set that it is no longer allowed to write PayloadIdToCache since we are writing it now
 				RegisteredBulk->bAllowedToWritePayloadIdToCache = false;
+
+				// Save the BulkData for writing to the cache outside of the lock
+				WritePayloadIdBulkData.Emplace(BulkData);
+				WritePayloadIdPackageName = RegisteredBulk->PackageName;
 			}
-			break;
 		}
+		if (WritePayloadIdBulkData)
+		{
+			FBulkDataRegistryImpl::WritePayloadIdToCache(WritePayloadIdPackageName, *WritePayloadIdBulkData);
+		}
+		break;
 	}
 
 	if (!bValid)
diff --git a/Engine/Source/Editor/UnrealEd/Private/EditorDomain/BulkDataRegistryImpl.h b/Engine/Source/Editor/UnrealEd/Private/EditorDomain/BulkDataRegistryImpl.h
--- a/Engine/Source/Editor/UnrealEd/Private/EditorDomain/BulkDataRegistryImpl.h
+++ b/Engine/Source/Editor/UnrealEd/Private/EditorDomain/BulkDataRegistryImpl.h
@@ -151,7 +151,7 @@ class FUpdatePayloadWorker : public FNonAbandonableTask
 {
 public:
 	FUpdatePayloadWorker(FBulkDataRegistryImpl* InBulkDataRegistry,
-		const UE::Serialization::FEditorBulkData& InSourceBulk);
+		const UE::Serialization::FEditorBulkData& InSourceBulk, bool bInKeepTempLoadedPayload);
 
 	void DoWork();
 	TStatId GetStatId() const
@@ -163,6 +163,7 @@ private:
 	UE::Serialization::FEditorBulkData BulkData;
 	TRefCountPtr<FTaskSharedDataLock> SharedDataLock;
 	FBulkDataRegistryImpl* BulkDataRegistry;
+	bool bKeepTempLoadedPayload;
 };
 
 /** Data storage for the FUpdatePayloadWorker that is updated while in flight for additional requesters. */
@@ -231,7 +232,7 @@ private:
 	void AddTempLoadedPayload(const FGuid& RegistryKey, uint64 PayloadSize);
 	void PruneTempLoadedPayloads();
 	void OnEndLoadPackage(const FEndLoadPackageContext& Context);
-	void WritePayloadIdToCache(FName PackageName, const UE::Serialization::FEditorBulkData& BulkData) const;
+	static void WritePayloadIdToCache(FName PackageName, const UE::Serialization::FEditorBulkData& BulkData);
 	void ReadPayloadIdsFromCache(FName PackageName, TArray<TRefCountPtr<FPendingPayloadId>>&& OldPendings,
 		TArray<TRefCountPtr<FPendingPayloadId>>&& NewPendings);
 
diff --git a/Engine/Source/Editor/UnrealEd/Private/EditorDomain/EditorDomainUtils.cpp b/Engine/Source/Editor/UnrealEd/Private/EditorDomain/EditorDomainUtils.cpp
--- a/Engine/Source/Editor/UnrealEd/Private/EditorDomain/EditorDomainUtils.cpp
+++ b/Engine/Source/Editor/UnrealEd/Private/EditorDomain/EditorDomainUtils.cpp
@@ -444,7 +444,7 @@ public:
 	FPrecacheClassDigest()
 		: ClassDigestsMap(GetClassDigests())
 		, ClassDigests(ClassDigestsMap.Map)
-		, AssetRegistry(*IAssetRegistry::Get())
+		, AssetRegistry(IAssetRegistry::GetChecked())
 	{
 		ClassDigestsMap.Lock.WriteLock();
 	}
@@ -1137,7 +1137,7 @@ public:
 				return *PackageDigest;
 			}
 		}
-		FPackageDigest PackageDigest = CalculatePackageDigest(*IAssetRegistry::Get(), PackageName);
+		FPackageDigest PackageDigest = CalculatePackageDigest(IAssetRegistry::GetChecked(), PackageName);
 		{
 			FWriteScopeLock ScopeLock(Lock);
 			PackageDigests.AddByHash(TypeHash, PackageName, PackageDigest);
@@ -2181,7 +2181,7 @@ void GetBulkDataList(FName PackageName, UE::DerivedData::IRequestOwner& Owner, T
 	IPackageDigestCache* DigestCache = IPackageDigestCache::Get();
 	FPackageDigest PackageDigest = DigestCache
 		? DigestCache->GetPackageDigest(PackageName)
-		: CalculatePackageDigest(*IAssetRegistry::Get(), PackageName);
+		: CalculatePackageDigest(IAssetRegistry::GetChecked(), PackageName);
 	if (!PackageDigest.IsSuccessful())
 	{
 		Callback(FSharedBuffer());
@@ -2209,7 +2209,7 @@ void PutBulkDataList(FName PackageName, FSharedBuffer Buffer)
 	IPackageDigestCache* DigestCache = IPackageDigestCache::Get();
 	FPackageDigest PackageDigest = DigestCache
 		? DigestCache->GetPackageDigest(PackageName)
-		: CalculatePackageDigest(*IAssetRegistry::Get(), PackageName);
+		: CalculatePackageDigest(IAssetRegistry::GetChecked(), PackageName);
 	if (!PackageDigest.IsSuccessful())
 	{
 		return;
@@ -2229,10 +2229,10 @@ void PutBulkDataList(FName PackageName, FSharedBuffer Buffer)
 	Owner.KeepAlive();
 }
 
-FIoHash GetPackageAndGuidHash(FIoHash& EditorDomainHash, const FGuid& BulkDataId)
+FIoHash GetPackageAndGuidHash(const FGuid& PackageGuid, const FGuid& BulkDataId)
 {
 	FBlake3 Builder;
-	Builder.Update(&EditorDomainHash, sizeof(EditorDomainHash));
+	Builder.Update(&PackageGuid, sizeof(PackageGuid));
 	Builder.Update(&BulkDataId, sizeof(BulkDataId));
 	return Builder.Finalize();
 }
@@ -2240,22 +2240,17 @@ FIoHash GetPackageAndGuidHash(FIoHash& EditorDomainHash, const FGuid& BulkDataId
 void GetBulkDataPayloadId(FName PackageName, const FGuid& BulkDataId, UE::DerivedData::IRequestOwner& Owner,
 	TUniqueFunction<void(FSharedBuffer Buffer)>&& Callback)
 {
-	IPackageDigestCache* DigestCache = IPackageDigestCache::Get();
-	FPackageDigest PackageDigest = DigestCache
-		? DigestCache->GetPackageDigest(PackageName)
-		: CalculatePackageDigest(*IAssetRegistry::Get(), PackageName);
-	if (!PackageDigest.IsSuccessful())
-	{
-		Callback(FSharedBuffer());
-		return;
-	}
-
-	if (!EnumHasAnyFlags(PackageDigest.DomainUse, EDomainUse::LoadEnabled))
+	IAssetRegistry& AssetRegistry = IAssetRegistry::GetChecked();
+	AssetRegistry.WaitForPackage(PackageName.ToString());
+	TOptional<FAssetPackageData> PackageData = AssetRegistry.GetAssetPackageDataCopy(PackageName);
+	if (!PackageData)
 	{
-		Callback(FSharedBuffer());
 		return;
 	}
-	FIoHash PackageAndGuidHash = GetPackageAndGuidHash(PackageDigest.Hash, BulkDataId);
+	PRAGMA_DISABLE_DEPRECATION_WARNINGS;
+	const FGuid& PackageGuid = PackageData->PackageGuid;
+	PRAGMA_ENABLE_DEPRECATION_WARNINGS;
+	FIoHash PackageAndGuidHash = GetPackageAndGuidHash(PackageGuid, BulkDataId);
 
 	using namespace UE::DerivedData;
 	ICache& Cache = GetCache();
@@ -2270,20 +2265,17 @@ void GetBulkDataPayloadId(FName PackageName, const FGuid& BulkDataId, UE::Derive
 
 void PutBulkDataPayloadId(FName PackageName, const FGuid& BulkDataId, FSharedBuffer Buffer)
 {
-	IPackageDigestCache* DigestCache = IPackageDigestCache::Get();
-	FPackageDigest PackageDigest = DigestCache
-		? DigestCache->GetPackageDigest(PackageName)
-		: CalculatePackageDigest(*IAssetRegistry::Get(), PackageName);
-	if (!PackageDigest.IsSuccessful())
-	{
-		return;
-	}
-
-	if (!EnumHasAnyFlags(PackageDigest.DomainUse, EDomainUse::SaveEnabled))
+	IAssetRegistry& AssetRegistry = IAssetRegistry::GetChecked();
+	AssetRegistry.WaitForPackage(PackageName.ToString());
+	TOptional<FAssetPackageData> PackageData = AssetRegistry.GetAssetPackageDataCopy(PackageName);
+	if (!PackageData)
 	{
 		return;
 	}
-	FIoHash PackageAndGuidHash = GetPackageAndGuidHash(PackageDigest.Hash, BulkDataId);
+	PRAGMA_DISABLE_DEPRECATION_WARNINGS;
+	const FGuid& PackageGuid = PackageData->PackageGuid;
+	PRAGMA_ENABLE_DEPRECATION_WARNINGS;
+	FIoHash PackageAndGuidHash = GetPackageAndGuidHash(PackageGuid, BulkDataId);
 
 	using namespace UE::DerivedData;
 	ICache& Cache = GetCache();
