Cooker: Change GetAllPackageFilenamesFromAssetRegistry to use an empty FAssetRegistryState when the preloaded data in GPreloadedARState is not being used.
#rnx

diff --git a/Engine/Source/Editor/UnrealEd/Private/CookOnTheFlyServer.cpp b/Engine/Source/Editor/UnrealEd/Private/CookOnTheFlyServer.cpp
--- a/Engine/Source/Editor/UnrealEd/Private/CookOnTheFlyServer.cpp
+++ b/Engine/Source/Editor/UnrealEd/Private/CookOnTheFlyServer.cpp
@@ -10311,6 +10311,8 @@ bool UCookOnTheFlyServer::GetAllPackageFilenamesFromAssetRegistry(const FString&
 		GPreloadARInfoEvent->Wait();
 
 		bool bHadPreloadedAR = false;
+		FAssetRegistryState* SerializedState = nullptr;
+		TOptional<FAssetRegistryState> NonPreloadedState;
 		if (AssetRegistryPath == GPreloadedARPath)
 		{
 			// make sure the Serialize call is done
@@ -10321,25 +10323,31 @@ bool UCookOnTheFlyServer::GetAllPackageFilenamesFromAssetRegistry(const FString&
 			
 			// if something went wrong, the num assets may be zero, in which case we do the normal load 
 			bHadPreloadedAR = GPreloadedARState.GetNumAssets() > 0;
+			SerializedState = &GPreloadedARState;
+		}
+		else
+		{
+			NonPreloadedState.Emplace();
+			SerializedState = &NonPreloadedState.GetValue();
 		}
 
 		// if we didn't preload an AR, then we need to do a blocking load now
 		if (!bHadPreloadedAR)
 		{
-			GPreloadedARState.Serialize(*Reader.Get(), FAssetRegistrySerializationOptions());
+			SerializedState->Serialize(*Reader.Get(), FAssetRegistrySerializationOptions());
 		}
 
 		
 		check(OutPackageDatas.Num() == 0);
 
-		int32 NumPackages = GPreloadedARState.GetNumAssets();
+		int32 NumPackages = SerializedState->GetNumAssets();
 		TArray<const FAssetData*> AssetDatas;
 		TSet<FName> PackageNames;
 		AssetDatas.Reserve(NumPackages);
 		OutPackageDatas.Reserve(NumPackages);
 
 		// Convert the Map of RegistryData into an Array of FAssetData and populate PackageNames in the output array
-		GPreloadedARState.EnumerateAllAssets([&](const FAssetData& RegistryData)
+		SerializedState->EnumerateAllAssets([&](const FAssetData& RegistryData)
 		{
 			// If we want to reevaluate (try cooking again) the uncooked packages (packages that were found to be empty when we cooked them before),
 			// then remove the uncooked packages from the set of known packages. Uncooked packages are identified by PackageFlags == 0.
