Allow auto-save to include newly created packages

This was disabled as part of the OFPA changes due to auto-save interacting poorly with the undo buffer, since the auto-save could mark objects within the package as RF_Transient when saving.

To avoid this, objects now store their transient state as part of the transaction so that it can be restored on undo/redo. This also fixes the auto-save list not updating correctly after an undo/redo.

diff --git a/Engine/Source/Editor/UnrealEd/Private/FileHelpers.cpp b/Engine/Source/Editor/UnrealEd/Private/FileHelpers.cpp
--- a/Engine/Source/Editor/UnrealEd/Private/FileHelpers.cpp
+++ b/Engine/Source/Editor/UnrealEd/Private/FileHelpers.cpp
@@ -2980,7 +2980,7 @@ void FEditorFileUtils::ResetLevelFilenames()
 	MainFrameModule.SetLevelNameForWindowTitle(EmptyFilename);
 }
 
-bool FEditorFileUtils::AutosaveMap(const FString& AbsoluteAutosaveDir, const int32 AutosaveIndex, const bool bForceIfNotInList, const TSet< TWeakObjectPtr<UPackage> >& DirtyPackagesForAutoSave)
+bool FEditorFileUtils::AutosaveMap(const FString& AbsoluteAutosaveDir, const int32 AutosaveIndex, const bool bForceIfNotInList, const TSet< TWeakObjectPtr<UPackage>, TWeakObjectPtrSetKeyFuncs<TWeakObjectPtr<UPackage>> >& DirtyPackagesForAutoSave)
 {
 	auto Result = AutosaveMapEx(AbsoluteAutosaveDir, AutosaveIndex, bForceIfNotInList, DirtyPackagesForAutoSave);
 
@@ -2989,7 +2989,7 @@ bool FEditorFileUtils::AutosaveMap(const FString& AbsoluteAutosaveDir, const int
 	return Result == EAutosaveContentPackagesResult::Success;
 }
 
-EAutosaveContentPackagesResult::Type FEditorFileUtils::AutosaveMapEx(const FString& AbsoluteAutosaveDir, const int32 AutosaveIndex, const bool bForceIfNotInList, const TSet< TWeakObjectPtr<UPackage> >& DirtyPackagesForAutoSave)
+EAutosaveContentPackagesResult::Type FEditorFileUtils::AutosaveMapEx(const FString& AbsoluteAutosaveDir, const int32 AutosaveIndex, const bool bForceIfNotInList, const TSet< TWeakObjectPtr<UPackage>, TWeakObjectPtrSetKeyFuncs<TWeakObjectPtr<UPackage>> >& DirtyPackagesForAutoSave)
 {
 	const FScopedBusyCursor BusyCursor;
 	bool bResult  = false;
@@ -3100,7 +3100,7 @@ EAutosaveContentPackagesResult::Type FEditorFileUtils::AutosaveMapEx(const FStri
 	return bResult ? EAutosaveContentPackagesResult::Success : EAutosaveContentPackagesResult::NothingToDo;
 }
 
-bool FEditorFileUtils::AutosaveContentPackages(const FString& AbsoluteAutosaveDir, const int32 AutosaveIndex, const bool bForceIfNotInList, const TSet< TWeakObjectPtr<UPackage> >& DirtyPackagesForAutoSave)
+bool FEditorFileUtils::AutosaveContentPackages(const FString& AbsoluteAutosaveDir, const int32 AutosaveIndex, const bool bForceIfNotInList, const TSet< TWeakObjectPtr<UPackage>, TWeakObjectPtrSetKeyFuncs<TWeakObjectPtr<UPackage>> >& DirtyPackagesForAutoSave)
 {
 	auto Result = AutosaveContentPackagesEx(AbsoluteAutosaveDir, AutosaveIndex, bForceIfNotInList, DirtyPackagesForAutoSave);
 
@@ -3109,7 +3109,7 @@ bool FEditorFileUtils::AutosaveContentPackages(const FString& AbsoluteAutosaveDi
 	return Result == EAutosaveContentPackagesResult::Success;
 }
 
-EAutosaveContentPackagesResult::Type FEditorFileUtils::AutosaveContentPackagesEx(const FString& AbsoluteAutosaveDir, const int32 AutosaveIndex, const bool bForceIfNotInList, const TSet< TWeakObjectPtr<UPackage> >& DirtyPackagesForAutoSave)
+EAutosaveContentPackagesResult::Type FEditorFileUtils::AutosaveContentPackagesEx(const FString& AbsoluteAutosaveDir, const int32 AutosaveIndex, const bool bForceIfNotInList, const TSet< TWeakObjectPtr<UPackage>, TWeakObjectPtrSetKeyFuncs<TWeakObjectPtr<UPackage>> >& DirtyPackagesForAutoSave)
 {
 	const FScopedBusyCursor BusyCursor;
 	double SaveStartTime = FPlatformTime::Seconds();
diff --git a/Engine/Source/Editor/UnrealEd/Private/PackageAutoSaver.cpp b/Engine/Source/Editor/UnrealEd/Private/PackageAutoSaver.cpp
--- a/Engine/Source/Editor/UnrealEd/Private/PackageAutoSaver.cpp
+++ b/Engine/Source/Editor/UnrealEd/Private/PackageAutoSaver.cpp
@@ -70,7 +70,7 @@ namespace PackageAutoSaverJson
 	 * @param bRestoreEnabled	Is the restore enabled, or is it disabled because we've shut-down cleanly, or are running under the debugger?
 	 * @param DirtyPackages		Packages that may have auto-saves that they could be restored from
 	 */
-	void SaveRestoreFile(const bool bRestoreEnabled, const TMap< TWeakObjectPtr<UPackage>, FString >& DirtyPackages);
+	void SaveRestoreFile(const bool bRestoreEnabled, const TMap< TWeakObjectPtr<UPackage>, FString, FDefaultSetAllocator, TWeakObjectPtrMapKeyFuncs<TWeakObjectPtr<UPackage>, FString> >& DirtyPackages);
 
 	/** @return whether the auto-save restore should be enabled (you can force this to true when testing with a debugger attached) */
 	bool IsRestoreEnabled()
@@ -103,6 +103,9 @@ FPackageAutoSaver::FPackageAutoSaver()
 
 	// Register for the package modified callback to catch packages that have been saved
 	UPackage::PackageSavedWithContextEvent.AddRaw(this, &FPackageAutoSaver::OnPackageSaved);
+
+	// Register to detect when an Undo/Redo changes the dirty state of a package
+	FEditorDelegates::PostUndoRedo.AddRaw(this, &FPackageAutoSaver::OnUndoRedo);
 }
 
 FPackageAutoSaver::~FPackageAutoSaver()
@@ -110,6 +113,7 @@ FPackageAutoSaver::~FPackageAutoSaver()
 	UPackage::PackageDirtyStateChangedEvent.RemoveAll(this);
 	UPackage::PackageMarkedDirtyEvent.RemoveAll(this);
 	UPackage::PackageSavedWithContextEvent.RemoveAll(this);
+	FEditorDelegates::PostUndoRedo.RemoveAll(this);
 }
 
 void FPackageAutoSaver::UpdateAutoSaveCount(const float DeltaSeconds)
@@ -127,12 +131,6 @@ void FPackageAutoSaver::UpdateAutoSaveCount(const float DeltaSeconds)
 	{
 		AutoSaveCount += DeltaSeconds;
 	}
-
-	// Update the restore information too, if needed
-	if (bNeedRestoreFileUpdate)
-	{
-		UpdateRestoreFile(PackageAutoSaverJson::IsRestoreEnabled());
-	}
 }
 
 void FPackageAutoSaver::ResetAutoSaveTimer()
@@ -157,6 +155,54 @@ void FPackageAutoSaver::AttemptAutoSave()
 	const UEditorLoadingSavingSettings* LoadingSavingSettings = GetDefault<UEditorLoadingSavingSettings>();
 	FUnrealEdMisc& UnrealEdMisc = FUnrealEdMisc::Get();
 
+	// Re-sync if needed
+	if (bSyncWithDirtyPackageList)
+	{
+		bSyncWithDirtyPackageList = false;
+		PackagesPendingUpdate.Reset();
+
+		DirtyMapsForAutoSave.Reset();
+		DirtyContentForAutoSave.Reset();
+
+		// The the list of dirty packages tracked by the engine (considered source of truth)
+		TArray<UPackage*> DirtyPackages;
+		FEditorFileUtils::GetDirtyPackages(DirtyPackages);
+		for (UPackage* Pkg : DirtyPackages)
+		{
+			UpdateDirtyListsForPackage(Pkg);
+		}
+
+		// Remove any clean package from the user-restore list
+		for (auto It = DirtyPackagesForUserSave.CreateIterator(); It; ++It)
+		{
+			UPackage* Pkg = It->Key.Get();
+			if (!Pkg || !Pkg->IsDirty() || (PackagesToIgnoreIfEmpty.Contains(Pkg->GetFName()) && UPackage::IsEmptyPackage(Pkg)))
+			{
+				bNeedRestoreFileUpdate = true;
+				It.RemoveCurrent();
+			}
+		}
+	}
+
+	// Process any packages that are pending an update
+	if (PackagesPendingUpdate.Num() > 0)
+	{
+		for (const TWeakObjectPtr<UPackage>& WeakPkg : PackagesPendingUpdate)
+		{
+			if (UPackage* Pkg = WeakPkg.Get())
+			{
+				UpdateDirtyListsForPackage(Pkg);
+			}
+		}
+		PackagesPendingUpdate.Reset();
+	}
+
+	// Update the restore information too, if needed
+	if (bNeedRestoreFileUpdate)
+	{
+		UpdateRestoreFile(PackageAutoSaverJson::IsRestoreEnabled());
+	}
+
 	// Don't auto-save if disabled or if it is not yet time to auto-save.
 	const bool bTimeToAutosave = (LoadingSavingSettings->bAutoSaveEnable && AutoSaveCount >= LoadingSavingSettings->AutoSaveTimeMinutes * 60.0f);
 	bool bAutosaveHandled = false;
@@ -320,34 +366,43 @@ void FPackageAutoSaver::OfferToRestorePackages()
 
 void FPackageAutoSaver::OnPackagesDeleted(const TArray<UPackage*>& DeletedPackages)
 {
-	ClearStalePointers();
-
-	for(UPackage* DeletedPackage : DeletedPackages)
+	for (UPackage* DeletedPackage : DeletedPackages)
 	{
+		PackagesPendingUpdate.Remove(DeletedPackage);
+		PackagesToIgnoreIfEmpty.Add(DeletedPackage->GetFName());
+
+		// We remove the package immediately as it may not survive to the next tick if queued for update via PackagesPendingUpdate
 		DirtyMapsForAutoSave.Remove(DeletedPackage);
 		DirtyContentForAutoSave.Remove(DeletedPackage);
-		DirtyPackagesForUserSave.Remove(DeletedPackage);
+		if (DirtyPackagesForUserSave.Remove(DeletedPackage) > 0)
+		{
+			bNeedRestoreFileUpdate = true;
+		}
 	}
-	bNeedRestoreFileUpdate = true;
 }
 
 void FPackageAutoSaver::OnPackageDirtyStateUpdated(UPackage* Pkg)
 {
-	UpdateDirtyListsForPackage(Pkg);
+	if (!IsAutoSaving())
+	{
+		PackagesPendingUpdate.Add(Pkg);
+	}
 }
 
 void FPackageAutoSaver::OnMarkPackageDirty(UPackage* Pkg, bool bWasDirty)
 {
-	UpdateDirtyListsForPackage(Pkg);
+	if (!IsAutoSaving())
+	{
+		PackagesPendingUpdate.Add(Pkg);
+	}
 }
 
 void FPackageAutoSaver::OnPackageSaved(const FString& Filename, UPackage* Pkg, FObjectPostSaveContext ObjectSaveContext)
 {
 	// If this has come from an auto-save, update the last known filename in the user dirty list so that we can offer is up as a restore file later
-	if(IsAutoSaving())
+	if (IsAutoSaving())
 	{
-		FString* const AutoSaveFilename = DirtyPackagesForUserSave.Find(Pkg);
-		if(AutoSaveFilename)
+		if (FString* const AutoSaveFilename = DirtyPackagesForUserSave.Find(Pkg))
 		{
 			// Make the filename relative to the auto-save directory
 			// Note: MakePathRelativeTo modifies in-place, hence the copy of Filename
@@ -356,9 +411,32 @@ void FPackageAutoSaver::OnPackageSaved(const FString& Filename, UPackage* Pkg, F
 			FPaths::MakePathRelativeTo(RelativeFilename, *AutoSaveDir);
 
 			(*AutoSaveFilename) = RelativeFilename;
+			bNeedRestoreFileUpdate = true;
+		}
+	}
+	else
+	{
+		// If the package was previously deleted, then it's certainly back after being saved!
+		PackagesToIgnoreIfEmpty.Remove(Pkg->GetFName());
+
+		// Remove the saved package from the user-restore list when this was a full save
+		if (DirtyPackagesForUserSave.Remove(Pkg) > 0)
+		{
+			bNeedRestoreFileUpdate = true;
 		}
 	}
-	UpdateDirtyListsForPackage(Pkg);
+
+	// Always remove a saved package from the auto-save lists
+	DirtyMapsForAutoSave.Remove(Pkg);
+	DirtyContentForAutoSave.Remove(Pkg);
+
+	// Discard any pending update since the save has already handled it
+	PackagesPendingUpdate.Remove(Pkg);
+}
+
+void FPackageAutoSaver::OnUndoRedo()
+{
+	bSyncWithDirtyPackageList = true;
 }
 
 void FPackageAutoSaver::UpdateDirtyListsForPackage(UPackage* Pkg)
@@ -366,39 +444,43 @@ void FPackageAutoSaver::UpdateDirtyListsForPackage(UPackage* Pkg)
 	const UPackage* TransientPackage = GetTransientPackage();
 
 	// Don't auto-save the transient package or packages with the transient flag.
-	if ( Pkg == TransientPackage || Pkg->HasAnyFlags(RF_Transient) || Pkg->HasAnyPackageFlags(PKG_InMemoryOnly) )
+	if ( Pkg == TransientPackage || Pkg->HasAnyFlags(RF_Transient) || Pkg->HasAnyPackageFlags(PKG_CompiledIn) )
+	{
+		return;
+	}
+
+	// Should this package be ignored because it was previously deleted and is still empty?
+	if (PackagesToIgnoreIfEmpty.Contains(Pkg->GetFName()) && UPackage::IsEmptyPackage(Pkg))
 	{
 		return;
 	}
 
 	if ( Pkg->IsDirty() )
 	{
-		// Always add the package to the user list
-		DirtyPackagesForUserSave.FindOrAdd(Pkg);
+		// Add the package to the user-restore list
+		if (!DirtyPackagesForUserSave.Contains(Pkg))
+		{
+			DirtyPackagesForUserSave.Add(Pkg);
+			bNeedRestoreFileUpdate = true;
+		}
 
-		// Only add the package to the auto-save list if we're not auto-saving
-		// Note: Packages get dirtied again after they're auto-saved, so this would add them back again, which we don't want
-		if ( !IsAutoSaving() )
+		// Add package into the appropriate list (map or content)
+		if (UWorld::IsWorldOrExternalActorPackage(Pkg))
 		{
-			// Add package into the appropriate list (map or content)
-			if (UWorld::IsWorldOrExternalActorPackage(Pkg))
-			{
-				DirtyMapsForAutoSave.Add(Pkg);
-			}
-			else
-			{
-				DirtyContentForAutoSave.Add(Pkg);
-			}
+			DirtyMapsForAutoSave.Add(Pkg);
+		}
+		else
+		{
+			DirtyContentForAutoSave.Add(Pkg);
 		}
 	}
 	else
 	{
-		// Always remove the package from the auto-save list
+		// Always remove a clean package from the auto-save and user-restore lists
 		DirtyMapsForAutoSave.Remove(Pkg);
 		DirtyContentForAutoSave.Remove(Pkg);
-		if (!IsAutoSaving())
+		if (DirtyPackagesForUserSave.Remove(Pkg) > 0)
 		{
-			DirtyPackagesForUserSave.Remove(Pkg);
 			bNeedRestoreFileUpdate = true;
 		}
 	}
@@ -664,33 +746,31 @@ void FPackageAutoSaver::OnAutoSaveCancel()
 
 void FPackageAutoSaver::ClearStalePointers()
 {
-	auto DirtyPackagesForUserSaveTmp = DirtyPackagesForUserSave;
-	for(auto It = DirtyPackagesForUserSaveTmp.CreateConstIterator(); It; ++It)
+	for(auto It = DirtyPackagesForUserSave.CreateIterator(); It; ++It)
 	{
 		const TWeakObjectPtr<UPackage>& Package = It->Key;
 		if(!Package.IsValid())
 		{
-			DirtyPackagesForUserSave.Remove(Package);
+			bNeedRestoreFileUpdate = true;
+			It.RemoveCurrent();
 		}
 	}
 
-	auto DirtyMapsForAutoSaveTmp = DirtyMapsForAutoSave;
-	for(auto It = DirtyMapsForAutoSaveTmp.CreateConstIterator(); It; ++It)
+	for(auto It = DirtyMapsForAutoSave.CreateIterator(); It; ++It)
 	{
 		const TWeakObjectPtr<UPackage>& Package = *It;
 		if(!Package.IsValid())
 		{
-			DirtyMapsForAutoSave.Remove(Package);
+			It.RemoveCurrent();
 		}
 	}
 
-	auto DirtyContentForAutoSaveTmp = DirtyContentForAutoSave;
-	for (auto It = DirtyContentForAutoSaveTmp.CreateConstIterator(); It; ++It)
+	for (auto It = DirtyContentForAutoSave.CreateIterator(); It; ++It)
 	{
 		const TWeakObjectPtr<UPackage>& Package = *It;
 		if (!Package.IsValid())
 		{
-			DirtyContentForAutoSave.Remove(Package);
+			It.RemoveCurrent();
 		}
 	}
 }
@@ -752,7 +832,7 @@ TMap<FString, FString> PackageAutoSaverJson::LoadRestoreFile()
 	return PackagesThatCanBeRestored;
 }
 
-void PackageAutoSaverJson::SaveRestoreFile(const bool bRestoreEnabled, const TMap< TWeakObjectPtr<UPackage>, FString >& DirtyPackages)
+void PackageAutoSaverJson::SaveRestoreFile(const bool bRestoreEnabled, const TMap< TWeakObjectPtr<UPackage>, FString, FDefaultSetAllocator, TWeakObjectPtrMapKeyFuncs<TWeakObjectPtr<UPackage>, FString> >& DirtyPackages)
 {
 	TSharedPtr<FJsonObject> RootObject = MakeShareable(new FJsonObject);
 
diff --git a/Engine/Source/Editor/UnrealEd/Private/PackageAutoSaver.h b/Engine/Source/Editor/UnrealEd/Private/PackageAutoSaver.h
--- a/Engine/Source/Editor/UnrealEd/Private/PackageAutoSaver.h
+++ b/Engine/Source/Editor/UnrealEd/Private/PackageAutoSaver.h
@@ -8,6 +8,7 @@
 #include "IPackageAutoSaver.h"
 
 class FObjectPostSaveContext;
+struct FTransactionContext;
 
 namespace ECloseNotification
 {
@@ -65,6 +66,11 @@ private:
 	 */
 	void OnPackageSaved(const FString& Filename, UPackage* Pkg, FObjectPostSaveContext ObjectSaveContext);
 
+	/**
+	 * Called when an undo/redo operation happens.
+	 */
+	void OnUndoRedo();
+
 	/**
 	 * Update the dirty lists for the current package
 	 *
@@ -130,17 +136,26 @@ private:
 	/** Indicate that we need to update the restore file. */
 	bool bNeedRestoreFileUpdate = false;
 
+	/** True if we need to invalidate and re-cache our list of dirty files. */
+	bool bSyncWithDirtyPackageList = false;
+
+	/** Set of packages that are pending an update in the auto-save lists */
+	TSet<TWeakObjectPtr<UPackage>, TWeakObjectPtrSetKeyFuncs<TWeakObjectPtr<UPackage>>> PackagesPendingUpdate;
+
+	/** Set of package names that have been deleted and should be ignored while in an empty state */
+	TSet<FName> PackagesToIgnoreIfEmpty;
+
 	/** Used to reference to the active auto-save warning notification */
 	TWeakPtr<SNotificationItem> AutoSaveNotificationPtr;
 
 	/** Packages that have been dirtied and not saved by the user, mapped to their latest auto-save file */
-	TMap<TWeakObjectPtr<UPackage>, FString> DirtyPackagesForUserSave;
+	TMap<TWeakObjectPtr<UPackage>, FString, FDefaultSetAllocator, TWeakObjectPtrMapKeyFuncs<TWeakObjectPtr<UPackage>, FString>> DirtyPackagesForUserSave;
 
 	/** Maps that have been dirtied and not saved by the auto-saver */
-	TSet<TWeakObjectPtr<UPackage>> DirtyMapsForAutoSave;
+	TSet<TWeakObjectPtr<UPackage>, TWeakObjectPtrSetKeyFuncs<TWeakObjectPtr<UPackage>>> DirtyMapsForAutoSave;
 
 	/** Content that has been dirtied and not saved by the auto-saver */
-	TSet<TWeakObjectPtr<UPackage>> DirtyContentForAutoSave;
+	TSet<TWeakObjectPtr<UPackage>, TWeakObjectPtrSetKeyFuncs<TWeakObjectPtr<UPackage>>> DirtyContentForAutoSave;
 
 	/** Restore information that was loaded following a crash */
 	TMap<FString, FString> PackagesThatCanBeRestored;
diff --git a/Engine/Source/Editor/UnrealEd/Public/FileHelpers.h b/Engine/Source/Editor/UnrealEd/Public/FileHelpers.h
--- a/Engine/Source/Editor/UnrealEd/Public/FileHelpers.h
+++ b/Engine/Source/Editor/UnrealEd/Public/FileHelpers.h
@@ -315,7 +315,7 @@ public:
 	 * @param	bForceIfNotInList			Should the save be forced if the package is dirty, but not in DirtyPackagesForAutoSave?
 	 * @param	DirtyPackagesForAutoSave	A set of packages that are considered by the auto-save system to be dirty, you should check this to see if a package needs saving
 	 */
-	static bool AutosaveMap(const FString& AbsoluteAutosaveDir, const int32 AutosaveIndex, const bool bForceIfNotInList, const TSet< TWeakObjectPtr<UPackage> >& DirtyPackagesForAutoSave);
+	static bool AutosaveMap(const FString& AbsoluteAutosaveDir, const int32 AutosaveIndex, const bool bForceIfNotInList, const TSet< TWeakObjectPtr<UPackage>, TWeakObjectPtrSetKeyFuncs<TWeakObjectPtr<UPackage>> >& DirtyPackagesForAutoSave);
 
 	/**
 	 * Saves all levels to the specified directory.
@@ -325,7 +325,7 @@ public:
 	 * @param	bForceIfNotInList			Should the save be forced if the package is dirty, but not in DirtyPackagesForAutoSave?
 	 * @param	DirtyPackagesForAutoSave	A set of packages that are considered by the auto-save system to be dirty, you should check this to see if a package needs saving
 	 */
-	static EAutosaveContentPackagesResult::Type AutosaveMapEx(const FString& AbsoluteAutosaveDir, const int32 AutosaveIndex, const bool bForceIfNotInList, const TSet< TWeakObjectPtr<UPackage> >& DirtyPackagesForAutoSave);
+	static EAutosaveContentPackagesResult::Type AutosaveMapEx(const FString& AbsoluteAutosaveDir, const int32 AutosaveIndex, const bool bForceIfNotInList, const TSet< TWeakObjectPtr<UPackage>, TWeakObjectPtrSetKeyFuncs<TWeakObjectPtr<UPackage>> >& DirtyPackagesForAutoSave);
 
 	/**
 	 * Saves all asset packages to the specified directory.
@@ -337,7 +337,7 @@ public:
 	 *
 	 * @return	true if one or more packages were autosaved; false otherwise
 	 */
-	static bool AutosaveContentPackages(const FString& AbsoluteAutosaveDir, const int32 AutosaveIndex, const bool bForceIfNotInList, const TSet< TWeakObjectPtr<UPackage> >& DirtyPackagesForAutoSave);
+	static bool AutosaveContentPackages(const FString& AbsoluteAutosaveDir, const int32 AutosaveIndex, const bool bForceIfNotInList, const TSet< TWeakObjectPtr<UPackage>, TWeakObjectPtrSetKeyFuncs<TWeakObjectPtr<UPackage>> >& DirtyPackagesForAutoSave);
 
 	/**
 	 * Saves all asset packages to the specified directory.
@@ -349,7 +349,7 @@ public:
 	 *
 	 * @return	Success if saved at least one faile. NothingToDo if there was nothing to save. Failure on at least one auto-save failure.
 	 */
-	static EAutosaveContentPackagesResult::Type AutosaveContentPackagesEx(const FString& AbsoluteAutosaveDir, const int32 AutosaveIndex, const bool bForceIfNotInList, const TSet< TWeakObjectPtr<UPackage> >& DirtyPackagesForAutoSave);
+	static EAutosaveContentPackagesResult::Type AutosaveContentPackagesEx(const FString& AbsoluteAutosaveDir, const int32 AutosaveIndex, const bool bForceIfNotInList, const TSet< TWeakObjectPtr<UPackage>, TWeakObjectPtrSetKeyFuncs<TWeakObjectPtr<UPackage>> >& DirtyPackagesForAutoSave);
 
 	/**
 	 * Looks at all currently loaded packages and saves them if their "bDirty" flag is set, optionally prompting the user to select which packages to save)
diff --git a/Engine/Source/Runtime/Core/Public/UObject/WeakObjectPtrTemplates.h b/Engine/Source/Runtime/Core/Public/UObject/WeakObjectPtrTemplates.h
--- a/Engine/Source/Runtime/Core/Public/UObject/WeakObjectPtrTemplates.h
+++ b/Engine/Source/Runtime/Core/Public/UObject/WeakObjectPtrTemplates.h
@@ -324,6 +324,25 @@ template<class T> struct TIsZeroConstructType<TWeakObjectPtr<T> > { enum { Value
 template<class T> struct TIsWeakPointerType<TWeakObjectPtr<T> > { enum { Value = true }; };
 
 
+/**
+ * SetKeyFuncs for TWeakObjectPtrs which allow the key to become stale without invalidating the set.
+ */
+template <typename ElementType, bool bInAllowDuplicateKeys = false>
+struct TWeakObjectPtrSetKeyFuncs : DefaultKeyFuncs<ElementType, bInAllowDuplicateKeys>
+{
+	typedef typename DefaultKeyFuncs<ElementType, bInAllowDuplicateKeys>::KeyInitType KeyInitType;
+
+	static FORCEINLINE bool Matches(KeyInitType A, KeyInitType B)
+	{
+		return A.HasSameIndexAndSerialNumber(B);
+	}
+
+	static FORCEINLINE uint32 GetKeyHash(KeyInitType Key)
+	{
+		return GetTypeHash(Key);
+	}
+};
+
 /**
  * MapKeyFuncs for TWeakObjectPtrs which allow the key to become stale without invalidating the map.
  */
diff --git a/Engine/Source/Runtime/CoreUObject/Private/UObject/Obj.cpp b/Engine/Source/Runtime/CoreUObject/Private/UObject/Obj.cpp
--- a/Engine/Source/Runtime/CoreUObject/Private/UObject/Obj.cpp
+++ b/Engine/Source/Runtime/CoreUObject/Private/UObject/Obj.cpp
@@ -1478,6 +1478,28 @@ void UObject::Serialize(FStructuredArchive::FRecord Record)
 			}
 		}
 
+		// Keep track of transient
+		if (UnderlyingArchive.IsTransacting())
+		{
+			bool WasTransient = HasAnyFlags(RF_Transient);
+			if (UnderlyingArchive.IsLoading())
+			{
+				Record << SA_VALUE(TEXT("WasTransient"), WasTransient);
+				if (WasTransient)
+				{
+					SetFlags(RF_Transient);
+				}
+				else
+				{
+					ClearFlags(RF_Transient);
+				}
+			}
+			else if (UnderlyingArchive.IsSaving())
+			{
+				Record << SA_VALUE(TEXT("WasTransient"), WasTransient);
+			}
+		}
+
 		// Serialize a GUID if this object has one mapped to it
 		FLazyObjectPtr::PossiblySerializeObjectGuid(this, Record);
 
diff --git a/Engine/Source/Runtime/CoreUObject/Private/UObject/SavePackage2.cpp b/Engine/Source/Runtime/CoreUObject/Private/UObject/SavePackage2.cpp
--- a/Engine/Source/Runtime/CoreUObject/Private/UObject/SavePackage2.cpp
+++ b/Engine/Source/Runtime/CoreUObject/Private/UObject/SavePackage2.cpp
@@ -2284,8 +2284,12 @@ void ClearCachedPlatformCookedData(FSaveContext& SaveContext)
 void PostSavePackage(FSaveContext& SaveContext)
 {
 	UPackage* Package = SaveContext.GetPackage();
-	// Package has been saved, so unmark the NewlyCreated flag.
-	Package->ClearPackageFlags(PKG_NewlyCreated);
+
+	if (!SaveContext.IsFromAutoSave() && !SaveContext.IsProceduralSave())
+	{
+		// Package has been saved, so unmark the NewlyCreated flag.
+		Package->ClearPackageFlags(PKG_NewlyCreated);
+	}
 
 	// Copy and modify the output SerializedPackageFlags from the PackageFlags written into the default realm summary
 	uint32 SerializedPackageFlags = SaveContext.GetLinker()->Summary.GetPackageFlags();
