Packages submitted via the ResavePackage commandlet can now be virtualized if the feature is enabled.

#rnx

- When the ResavePackage commandlet is run with the -autosubmit/-autocheckin commands then the packages that are resaved will also be submitted to source control. When this happens we need to make sure that the virtualization process can be run.
- The new virtualization code mirrors the submission code found in SSourceControlSubmit.cpp
- We only run the virtualization pass if the project has enabled both virtualization and the virtualization of new packages.
- Adds a new command line to the commandlet '-SkipVirtualization' which when used causes the virtualization process to be skipped over. The option to skip virtualization has been added as the commandlet is most likely only used by automated build systems or power users who are better places to decide if it is needed or not.

diff --git a/Engine/Source/Editor/UnrealEd/Classes/Commandlets/ResavePackagesCommandlet.h b/Engine/Source/Editor/UnrealEd/Classes/Commandlets/ResavePackagesCommandlet.h
--- a/Engine/Source/Editor/UnrealEd/Classes/Commandlets/ResavePackagesCommandlet.h
+++ b/Engine/Source/Editor/UnrealEd/Classes/Commandlets/ResavePackagesCommandlet.h
@@ -63,7 +63,7 @@ protected:
 	// skip the assert when a package can not be opened
 	bool bCanIgnoreFails;
 
-	/** load all packages, and display warnings for those packages which would have been resaved but were read-only */
+	/** load all packages, and display warnings for those packages which would have been resaved but were read-only **/
 	bool bVerifyContent;
 
 	/** if we should only save dirty packages **/
@@ -72,15 +72,18 @@ protected:
 	/** if we should auto checkout packages that need to be saved**/
 	bool bAutoCheckOut;
 
-	/** if we should batch together source control operations rather than submit them one at a time */
+	/** if we should batch together source control operations rather than submit them one at a time **/
 	bool bBatchSourceControl;
 
-	/** if we should simply skip checked out files rather than error-ing out */
+	/** if we should simply skip checked out files rather than error-ing out **/
 	bool bSkipCheckedOutFiles;
 
 	/** if we should auto checkin packages that were checked out**/
 	bool bAutoCheckIn;
 
+	/** if we should skip trying to virtualize packages being checked in **/
+	bool bSkipVirtualization;
+
 	/** Should we build lighting for the packages we are saving? **/
 	bool bShouldBuildLighting;
 
@@ -243,6 +246,7 @@ protected:
 	bool CanCheckoutFile(const FString& Filename, FString& CheckedOutUser);
 	void CheckoutAndSavePackage(UPackage* Package, TArray<FString>& SublevelFilenames, bool bIgnoreAlreadyCheckedOut = false);
 	void CheckInFiles(const TArray<FString>& InFilesToSubmit, const FText& InDescription) const;
+	static bool TryVirtualization(const TArray<FString>& FilesToSubmit, FText& InOutDescription);
 
 	/**
 	 * Creates and returns a unique filename in the temporary file directory.
diff --git a/Engine/Source/Editor/UnrealEd/Private/Commandlets/ContentCommandlets.cpp b/Engine/Source/Editor/UnrealEd/Private/Commandlets/ContentCommandlets.cpp
--- a/Engine/Source/Editor/UnrealEd/Private/Commandlets/ContentCommandlets.cpp
+++ b/Engine/Source/Editor/UnrealEd/Private/Commandlets/ContentCommandlets.cpp
@@ -62,6 +62,7 @@
 #include "WorldPartition/ActorDescContainer.h"
 #include "WorldPartition/WorldPartition.h"
 #include "WorldPartition/WorldPartitionHelpers.h"
+#include "Virtualization/VirtualizationSystem.h"
 
 DEFINE_LOG_CATEGORY(LogContentCommandlet);
 
@@ -100,6 +101,7 @@ DEFINE_LOG_CATEGORY(LogContentCommandlet);
 #include "PerQualityLevelProperties.h"
 #include "Misc/RedirectCollector.h"
 
+
 /**-----------------------------------------------------------------------------
  *	UResavePackages commandlet.
  *
@@ -1113,6 +1115,8 @@ int32 UResavePackagesCommandlet::Main( const FString& Params )
 	bSkipCheckedOutFiles = Switches.Contains(TEXT("SkipCheckedOutPackages"));
 	/** if we should auto checkin packages that were checked out**/
 	bAutoCheckIn = bAutoCheckOut && (Switches.Contains(TEXT("AutoCheckIn")) || Switches.Contains(TEXT("AutoSubmit")));
+	/** if we should skip trying to virtualize packages being checked in **/
+	bSkipVirtualization = Switches.Contains(TEXT("SkipVirtualization"));
 	/** determine if we are building lighting for the map packages on the pass. **/
 	bShouldBuildLighting = Switches.Contains(TEXT("buildlighting"));
 	/** determine if we are building reflection captures for the map packages on the pass. **/
@@ -1748,19 +1752,78 @@ void UResavePackagesCommandlet::CheckoutAndSavePackage(UPackage* Package, TArray
 
 void UResavePackagesCommandlet::CheckInFiles(const TArray<FString>& InFilesToSubmit, const FText& InDescription) const
 {
-	if (!bAutoCheckIn)
+	if (!bAutoCheckIn || InFilesToSubmit.IsEmpty())
 	{
 		return;
 	}
 
-	// Check in all changed files
-	if (InFilesToSubmit.Num() > 0)
+	FText FinalDescription = InDescription;
+	if (!bSkipVirtualization)
+	{
+		if (!TryVirtualization(InFilesToSubmit, FinalDescription))
+		{
+			UE_LOG(LogContentCommandlet, Error, TEXT("Files will not be checked in due to virtualization failure!"));
+			return;
+		}
+	}
+	else
+	{
+		UE_LOG(LogContentCommandlet, Display, TEXT("Skipping virtualization due to the cmdline"));
+	}
+
+	TSharedRef<FCheckIn, ESPMode::ThreadSafe> CheckInOperation = ISourceControlOperation::Create<FCheckIn>();
+	CheckInOperation->SetDescription(FinalDescription);
+
+	ISourceControlProvider& SourceControlProvider = ISourceControlModule::Get().GetProvider();
+	SourceControlProvider.Execute(CheckInOperation, SourceControlHelpers::PackageFilenames(InFilesToSubmit));
+}
+
+bool UResavePackagesCommandlet::TryVirtualization(const TArray<FString>& FilesToSubmit, FText& InOutDescription)
+{
+	TArray<FText> PayloadErrors;
+	TArray<FText> DescriptionTags;
+
+	UE::Virtualization::IVirtualizationSystem& System = UE::Virtualization::IVirtualizationSystem::Get();
+	if (!System.IsEnabled())
+	{
+		return true;
+	}
+
+	UE::Virtualization::EVirtualizationResult Result = System.TryVirtualizePackages(FilesToSubmit, DescriptionTags, PayloadErrors);
+	if (Result == UE::Virtualization::EVirtualizationResult::Success)
+	{
+		FTextBuilder NewDescription;
+		NewDescription.AppendLine(InOutDescription);
+
+		for (const FText& Line : DescriptionTags)
+		{
+			NewDescription.AppendLine(Line);
+		}
+
+		InOutDescription = NewDescription.ToText();
+
+		return true;
+	}
+	else if (System.AllowSubmitIfVirtualizationFailed())
 	{
-		TSharedRef<FCheckIn, ESPMode::ThreadSafe> CheckInOperation = ISourceControlOperation::Create<FCheckIn>();
-		CheckInOperation->SetDescription(InDescription);
+		for (const FText& Error : PayloadErrors)
+		{
+			UE_LOG(LogContentCommandlet, Warning, TEXT("%s"), *Error.ToString());
+		}
+
+		// Even though the virtualization process had problems we should continue submitting
+		return true;
+	}
+	else
+	{
+		for (const FText& Error : PayloadErrors)
+		{
+			UE_LOG(LogContentCommandlet, Error, TEXT("%s"), *Error.ToString());
+		}
+
+		UE_LOG(LogContentCommandlet, Error, TEXT("Failed to virtualize the files being submitted!"));
 
-		ISourceControlProvider& SourceControlProvider = ISourceControlModule::Get().GetProvider();
-		SourceControlProvider.Execute(CheckInOperation, SourceControlHelpers::PackageFilenames(InFilesToSubmit));
+		return false;
 	}
 }
 
