Fix 'reachable garbage object' errors in anim BP thumbnail preview scenes

Cherry-picking fix for dangling ptrs to preview instances in debug skeletal DebugSkelMeshComponent, also removing the explicit MarkAsGarbage call for preview instances.

Original change, CL 23407330:
DebugSkelMeshComponent: Clear preview instance correctly

diff --git a/Engine/Source/Editor/UnrealEd/Private/ThumbnailHelpers.cpp b/Engine/Source/Editor/UnrealEd/Private/ThumbnailHelpers.cpp
--- a/Engine/Source/Editor/UnrealEd/Private/ThumbnailHelpers.cpp
+++ b/Engine/Source/Editor/UnrealEd/Private/ThumbnailHelpers.cpp
@@ -898,16 +898,8 @@ bool FAnimBlueprintThumbnailScene::SetAnimBlueprint(class UAnimBlueprint* InBlue
 			{
 				bSetSucessfully = true;
 
-				UAnimInstance* PreviousInstance = PreviewActor->GetSkeletalMeshComponent()->GetAnimInstance();
-
 				PreviewActor->GetSkeletalMeshComponent()->SetAnimInstanceClass(InBlueprint->GeneratedClass);
 
-				if (PreviousInstance && PreviousInstance != PreviewActor->GetSkeletalMeshComponent()->GetAnimInstance())
-				{
-					//Mark this as gone!
-					PreviousInstance->MarkAsGarbage();
-				}
-
 				FTransform MeshTransform = FTransform::Identity;
 
 				PreviewActor->SetActorLocation(FVector(0, 0, 0), false);
