AssetRegistry: Disable ClearGathererCache; it is not threadsafe because the AssetGathererThread can be reading the AssetDatas for save outside of the TickLock while we delete them inside the TickLock.
#rnx

diff --git a/Engine/Source/Editor/UnrealEd/Private/CookOnTheFlyServer.cpp b/Engine/Source/Editor/UnrealEd/Private/CookOnTheFlyServer.cpp
--- a/Engine/Source/Editor/UnrealEd/Private/CookOnTheFlyServer.cpp
+++ b/Engine/Source/Editor/UnrealEd/Private/CookOnTheFlyServer.cpp
@@ -7306,7 +7306,7 @@ void UCookOnTheFlyServer::BlockOnAssetRegistry()
 #endif
 
 	FAssetRegistryGenerator::UpdateAssetManagerDatabase();
-	AssetRegistry->ClearGathererCache();
+	//AssetRegistry->ClearGathererCache();
 }
 
 void UCookOnTheFlyServer::RefreshPlatformAssetRegistries(const TArrayView<const ITargetPlatform* const>& TargetPlatforms)
