Fixing packages not being marked as dirty when undoing post-save

This was because the dirty state of the package was cached in the object transaction for the package, so if you undid back to a state where a package was previously clean, then the package would just get marked as clean again even if the in-memory state no longer matched the file on disk.

To avoid this we now "fence" the dirty state of a package on certain events (like save or delete) so we can tell if the dirty state cached in the transaction is still valid to re-apply, or if the package just needs to be marked dirty.

This also moved the save/restore of the dirty state out of the object transaction for the package and into the transaction itself, as the package object was only transacted when it became dirty, but we need each transaction to know the cached dirty state of any packages related to objects within the current transaction. This allows any undo re-evaluate the dirty state of a package, and not just an undo where the dirty state changed in the past.

diff --git a/Engine/Source/Editor/UnrealEd/Classes/Editor/Transactor.h b/Engine/Source/Editor/UnrealEd/Classes/Editor/Transactor.h
--- a/Engine/Source/Editor/UnrealEd/Classes/Editor/Transactor.h
+++ b/Engine/Source/Editor/UnrealEd/Classes/Editor/Transactor.h
@@ -225,6 +225,22 @@ protected:
 		int32 SaveCount = 0;
 	};
 
+	struct FPackageRecord
+	{
+		friend FArchive& operator<<(FArchive& Ar, FPackageRecord& PackageRecord)
+		{
+			Ar << PackageRecord.DirtyFenceCount;
+			Ar << PackageRecord.bWasDirty;
+			return Ar;
+		}
+
+		/** The fence count used to determine if the cached dirty state is still valid on undo/redo (ie, was the package saved or deleted since the undo was made) */
+		int32 DirtyFenceCount = 0;
+
+		/** The cached dirty state of this package, to be restored on undo/redo as long as the DirtyFenceCount still matches */
+		bool bWasDirty = false;
+	};
+
 	// Transaction variables.
 	/** List of object records in this transaction */
 	TIndirectArray<FObjectRecord> Records;
@@ -232,6 +248,9 @@ protected:
 	/** Map of object records (non-array), for optimized look-up and to prevent an object being serialized to a transaction more than once. */
 	TMap<UE::Transaction::FPersistentObjectRef, FObjectRecords> ObjectRecordsMap;
 
+	/** Map of package records for tracking and restoring the dirty state of packages on undo/redo. The package itself may also be stored in ObjectRecordsMap. */
+	TMap<UE::Transaction::FPersistentObjectRef, FPackageRecord> PackageRecordMap;
+
 	/** Unique identifier for this transaction, used to track it during its lifetime */
 	FGuid		Id;
 
@@ -296,6 +315,8 @@ private:
 	FTransaction( const FTransaction& ) = delete;
 	FTransaction& operator=( const FTransaction& ) = delete;
 
+	void SavePackage(UPackage* Package);
+
 public:
 	
 	// FTransactionBase interface.
@@ -391,7 +412,7 @@ public:
 	/** Serializes a reference to a transaction in a given archive. */
 	friend FArchive& operator<<( FArchive& Ar, FTransaction& T )
 	{
-		Ar << T.Records << T.ObjectRecordsMap << T.Id << T.Title << T.Context << T.PrimaryObject;
+		Ar << T.Records << T.ObjectRecordsMap << T.PackageRecordMap << T.Id << T.Title << T.Context << T.PrimaryObject;
 
 		if (Ar.IsLoading())
 		{
@@ -425,7 +446,6 @@ public:
 	 * @param	Objects		[out] Receives the object list.  Previous contents are cleared.
 	 */
 	void GetTransactionObjects(TArray<UObject*>& Objects) const;
-	void RemoveRecords( int32 Count = 1 );
 	int32 GetRecordCount() const;
 
 	const UObject* GetPrimaryObject() const { return PrimaryObject; }
diff --git a/Engine/Source/Editor/UnrealEd/Private/EditorTransaction.cpp b/Engine/Source/Editor/UnrealEd/Private/EditorTransaction.cpp
--- a/Engine/Source/Editor/UnrealEd/Private/EditorTransaction.cpp
+++ b/Engine/Source/Editor/UnrealEd/Private/EditorTransaction.cpp
@@ -4,6 +4,7 @@
 #include "CoreMinimal.h"
 #include "Misc/MemStack.h"
 #include "UObject/Object.h"
+#include "UObject/ObjectSaveContext.h"
 #include "UObject/Package.h"
 #include "Algo/Find.h"
 #include "Algo/Reverse.h"
@@ -11,6 +12,7 @@
 #include "Components/ActorComponent.h"
 #include "Model.h"
 #include "Misc/ITransactionObjectAnnotation.h"
+#include "Editor.h"
 #include "Editor/Transactor.h"
 #include "Editor/TransBuffer.h"
 #include "Components/ModelComponent.h"
@@ -20,6 +22,68 @@
 
 DEFINE_LOG_CATEGORY_STATIC(LogEditorTransaction, Log, All);
 
+struct FTransactionPackageDirtyFenceCounter
+{
+public:
+	static FTransactionPackageDirtyFenceCounter Get()
+	{
+		static FTransactionPackageDirtyFenceCounter Instance;
+		return Instance;
+	}
+
+	int32 GetFenceCount(const UPackage* Pkg) const
+	{
+		return PackageFenceCounts.FindRef(Pkg);
+	}
+
+	~FTransactionPackageDirtyFenceCounter()
+	{
+		UPackage::PackageSavedWithContextEvent.RemoveAll(this);
+		FEditorDelegates::OnPackageDeleted.RemoveAll(this);
+		FCoreUObjectDelegates::GetPostGarbageCollect().RemoveAll(this);
+	}
+
+private:
+	FTransactionPackageDirtyFenceCounter()
+	{
+		UPackage::PackageSavedWithContextEvent.AddRaw(this, &FTransactionPackageDirtyFenceCounter::OnPackageSaved);
+		FEditorDelegates::OnPackageDeleted.AddRaw(this, &FTransactionPackageDirtyFenceCounter::OnPackageDeleted);
+		FCoreUObjectDelegates::GetPostGarbageCollect().AddRaw(this, &FTransactionPackageDirtyFenceCounter::OnPostGarbageCollect);
+	}
+
+	void IncrementCount(UPackage* Pkg)
+	{
+		int32& PackageSaveCount = PackageFenceCounts.FindOrAdd(Pkg);
+		++PackageSaveCount;
+	}
+
+	void OnPackageSaved(const FString& Filename, UPackage* Pkg, FObjectPostSaveContext ObjectSaveContext)
+	{
+		if (!(ObjectSaveContext.GetSaveFlags() & SAVE_FromAutosave) && !ObjectSaveContext.IsProceduralSave())
+		{
+			IncrementCount(Pkg);
+		}
+	}
+
+	void OnPackageDeleted(UPackage* Pkg)
+	{
+		IncrementCount(Pkg);
+	}
+
+	void OnPostGarbageCollect()
+	{
+		for (auto It = PackageFenceCounts.CreateIterator(); It; ++It)
+		{
+			if (!It->Key.IsValid())
+			{
+				It.RemoveCurrent();
+			}
+		}
+	}
+
+	TMap<TWeakObjectPtr<const UPackage>, int32, FDefaultSetAllocator, TWeakObjectPtrMapKeyFuncs<TWeakObjectPtr<const UPackage>, int32>> PackageFenceCounts;
+};
+
 /*-----------------------------------------------------------------------------
 	A single transaction.
 -----------------------------------------------------------------------------*/
@@ -489,28 +553,6 @@ bool FTransaction::IsObjectTransacting(const UObject* Object) const
 	return ChangedObjects.Contains(Object);
 }
 
-void FTransaction::RemoveRecords( int32 Count /* = 1  */ )
-{
-	if ( Count > 0 && Records.Num() >= Count )
-	{
-		// Remove anything from the ObjectRecordMap which is about to be removed from the Records array
-		for (int32 Index = 0; Index < Count; Index++)
-		{
-			FObjectRecord& Record = Records[Records.Num() - Count + Index];
-			if (FObjectRecords* ObjectRecords = ObjectRecordsMap.Find(Record.Object))
-			{
-				ObjectRecords->Records.RemoveSingle(&Record);
-				if (ObjectRecords->Records.Num() == 0)
-				{
-					ObjectRecordsMap.Remove(Record.Object);
-				}
-			}
-		}
-
-		Records.RemoveAt( Records.Num() - Count, Count );
-	}
-}
-
 /**
  * Outputs the contents of the ObjectMap to the specified output device.
  */
@@ -624,11 +666,33 @@ void FTransaction::AddReferencedObjects( FReferenceCollector& Collector )
 	}
 }
 
+void FTransaction::SavePackage(UPackage* Package)
+{
+	check(Package);
+
+	const bool bIsTransactional = Package->HasAnyFlags(RF_Transactional);
+	const bool bIsTransient = Package->HasAnyFlags(RF_Transient);
+	const bool bIsScriptPackage = Package->HasAnyPackageFlags(PKG_ContainsScript);
+
+	if (bIsTransactional && !bIsTransient && !bIsScriptPackage)
+	{
+		UE::Transaction::FPersistentObjectRef PackageRef(Package);
+		if (!PackageRecordMap.Contains(PackageRef))
+		{
+			FPackageRecord& PackageRecord = PackageRecordMap.Add(PackageRef);
+			PackageRecord.DirtyFenceCount = FTransactionPackageDirtyFenceCounter::Get().GetFenceCount(Package);
+			PackageRecord.bWasDirty = Package->IsDirty();
+		}
+	}
+}
+
 void FTransaction::SaveObject( UObject* Object )
 {
 	check(Object);
 	Object->CheckDefaultSubobjects();
 
+	SavePackage(Object->GetPackage());
+
 	FObjectRecords* ObjectRecords = &ObjectRecordsMap.FindOrAdd(UE::Transaction::FPersistentObjectRef(Object));
 	if (ObjectRecords->Records.Num() == 0)
 	{
@@ -671,6 +735,8 @@ void FTransaction::StoreUndo(UObject* Object, TUniquePtr<FChange> UndoChange)
 	check(Object);
 	Object->CheckDefaultSubobjects();
 
+	SavePackage(Object->GetPackage());
+
 	// Save the undo record
 	FObjectRecords& ObjectRecords = ObjectRecordsMap.FindOrAdd(UE::Transaction::FPersistentObjectRef(Object));
 	FObjectRecord* UndoRecord = ObjectRecords.Records.Add_GetRef(new FObjectRecord(this, Object, MoveTemp(UndoChange), nullptr, 0, 0, 0, 0, 0, nullptr, nullptr, nullptr));
@@ -743,6 +809,29 @@ void FTransaction::Apply()
 
 	UE::Transaction::DiffUtil::FDiffableObjectArchetypeCache ArchetypeCache;
 
+	// Update the package dirty states
+	// We do this prior to applying any object updates as we want to respect if an undo operation causes an object to dirty its own package
+	if (bFlip)
+	{
+		for (TTuple<UE::Transaction::FPersistentObjectRef, FPackageRecord>& PackageRecordPair : PackageRecordMap)
+		{
+			if (UPackage* Package = Cast<UPackage>(PackageRecordPair.Key.Get()))
+			{
+				const int32 CurrentDirtyFenceCount = FTransactionPackageDirtyFenceCounter::Get().GetFenceCount(Package);
+				const bool bCurrentDirtyFlag = Package->IsDirty();
+
+				// When restoring an undo, any package that has been "fenced" since this transaction was made needs to be considered dirty
+				// since the undo may restore it to a state that no longer matches the file on disk
+				const bool bNewDirtyFlag = PackageRecordPair.Value.bWasDirty || PackageRecordPair.Value.DirtyFenceCount != CurrentDirtyFenceCount;
+				Package->SetDirtyFlag(bNewDirtyFlag);
+				
+				// Store the current state for any inverse undo/redo operation
+				PackageRecordPair.Value.DirtyFenceCount = CurrentDirtyFenceCount;
+				PackageRecordPair.Value.bWasDirty = bCurrentDirtyFlag;
+			}
+		}
+	}
+
 	// Init objects.
 	for( int32 i=Start; i!=End; i+=Inc )
 	{
diff --git a/Engine/Source/Runtime/CoreUObject/Private/UObject/Package.cpp b/Engine/Source/Runtime/CoreUObject/Private/UObject/Package.cpp
--- a/Engine/Source/Runtime/CoreUObject/Private/UObject/Package.cpp
+++ b/Engine/Source/Runtime/CoreUObject/Private/UObject/Package.cpp
@@ -108,12 +108,6 @@ void UPackage::Serialize( FArchive& Ar )
 {
 	Super::Serialize(Ar);
 
-	if ( Ar.IsTransacting() )
-	{
-		bool bTempDirty = bDirty;
-		Ar << bTempDirty;
-		bDirty = bTempDirty;
-	}
 	if (Ar.IsCountingMemory())
 	{		
 		if (FLinker* Loader = GetLinker())
