Fix PhAT preview in RBAN simulation mode

diff --git a/Engine/Source/Editor/PhysicsAssetEditor/Private/PhysicsAssetEditor.cpp b/Engine/Source/Editor/PhysicsAssetEditor/Private/PhysicsAssetEditor.cpp
--- a/Engine/Source/Editor/PhysicsAssetEditor/Private/PhysicsAssetEditor.cpp
+++ b/Engine/Source/Editor/PhysicsAssetEditor/Private/PhysicsAssetEditor.cpp
@@ -3531,8 +3531,6 @@ void FPhysicsAssetEditor::HandlePreviewSceneCreated(const TSharedRef<IPersonaPre
 
 	SharedData->EditorSkelComp->Stop();
 
-	SharedData->EditorSkelComp->PreviewInstance = NewObject<UPhysicsAssetEditorAnimInstance>(SharedData->EditorSkelComp, TEXT("PhatAnimScriptInstance"));
-
 	SharedData->PhysicalAnimationComponent = NewObject<UPhysicalAnimationComponent>(Actor);
 	SharedData->PhysicalAnimationComponent->SetSkeletalMeshComponent(SharedData->EditorSkelComp);
 	InPersonaPreviewScene->AddComponent(SharedData->PhysicalAnimationComponent, FTransform::Identity);
diff --git a/Engine/Source/Editor/PhysicsAssetEditor/Private/PhysicsAssetEditorSkeletalMeshComponent.cpp b/Engine/Source/Editor/PhysicsAssetEditor/Private/PhysicsAssetEditorSkeletalMeshComponent.cpp
--- a/Engine/Source/Editor/PhysicsAssetEditor/Private/PhysicsAssetEditorSkeletalMeshComponent.cpp
+++ b/Engine/Source/Editor/PhysicsAssetEditor/Private/PhysicsAssetEditorSkeletalMeshComponent.cpp
@@ -68,6 +68,12 @@ UPhysicsAssetEditorSkeletalMeshComponent::UPhysicsAssetEditorSkeletalMeshCompone
 
 	bSelectable = false;
 }
+
+TObjectPtr<UAnimPreviewInstance> UPhysicsAssetEditorSkeletalMeshComponent::CreatePreviewInstance()
+{
+	return NewObject<UPhysicsAssetEditorAnimInstance>(this, TEXT("PhatAnimScriptInstance"));
+}
+
 void UPhysicsAssetEditorSkeletalMeshComponent::DebugDraw(const FSceneView* View, FPrimitiveDrawInterface* PDI)
 {
 	check(SharedData);
diff --git a/Engine/Source/Editor/PhysicsAssetEditor/Private/PhysicsAssetEditorSkeletalMeshComponent.h b/Engine/Source/Editor/PhysicsAssetEditor/Private/PhysicsAssetEditorSkeletalMeshComponent.h
--- a/Engine/Source/Editor/PhysicsAssetEditor/Private/PhysicsAssetEditorSkeletalMeshComponent.h
+++ b/Engine/Source/Editor/PhysicsAssetEditor/Private/PhysicsAssetEditorSkeletalMeshComponent.h
@@ -48,6 +48,9 @@ class UPhysicsAssetEditorSkeletalMeshComponent : public UDebugSkelMeshComponent
 	/** Mesh-space matrices showing state of just animation (ie before physics) - useful for debugging! */
 	TArray<FTransform> AnimationSpaceBases;
 
+	/** UDebugSkelMeshComponent interface */
+	virtual TObjectPtr<UAnimPreviewInstance> CreatePreviewInstance() override;
+
 	/** UPrimitiveComponent interface */
 	virtual FPrimitiveSceneProxy* CreateSceneProxy() override;
 	virtual void AddImpulseAtLocation(FVector Impulse, FVector Location, FName BoneName = NAME_None) override;
diff --git a/Engine/Source/Editor/UnrealEd/Classes/Animation/DebugSkelMeshComponent.h b/Engine/Source/Editor/UnrealEd/Classes/Animation/DebugSkelMeshComponent.h
--- a/Engine/Source/Editor/UnrealEd/Classes/Animation/DebugSkelMeshComponent.h
+++ b/Engine/Source/Editor/UnrealEd/Classes/Animation/DebugSkelMeshComponent.h
@@ -379,6 +379,9 @@ class UNREALED_API UDebugSkelMeshComponent : public USkeletalMeshComponent
 	// we can disable per asset, so that if some other window disabled before me, I don't accidently turn it off
 	virtual void EnablePreview(bool bEnable, class UAnimationAsset * PreviewAsset);
 
+	// Create the preview instance to use (default UAnimPreviewInstance)
+	virtual TObjectPtr<UAnimPreviewInstance> CreatePreviewInstance();
+
 	// reference pose for this component
 	// we don't want to use default refpose because you still want to move joint when this mode is on
 	virtual void ShowReferencePose(bool bRefPose);
diff --git a/Engine/Source/Editor/UnrealEd/Private/Animation/DebugSkelMeshComponent.cpp b/Engine/Source/Editor/UnrealEd/Private/Animation/DebugSkelMeshComponent.cpp
--- a/Engine/Source/Editor/UnrealEd/Private/Animation/DebugSkelMeshComponent.cpp
+++ b/Engine/Source/Editor/UnrealEd/Private/Animation/DebugSkelMeshComponent.cpp
@@ -515,6 +515,11 @@ FString UDebugSkelMeshComponent::GetPreviewText() const
 #undef LOCTEXT_NAMESPACE
 }
 
+TObjectPtr<UAnimPreviewInstance> UDebugSkelMeshComponent::CreatePreviewInstance()
+{
+	return NewObject<UAnimPreviewInstance>(this);
+}
+
 void UDebugSkelMeshComponent::InitAnim(bool bForceReinit)
 {
 	if (PreviewInstance != nullptr && AnimScriptInstance == PreviewInstance && bForceReinit)
@@ -529,7 +534,7 @@ void UDebugSkelMeshComponent::InitAnim(bool bForceReinit)
 	// if PreviewInstance is nullptr, create here once
 	if (PreviewInstance == nullptr)
 	{
-		PreviewInstance = NewObject<UAnimPreviewInstance>(this);
+		PreviewInstance = CreatePreviewInstance();
 		check(PreviewInstance);
 
 		//Set transactional flag in order to restore slider position when undo operation is performed
