Optimize thumbnails to avoid loading from disk unless the package they are in is going to be resaved

#rnx

diff --git a/Engine/Source/Editor/UnrealEd/Private/ThumbnailManager.cpp b/Engine/Source/Editor/UnrealEd/Private/ThumbnailManager.cpp
--- a/Engine/Source/Editor/UnrealEd/Private/ThumbnailManager.cpp
+++ b/Engine/Source/Editor/UnrealEd/Private/ThumbnailManager.cpp
@@ -440,6 +440,7 @@ void UThumbnailManager::DirtyThumbnailForObject(UObject* ObjectBeingModified)
 
 			const bool bMemoryPackage = FPackageName::IsMemoryPackage(ObjectBeingModified->GetPathName());  // Don't try to load from disk if the package is a memory package
 			const bool bUnsavedPackage = ObjectPackage->HasAnyPackageFlags(PKG_NewlyCreated);               // Don't try loading thumbnails for package that have never been saved
+			const bool bPackageDirty = ObjectPackage->IsDirty();                                            // Don't try loading thumbnails for package that we have no intention of saving
 			const bool bIsGarbageCollecting = IsGarbageCollecting();                                        // Don't attempt to do this while garbage collecting since loading or finding objects during GC is illegal
 			const bool bUsesGenericThumbnail = [ObjectBeingModified, this]() -> bool                        // No need to dirty generic thumbnails
 			{
@@ -453,7 +454,7 @@ void UThumbnailManager::DirtyThumbnailForObject(UObject* ObjectBeingModified)
 				}
 			}();
 
-			const bool bTryLoadThumbnailFromDisk = !bIsGarbageCollecting && !bMemoryPackage && !bUnsavedPackage && !bUsesGenericThumbnail;
+			const bool bTryLoadThumbnailFromDisk = !bIsGarbageCollecting && !bMemoryPackage && !bUnsavedPackage && !bUsesGenericThumbnail && bPackageDirty;
 			if (bTryLoadThumbnailFromDisk)
 			{
 				FName ObjectFullName = FName(*ObjectBeingModified->GetFullName());
