Cooker LLM Tags: Add more tags around allocations made when cooking.
#rnx

diff --git a/Engine/Plugins/FX/Niagara/Source/Niagara/Private/NiagaraEmitter.cpp b/Engine/Plugins/FX/Niagara/Source/Niagara/Private/NiagaraEmitter.cpp
--- a/Engine/Plugins/FX/Niagara/Source/Niagara/Private/NiagaraEmitter.cpp
+++ b/Engine/Plugins/FX/Niagara/Source/Niagara/Private/NiagaraEmitter.cpp
@@ -17,6 +17,7 @@
 #include "NiagaraStats.h"
 #include "NiagaraSystem.h"
 #include "NiagaraTrace.h"
+#include "HAL/LowLevelMemTracker.h"
 #include "Interfaces/ITargetPlatform.h"
 #include "Modules/ModuleManager.h"
 #include "UObject/LinkerLoad.h"
@@ -1810,6 +1811,7 @@ void UNiagaraEmitter::UpdateEmitterAfterLoad()
 		return;
 	}
 	bFullyLoaded = true;
+	LLM_SCOPE(ELLMTag::Niagara);
 	
 #if WITH_EDITORONLY_DATA
 	check(IsInGameThread());
diff --git a/Engine/Plugins/FX/Niagara/Source/Niagara/Private/NiagaraSystem.cpp b/Engine/Plugins/FX/Niagara/Source/Niagara/Private/NiagaraSystem.cpp
--- a/Engine/Plugins/FX/Niagara/Source/Niagara/Private/NiagaraSystem.cpp
+++ b/Engine/Plugins/FX/Niagara/Source/Niagara/Private/NiagaraSystem.cpp
@@ -28,6 +28,7 @@
 #include "NiagaraWorldManager.h"
 #include "Algo/RemoveIf.h"
 #include "Async/Async.h"
+#include "HAL/LowLevelMemTracker.h"
 #include "Interfaces/ITargetPlatform.h"
 #include "Misc/ScopedSlowTask.h"
 #include "Misc/ScopeExit.h"
@@ -1772,6 +1773,7 @@ void UNiagaraSystem::ComputeRenderersDrawOrder()
 
 void UNiagaraSystem::CacheFromCompiledData()
 {
+	LLM_SCOPE(ELLMTag::Niagara);
 	const FNiagaraDataSetCompiledData& SystemDataSet = SystemCompiledData.DataSetCompiledData;
 
 	bNeedsAsyncOptimize = true;
diff --git a/Engine/Plugins/FX/Niagara/Source/NiagaraEditor/Private/NiagaraGraph.cpp b/Engine/Plugins/FX/Niagara/Source/NiagaraEditor/Private/NiagaraGraph.cpp
--- a/Engine/Plugins/FX/Niagara/Source/NiagaraEditor/Private/NiagaraGraph.cpp
+++ b/Engine/Plugins/FX/Niagara/Source/NiagaraEditor/Private/NiagaraGraph.cpp
@@ -31,6 +31,7 @@
 #include "NiagaraScriptVariable.h"
 #include "NiagaraSettings.h"
 #include "HAL/FileManager.h"
+#include "HAL/LowLevelMemTracker.h"
 #include "Misc/Paths.h"
 #include "Misc/SecureHash.h"
 #include "Modules/ModuleManager.h"
@@ -282,6 +283,7 @@ void UNiagaraGraph::PostLoad_LWCFixup(int32 NiagaraVersion)
 
 void UNiagaraGraph::PostLoad()
 {
+	LLM_SCOPE(ELLMTag::Niagara);
 	check(!bIsForCompilationOnly);
 
 	Super::PostLoad();
diff --git a/Engine/Plugins/Runtime/GameplayAbilities/Source/GameplayAbilities/Private/AbilitySystemGlobals.cpp b/Engine/Plugins/Runtime/GameplayAbilities/Source/GameplayAbilities/Private/AbilitySystemGlobals.cpp
--- a/Engine/Plugins/Runtime/GameplayAbilities/Source/GameplayAbilities/Private/AbilitySystemGlobals.cpp
+++ b/Engine/Plugins/Runtime/GameplayAbilities/Source/GameplayAbilities/Private/AbilitySystemGlobals.cpp
@@ -8,6 +8,7 @@
 #include "AbilitySystemComponent.h"
 #include "AbilitySystemInterface.h"
 #include "AbilitySystemLog.h"
+#include "HAL/LowLevelMemTracker.h"
 #include "GameFramework/PlayerController.h"
 #include "GameplayCueManager.h"
 #include "GameplayTagResponseTable.h"
@@ -53,6 +54,7 @@ UAbilitySystemGlobals::UAbilitySystemGlobals(const FObjectInitializer& ObjectIni
 
 void UAbilitySystemGlobals::InitGlobalData()
 {
+	LLM_SCOPE(TEXT("AbilitySystem"));
 	GetGlobalCurveTable();
 	GetGlobalAttributeMetaDataTable();
 	
diff --git a/Engine/Source/Developer/IoStoreUtilities/Private/PackageStoreManifest.cpp b/Engine/Source/Developer/IoStoreUtilities/Private/PackageStoreManifest.cpp
--- a/Engine/Source/Developer/IoStoreUtilities/Private/PackageStoreManifest.cpp
+++ b/Engine/Source/Developer/IoStoreUtilities/Private/PackageStoreManifest.cpp
@@ -5,18 +5,22 @@
 #include "Serialization/JsonReader.h"
 #include "Serialization/JsonSerializer.h"
 #include "Dom/JsonObject.h"
+#include "HAL/LowLevelMemTracker.h"
 #include "Misc/FileHelper.h"
 #include "Misc/ScopeLock.h"
 #include "Misc/Paths.h"
 
+LLM_DEFINE_TAG(Cooker_PackageStoreManifest);
 FPackageStoreManifest::FPackageStoreManifest(const FString& InCookedOutputPath)
 	: CookedOutputPath(InCookedOutputPath)
 {
+	LLM_SCOPE_BYTAG(Cooker_PackageStoreManifest);
 	FPaths::NormalizeFilename(CookedOutputPath);
 }
 
 void FPackageStoreManifest::BeginPackage(FName PackageName)
 {
+	LLM_SCOPE_BYTAG(Cooker_PackageStoreManifest);
 	FScopeLock Lock(&CriticalSection);
 	FPackageInfo& PackageInfo = PackageInfoByNameMap.FindOrAdd(PackageName);
 	PackageInfo.PackageName = PackageName;
@@ -33,6 +37,7 @@ void FPackageStoreManifest::BeginPackage(FName PackageName)
 
 void FPackageStoreManifest::AddPackageData(FName PackageName, const FString& FileName, const FIoChunkId& ChunkId)
 {
+	LLM_SCOPE_BYTAG(Cooker_PackageStoreManifest);
 	FScopeLock Lock(&CriticalSection);
 	FPackageInfo* PackageInfo = GetPackageInfo_NoLock(PackageName);
 	check(PackageInfo);
@@ -45,6 +50,7 @@ void FPackageStoreManifest::AddPackageData(FName PackageName, const FString& Fil
 
 void FPackageStoreManifest::AddBulkData(FName PackageName, const FString& FileName, const FIoChunkId& ChunkId)
 {
+	LLM_SCOPE_BYTAG(Cooker_PackageStoreManifest);
 	FScopeLock Lock(&CriticalSection);
 	FPackageInfo* PackageInfo = GetPackageInfo_NoLock(PackageName);
 	check(PackageInfo);
@@ -57,6 +63,7 @@ void FPackageStoreManifest::AddBulkData(FName PackageName, const FString& FileNa
 
 FIoStatus FPackageStoreManifest::Save(const TCHAR* Filename) const
 {
+	LLM_SCOPE_BYTAG(Cooker_PackageStoreManifest);
 	FScopeLock Lock(&CriticalSection);
 	TStringBuilder<64> ChunkIdStringBuilder;
 	auto ChunkIdToString = [&ChunkIdStringBuilder](const FIoChunkId& ChunkId)
@@ -146,6 +153,7 @@ FIoStatus FPackageStoreManifest::Save(const TCHAR* Filename) const
 
 FIoStatus FPackageStoreManifest::Load(const TCHAR* Filename)
 {
+	LLM_SCOPE_BYTAG(Cooker_PackageStoreManifest);
 	FScopeLock Lock(&CriticalSection);
 	PackageInfoByNameMap.Empty();
 	FileNameByChunkIdMap.Empty();
diff --git a/Engine/Source/Editor/ContentBrowser/Private/CollectionAssetRegistryBridge.cpp b/Engine/Source/Editor/ContentBrowser/Private/CollectionAssetRegistryBridge.cpp
--- a/Engine/Source/Editor/ContentBrowser/Private/CollectionAssetRegistryBridge.cpp
+++ b/Engine/Source/Editor/ContentBrowser/Private/CollectionAssetRegistryBridge.cpp
@@ -10,6 +10,7 @@
 #include "Containers/Set.h"
 #include "ContentBrowserLog.h"
 #include "Delegates/Delegate.h"
+#include "HAL/LowLevelMemTracker.h"
 #include "HAL/Platform.h"
 #include "ICollectionManager.h"
 #include "Logging/LogCategory.h"
@@ -172,6 +173,7 @@ FCollectionAssetRegistryBridge::~FCollectionAssetRegistryBridge()
 
 void FCollectionAssetRegistryBridge::OnAssetRegistryLoadComplete()
 {
+	LLM_SCOPE_BYNAME(TEXT("CollectionManager"));
 	FCollectionManagerModule& CollectionManagerModule = FCollectionManagerModule::GetModule();
 
 	// We've found all the assets, let the collections manager fix up its references now so that it doesn't reference any redirectors
diff --git a/Engine/Source/Editor/UnrealEd/Private/Commandlets/AssetRegistryGenerator.cpp b/Engine/Source/Editor/UnrealEd/Private/Commandlets/AssetRegistryGenerator.cpp
--- a/Engine/Source/Editor/UnrealEd/Private/Commandlets/AssetRegistryGenerator.cpp
+++ b/Engine/Source/Editor/UnrealEd/Private/Commandlets/AssetRegistryGenerator.cpp
@@ -17,6 +17,7 @@
 #include "Engine/World.h"
 #include "GameDelegates.h"
 #include "HAL/FileManager.h"
+#include "HAL/LowLevelMemTracker.h"
 #include "HAL/PlatformFileManager.h"
 #include "ICollectionManager.h"
 #include "Interfaces/ITargetPlatform.h"
@@ -52,6 +53,8 @@ DEFINE_LOG_CATEGORY_STATIC(LogAssetRegistryGenerator, Log, All);
 
 #define LOCTEXT_NAMESPACE "AssetRegistryGenerator"
 
+LLM_DEFINE_TAG(Cooker_GeneratedAssetRegistry);
+
 //////////////////////////////////////////////////////////////////////////
 // Static functions
 FName GetPackageNameFromDependencyPackageName(const FName RawPackageFName)
@@ -160,6 +163,7 @@ private:
 
 void FAssetRegistryGenerator::UpdateAssetManagerDatabase()
 {
+	LLM_SCOPE_BYTAG(Cooker_GeneratedAssetRegistry);
 	InitializeUseAssetManager();
 	if (bUseAssetManager)
 	{
@@ -185,6 +189,7 @@ FAssetRegistryGenerator::FAssetRegistryGenerator(const ITargetPlatform* InPlatfo
 	, HighestChunkId(0)
 	, DependencyInfo(*GetMutableDefault<UChunkDependencyInfo>())
 {
+	LLM_SCOPE_BYTAG(Cooker_GeneratedAssetRegistry);
 	bool bOnlyHardReferences = false;
 	const UProjectPackagingSettings* const PackagingSettings = GetDefault<UProjectPackagingSettings>();
 	if (PackagingSettings)
@@ -924,11 +929,13 @@ void FAssetRegistryGenerator::CalculateChunkIdsAndAssignToManifest(const FName&
 
 void FAssetRegistryGenerator::CleanManifestDirectories()
 {
+	LLM_SCOPE_BYTAG(Cooker_GeneratedAssetRegistry);
 	CleanTempPackagingDirectory(TargetPlatform->PlatformName());
 }
 
 void FAssetRegistryGenerator::SetPreviousAssetRegistry(TUniquePtr<FAssetRegistryState>&& InPreviousState)
 {
+	LLM_SCOPE_BYTAG(Cooker_GeneratedAssetRegistry);
 	PreviousPackagesToUpdate.Empty();
 	if (InPreviousState)
 	{
@@ -1020,6 +1027,7 @@ void FAssetRegistryGenerator::InjectEncryptionData(FAssetRegistryState& TargetSt
 bool FAssetRegistryGenerator::SaveManifests(FSandboxPlatformFile& InSandboxFile, int64 InOverrideChunkSize,
 	const TCHAR* InManifestSubDir)
 {
+	LLM_SCOPE_BYTAG(Cooker_GeneratedAssetRegistry);
 	if (!bGenerateChunks)
 	{
 		return true;
@@ -1213,6 +1221,7 @@ void FAssetRegistryGenerator::UpdateCollectionAssetData()
 
 void FAssetRegistryGenerator::Initialize(const TArray<FName> &InStartupPackages, bool bInitializeFromExisting)
 {
+	LLM_SCOPE_BYTAG(Cooker_GeneratedAssetRegistry);
 	StartupPackages.Append(InStartupPackages);
 
 	FAssetRegistrySerializationOptions SaveOptions;
@@ -1339,6 +1348,7 @@ void FAssetRegistryGenerator::FinalizeChunkIDs(const TSet<FName>& InCookedPackag
 	const TSet<FName>& InDevelopmentOnlyPackages, FSandboxPlatformFile& InSandboxFile,
 	bool bGenerateStreamingInstallManifest)
 {
+	LLM_SCOPE_BYTAG(Cooker_GeneratedAssetRegistry);
 	// bGenerateNoChunks overrides bGenerateStreamingInstallManifest overrides ShouldPlatformGenerateStreamingInstallManifest
 	// bGenerateChunks means we allow chunks other than 0 based on package ChunkIds, AND we generate a manifest for each chunk
 	const UProjectPackagingSettings* PackagingSettings = Cast<UProjectPackagingSettings>(UProjectPackagingSettings::StaticClass()->GetDefaultObject());
@@ -1484,11 +1494,13 @@ void FAssetRegistryGenerator::FinalizeChunkIDs(const TSet<FName>& InCookedPackag
 
 void FAssetRegistryGenerator::RegisterChunkDataGenerator(TSharedRef<IChunkDataGenerator> InChunkDataGenerator)
 {
+	LLM_SCOPE_BYTAG(Cooker_GeneratedAssetRegistry);
 	ChunkDataGenerators.Add(MoveTemp(InChunkDataGenerator));
 }
 
 void FAssetRegistryGenerator::PreSave(const TSet<FName>& InCookedPackages)
 {
+	LLM_SCOPE_BYTAG(Cooker_GeneratedAssetRegistry);
 	if (bUseAssetManager)
 	{
 		UAssetManager::Get().PreSaveAssetRegistry(TargetPlatform, InCookedPackages);
@@ -1497,6 +1509,7 @@ void FAssetRegistryGenerator::PreSave(const TSet<FName>& InCookedPackages)
 
 void FAssetRegistryGenerator::PostSave()
 {
+	LLM_SCOPE_BYTAG(Cooker_GeneratedAssetRegistry);
 	if (bUseAssetManager)
 	{
 		UAssetManager::Get().PostSaveAssetRegistry();
@@ -1529,6 +1542,7 @@ void FAssetRegistryGenerator::AddAssetToFileOrderRecursive(const FName& InPackag
 
 bool FAssetRegistryGenerator::SaveAssetRegistry(const FString& SandboxPath, bool bSerializeDevelopmentAssetRegistry, bool bForceNoFilter)
 {
+	LLM_SCOPE_BYTAG(Cooker_GeneratedAssetRegistry);
 	UE_LOG(LogAssetRegistryGenerator, Display, TEXT("Saving asset registry v%d."), FAssetRegistryVersion::Type::LatestVersion);
 	FAutoScopedDurationTimer Timer;
 
@@ -1750,6 +1764,7 @@ public:
 
 bool FAssetRegistryGenerator::WriteCookerOpenOrder(FSandboxPlatformFile& InSandboxFile)
 {
+	LLM_SCOPE_BYTAG(Cooker_GeneratedAssetRegistry);
 	TSet<FName> PackageNameSet;
 	TSet<FName> MapList;
 	State.EnumerateAllAssets([this, &PackageNameSet, &MapList](const FAssetData& AssetData)
@@ -2520,6 +2535,7 @@ void FAssetRegistryGenerator::UpdateAssetRegistryData(const UPackage& Package,
 	bool bIncludeOnlyDiskAssets
 	)
 {
+	LLM_SCOPE_BYTAG(Cooker_GeneratedAssetRegistry);
 	const FName PackageName = Package.GetFName();
 	PreviousPackagesToUpdate.Remove(PackageName);
 
@@ -2568,6 +2584,7 @@ void FAssetRegistryGenerator::UpdateAssetRegistryData(const UPackage& Package,
 
 void FAssetRegistryGenerator::UpdateAssetRegistryData(UE::Cook::FPackageData& PackageData, UE::Cook::FAssetRegistryPackageMessage&& Message)
 {
+	LLM_SCOPE_BYTAG(Cooker_GeneratedAssetRegistry);
 	const FName PackageName = PackageData.GetPackageName();
 	PreviousPackagesToUpdate.Remove(PackageName);
 
@@ -2664,6 +2681,7 @@ void FAssetRegistryPackageMessage::Write(FCbWriter& Writer, const FPackageData&
 
 bool FAssetRegistryPackageMessage::TryRead(FCbObjectView Object, FPackageData& PackageData, const ITargetPlatform* TargetPlatform)
 {
+	LLM_SCOPE_BYTAG(Cooker_GeneratedAssetRegistry);
 	FCbFieldView AssetDatasField = Object["A"];
 	FCbArrayView AssetDatasArray = AssetDatasField.AsArrayView();
 	if (AssetDatasField.HasError())
diff --git a/Engine/Source/Editor/UnrealEd/Private/CookOnTheFlyServer.cpp b/Engine/Source/Editor/UnrealEd/Private/CookOnTheFlyServer.cpp
--- a/Engine/Source/Editor/UnrealEd/Private/CookOnTheFlyServer.cpp
+++ b/Engine/Source/Editor/UnrealEd/Private/CookOnTheFlyServer.cpp
@@ -142,6 +142,7 @@
 
 DEFINE_LOG_CATEGORY(LogCook);
 LLM_DEFINE_TAG(Cooker);
+LLM_DEFINE_TAG(Cooker_SavePackage);
 
 
 int32 GCookProgressDisplay = (int32)ECookProgressDisplayMode::RemainingPackages;
@@ -5442,6 +5443,7 @@ void UCookOnTheFlyServer::SaveCookedPackage(UE::Cook::FSaveCookedPackageContext&
 			{
 				try
 				{
+					LLM_SCOPE_BYTAG(Cooker_SavePackage);
 					Context.SavePackageResult = GEditor->Save(Package, Context.World, *Context.PlatFilename, SaveArgs);
 				}
 				catch (std::exception&)
@@ -5806,6 +5808,7 @@ void FSaveCookedPackageContext::FinishPackage()
 void UCookOnTheFlyServer::Initialize( ECookMode::Type DesiredCookMode, ECookInitializationFlags InCookFlags, const FString &InOutputDirectoryOverride )
 {
 	TRACE_CPUPROFILER_EVENT_SCOPE(UCookOnTheFlyServer::Initialize);
+	LLM_SCOPE_BYTAG(Cooker);
 
 	CurrentCookMode = DesiredCookMode;
 	CookFlags = InCookFlags;
@@ -7878,6 +7881,7 @@ bool UCookOnTheFlyServer::IsCookByTheBookRunning() const
 
 void UCookOnTheFlyServer::SaveGlobalShaderMapFiles(const TArrayView<const ITargetPlatform* const>& Platforms, ODSCRecompileCommand RecompileCommand)
 {
+	LLM_SCOPE(ELLMTag::Shaders);
 	check(!IsCookingDLC()); // GlobalShaderMapFiles are not supported when cooking DLC
 	check(IsInGameThread());
 	for (const ITargetPlatform* TargetPlatform : Platforms)
diff --git a/Engine/Source/Editor/UnrealEd/Private/Cooker/CookDirector.cpp b/Engine/Source/Editor/UnrealEd/Private/Cooker/CookDirector.cpp
--- a/Engine/Source/Editor/UnrealEd/Private/Cooker/CookDirector.cpp
+++ b/Engine/Source/Editor/UnrealEd/Private/Cooker/CookDirector.cpp
@@ -27,6 +27,8 @@
 
 extern CORE_API int32 GNumForegroundWorkers; // TaskGraph.cpp
 
+LLM_DEFINE_TAG(Cooker_MPCook);
+
 namespace UE::Cook
 {
 
diff --git a/Engine/Source/Editor/UnrealEd/Private/Cooker/CookDirector.h b/Engine/Source/Editor/UnrealEd/Private/Cooker/CookDirector.h
--- a/Engine/Source/Editor/UnrealEd/Private/Cooker/CookDirector.h
+++ b/Engine/Source/Editor/UnrealEd/Private/Cooker/CookDirector.h
@@ -11,6 +11,7 @@
 #include "CookTypes.h"
 #include "HAL/CriticalSection.h"
 #include "HAL/Event.h"
+#include "HAL/LowLevelMemTracker.h"
 #include "HAL/Runnable.h"
 #include "Memory/SharedBuffer.h"
 #include "Misc/Guid.h"
@@ -34,6 +35,8 @@ namespace UE::Cook { struct FPackageData; }
 namespace UE::Cook { struct FRetractionResultsMessage; }
 namespace UE::Cook { struct FWorkerId; }
 
+LLM_DECLARE_TAG(Cooker_MPCook);
+
 namespace UE::Cook
 {
 
diff --git a/Engine/Source/Editor/UnrealEd/Private/Cooker/CookTypes.cpp b/Engine/Source/Editor/UnrealEd/Private/Cooker/CookTypes.cpp
--- a/Engine/Source/Editor/UnrealEd/Private/Cooker/CookTypes.cpp
+++ b/Engine/Source/Editor/UnrealEd/Private/Cooker/CookTypes.cpp
@@ -15,6 +15,8 @@
 #include "Misc/Parse.h"
 #include "PackageTracker.h"
 
+LLM_DEFINE_TAG(Cooker_CachedPlatformData);
+
 namespace UE::Cook
 {
 
diff --git a/Engine/Source/Editor/UnrealEd/Private/Cooker/CookTypes.h b/Engine/Source/Editor/UnrealEd/Private/Cooker/CookTypes.h
--- a/Engine/Source/Editor/UnrealEd/Private/Cooker/CookTypes.h
+++ b/Engine/Source/Editor/UnrealEd/Private/Cooker/CookTypes.h
@@ -39,6 +39,8 @@ bool LoadFromCompactBinary(FCbFieldView Field, UE::Cook::FInitializeConfigSettin
 #define COOK_CHECKSLOW_PACKAGEDATA 0
 #define DEBUG_COOKONTHEFLY 0
 
+LLM_DECLARE_TAG(Cooker_CachedPlatformData);
+
 /** A BaseKeyFuncs for Maps and Sets with a quicker hash function for pointers than TDefaultMapKeyFuncs */
 template<typename KeyType>
 struct TFastPointerSetKeyFuncs : public DefaultKeyFuncs<KeyType>
@@ -363,6 +365,7 @@ FString LexToString(FPlatformMemoryStats::EMemoryPressureStatus Value);
 
 inline void RouteBeginCacheForCookedPlatformData(UObject* Obj, const ITargetPlatform* TargetPlatform)
 {
+	LLM_SCOPE_BYTAG(Cooker_CachedPlatformData);
 	UE_SCOPED_TEXT_COOKTIMER(*WriteToString<128>(GetClassTraceScope(Obj), TEXT("_BeginCacheForCookedPlatformData")));
 	UE_SCOPED_COOK_STAT(Obj->GetPackage()->GetFName(), EPackageEventStatType::BeginCacheForCookedPlatformData);
 	Obj->BeginCacheForCookedPlatformData(TargetPlatform);
@@ -370,6 +373,7 @@ inline void RouteBeginCacheForCookedPlatformData(UObject* Obj, const ITargetPlat
 
 inline bool RouteIsCachedCookedPlatformDataLoaded(UObject* Obj, const ITargetPlatform* TargetPlatform)
 {
+	LLM_SCOPE_BYTAG(Cooker_CachedPlatformData);
 	UE_SCOPED_TEXT_COOKTIMER(*WriteToString<128>(GetClassTraceScope(Obj), TEXT("_IsCachedCookedPlatformDataLoaded")));
 	UE_SCOPED_COOK_STAT(Obj->GetPackage()->GetFName(), EPackageEventStatType::IsCachedCookedPlatformDataLoaded);
 	return Obj->IsCachedCookedPlatformDataLoaded(TargetPlatform);
diff --git a/Engine/Source/Editor/UnrealEd/Private/Cooker/CookWorkerServer.cpp b/Engine/Source/Editor/UnrealEd/Private/Cooker/CookWorkerServer.cpp
--- a/Engine/Source/Editor/UnrealEd/Private/Cooker/CookWorkerServer.cpp
+++ b/Engine/Source/Editor/UnrealEd/Private/Cooker/CookWorkerServer.cpp
@@ -422,6 +422,7 @@ void FCookWorkerServer::SendPendingPackages()
 	{
 		return;
 	}
+	LLM_SCOPE_BYTAG(Cooker_MPCook);
 
 	TArray<FAssignPackageData> AssignDatas;
 	AssignDatas.Reserve(PackagesToAssign.Num());
@@ -439,6 +440,7 @@ void FCookWorkerServer::SendPendingPackages()
 void FCookWorkerServer::PumpReceiveMessages()
 {
 	using namespace UE::CompactBinaryTCP;
+	LLM_SCOPE_BYTAG(Cooker_MPCook);
 	TArray<FMarshalledMessage> Messages;
 	EConnectionStatus SocketStatus = TryReadPacket(Socket, ReceiveBuffer, Messages);
 	if (SocketStatus != EConnectionStatus::Okay && SocketStatus != EConnectionStatus::Incomplete)
diff --git a/Engine/Source/Editor/UnrealEd/Private/Cooker/PackageTracker.cpp b/Engine/Source/Editor/UnrealEd/Private/Cooker/PackageTracker.cpp
--- a/Engine/Source/Editor/UnrealEd/Private/Cooker/PackageTracker.cpp
+++ b/Engine/Source/Editor/UnrealEd/Private/Cooker/PackageTracker.cpp
@@ -62,6 +62,7 @@ namespace UE::Cook
 	FPackageTracker::FPackageTracker(FPackageDatas& InPackageDatas)
 		:PackageDatas(InPackageDatas)
 	{
+		LLM_SCOPE_BYTAG(Cooker);
 		for (TObjectIterator<UPackage> It; It; ++It)
 		{
 			UPackage* Package = *It;
@@ -105,6 +106,7 @@ namespace UE::Cook
 	{
 		if (Object->GetClass() == UPackage::StaticClass())
 		{
+			LLM_SCOPE_BYTAG(Cooker);
 			auto Package = const_cast<UPackage*>(static_cast<const UPackage*>(Object));
 
 			if (Package->GetOuter() == nullptr)
diff --git a/Engine/Source/Runtime/AssetRegistry/Private/AssetDataGatherer.cpp b/Engine/Source/Runtime/AssetRegistry/Private/AssetDataGatherer.cpp
--- a/Engine/Source/Runtime/AssetRegistry/Private/AssetDataGatherer.cpp
+++ b/Engine/Source/Runtime/AssetRegistry/Private/AssetDataGatherer.cpp
@@ -11,6 +11,7 @@
 #include "Async/MappedFileHandle.h"
 #include "Async/ParallelFor.h"
 #include "Containers/BinaryHeap.h"
+#include "HAL/LowLevelMemTracker.h"
 #include "HAL/PlatformFileManager.h"
 #include "HAL/PlatformProcess.h"
 #include "HAL/RunnableThread.h"
@@ -1380,6 +1381,7 @@ bool FAssetDataDiscovery::Init()
 
 uint32 FAssetDataDiscovery::Run()
 {
+	LLM_SCOPE(ELLMTag::AssetRegistry);
 	constexpr float IdleSleepTime = 0.1f;
 	{
 		FGathererScopeLock ResultsScopeLock(&ResultsLock);
@@ -2862,6 +2864,7 @@ private:
 
 	void LoadAsync()
 	{
+		LLM_SCOPE(ELLMTag::AssetRegistry);
 		Payload = LoadCacheFile(GPreloadSettings.GetMonolithicCacheFilename());
 	}
 
@@ -2972,6 +2975,8 @@ bool FAssetDataGatherer::Init()
 uint32 FAssetDataGatherer::Run()
 {
 	constexpr float IdleSleepTime = 0.1f;
+	LLM_SCOPE(ELLMTag::AssetRegistry);
+
 	while (!IsStopped)
 	{
 		InnerTickLoop(false /* bInIsSynchronousTick */, true /* bContributeToCacheSave */);
diff --git a/Engine/Source/Runtime/Core/Private/Internationalization/TextFormatter.cpp b/Engine/Source/Runtime/Core/Private/Internationalization/TextFormatter.cpp
--- a/Engine/Source/Runtime/Core/Private/Internationalization/TextFormatter.cpp
+++ b/Engine/Source/Runtime/Core/Private/Internationalization/TextFormatter.cpp
@@ -4,6 +4,7 @@
 
 #include "CoreGlobals.h"
 #include "CoreTypes.h"
+#include "HAL/LowLevelMemTracker.h"
 #include "HAL/PlatformProcess.h"
 #include "Internationalization/CulturePointer.h"
 #include "Internationalization/Internationalization.h"
@@ -696,6 +697,7 @@ bool FTextFormatData::IsValid_NoLock() const
 void FTextFormatData::Compile_NoLock()
 {
 	SCOPE_CYCLE_COUNTER(STAT_TextFormatData_Compile);
+	LLM_SCOPE(ELLMTag::Localization);
 
 	LexedExpression.Reset();
 	if (SourceType == ESourceType::Text)
diff --git a/Engine/Source/Runtime/CoreUObject/Private/UObject/SavePackage/SavePackageUtilities.cpp b/Engine/Source/Runtime/CoreUObject/Private/UObject/SavePackage/SavePackageUtilities.cpp
--- a/Engine/Source/Runtime/CoreUObject/Private/UObject/SavePackage/SavePackageUtilities.cpp
+++ b/Engine/Source/Runtime/CoreUObject/Private/UObject/SavePackage/SavePackageUtilities.cpp
@@ -1598,26 +1598,31 @@ namespace UE::SavePackageUtilities
 
 void StartSavingEDLCookInfoForVerification()
 {
+	LLM_SCOPE_BYTAG(EDLCookChecker);
 	FEDLCookChecker::StartSavingEDLCookInfoForVerification();
 }
 
 void VerifyEDLCookInfo(bool bFullReferencesExpected)
 {
+	LLM_SCOPE_BYTAG(EDLCookChecker);
 	FEDLCookChecker::Verify(bFullReferencesExpected);
 }
 
 void EDLCookInfoAddIterativelySkippedPackage(FName LongPackageName)
 {
+	LLM_SCOPE_BYTAG(EDLCookChecker);
 	FEDLCookCheckerThreadState::Get().AddPackageWithUnknownExports(LongPackageName);
 }
 
 void EDLCookInfoMoveToCompactBinaryAndClear(FCbWriter& Writer, bool& bOutHasData)
 {
+	LLM_SCOPE_BYTAG(EDLCookChecker);
 	FEDLCookChecker::MoveToCompactBinaryAndClear(Writer, bOutHasData);
 }
 
 bool EDLCookInfoAppendFromCompactBinary(FCbFieldView Field)
 {
+	LLM_SCOPE_BYTAG(EDLCookChecker);
 	return FEDLCookChecker::AppendFromCompactBinary(Field);
 }
 
diff --git a/Engine/Source/Runtime/Engine/Private/Animation/Skeleton.cpp b/Engine/Source/Runtime/Engine/Private/Animation/Skeleton.cpp
--- a/Engine/Source/Runtime/Engine/Private/Animation/Skeleton.cpp
+++ b/Engine/Source/Runtime/Engine/Private/Animation/Skeleton.cpp
@@ -350,6 +350,7 @@ bool USkeleton::IsPostLoadThreadSafe() const
 void USkeleton::PostLoad()
 {
 	Super::PostLoad();
+	LLM_SCOPE(ELLMTag::Animation);
 
 	if( GetLinker() && (GetLinker()->UEVer() < VER_UE4_REFERENCE_SKELETON_REFACTOR) )
 	{
diff --git a/Engine/Source/Runtime/Engine/Private/AssetManager.cpp b/Engine/Source/Runtime/Engine/Private/AssetManager.cpp
--- a/Engine/Source/Runtime/Engine/Private/AssetManager.cpp
+++ b/Engine/Source/Runtime/Engine/Private/AssetManager.cpp
@@ -1276,6 +1276,7 @@ void UAssetManager::SetPrimaryAssetTypeRules(FPrimaryAssetType PrimaryAssetType,
 void UAssetManager::SetPrimaryAssetRules(FPrimaryAssetId PrimaryAssetId, const FPrimaryAssetRules& Rules)
 {
 	static FPrimaryAssetRules DefaultRules;
+	LLM_SCOPE_BYTAG(AssetManager);
 
 	FPrimaryAssetRulesExplicitOverride ExplicitRules;
 	ExplicitRules.Rules = Rules;
diff --git a/Engine/Source/Runtime/Engine/Private/Curves/CurveLinearColorAtlas.cpp b/Engine/Source/Runtime/Engine/Private/Curves/CurveLinearColorAtlas.cpp
--- a/Engine/Source/Runtime/Engine/Private/Curves/CurveLinearColorAtlas.cpp
+++ b/Engine/Source/Runtime/Engine/Private/Curves/CurveLinearColorAtlas.cpp
@@ -7,6 +7,7 @@
 #include "Curves/CurveLinearColorAtlas.h"
 #include "Curves/CurveLinearColor.h"
 #include "Components/MeshComponent.h"
+#include "HAL/LowLevelMemTracker.h"
 #include "Materials/MaterialInstanceDynamic.h"
 
 #include UE_INLINE_GENERATED_CPP_BY_NAME(CurveLinearColorAtlas)
@@ -325,6 +326,8 @@ void UCurveLinearColorAtlas::OnCurveUpdated(UCurveBase* Curve, EPropertyChangeTy
 // Render any textures
 void UCurveLinearColorAtlas::UpdateTextures()
 {
+	LLM_SCOPE(ELLMTag::Textures);
+
 	// Save off the data needed to render each gradient.
 	// Callback into the section owner to get the Gradients array
 	const int32 TextureDataSize = Source.CalcMipSize(0);
diff --git a/Engine/Source/Runtime/Engine/Private/Materials/Material.cpp b/Engine/Source/Runtime/Engine/Private/Materials/Material.cpp
--- a/Engine/Source/Runtime/Engine/Private/Materials/Material.cpp
+++ b/Engine/Source/Runtime/Engine/Private/Materials/Material.cpp
@@ -4206,6 +4206,7 @@ void UMaterial::PropagateDataToMaterialProxy()
 #if WITH_EDITOR
 void UMaterial::BeginCacheForCookedPlatformData( const ITargetPlatform *TargetPlatform )
 {
+	LLM_SCOPE(ELLMTag::Materials);
 	TArray<FName> DesiredShaderFormats;
 	TargetPlatform->GetAllTargetedShaderFormats(DesiredShaderFormats);
 
@@ -4238,6 +4239,7 @@ void UMaterial::BeginCacheForCookedPlatformData( const ITargetPlatform *TargetPl
 
 bool UMaterial::IsCachedCookedPlatformDataLoaded( const ITargetPlatform* TargetPlatform ) 
 {
+	LLM_SCOPE(ELLMTag::Materials);
 	const TArray<FMaterialResource*>* CachedMaterialResourcesForPlatform = CachedMaterialResourcesForCooking.Find( TargetPlatform );
 
 	if ( CachedMaterialResourcesForPlatform != NULL ) // this should always succeed if begincacheforcookedcplatformdata is called first
diff --git a/Engine/Source/Runtime/Engine/Private/Materials/MaterialInstance.cpp b/Engine/Source/Runtime/Engine/Private/Materials/MaterialInstance.cpp
--- a/Engine/Source/Runtime/Engine/Private/Materials/MaterialInstance.cpp
+++ b/Engine/Source/Runtime/Engine/Private/Materials/MaterialInstance.cpp
@@ -2525,6 +2525,7 @@ void TrimToOverriddenOnly(TArray<ParameterType>& Parameters)
 
 void UMaterialInstance::BeginCacheForCookedPlatformData( const ITargetPlatform *TargetPlatform )
 {
+	LLM_SCOPE(ELLMTag::Materials);
 	TArray<FMaterialResource*> *CachedMaterialResourcesForPlatform = CachedMaterialResourcesForCooking.Find( TargetPlatform );
 
 	if ( CachedMaterialResourcesForPlatform == NULL )
@@ -2555,6 +2556,7 @@ void UMaterialInstance::BeginCacheForCookedPlatformData( const ITargetPlatform *
 
 bool UMaterialInstance::IsCachedCookedPlatformDataLoaded( const ITargetPlatform* TargetPlatform ) 
 {
+	LLM_SCOPE(ELLMTag::Materials);
 	const TArray<FMaterialResource*> *CachedMaterialResourcesForPlatform = CachedMaterialResourcesForCooking.Find( TargetPlatform );
 	if ( CachedMaterialResourcesForPlatform != NULL )
 	{
diff --git a/Engine/Source/Runtime/Engine/Private/PhysicsEngine/BodySetup.cpp b/Engine/Source/Runtime/Engine/Private/PhysicsEngine/BodySetup.cpp
--- a/Engine/Source/Runtime/Engine/Private/PhysicsEngine/BodySetup.cpp
+++ b/Engine/Source/Runtime/Engine/Private/PhysicsEngine/BodySetup.cpp
@@ -7,6 +7,7 @@
 #include "PhysicsEngine/BodySetup.h"
 #include "EngineGlobals.h"
 #include "HAL/IConsoleManager.h"
+#include "HAL/LowLevelMemTracker.h"
 #include "Components/PrimitiveComponent.h"
 #include "Engine/Engine.h"
 #include "Engine/StaticMesh.h"
@@ -997,6 +998,7 @@ void UBodySetup::Serialize(FArchive& Ar)
 
 void UBodySetup::PostLoad()
 {
+	LLM_SCOPE(ELLMTag::Physics);
 	Super::PostLoad();
 
 	// Our owner needs to be post-loaded before us else they may not have loaded
diff --git a/Engine/Source/Runtime/Engine/Private/SkeletalMesh.cpp b/Engine/Source/Runtime/Engine/Private/SkeletalMesh.cpp
--- a/Engine/Source/Runtime/Engine/Private/SkeletalMesh.cpp
+++ b/Engine/Source/Runtime/Engine/Private/SkeletalMesh.cpp
@@ -4412,6 +4412,7 @@ void USkeletalMesh::ValidateBoneWeights(const ITargetPlatform* TargetPlatform)
 
 void USkeletalMesh::BeginCacheForCookedPlatformData(const ITargetPlatform* TargetPlatform)
 {
+	LLM_SCOPE(ELLMTag::SkeletalMesh);
 	// Make sure to cache platform data so it doesn't happen lazily during serialization of the skeletal mesh
 	constexpr bool bIsSerializeSaving = false;
 	GetPlatformSkeletalMeshRenderData(this, TargetPlatform, bIsSerializeSaving);
@@ -4420,6 +4421,7 @@ void USkeletalMesh::BeginCacheForCookedPlatformData(const ITargetPlatform* Targe
 
 void USkeletalMesh::ClearAllCachedCookedPlatformData()
 {
+	LLM_SCOPE(ELLMTag::SkeletalMesh);
 	GetResourceForRendering()->NextCachedRenderData.Reset();
 	
 	if (FApp::CanEverRender())
diff --git a/Engine/Source/Runtime/Engine/Private/StaticMeshCompiler.cpp b/Engine/Source/Runtime/Engine/Private/StaticMeshCompiler.cpp
--- a/Engine/Source/Runtime/Engine/Private/StaticMeshCompiler.cpp
+++ b/Engine/Source/Runtime/Engine/Private/StaticMeshCompiler.cpp
@@ -9,6 +9,7 @@
 #include "Engine/StaticMesh.h"
 #include "ObjectCacheContext.h"
 #include "Settings/EditorExperimentalSettings.h"
+#include "HAL/LowLevelMemTracker.h"
 #include "Misc/QueuedThreadPoolWrapper.h"
 #include "Misc/Optional.h"
 #include "EngineModule.h"
@@ -699,6 +700,7 @@ void FStaticMeshCompilingManager::Reschedule()
 void FStaticMeshCompilingManager::ProcessStaticMeshes(bool bLimitExecutionTime, int32 MinBatchSize)
 {
 	using namespace StaticMeshCompilingManagerImpl;
+	LLM_SCOPE(ELLMTag::StaticMesh);
 	TRACE_CPUPROFILER_EVENT_SCOPE(FStaticMeshCompilingManager::ProcessStaticMeshes);
 	const int32 NumRemainingMeshes = GetNumRemainingMeshes();
 	// Spread out the load over multiple frames but if too many meshes, convergence is more important than frame time
diff --git a/Engine/Source/Runtime/RenderCore/Private/ShaderCodeLibrary.cpp b/Engine/Source/Runtime/RenderCore/Private/ShaderCodeLibrary.cpp
--- a/Engine/Source/Runtime/RenderCore/Private/ShaderCodeLibrary.cpp
+++ b/Engine/Source/Runtime/RenderCore/Private/ShaderCodeLibrary.cpp
@@ -2932,6 +2932,7 @@ public:
 
 	bool AppendFromCompactBinary(FCbFieldView Field)
 	{
+		LLM_SCOPE(ELLMTag::Shaders);
 		bool bOk = true;
 		for (FCbFieldView PlatformField : Field)
 		{
