Removed legacy cook anlaytics feature as remote endpoint is no longer available and this is causing a bunch of automated tests to fail.

diff --git a/Engine/Source/Editor/UnrealEd/Private/Cooker/CookProfiling.cpp b/Engine/Source/Editor/UnrealEd/Private/Cooker/CookProfiling.cpp
--- a/Engine/Source/Editor/UnrealEd/Private/Cooker/CookProfiling.cpp
+++ b/Engine/Source/Editor/UnrealEd/Private/Cooker/CookProfiling.cpp
@@ -718,56 +718,6 @@ void LogCookStats(ECookMode::Type CookMode)
 		FStudioAnalytics::GetProvider().BlockUntilFlushed(60.0f);
 	}
 
-	bool bSendCookAnalytics = false;
-	GConfig->GetBool(TEXT("CookAnalytics"), TEXT("SendAnalytics"), bSendCookAnalytics, GEngineIni);
-
-	if (IsCookByTheBookMode(CookMode) &&
-		(GIsBuildMachine || FParse::Param(FCommandLine::Get(), TEXT("SendCookAnalytics")) || bSendCookAnalytics))
-	{
-		FString APIServerET;
-		if (GConfig->GetString(TEXT("CookAnalytics"), TEXT("APIServer"), APIServerET, GEngineIni))
-		{
-			FString AppId(TEXT("Cook"));
-			bool bUseLegacyCookProtocol = !GConfig->GetString(TEXT("CookAnalytics"), TEXT("AppId"), AppId, GEngineIni);
-
-			// Optionally create an analytics provider to send stats to for central collection.
-			TSharedPtr<IAnalyticsProviderET> CookAnalytics = FAnalyticsET::Get().CreateAnalyticsProvider(FAnalyticsET::Config(AppId, APIServerET, FString(), bUseLegacyCookProtocol));
-			if (CookAnalytics.IsValid())
-			{
-				CookAnalytics->SetUserID(FString::Printf(TEXT("%s\\%s"), FPlatformProcess::ComputerName(), FPlatformProcess::UserName(false)));
-				CookAnalytics->StartSession(MakeAnalyticsEventAttributeArray(
-					TEXT("Project"), CookProject,
-					TEXT("CmdLine"), FString(FCommandLine::Get()),
-					TEXT("IsBuildMachine"), GIsBuildMachine,
-					TEXT("TargetPlatforms"), TargetPlatforms
-				));
-
-				TArray<FString> CookStatsToSend;
-				const bool bFilterStats = GConfig->GetArray(TEXT("CookAnalytics"), TEXT("CookStats"), CookStatsToSend, GEngineIni) > 0;
-				// Sends each cook stat to the analytics provider.
-				auto SendCookStatsToAnalytics = [CookAnalytics, &CookStatsToSend, bFilterStats](const FString& StatName, const TArray<FCookStatsManager::StringKeyValue>& StatAttributes)
-				{
-					if (!bFilterStats || CookStatsToSend.Contains(StatName))
-					{
-						// convert filtered stats directly to an analytics event
-						TArray<FAnalyticsEventAttribute> StatAttrs;
-						StatAttrs.Reset(StatAttributes.Num());
-						for (const auto& Attr : StatAttributes)
-						{
-							StatAttrs.Emplace(Attr.Key, Attr.Value);
-						}
-						CookAnalytics->RecordEvent(StatName, StatAttrs);
-					}
-					else
-					{
-						UE_LOG(LogCook, Verbose, TEXT("[%s] not present in analytics CookStats filter"), *StatName);
-					}
-				};
-				FCookStatsManager::LogCookStats(SendCookStatsToAnalytics);
-			}
-		}
-	}
-
 	/** Used for custom logging of DDC Resource usage stats. */
 	struct FDDCResourceUsageStat
 	{
