Added support for additional errors reported in content workers.

diff --git a/Engine/Source/Editor/UnrealEd/Private/Cooker/CookedSavePackageValidator.cpp b/Engine/Source/Editor/UnrealEd/Private/Cooker/CookedSavePackageValidator.cpp
--- a/Engine/Source/Editor/UnrealEd/Private/Cooker/CookedSavePackageValidator.cpp
+++ b/Engine/Source/Editor/UnrealEd/Private/Cooker/CookedSavePackageValidator.cpp
@@ -62,13 +62,13 @@ ESavePackageResult FCookedSavePackageValidator::ValidateImports(const UPackage*
 			// unavailable class by overriding UObject::NeedsLoadForTargetPlatform on that class to return false.
 			if (bIsError)
 			{
-				UE_LOG(LogCook, Error, TEXT("Failed to cook %s for platform %s. It imports class %s, which is in a module that is not available on the platform."),
+				UE_ASSET_LOG(LogCook, Error, Package, TEXT("Failed to cook %s for platform %s. It imports class %s, which is in a module that is not available on the platform."),
 					*Package->GetName(), *TargetPlatform->PlatformName(), *Class->GetPathName());
 				return ESavePackageResult::ValidatorError;
 			}
 			else
 			{
-				UE_LOG(LogCook, Display, TEXT("Skipping package %s for platform %s. It imports class %s, which is in a module that is not available on the platform."),
+				UE_ASSET_LOG(LogCook, Display, Package, TEXT("Skipping package %s for platform %s. It imports class %s, which is in a module that is not available on the platform."),
 					*Package->GetName(), *TargetPlatform->PlatformName(), *Class->GetPathName());
 				return ESavePackageResult::ValidatorSuppress;
 			}
