Fixed crash in moderation commandlet.

diff --git a/Engine/Source/Editor/UnrealEd/Private/Commandlets/BaseIteratePackagesCommandlet.cpp b/Engine/Source/Editor/UnrealEd/Private/Commandlets/BaseIteratePackagesCommandlet.cpp
--- a/Engine/Source/Editor/UnrealEd/Private/Commandlets/BaseIteratePackagesCommandlet.cpp
+++ b/Engine/Source/Editor/UnrealEd/Private/Commandlets/BaseIteratePackagesCommandlet.cpp
@@ -761,9 +761,13 @@ void UBaseIteratePackagesCommandlet::LoadAndSaveOnePackage(const FString& Filena
 
 				// Check for any special per object operations
 				ForEachObjectWithOuter(Package, [this, &bSavePackage](UObject* Object)
-				{
-					PerformAdditionalOperations(Object, bSavePackage);
-				});
+					{
+						if (!IsValid(Object))
+						{
+							return;
+						}
+						PerformAdditionalOperations(Object, bSavePackage);
+					});
 			}
 
 			PostPerformAdditionalOperations(Package);
