implemented a Static Switch Commandlet that runs fast and can give a detailed HTML report on what shaders with permuations are being used

diff --git a/Engine/Source/Editor/UnrealEd/Private/Commandlets/DumpMaterialShaderTypes.cpp b/Engine/Source/Editor/UnrealEd/Private/Commandlets/DumpMaterialShaderTypes.cpp
--- a/Engine/Source/Editor/UnrealEd/Private/Commandlets/DumpMaterialShaderTypes.cpp
+++ b/Engine/Source/Editor/UnrealEd/Private/Commandlets/DumpMaterialShaderTypes.cpp
@@ -16,6 +16,7 @@
 #include "ICollectionManager.h"
 #include "CollectionManagerModule.h"
 #include "Misc/EngineVersion.h"
+#include "Experimental/Containers/RobinHoodHashTable.h"
 
 DEFINE_LOG_CATEGORY_STATIC(LogDumpMaterialShaderTypesCommandlet, Log, All);
 
@@ -276,7 +277,7 @@ UDumpMaterialShaderTypesCommandlet::UDumpMaterialShaderTypesCommandlet(const FOb
 {
 }
 
-int GetTotalShaders(const TArray<FDebugShaderTypeInfo>& OutShaderInfo)
+static int GetTotalShaders(const TArray<FDebugShaderTypeInfo>& OutShaderInfo)
 {
 	int TotalShadersForMaterial = 0;
 	for (const FDebugShaderTypeInfo& ShaderInfo : OutShaderInfo)
@@ -291,7 +292,7 @@ int GetTotalShaders(const TArray<FDebugShaderTypeInfo>& OutShaderInfo)
 	return TotalShadersForMaterial;
 }
 
-void PrintDebugShaderInfo(FShaderStatsGatheringContext& Output, const TArray<FDebugShaderTypeInfo>& OutShaderInfo)
+static void PrintDebugShaderInfo(FShaderStatsGatheringContext& Output, const TArray<FDebugShaderTypeInfo>& OutShaderInfo)
 {
 	for (const FDebugShaderTypeInfo& ShaderInfo : OutShaderInfo)
 	{
@@ -362,7 +363,7 @@ void PrintDebugShaderInfo(FShaderStatsGatheringContext& Output, const TArray<FDe
 	}
 }
 
-int ProcessMaterials(const ITargetPlatform* TargetPlatform, const EShaderPlatform ShaderPlatform, FShaderStatsGatheringContext& Output, const TArray<FAssetData>& MaterialList)
+static int ProcessMaterials(const ITargetPlatform* TargetPlatform, const EShaderPlatform ShaderPlatform, FShaderStatsGatheringContext& Output, const TArray<FAssetData>& MaterialList)
 {
 	TRACE_CPUPROFILER_EVENT_SCOPE(ProcessMaterials);
 
@@ -393,7 +394,253 @@ int ProcessMaterials(const ITargetPlatform* TargetPlatform, const EShaderPlatfor
 	return TotalShaders;
 }
 
-int ProcessMaterialInstances(const ITargetPlatform* TargetPlatform, const EShaderPlatform ShaderPlatform, FShaderStatsGatheringContext& Output, const TArray<FAssetData>& MaterialInstanceList)
+static void ProcessSwitchOptimizer(const ITargetPlatform* TargetPlatform, const EShaderPlatform ShaderPlatform, const TArray<FAssetData>& MaterialList, const TArray<FAssetData>& MaterialInstanceList)
+{
+	TRACE_CPUPROFILER_EVENT_SCOPE(ProcessSwitchOptimizer);
+
+	const FString TimeNow = FDateTime::Now().ToString();
+	FString FileName = FString::Printf(TEXT("%s-StaticSwitches-%s-%s-%s.html"), FApp::GetProjectName(), *TargetPlatform->PlatformName(), *LexToString(ShaderPlatform), *TimeNow);
+	FShaderStatsGatheringContext Output(FileName);
+
+	Output.Log(TEXT("<!DOCTYPE html>"));
+	Output.Log(TEXT("<html>"));
+	Output.Log(TEXT("<head>"));
+	Output.Log(TEXT("\t<title>StaticSwitchOptimizer</title>"));
+	Output.Log(TEXT("\t<style>"));
+	Output.Log(TEXT("\t\ttable, th, td {border: 1px solid black;}"));
+	Output.Log(TEXT("\t</style>"));
+	Output.Log(TEXT("</head>"));
+	Output.Log(TEXT("<body>"));
+
+	struct FOuterKeyFuncs
+	{
+		static inline bool Matches(const FMaterialShaderMapId& A, const FMaterialShaderMapId& B)
+		{
+			return A.Equals(B, false);
+		}
+
+		static inline uint32 GetKeyHash(const FMaterialShaderMapId& Key)
+		{
+			FSHAHash Hash;
+			Key.GetMaterialHash(Hash, false);
+			return CityHash32((char*)Hash.Hash, 20);
+		}
+	};
+
+	using StaticSwitchArrayType = TArray<FStaticSwitchParameter>;
+	struct FInnerKeyFuncs
+	{
+		static inline bool Matches(const StaticSwitchArrayType& A, const StaticSwitchArrayType& B)
+		{
+			return A == B;
+		}
+
+		static inline uint32 GetKeyHash(const StaticSwitchArrayType& Keys)
+		{
+			FSHA1 Hash;
+			for(const FStaticSwitchParameter& Key : Keys)
+			{
+				Key.UpdateHash(Hash);
+			}
+			Hash.Final();
+			return CityHash32((char*)Hash.m_digest, FSHA1::DigestSize);
+		}
+	};
+
+	FPlatformTypeLayoutParameters LayoutParams;
+	LayoutParams.InitializeForPlatform(TargetPlatform);
+	TArray<FDebugShaderTypeInfo> OutShaderInfo;
+
+	using MaterialAndSizeSet = Experimental::TRobinHoodHashMap<UMaterialInterface*, int32>;
+	using StaticSwitchGroupSet = Experimental::TRobinHoodHashMap<StaticSwitchArrayType, MaterialAndSizeSet, FInnerKeyFuncs>;
+	using ShaderIdGroupSet = Experimental::TRobinHoodHashMap<FMaterialShaderMapId, StaticSwitchGroupSet, FOuterKeyFuncs>;
+	ShaderIdGroupSet ShaderMapHashMap;
+
+	for (const FAssetData& AssetData : MaterialList)
+	{
+		if (UMaterial* Material = Cast<UMaterial>(AssetData.GetAsset()))
+		{
+			TArray<FMaterialResource*> ResourcesToCache;
+			for (int32 QualityLevel = 0; QualityLevel < EMaterialQualityLevel::Num; QualityLevel++)
+			{
+				FMaterialResource* Resource = FindOrCreateMaterialResource(ResourcesToCache, Material, nullptr, GMaxRHIFeatureLevel, EMaterialQualityLevel::Type(QualityLevel));
+				FMaterialShaderMapId ShaderMapId;
+				Resource->GetShaderMapId(ShaderPlatform, TargetPlatform, ShaderMapId);
+				StaticSwitchGroupSet& InnerHashMap = *ShaderMapHashMap.FindOrAdd(ShaderMapId, StaticSwitchGroupSet());
+				MaterialAndSizeSet& InnerValue = *InnerHashMap.FindOrAdd(ShaderMapId.GetStaticSwitchParameters(), MaterialAndSizeSet());
+
+				OutShaderInfo.Reset();
+				Resource->GetShaderTypes(ShaderPlatform, LayoutParams, OutShaderInfo);
+				uint32 TotalNumShaders = GetTotalShaders(OutShaderInfo);
+				if(TotalNumShaders)
+				{
+					InnerValue.FindOrAdd(Material, TotalNumShaders);
+				}
+			}
+			FMaterial::DeferredDeleteArray(ResourcesToCache);
+		}
+	}
+
+	for (const FAssetData& AssetData : MaterialInstanceList)
+	{
+		if (UMaterialInstance* MaterialInstance = Cast<UMaterialInstance>(AssetData.GetAsset()))
+		{
+			TArray<FMaterialResource*> ResourcesToCache;
+			for (int32 QualityLevel = 0; QualityLevel < EMaterialQualityLevel::Num; QualityLevel++)
+			{
+				UMaterial* BaseMaterial = MaterialInstance->GetMaterial();
+				FMaterialResource* Resource = FindOrCreateMaterialResource(ResourcesToCache, BaseMaterial, MaterialInstance, GMaxRHIFeatureLevel, EMaterialQualityLevel::Type(QualityLevel));
+				FMaterialShaderMapId ShaderMapId;
+				Resource->GetShaderMapId(ShaderPlatform, TargetPlatform, ShaderMapId);
+				StaticSwitchGroupSet& InnerHashMap = *ShaderMapHashMap.FindOrAdd(ShaderMapId, StaticSwitchGroupSet());
+				MaterialAndSizeSet& InnerValue = *InnerHashMap.FindOrAdd(ShaderMapId.GetStaticSwitchParameters(), MaterialAndSizeSet());
+
+				OutShaderInfo.Reset();
+				Resource->GetShaderTypes(ShaderPlatform, LayoutParams, OutShaderInfo);
+				uint32 TotalNumShaders = GetTotalShaders(OutShaderInfo);
+				if (TotalNumShaders)
+				{
+					InnerValue.FindOrAdd(MaterialInstance, TotalNumShaders);
+				}
+			}
+			FMaterial::DeferredDeleteArray(ResourcesToCache);
+		}
+	}
+
+	using ShaderIdAndSwitchGroup = TPair<const FMaterialShaderMapId, StaticSwitchGroupSet>;
+	using StaticSwitchGroupType = TPair<const StaticSwitchArrayType, MaterialAndSizeSet>;
+
+	using ShaderIdAndStaticSwitchGroupType = TPair<const FMaterialShaderMapId, StaticSwitchGroupType>;
+	using ShaderIdGroupArray = TArray<ShaderIdAndStaticSwitchGroupType>;
+	using VaryingSwitchesType = Experimental::TRobinHoodHashMap<FName, TArray<bool>>;
+	using InnerFilteredType = TPair<const ShaderIdGroupArray, VaryingSwitchesType>;
+	
+	using FlattenedArrayType = TArray<InnerFilteredType>;
+	FlattenedArrayType FilteredHashMap;
+
+	for(const ShaderIdAndSwitchGroup& OuterElement : ShaderMapHashMap)
+	{
+		if(OuterElement.Value.Num() > 1)
+		{
+			ShaderIdGroupArray InnerMap;
+			VaryingSwitchesType VaryingSwitches;
+			StaticSwitchArrayType First;
+			for(const StaticSwitchGroupType& Value : OuterElement.Value)
+			{
+				if(Value.Key.Num())
+				{
+					First = Value.Key;
+					break;
+				}
+			}
+
+			for (const StaticSwitchGroupType& Inner : OuterElement.Value)
+			{
+				for(int32 i = 0; i < Inner.Key.Num(); i++)
+				{
+					const FStaticSwitchParameter& Key = Inner.Key[i];
+					if(Key.Value != First[i].Value)
+					{
+						VaryingSwitches.FindOrAdd(Key.ParameterInfo.Name, TArray<bool>());
+					}
+				}
+				InnerMap.Emplace(OuterElement.Key, Inner);
+			}
+			InnerMap.Sort([](const ShaderIdAndStaticSwitchGroupType& A, const ShaderIdAndStaticSwitchGroupType& B)
+			{
+				if(A.Value.Key.Num() != B.Value.Key.Num())
+				{
+					return A.Value.Key.Num() < B.Value.Key.Num();
+				}
+
+				for(int32 i = 0; i < A.Value.Key.Num(); i++)
+				{
+					if(A.Value.Key[i].Value != B.Value.Key[i].Value)
+					{
+						return A.Value.Key[i].Value < B.Value.Key[i].Value;
+					}
+				}
+				return false;
+			});
+
+			if(VaryingSwitches.Num())
+			{
+				FilteredHashMap.Emplace(MoveTemp(InnerMap), MoveTemp(VaryingSwitches));
+			}
+		}
+	}
+
+	FilteredHashMap.Sort([](const InnerFilteredType& A, const InnerFilteredType& B)
+	{
+		int32 NumA = 0;
+		for (const ShaderIdAndStaticSwitchGroupType& InnerA : A.Key)
+		{
+			NumA += (*InnerA.Value.Value.begin()).Value;
+		}
+
+		int32 NumB = 0;
+		for (const ShaderIdAndStaticSwitchGroupType& InnerB : B.Key)
+		{
+			NumB += (*InnerB.Value.Value.begin()).Value;
+		}
+		return NumA > NumB;
+	});
+
+	for(int32 i = 0; i < FilteredHashMap.Num(); i++)
+	{
+		const ShaderIdGroupArray& InnerMap = FilteredHashMap[i].Key;
+		VaryingSwitchesType& Varying = FilteredHashMap[i].Value;
+		const UMaterialInterface* Parent = (*(*InnerMap.begin()).Value.Value.begin()).Key->GetMaterial();
+		const FMaterialShaderMapId& ShaderId = (*InnerMap.begin()).Key;
+
+		int32 Num = 0;
+		for (const ShaderIdAndStaticSwitchGroupType& Inner : InnerMap)
+		{
+			Num += (*Inner.Value.Value.begin()).Value;
+		}
+
+		Output.Log(TEXT("<table>"));
+		Output.Log(FString::Printf(TEXT("<tr><th><h3>Candidate %s with NumShaders: %d</h3></th></tr>"), *Parent->GetOuter()->GetFName().ToString(), Num));
+		Output.Log(TEXT("<tr><td><table>"));
+
+		for (const ShaderIdAndStaticSwitchGroupType& Inner : InnerMap)
+		{
+			const StaticSwitchGroupType& Groups = Inner.Value;
+
+			if(Groups.Key.Num())
+			{
+				for(const FStaticSwitchParameter& Param : Groups.Key)
+				{
+					if(TArray<bool>* ValueArray = Varying.Find(Param.ParameterInfo.Name))
+					{
+						ValueArray->Add(Param.Value);
+					}
+				}
+			}
+		}
+
+		for (const TPair<const FName, TArray<bool>>& Param : Varying)
+		{
+			Output.Log(TEXT("\t<tr>"));
+			Output.Log(FString::Printf(TEXT("\t\t<th>%s</th>"), *Param.Key.ToString()));
+
+			for(bool Switch : Param.Value)
+			{
+				Output.Log(Switch ? TEXT("\t\t<td style=\"background-color:red;\">1</td>") : TEXT("\t\t<td style=\"background-color:green;\">0</td>"));
+			}
+			Output.Log(TEXT("\t</tr>"));
+		}
+
+		Output.Log(TEXT("</table></td></tr>"));
+		Output.Log(TEXT("</table>"));
+		Output.Log(TEXT("<br></br>"));
+	}
+
+	Output.Log(TEXT("</body>"));
+	Output.Log(TEXT("</html>"));
+}
+
+static int ProcessMaterialInstances(const ITargetPlatform* TargetPlatform, const EShaderPlatform ShaderPlatform, FShaderStatsGatheringContext& Output, const TArray<FAssetData>& MaterialInstanceList)
 {
 	TRACE_CPUPROFILER_EVENT_SCOPE(ProcessMaterialInstances);
 
@@ -480,7 +727,7 @@ int ProcessMaterialInstances(const ITargetPlatform* TargetPlatform, const EShade
 	return TotalShaders;
 }
 
-int ProcessGlobalShaders(const ITargetPlatform* TargetPlatform, const EShaderPlatform ShaderPlatform, FShaderStatsGatheringContext& Output)
+static int ProcessGlobalShaders(const ITargetPlatform* TargetPlatform, const EShaderPlatform ShaderPlatform, FShaderStatsGatheringContext& Output)
 {
 	TRACE_CPUPROFILER_EVENT_SCOPE(ProcessGlobalShaders);
 
@@ -532,7 +779,7 @@ int ProcessGlobalShaders(const ITargetPlatform* TargetPlatform, const EShaderPla
 	return TotalShaders;
 }
 
-void ProcessForTargetAndShaderPlatform(const ITargetPlatform* TargetPlatform, const EShaderPlatform ShaderPlatform, const FString& Params, const TArray<FAssetData>& MaterialList, const TArray<FAssetData>& MaterialInstanceList, TSharedPtr<IAnalyticsProviderET> Provider)
+static void ProcessForTargetAndShaderPlatform(const ITargetPlatform* TargetPlatform, const EShaderPlatform ShaderPlatform, const FString& Params, const TArray<FAssetData>& MaterialList, const TArray<FAssetData>& MaterialInstanceList, TSharedPtr<IAnalyticsProviderET> Provider)
 {
 	TRACE_CPUPROFILER_EVENT_SCOPE(ProcessForTargetAndShaderPlatform);
 
@@ -610,9 +857,11 @@ int32 UDumpMaterialShaderTypesCommandlet::Main(const FString& Params)
 		UE_LOG(LogDumpMaterialShaderTypesCommandlet, Log, TEXT(" Required: -targetplatform=<platform(s)>     (Which target platform do you want results, e.g. WindowsClient, WindowsEditor. Multiple shader platforms are allowed)."));
 		UE_LOG(LogDumpMaterialShaderTypesCommandlet, Log, TEXT(" Optional: -collection=<name>                (You can also specify a collection of assets to narrow down the results e.g. if you maintain a collection that represents the actually used in-game assets)."));
 		UE_LOG(LogDumpMaterialShaderTypesCommandlet, Log, TEXT(" Optional: -analytics                        (Whether or not to send analytics data for tracking purposes)."));
+		UE_LOG(LogDumpMaterialShaderTypesCommandlet, Log, TEXT(" Optional: -staticswitches					 (Gain more detailed information of StaticSwitch use and cost)."));
 		return 0;
 	}
 
+	const bool bStaticSwitches = FParse::Param(FCommandLine::Get(), TEXT("staticswitches"));
 	const bool bSendAnalytics = FParse::Param(FCommandLine::Get(), TEXT("analytics"));
 
 	const double AssetRegistryStart = FPlatformTime::Seconds();
@@ -705,7 +954,14 @@ PRAGMA_ENABLE_DEPRECATION_WARNINGS
 			const EShaderPlatform ShaderPlatform = ShaderFormatToLegacyShaderPlatform(DesiredShaderFormats[FormatIndex]);
 
 			UE_LOG(LogDumpMaterialShaderTypesCommandlet, Display, TEXT("Dumping material shader types for '%s' - '%s'..."), *Platforms[Index]->PlatformName(), *LexToString(ShaderPlatform));
-			ProcessForTargetAndShaderPlatform(Platforms[Index], ShaderPlatform, Params, MaterialList, MaterialInstanceList, Provider);
+			if(bStaticSwitches)
+			{
+				ProcessSwitchOptimizer(Platforms[Index], ShaderPlatform, MaterialList, MaterialInstanceList);
+			}
+			else
+			{
+				ProcessForTargetAndShaderPlatform(Platforms[Index], ShaderPlatform, Params, MaterialList, MaterialInstanceList, Provider);
+			}
 		}
 	}
 
diff --git a/Engine/Source/Runtime/Engine/Classes/Materials/MaterialInterface.h b/Engine/Source/Runtime/Engine/Classes/Materials/MaterialInterface.h
--- a/Engine/Source/Runtime/Engine/Classes/Materials/MaterialInterface.h
+++ b/Engine/Source/Runtime/Engine/Classes/Materials/MaterialInterface.h
@@ -1112,7 +1112,7 @@ extern void ProcessSerializedInlineShaderMaps(UMaterialInterface* Owner, TArray<
 extern FMaterialResource* FindMaterialResource(const TArray<FMaterialResource*>& MaterialResources, ERHIFeatureLevel::Type InFeatureLevel, EMaterialQualityLevel::Type QualityLevel, bool bAllowDefaultQuality);
 extern FMaterialResource* FindMaterialResource(TArray<FMaterialResource*>& MaterialResources, ERHIFeatureLevel::Type InFeatureLevel, EMaterialQualityLevel::Type QualityLevel, bool bAllowDefaultQuality);
 
-extern FMaterialResource* FindOrCreateMaterialResource(TArray<FMaterialResource*>& MaterialResources,
+ENGINE_API FMaterialResource* FindOrCreateMaterialResource(TArray<FMaterialResource*>& MaterialResources,
 	UMaterial* OwnerMaterial,
 	UMaterialInstance* OwnerMaterialInstance,
 	ERHIFeatureLevel::Type InFeatureLevel,
diff --git a/Engine/Source/Runtime/Engine/Private/Materials/MaterialShader.cpp b/Engine/Source/Runtime/Engine/Private/Materials/MaterialShader.cpp
--- a/Engine/Source/Runtime/Engine/Private/Materials/MaterialShader.cpp
+++ b/Engine/Source/Runtime/Engine/Private/Materials/MaterialShader.cpp
@@ -725,7 +725,7 @@ void FMaterialShaderMapId::Serialize(FArchive& Ar, bool bLoadedByCookedMaterial)
 
 #if WITH_EDITOR
 /** Hashes the material-specific part of this shader map Id. */
-void FMaterialShaderMapId::GetMaterialHash(FSHAHash& OutHash) const
+void FMaterialShaderMapId::GetMaterialHash(FSHAHash& OutHash, bool bWithStaticParameters) const
 {
 	check(IsContentValid());
 	FSHA1 HashState;
@@ -746,9 +746,12 @@ void FMaterialShaderMapId::GetMaterialHash(FSHAHash& OutHash) const
 
 
 	//Hash the static parameters
-	for (const FStaticSwitchParameter& StaticSwitchParameter : StaticSwitchParameters)
+	if (bWithStaticParameters)
 	{
-		StaticSwitchParameter.UpdateHash(HashState);
+		for (const FStaticSwitchParameter& StaticSwitchParameter : StaticSwitchParameters)
+		{
+			StaticSwitchParameter.UpdateHash(HashState);
+		}
 	}
 	for (const FStaticComponentMaskParameter& StaticComponentMaskParameter : StaticComponentMaskParameters)
 	{
@@ -795,7 +798,7 @@ void FMaterialShaderMapId::GetMaterialHash(FSHAHash& OutHash) const
 * @param ReferenceSet	The set to compare against
 * @return				true if the sets are equal
 */
-bool FMaterialShaderMapId::operator==(const FMaterialShaderMapId& ReferenceSet) const
+bool FMaterialShaderMapId::Equals(const FMaterialShaderMapId& ReferenceSet, bool bWithStaticParameters) const
 {
 	// Ensure data is in valid state for comparison
 	check(IsContentValid() && ReferenceSet.IsContentValid());
@@ -849,7 +852,7 @@ bool FMaterialShaderMapId::operator==(const FMaterialShaderMapId& ReferenceSet)
 #if WITH_EDITOR
 	if (!IsCookedId())
 	{
-		if (StaticSwitchParameters.Num() != ReferenceSet.StaticSwitchParameters.Num()
+		if ((bWithStaticParameters && StaticSwitchParameters.Num() != ReferenceSet.StaticSwitchParameters.Num())
 			|| StaticComponentMaskParameters.Num() != ReferenceSet.StaticComponentMaskParameters.Num()
 			|| TerrainLayerWeightParameters.Num() != ReferenceSet.TerrainLayerWeightParameters.Num()
 			|| ReferencedFunctions.Num() != ReferenceSet.ReferencedFunctions.Num()
@@ -861,7 +864,7 @@ bool FMaterialShaderMapId::operator==(const FMaterialShaderMapId& ReferenceSet)
 			return false;
 		}
 
-		if (StaticSwitchParameters != ReferenceSet.StaticSwitchParameters
+		if ((bWithStaticParameters && StaticSwitchParameters != ReferenceSet.StaticSwitchParameters)
 			|| StaticComponentMaskParameters != ReferenceSet.StaticComponentMaskParameters
 			|| TerrainLayerWeightParameters != ReferenceSet.TerrainLayerWeightParameters
 			|| MaterialLayersId != ReferenceSet.MaterialLayersId)
diff --git a/Engine/Source/Runtime/Engine/Public/MaterialShared.h b/Engine/Source/Runtime/Engine/Public/MaterialShared.h
--- a/Engine/Source/Runtime/Engine/Public/MaterialShared.h
+++ b/Engine/Source/Runtime/Engine/Public/MaterialShared.h
@@ -960,20 +960,30 @@ public:
 
 #if WITH_EDITOR
 	/** Hashes the material-specific part of this shader map Id. */
-	ENGINE_API void GetMaterialHash(FSHAHash& OutHash) const;
+	ENGINE_API void GetMaterialHash(FSHAHash& OutHash, bool bWithStaticParameters) const;
+
+	ENGINE_API void GetMaterialHash(FSHAHash& OutHash) const
+	{
+		GetMaterialHash(OutHash, true);
+	}
 #endif
 
+	ENGINE_API bool Equals(const FMaterialShaderMapId& ReferenceSet, bool bWithStaticParameters) const;
+
 	/** 
 	* Tests this set against another for equality
 	* 
 	* @param ReferenceSet	The set to compare against
 	* @return				true if the sets are equal
 	*/
-	bool operator==(const FMaterialShaderMapId& ReferenceSet) const;
+	bool operator==(const FMaterialShaderMapId& ReferenceSet) const
+	{
+		return Equals(ReferenceSet, true);
+	}
 
 	bool operator!=(const FMaterialShaderMapId& ReferenceSet) const
 	{
-		return !(*this == ReferenceSet);
+		return !Equals(ReferenceSet, true);
 	}
 
 	/** Ensure content is valid - for example overrides are set deterministically for serialization and sorting */
