Cooker: Add timer log for PollPendingCookedPlatformDatas.
#rnx

diff --git a/Engine/Source/Editor/UnrealEd/Private/CookOnTheFlyServer.cpp b/Engine/Source/Editor/UnrealEd/Private/CookOnTheFlyServer.cpp
--- a/Engine/Source/Editor/UnrealEd/Private/CookOnTheFlyServer.cpp
+++ b/Engine/Source/Editor/UnrealEd/Private/CookOnTheFlyServer.cpp
@@ -4131,7 +4131,10 @@ void UCookOnTheFlyServer::PumpSaves(UE::Cook::FTickStackData& StackData, uint32
 
 		// Release any completed pending CookedPlatformDatas, so that slots in the per-class limits on calls to BeginCacheForCookedPlatformData are freed up for new objects to use
 		bool bForce = IsCookOnTheFlyMode() && !IsRealtimeMode();
-		PackageDatas->PollPendingCookedPlatformDatas(bForce, LastCookableObjectTickTime);
+		{
+			UE_SCOPED_HIERARCHICAL_COOKTIMER(PollPendingCookedPlatformDatas);
+			PackageDatas->PollPendingCookedPlatformDatas(bForce, LastCookableObjectTickTime);
+		}
 
 		// If BeginCacheCookPlatformData is not ready then postpone the package, exit, or wait for it as appropriate
 		EPollStatus PrepareSaveStatus = PrepareSave(PackageData, StackData.Timer, false /* bPrecaching */);
