MPCook: Change 1-length TArrayView argument from TInlineAllocator to more standard address of variable.
#rnx

diff --git a/Engine/Source/Editor/UnrealEd/Private/Cooker/CookWorkerServer.cpp b/Engine/Source/Editor/UnrealEd/Private/Cooker/CookWorkerServer.cpp
--- a/Engine/Source/Editor/UnrealEd/Private/Cooker/CookWorkerServer.cpp
+++ b/Engine/Source/Editor/UnrealEd/Private/Cooker/CookWorkerServer.cpp
@@ -100,9 +100,8 @@ void FCookWorkerServer::AbortAllAssignmentsInLock(TSet<FPackageData*>& OutPendin
 void FCookWorkerServer::AbortAssignment(FPackageData& PackageData, ECookDirectorThread TickThread,
 	ENotifyRemote NotifyRemote)
 {
-	TArray<FPackageData*, TInlineAllocator<1>> PackageDatas;
-	PackageDatas.Add(&PackageData);
-	AbortAssignments(PackageDatas, TickThread, NotifyRemote);
+	FPackageData* PackageDataPtr = &PackageData;
+	AbortAssignments(TConstArrayView<FPackageData*>(&PackageDataPtr, 1), TickThread, NotifyRemote);
 }
 
 void FCookWorkerServer::AbortAssignments(TConstArrayView<FPackageData*> PackageDatas, ECookDirectorThread TickThread,
