UEditorLevelUtils::CopyActorsToLevel() wasn't fowarding the OutActors array to the internal call to UEditorLevelUtils::CopyOrMoveActorsToLevel()

diff --git a/Engine/Source/Editor/UnrealEd/Private/EditorLevelUtils.cpp b/Engine/Source/Editor/UnrealEd/Private/EditorLevelUtils.cpp
--- a/Engine/Source/Editor/UnrealEd/Private/EditorLevelUtils.cpp
+++ b/Engine/Source/Editor/UnrealEd/Private/EditorLevelUtils.cpp
@@ -85,7 +85,7 @@ int32 UEditorLevelUtils::MoveActorsToLevel(const TArray<AActor*>& ActorsToMove,
 int32 UEditorLevelUtils::CopyActorsToLevel(const TArray<AActor*>& ActorsToMove, ULevel* DestLevel, bool bWarnAboutReferences, bool bWarnAboutRenaming, bool bMoveAllOrFail, TArray<AActor*>* OutActors /*=nullptr*/)
 {
 	const bool bMoveActors = false;
-	return CopyOrMoveActorsToLevel(ActorsToMove, DestLevel, bMoveActors, bWarnAboutReferences, bWarnAboutRenaming, bMoveAllOrFail);
+	return CopyOrMoveActorsToLevel(ActorsToMove, DestLevel, bMoveActors, bWarnAboutReferences, bWarnAboutRenaming, bMoveAllOrFail, OutActors);
 }
 
 int32 UEditorLevelUtils::CopyOrMoveActorsToLevel(const TArray<AActor*>& ActorsToMove, ULevel* DestLevel, bool bMoveActors, bool bWarnAboutReferences /*= true*/, bool bWarnAboutRenaming /*= true*/, bool bMoveAllOrFail /*= false*/, TArray<AActor*>* OutActors /*=nullptr*/)
