MPCook: Reduce which packages can be retracted to avoid cancel during save and cancel of generated packages.
#rnx

diff --git a/Engine/Source/Editor/UnrealEd/Private/CookOnTheFlyServer.cpp b/Engine/Source/Editor/UnrealEd/Private/CookOnTheFlyServer.cpp
--- a/Engine/Source/Editor/UnrealEd/Private/CookOnTheFlyServer.cpp
+++ b/Engine/Source/Editor/UnrealEd/Private/CookOnTheFlyServer.cpp
@@ -9605,7 +9605,7 @@ void UCookOnTheFlyServer::GetPackagesToRetract(int32 NumToRetract, TArray<FName>
 			return true;
 		}
 
-		if (PackageData->GetWorkerAssignmentConstraint().IsValid())
+		if (PackageData->GetWorkerAssignmentConstraint().IsValid() || PackageData->IsGenerated())
 		{
 			// Don't send back Packages that are constrained to this worker. Doing so will just
 			// cause the CookDirector to send it back to us, and this can cause the cooker to crash
@@ -9658,11 +9658,25 @@ void UCookOnTheFlyServer::GetPackagesToRetract(int32 NumToRetract, TArray<FName>
 			return;
 		}
 	}
+	// Send back all packages that have not started saving before sending back any that have
 	for (FPackageData* PackageData : PackageDatas->GetSaveQueue())
 	{
-		if (AddPackageIfPossibleAndReportDone(PackageData))
+		if (!PackageData->GetCookedPlatformDataStarted())
 		{
-			return;
+			if (AddPackageIfPossibleAndReportDone(PackageData))
+			{
+				return;
+			}
+		}
+	}
+	for (FPackageData* PackageData : PackageDatas->GetSaveQueue())
+	{
+		if (PackageData->GetCookedPlatformDataStarted())
+		{
+			if (AddPackageIfPossibleAndReportDone(PackageData))
+			{
+				return;
+			}
 		}
 	}
 }
