MPCook: CommunicationThread: Fix backwards scale from seconds to milliseconds and prevent casting negative number to uint32.
#rnx

diff --git a/Engine/Source/Editor/UnrealEd/Private/Cooker/CookDirector.cpp b/Engine/Source/Editor/UnrealEd/Private/Cooker/CookDirector.cpp
--- a/Engine/Source/Editor/UnrealEd/Private/Cooker/CookDirector.cpp
+++ b/Engine/Source/Editor/UnrealEd/Private/Cooker/CookDirector.cpp
@@ -240,10 +240,13 @@ uint32 FCookDirector::RunCommunicationThread()
 
 		double CurrentTime = FPlatformTime::Seconds();
 		float RemainingDuration = StartTime + TickPeriod - CurrentTime;
-		uint32 WaitTimeMilliseconds = static_cast<uint32>(RemainingDuration * .001f);
-		if (ShutdownEvent->Wait(WaitTimeMilliseconds))
+		if (RemainingDuration > .001f)
 		{
-			break;
+			uint32 WaitTimeMilliseconds = static_cast<uint32>(RemainingDuration * 1000);
+			if (ShutdownEvent->Wait(WaitTimeMilliseconds))
+			{
+				break;
+			}
 		}
 	}
 	return 0;
