Remove test against pass through vertex factory type.
That vertex factory type no longer exists since it was merged with the local vertex factory.

diff --git a/Engine/Source/Editor/UnrealEd/Private/PreviewMaterial.cpp b/Engine/Source/Editor/UnrealEd/Private/PreviewMaterial.cpp
--- a/Engine/Source/Editor/UnrealEd/Private/PreviewMaterial.cpp
+++ b/Engine/Source/Editor/UnrealEd/Private/PreviewMaterial.cpp
@@ -99,13 +99,9 @@ public:
 					return false;
 				}
 
-				extern ENGINE_API bool IsGPUSkinCacheAvailable(EShaderPlatform Platform);
-				bool bSkinCache = IsGPUSkinCacheAvailable(Platform) && (VertexFactoryType == FindVertexFactoryType(FName(TEXT("FGPUSkinPassthroughVertexFactory"), FNAME_Find)));
-					
 				if (
 					VertexFactoryType != FindVertexFactoryType(FName(TEXT("TGPUSkinVertexFactoryDefault"), FNAME_Find)) &&
-					VertexFactoryType != FindVertexFactoryType(FName(TEXT("TGPUSkinVertexFactoryUnlimited"), FNAME_Find)) &&
-					!bSkinCache
+					VertexFactoryType != FindVertexFactoryType(FName(TEXT("TGPUSkinVertexFactoryUnlimited"), FNAME_Find))
 					)
 				{
 					return false;
