Fix an issue where directories failed to rename because assets inside of them had not been loaded

PR #9890

diff --git a/Engine/Source/Editor/UnrealEd/Private/Subsystems/EditorAssetSubsystem.cpp b/Engine/Source/Editor/UnrealEd/Private/Subsystems/EditorAssetSubsystem.cpp
--- a/Engine/Source/Editor/UnrealEd/Private/Subsystems/EditorAssetSubsystem.cpp
+++ b/Engine/Source/Editor/UnrealEd/Private/Subsystems/EditorAssetSubsystem.cpp
@@ -12,6 +12,7 @@
 #include "FileHelpers.h"
 #include "HAL/FileManager.h"
 #include "Misc/PackageName.h"
+#include "Misc/ScopedSlowTask.h"
 #include "Modules/ModuleManager.h"
 #include "ObjectTools.h"
 #include "PackageTools.h"
@@ -342,7 +343,7 @@ PRAGMA_ENABLE_DEPRECATION_WARNINGS
 			}
 
 			// Keep AssetPath /Game/MyFolder
-			AssetPaths.DestinationDirectoryAssetPaths.Add(FPackageName::GetLongPackagePath(NewAssetPackageName));
+			AssetPaths.DestinationDirectoryAssetPaths.Add(NewAssetPackageName);
 		}
 		return MakeValue(AssetPaths);
 	}
@@ -934,7 +935,7 @@ bool UEditorAssetSubsystem::DuplicateDirectory(const FString& SourceDirectoryPat
 	FAssetToolsModule& Module = FModuleManager::LoadModuleChecked<FAssetToolsModule>("AssetTools");
 	for (int32 Index = 0; Index < AssetPaths.SourceDirectoryAssetDatas.Num(); ++Index)
 	{
-		UObject* DuplicatedAsset = Module.Get().DuplicateAsset(AssetPaths.SourceDirectoryAssetDatas[Index].AssetName.ToString(), AssetPaths.DestinationDirectoryAssetPaths[Index], AssetPaths.SourceDirectoryAssetDatas[Index].GetAsset());
+		UObject* DuplicatedAsset = Module.Get().DuplicateAsset(AssetPaths.SourceDirectoryAssetDatas[Index].AssetName.ToString(), FPackageName::GetLongPackagePath(AssetPaths.DestinationDirectoryAssetPaths[Index]), AssetPaths.SourceDirectoryAssetDatas[Index].GetAsset());
 		if (DuplicatedAsset == nullptr)
 		{
 			UE_LOG(LogEditorAssetSubsystem, Warning, TEXT("DuplicateDirectory failed: Could not duplicate object '%s'"), *AssetPaths.SourceDirectoryAssetDatas[Index].GetObjectPathString());
@@ -1055,20 +1056,34 @@ bool UEditorAssetSubsystem::RenameDirectory(const FString& SourceDirectoryPath,
 		UE_LOG(LogEditorAssetSubsystem, Log, TEXT("RenameDirectory. No assets to rename."));
 	}
 	
-	TArray<FAssetRenameData> AssetsToRename;
-	AssetsToRename.Reserve(AssetPathsResult.GetValue().SourceDirectoryAssetDatas.Num());
-
-	for (int32 Index = 0; Index < AssetPathsResult.GetValue().SourceDirectoryAssetDatas.Num(); ++Index)
+	const int32 NumAssetsToRename = AssetPathsResult.GetValue().SourceDirectoryAssetDatas.Num();
+	if (NumAssetsToRename > 0)
 	{
-		AssetsToRename.Add(FAssetRenameData(AssetPathsResult.GetValue().SourceDirectoryAssetDatas[Index].GetObjectPathString(), AssetPathsResult.GetValue().DestinationDirectoryAssetPaths[Index]));
-	}
+		TArray<FAssetRenameData> AssetsToRename;
+		AssetsToRename.Reserve(NumAssetsToRename);
 
-	// Rename the assets
-	FAssetToolsModule& Module = FModuleManager::GetModuleChecked<FAssetToolsModule>("AssetTools");
-	if (!Module.Get().RenameAssets(AssetsToRename))
-	{
-		UE_LOG(LogEditorAssetSubsystem, Error, TEXT("RenameDirectory failed: Could not rename the assets."));
-		return false;
+		FScopedSlowTask SlowTask((float)NumAssetsToRename, NSLOCTEXT("EditorUtilities", "RenameDirectorySlowTaskLabel", "Renaming directory and assets."));
+		for (int32 Index = 0; Index < AssetPathsResult.GetValue().SourceDirectoryAssetDatas.Num(); ++Index)
+		{
+			// Try to load the asset, since if it has not previously been loaded, the rename process will fail to resolve the object in AssetRenameManager
+			TValueOrError<UObject*, FString> LoadResult = UE::EditorAssetUtils::LoadAssetFromData(AssetPathsResult.GetValue().SourceDirectoryAssetDatas[Index]);
+			if (LoadResult.HasError())
+			{
+				UE_LOG(LogEditorAssetSubsystem, Error, TEXT("RenameDirectory failed: Some assets couldn't be renamed. %s"), *LoadResult.GetError());
+				return false;
+			}
+
+			AssetsToRename.Add(FAssetRenameData(AssetPathsResult.GetValue().SourceDirectoryAssetDatas[Index].GetSoftObjectPath(), FSoftObjectPath(AssetPathsResult.GetValue().DestinationDirectoryAssetPaths[Index])));
+			SlowTask.EnterProgressFrame(1.f);
+		}
+
+		// Rename the assets
+		FAssetToolsModule& Module = FModuleManager::GetModuleChecked<FAssetToolsModule>("AssetTools");
+		if (!Module.Get().RenameAssets(AssetsToRename))
+		{
+			UE_LOG(LogEditorAssetSubsystem, Error, TEXT("RenameDirectory failed: Could not rename the assets."));
+			return false;
+		}
 	}
 
 	// Delete the old directory
