Remove experimental level editor  and remove the corresponding setting

#rnx

diff --git a/Engine/Source/Editor/LevelAssetEditor/LevelAssetEditor.Build.cs b/Engine/Source/Editor/LevelAssetEditor/LevelAssetEditor.Build.cs
deleted file mode 100644
diff --git a/Engine/Source/Editor/LevelAssetEditor/Private/LevelAssetEditor.cpp b/Engine/Source/Editor/LevelAssetEditor/Private/LevelAssetEditor.cpp
deleted file mode 100644
diff --git a/Engine/Source/Editor/LevelAssetEditor/Private/LevelAssetEditorCommands.cpp b/Engine/Source/Editor/LevelAssetEditor/Private/LevelAssetEditorCommands.cpp
deleted file mode 100644
diff --git a/Engine/Source/Editor/LevelAssetEditor/Private/LevelAssetEditorViewport.cpp b/Engine/Source/Editor/LevelAssetEditor/Private/LevelAssetEditorViewport.cpp
deleted file mode 100644
diff --git a/Engine/Source/Editor/LevelAssetEditor/Private/LevelAssetEditorViewport.h b/Engine/Source/Editor/LevelAssetEditor/Private/LevelAssetEditorViewport.h
deleted file mode 100644
diff --git a/Engine/Source/Editor/LevelAssetEditor/Private/LevelAssetEditorViewportClient.cpp b/Engine/Source/Editor/LevelAssetEditor/Private/LevelAssetEditorViewportClient.cpp
deleted file mode 100644
diff --git a/Engine/Source/Editor/LevelAssetEditor/Private/LevelAssetEditorViewportClient.h b/Engine/Source/Editor/LevelAssetEditor/Private/LevelAssetEditorViewportClient.h
deleted file mode 100644
diff --git a/Engine/Source/Editor/LevelAssetEditor/Private/LevelEditorAssetToolkit.cpp b/Engine/Source/Editor/LevelAssetEditor/Private/LevelEditorAssetToolkit.cpp
deleted file mode 100644
diff --git a/Engine/Source/Editor/LevelAssetEditor/Private/LevelEditorViewportInterface.cpp b/Engine/Source/Editor/LevelAssetEditor/Private/LevelEditorViewportInterface.cpp
deleted file mode 100644
diff --git a/Engine/Source/Editor/LevelAssetEditor/Private/LevelEditorViewportInterface.h b/Engine/Source/Editor/LevelAssetEditor/Private/LevelEditorViewportInterface.h
deleted file mode 100644
diff --git a/Engine/Source/Editor/LevelAssetEditor/Private/ULevelAssetEditor.cpp b/Engine/Source/Editor/LevelAssetEditor/Private/ULevelAssetEditor.cpp
deleted file mode 100644
diff --git a/Engine/Source/Editor/LevelAssetEditor/Public/LevelAssetEditor.h b/Engine/Source/Editor/LevelAssetEditor/Public/LevelAssetEditor.h
deleted file mode 100644
diff --git a/Engine/Source/Editor/LevelAssetEditor/Public/LevelAssetEditorCommands.h b/Engine/Source/Editor/LevelAssetEditor/Public/LevelAssetEditorCommands.h
deleted file mode 100644
diff --git a/Engine/Source/Editor/LevelAssetEditor/Public/LevelAssetEditorToolkit.h b/Engine/Source/Editor/LevelAssetEditor/Public/LevelAssetEditorToolkit.h
deleted file mode 100644
diff --git a/Engine/Source/Editor/LevelAssetEditor/Public/ULevelAssetEditor.h b/Engine/Source/Editor/LevelAssetEditor/Public/ULevelAssetEditor.h
deleted file mode 100644
diff --git a/Engine/Source/Editor/LevelEditor/LevelEditor.Build.cs b/Engine/Source/Editor/LevelEditor/LevelEditor.Build.cs
--- a/Engine/Source/Editor/LevelEditor/LevelEditor.Build.cs
+++ b/Engine/Source/Editor/LevelEditor/LevelEditor.Build.cs
@@ -119,7 +119,6 @@ public class LevelEditor : ModuleRules
 				"HeadMountedDisplay",
 				"VREditor",
                 "Persona",
-				"LevelAssetEditor",
 				"MergeActors"
 			}
 		);
diff --git a/Engine/Source/Editor/LevelEditor/Private/LevelEditor.cpp b/Engine/Source/Editor/LevelEditor/Private/LevelEditor.cpp
--- a/Engine/Source/Editor/LevelEditor/Private/LevelEditor.cpp
+++ b/Engine/Source/Editor/LevelEditor/Private/LevelEditor.cpp
@@ -259,7 +259,6 @@ void FLevelEditorModule::StartupModule()
 	FModuleManager::LoadModuleChecked<IMainFrameModule>(MainFrame);
 
 	FModuleManager::LoadModuleChecked<FCommonMenuExtensionsModule>(CommonMenuExtensionsName);
-	FModuleManager::Get().LoadModule(TEXT("LevelAssetEditor"));
 
 	MenuExtensibilityManager = MakeShared<FExtensibilityManager>();
 	ToolBarExtensibilityManager = MakeShared<FExtensibilityManager>();
diff --git a/Engine/Source/Editor/UnrealEd/Classes/Settings/LevelEditorMiscSettings.h b/Engine/Source/Editor/UnrealEd/Classes/Settings/LevelEditorMiscSettings.h
--- a/Engine/Source/Editor/UnrealEd/Classes/Settings/LevelEditorMiscSettings.h
+++ b/Engine/Source/Editor/UnrealEd/Classes/Settings/LevelEditorMiscSettings.h
@@ -43,10 +43,6 @@ public:
 	UPROPERTY(EditAnywhere, config, Category=Editing, meta=( DisplayName = "Preserve Actor Scale on Replace" ))
 	uint32 bReplaceRespectsScale:1;
 
-	/** If enabled, the experimental level editor will be available */
-	UPROPERTY(EditAnywhere, config, AdvancedDisplay, Category = Editing)
-	bool bEnableExperimentalLevelEditor = false;
-
 	/** If enabled, will avoid relabeling actors in UUnrealEdEngine::edactPasteSelected */
 	UPROPERTY(EditAnywhere, config, Category = Editing, meta = (DisplayName = "Avoid Actor Relabel on Paste Selected"))
 	uint32 bAvoidRelabelOnPasteSelected:1;
