Avoid clearing the undo buffer when saving newly added/deleted OFPA actors

A newly created OFPA actor would trigger PostLoadedStateChanged to clear the undo buffer due to it creating a new actor desc. This now skips clearing the undo buffer unless an actor has been unloaded.

A deleted OFPA actor would trigger CleanupAfterSuccessfulDelete to clear the undo buffer so that the package can be GC'd. This makes sense for assets (where a delete isn't transacted), but not for actors (where the delete can be undone). This now skips clearing the undo buffer when cleaning-up OFPA packages, and also marks any lingering packages as being PKG_NewlyCreated (since they no longer have a file on disk) and dirty so that they can be saved again if the delete is undone.

diff --git a/Engine/Source/Editor/UnrealEd/Private/EditorEngine.cpp b/Engine/Source/Editor/UnrealEd/Private/EditorEngine.cpp
--- a/Engine/Source/Editor/UnrealEd/Private/EditorEngine.cpp
+++ b/Engine/Source/Editor/UnrealEd/Private/EditorEngine.cpp
@@ -4656,13 +4656,6 @@ FSavePackageResultStruct UEditorEngine::Save(UPackage* InOuter, UObject* InAsset
 
 	SlowTask.EnterProgressFrame(10);
 
-	const bool bAutosave = (SaveArgs.SaveFlags & SAVE_FromAutosave) != 0;
-	if (!bSavingConcurrent && !IsRunningCommandlet() && !bAutosave && Asset && (World || Asset->IsA<AActor>() || Asset->IsA<UActorFolder>()))
-	{
-		// Always reset the transaction buffer on level/actor save to avoid problems with deleted actors (marked pending kill) that gets marked transient by the saving code
-		ResetTransaction( World ? NSLOCTEXT("UnrealEd", "MapSaved", "Map Saved") : Asset->IsA<AActor>() ? NSLOCTEXT("UnrealEd", "ActorSaved", "Actor Saved") : NSLOCTEXT("UnrealEd", "ActorFolderSaved", "Actor Folder Saved"));
-	}
-
 	if ( World )
 	{
 		if (OriginalOwningWorld)
diff --git a/Engine/Source/Editor/UnrealEd/Private/ObjectTools.cpp b/Engine/Source/Editor/UnrealEd/Private/ObjectTools.cpp
--- a/Engine/Source/Editor/UnrealEd/Private/ObjectTools.cpp
+++ b/Engine/Source/Editor/UnrealEd/Private/ObjectTools.cpp
@@ -2329,15 +2329,17 @@ namespace ObjectTools
 
 			bool bIsReferenced = false;
 
-			if ( Package != nullptr && bPerformReferenceCheck )
+			// Skip external actor packages when considering whether to clear the transaction buffer, as you should be able to undo deleting an actor (but not an asset)
+			// If an external actor package is kept alive by the transaction buffer then it will be re-marked as "newly created" further down this function
+			if ( Package != nullptr && bPerformReferenceCheck && !Package->GetName().Contains(FPackagePath::GetExternalActorsFolderName()))
 			{
 				bool bIsReferencedByUndo = false;
 				GatherObjectReferencersForDeletion(Package, bIsReferenced, bIsReferencedByUndo);
 
 				// only ref to this object is the transaction buffer, clear the transaction buffer
-				if (!bIsReferenced && bIsReferencedByUndo && GEditor && GEditor->Trans)
+				if (!bIsReferenced && bIsReferencedByUndo && GEditor)
 				{
-					GEditor->Trans->Reset(NSLOCTEXT("UnrealEd", "DeleteSelectedItem", "Delete Selected Item"));
+					GEditor->ResetTransaction(NSLOCTEXT("UnrealEd", "DeleteSelectedItem", "Delete Selected Item"));
 				}
 			}
 
@@ -2519,6 +2521,24 @@ namespace ObjectTools
 			}
 		}
 
+		// Ensure that any packages that had their file deleted despite still existing in memory are marked "newly created" again, since they 
+		// no longer have an associated file on disk. This typically happens for OFPA packages that are kept alive by the transaction buffer.
+		for (const FString& PackageFilename : PackageFilesToDelete)
+		{
+			if (!FPaths::FileExists(PackageFilename))
+			{
+				FString PackageName;
+				if (FPackageName::TryConvertFilenameToLongPackageName(PackageFilename, PackageName))
+				{
+					if (UPackage* Package = FindPackage(nullptr, *PackageName))
+					{
+						Package->MarkAsNewlyCreated();
+						Package->SetDirtyFlag(true);
+					}
+				}
+			}
+		}
+
 		// Let the level browser that we deleted a level (must happen after physically deleting the package file as it will rescan the folders)
 		FEditorDelegates::RefreshLevelBrowser.Broadcast();
 	}
@@ -3065,9 +3085,9 @@ namespace ObjectTools
 			GatherObjectReferencersForDeletion(ObjectToDelete, bIsReferenced, bIsReferencedByUndo, &Refs, bRequireReferencedProperties);
 
 			// only ref to this object is the transaction buffer, clear the transaction buffer
-			if (!bIsReferenced && bIsReferencedByUndo && GEditor && GEditor->Trans)
+			if (!bIsReferenced && bIsReferencedByUndo && GEditor)
 			{
-				GEditor->Trans->Reset( NSLOCTEXT( "UnrealEd", "DeleteSelectedItem", "Delete Selected Item" ) );
+				GEditor->ResetTransaction( NSLOCTEXT( "UnrealEd", "DeleteSelectedItem", "Delete Selected Item" ) );
 			}
 
 			if ( bIsReferenced )
diff --git a/Engine/Source/Editor/UnrealEd/Private/PackageTools.cpp b/Engine/Source/Editor/UnrealEd/Private/PackageTools.cpp
--- a/Engine/Source/Editor/UnrealEd/Private/PackageTools.cpp
+++ b/Engine/Source/Editor/UnrealEd/Private/PackageTools.cpp
@@ -441,8 +441,7 @@ UPackageTools::UPackageTools(const FObjectInitializer& ObjectInitializer)
 				}
 				ObjectsInPackage.Reset();
 
-				PackageBeingUnloaded->bHasBeenFullyLoaded = false;
-				PackageBeingUnloaded->ClearFlags(RF_WasLoaded);
+				PackageBeingUnloaded->MarkAsUnloaded();
 				if ( PackageBeingUnloaded->HasAnyPackageFlags(PKG_ContainsScript) )
 				{
 					bScriptPackageWasUnloaded = true;
diff --git a/Engine/Source/Runtime/CoreUObject/Public/UObject/Package.h b/Engine/Source/Runtime/CoreUObject/Public/UObject/Package.h
--- a/Engine/Source/Runtime/CoreUObject/Public/UObject/Package.h
+++ b/Engine/Source/Runtime/CoreUObject/Public/UObject/Package.h
@@ -499,6 +499,27 @@ public:
 		return bDirty;
 	}
 
+#if WITH_EDITOR
+	/**
+	* Marks this package as newly created (has no corresponding file on disk).
+	*/
+	void MarkAsNewlyCreated()
+	{
+		MarkAsUnloaded();
+		SetPackageFlags(PKG_NewlyCreated);
+		SetFileSize(0);
+	}
+
+	/**
+	* Marks this package as unloaded.
+	*/
+	void MarkAsUnloaded()
+	{
+		bHasBeenFullyLoaded = false;
+		ClearFlags(RF_WasLoaded);
+	}
+#endif
+
 	/**
 	* Marks this package as being fully loaded.
 	*/
diff --git a/Engine/Source/Runtime/Engine/Private/WorldPartition/WorldPartitionActorLoaderInterface.cpp b/Engine/Source/Runtime/Engine/Private/WorldPartition/WorldPartitionActorLoaderInterface.cpp
--- a/Engine/Source/Runtime/Engine/Private/WorldPartition/WorldPartitionActorLoaderInterface.cpp
+++ b/Engine/Source/Runtime/Engine/Private/WorldPartition/WorldPartitionActorLoaderInterface.cpp
@@ -270,13 +270,13 @@ void IWorldPartitionActorLoaderInterface::ILoaderAdapter::PostLoadedStateChanged
 		GEngine->BroadcastLevelActorListChanged();
 		GEditor->NoteSelectionChange();
 
-		if (!GIsTransacting)
-		{
-			GEditor->ResetTransaction(LOCTEXT("LoadingEditorActorResetTrans", "Editor Actors Loading State Changed"));
-		}
-
 		if (NumUnloads)
 		{
+			if (!GIsTransacting)
+			{
+				GEditor->ResetTransaction(LOCTEXT("UnloadingEditorActorResetTrans", "Editor Actors Unloaded"));
+			}
+
 			GEngine->ForceGarbageCollection(true);
 		}
 	}
