Logging: Added structured logging support to non-platform feedback contexts

#rnx

diff --git a/Engine/Source/Editor/UnrealEd/Private/AutoReimport/ReimportFeedbackContext.h b/Engine/Source/Editor/UnrealEd/Private/AutoReimport/ReimportFeedbackContext.h
--- a/Engine/Source/Editor/UnrealEd/Private/AutoReimport/ReimportFeedbackContext.h
+++ b/Engine/Source/Editor/UnrealEd/Private/AutoReimport/ReimportFeedbackContext.h
@@ -41,7 +41,9 @@ public:
 	void AddWidget(const TSharedRef<SWidget>& Widget);
 
 	/** INotificationWidget and FFeedbackContext overrides */
-	virtual void Serialize( const TCHAR* V, ELogVerbosity::Type Verbosity, const class FName& Category ) override {}
+	virtual void Serialize(const TCHAR* V, ELogVerbosity::Type Verbosity, const FName& Category) override {}
+	virtual void Serialize(const TCHAR* V, ELogVerbosity::Type Verbosity, const FName& Category, double Time) override {}
+	virtual void SerializeRecord(const UE::FLogRecord& Record) override {}
 	virtual void OnSetCompletionState(SNotificationItem::ECompletionState State) override {}
 	virtual TSharedRef<SWidget> AsWidget() override;
 	virtual void StartSlowTask(const FText& Task, bool bShowCancelButton=false) override;
diff --git a/Engine/Source/Editor/UnrealEd/Private/Commandlets/GatherTextFromAssetsCommandlet.cpp b/Engine/Source/Editor/UnrealEd/Private/Commandlets/GatherTextFromAssetsCommandlet.cpp
--- a/Engine/Source/Editor/UnrealEd/Private/Commandlets/GatherTextFromAssetsCommandlet.cpp
+++ b/Engine/Source/Editor/UnrealEd/Private/Commandlets/GatherTextFromAssetsCommandlet.cpp
@@ -3,6 +3,7 @@
 #include "Commandlets/GatherTextFromAssetsCommandlet.h"
 #include "UObject/Class.h"
 #include "HAL/FileManager.h"
+#include "Logging/StructuredLog.h"
 #include "Misc/CommandLine.h"
 #include "Misc/ConfigCacheIni.h"
 #include "Misc/Paths.h"
@@ -100,6 +101,11 @@ public:
 	}
 
 	virtual void Serialize(const TCHAR* V, ELogVerbosity::Type Verbosity, const FName& Category) override
+	{
+		Serialize(V, Verbosity, Category, -1.0);
+	}
+
+	virtual void Serialize(const TCHAR* V, ELogVerbosity::Type Verbosity, const FName& Category, double Time) override
 	{
 		if (Verbosity == ELogVerbosity::Error)
 		{
@@ -116,12 +122,49 @@ public:
 		else if (Verbosity == ELogVerbosity::Display)
 		{
 			// Downgrade Display to Log while loading packages
-			OriginalWarningContext->Serialize(V, ELogVerbosity::Log, Category);
+			OriginalWarningContext->Serialize(V, ELogVerbosity::Log, Category, Time);
+		}
+		else
+		{
+			// Pass anything else on to GWarn so that it can handle them appropriately
+			OriginalWarningContext->Serialize(V, Verbosity, Category, Time);
+		}
+	}
+
+	virtual void SerializeRecord(const UE::FLogRecord& Record) override
+	{
+		const ELogVerbosity::Type Verbosity = Record.GetVerbosity();
+		if (Verbosity == ELogVerbosity::Error)
+		{
+			++ErrorCount;
+			// Downgrade Error to Log while loading packages to avoid false positives from things searching for "Error:" tokens in the log file
+			UE::FLogRecord LocalRecord = Record;
+			LocalRecord.SetVerbosity(ELogVerbosity::Log);
+			TStringBuilder<512> Line;
+			FormatRecordLine(Line, LocalRecord);
+			FormattedErrorsAndWarningsList.Emplace(Line);
+		}
+		else if (Verbosity == ELogVerbosity::Warning)
+		{
+			++WarningCount;
+			// Downgrade Warning to Log while loading packages to avoid false positives from things searching for "Warning:" tokens in the log file
+			UE::FLogRecord LocalRecord = Record;
+			LocalRecord.SetVerbosity(ELogVerbosity::Log);
+			TStringBuilder<512> Line;
+			FormatRecordLine(Line, LocalRecord);
+			FormattedErrorsAndWarningsList.Emplace(Line);
+		}
+		else if (Verbosity == ELogVerbosity::Display)
+		{
+			// Downgrade Display to Log while loading packages
+			UE::FLogRecord LocalRecord = Record;
+			LocalRecord.SetVerbosity(ELogVerbosity::Log);
+			OriginalWarningContext->SerializeRecord(LocalRecord);
 		}
 		else
 		{
 			// Pass anything else on to GWarn so that it can handle them appropriately
-			OriginalWarningContext->Serialize(V, Verbosity, Category);
+			OriginalWarningContext->SerializeRecord(Record);
 		}
 	}
 
diff --git a/Engine/Source/Editor/UnrealEd/Private/FeedbackContextEditor.cpp b/Engine/Source/Editor/UnrealEd/Private/FeedbackContextEditor.cpp
--- a/Engine/Source/Editor/UnrealEd/Private/FeedbackContextEditor.cpp
+++ b/Engine/Source/Editor/UnrealEd/Private/FeedbackContextEditor.cpp
@@ -398,14 +398,6 @@ FFeedbackContextEditor::FFeedbackContextEditor()
 	
 }
 
-void FFeedbackContextEditor::Serialize( const TCHAR* V, ELogVerbosity::Type Verbosity, const class FName& Category )
-{
-	if( !GLog->IsRedirectingTo( this ) )
-	{
-		GLog->Serialize( V, Verbosity, Category );
-	}
-}
-
 void FFeedbackContextEditor::StartSlowTask( const FText& Task, bool bShowCancelButton )
 {
 	FFeedbackContext::StartSlowTask( Task, bShowCancelButton );
diff --git a/Engine/Source/Editor/UnrealEd/Public/FeedbackContextEditor.h b/Engine/Source/Editor/UnrealEd/Public/FeedbackContextEditor.h
--- a/Engine/Source/Editor/UnrealEd/Public/FeedbackContextEditor.h
+++ b/Engine/Source/Editor/UnrealEd/Public/FeedbackContextEditor.h
@@ -35,8 +35,6 @@ public:
 
 	FFeedbackContextEditor();
 
-	virtual void Serialize( const TCHAR* V, ELogVerbosity::Type Verbosity, const class FName& Category ) override;
-
 	virtual void StartSlowTask( const FText& Task, bool bShowCancelButton=false ) override;
 	virtual void FinalizeSlowTask( ) override;
 	virtual void ProgressReported( const float TotalProgressInterp, FText DisplayMessage ) override;
diff --git a/Engine/Source/Programs/UnrealVirtualizationTool/Private/UnrealVirtualizationToolApp.cpp b/Engine/Source/Programs/UnrealVirtualizationTool/Private/UnrealVirtualizationToolApp.cpp
--- a/Engine/Source/Programs/UnrealVirtualizationTool/Private/UnrealVirtualizationToolApp.cpp
+++ b/Engine/Source/Programs/UnrealVirtualizationTool/Private/UnrealVirtualizationToolApp.cpp
@@ -9,6 +9,7 @@
 #include "HAL/FeedbackContextAnsi.h"
 #include "HAL/FileManager.h"
 #include "Interfaces/IPluginManager.h"
+#include "Logging/StructuredLog.h"
 #include "Misc/CommandLine.h"
 #include "Misc/ConfigCacheIni.h"
 #include "Misc/FileHelper.h"
@@ -84,14 +85,31 @@ public:
 		GWarn = OriginalLog;
 	}
 
-	virtual void Serialize(const TCHAR* V, ELogVerbosity::Type Verbosity, const class FName& Category) override
+	virtual void Serialize(const TCHAR* V, ELogVerbosity::Type Verbosity, const FName& Category) override
+	{
+		Serialize(V, Verbosity, Category, -1.0);
+	}
+
+	virtual void Serialize(const TCHAR* V, ELogVerbosity::Type Verbosity, const FName& Category, double Time) override
 	{
 		if (Verbosity == ELogVerbosity::Display && Category != LogVirtualizationTool.GetCategoryName())
 		{
 			Verbosity = ELogVerbosity::Log;
 		}
 
-		FFeedbackContextAnsi::Serialize(V, Verbosity, Category);
+		FFeedbackContextAnsi::Serialize(V, Verbosity, Category, Time);
+	}
+
+	virtual void SerializeRecord(const UE::FLogRecord& Record) override
+	{
+		if (Record.GetVerbosity() == ELogVerbosity::Display && Record.GetCategory() != LogVirtualizationTool.GetCategoryName())
+		{
+			UE::FLogRecord LocalRecord = Record;
+			LocalRecord.SetVerbosity(ELogVerbosity::Log);
+			return FFeedbackContextAnsi::SerializeRecord(LocalRecord);
+		}
+
+		FFeedbackContextAnsi::SerializeRecord(Record);
 	}
 
 private:
diff --git a/Engine/Source/Runtime/Core/Private/Misc/AutomationTest.cpp b/Engine/Source/Runtime/Core/Private/Misc/AutomationTest.cpp
--- a/Engine/Source/Runtime/Core/Private/Misc/AutomationTest.cpp
+++ b/Engine/Source/Runtime/Core/Private/Misc/AutomationTest.cpp
@@ -6,6 +6,7 @@
 #include "Misc/CommandLine.h"
 #include "Misc/Paths.h"
 #include "Internationalization/Internationalization.h"
+#include "Logging/StructuredLog.h"
 #include "Misc/ConfigCacheIni.h"
 #include "Misc/ScopedSlowTask.h"
 #include "Misc/App.h"
@@ -198,7 +199,12 @@ void FAutomationTestFramework::FAutomationTestOutputDevice::Serialize( const TCH
 	}
 }
 
-void FAutomationTestFramework::FAutomationTestMessageFilter::Serialize(const TCHAR* V, ELogVerbosity::Type Verbosity, const class FName& Category)
+void FAutomationTestFramework::FAutomationTestMessageFilter::Serialize(const TCHAR* V, ELogVerbosity::Type Verbosity, const FName& Category)
+{
+	Serialize(V, Verbosity, Category, -1.0);
+}
+
+void FAutomationTestFramework::FAutomationTestMessageFilter::Serialize(const TCHAR* V, ELogVerbosity::Type Verbosity, const FName& Category, double Time)
 {
 	// Prevent null dereference if logging happens in async tasks while changing DestinationContext
 	FFeedbackContext* const LocalDestinationContext = DestinationContext.load(std::memory_order_relaxed);
@@ -213,7 +219,30 @@ void FAutomationTestFramework::FAutomationTestMessageFilter::Serialize(const TCH
 				Verbosity = ELogVerbosity::Verbose;
 			}
 		}
-		LocalDestinationContext->Serialize(V, Verbosity, Category);
+		LocalDestinationContext->Serialize(V, Verbosity, Category, Time);
+	}
+}
+
+void FAutomationTestFramework::FAutomationTestMessageFilter::SerializeRecord(const UE::FLogRecord& Record)
+{
+	// Prevent null dereference if logging happens in async tasks while changing DestinationContext
+	FFeedbackContext* const LocalDestinationContext = DestinationContext.load(std::memory_order_relaxed);
+	FAutomationTestBase* const LocalCurTest = CurTest.load(std::memory_order_relaxed);
+	if (LocalDestinationContext)
+	{
+		UE::FLogRecord LocalRecord = Record;
+		const ELogVerbosity::Type Verbosity = LocalRecord.GetVerbosity();
+		FScopeLock Lock(&ActionCS);
+		if ((Verbosity == ELogVerbosity::Warning) || (Verbosity == ELogVerbosity::Error))
+		{
+			TStringBuilder<512> Line;
+			Record.FormatMessageTo(Line);
+			if (LocalCurTest->IsExpectedError(FString(Line)))
+			{
+				LocalRecord.SetVerbosity(ELogVerbosity::Verbose);
+			}
+		}
+		LocalDestinationContext->SerializeRecord(LocalRecord);
 	}
 }
 
diff --git a/Engine/Source/Runtime/Core/Public/Misc/AutomationTest.h b/Engine/Source/Runtime/Core/Public/Misc/AutomationTest.h
--- a/Engine/Source/Runtime/Core/Public/Misc/AutomationTest.h
+++ b/Engine/Source/Runtime/Core/Public/Misc/AutomationTest.h
@@ -1096,7 +1096,10 @@ private:
 		  * @param	V		String to serialize within the context
 		  * @param	Event	Event associated with the string
 		  */
-		 virtual void Serialize(const TCHAR* V, ELogVerbosity::Type Verbosity, const class FName& Category) override;
+		 virtual void Serialize(const TCHAR* V, ELogVerbosity::Type Verbosity, const FName& Category) override;
+		 virtual void Serialize(const TCHAR* V, ELogVerbosity::Type Verbosity, const FName& Category, double Time) override;
+
+		 virtual void SerializeRecord(const UE::FLogRecord& Record) override;
 
 		 /**
 		  * FOutputDevice interface
diff --git a/Engine/Source/Runtime/CoreUObject/Private/UObject/Class.cpp b/Engine/Source/Runtime/CoreUObject/Private/UObject/Class.cpp
--- a/Engine/Source/Runtime/CoreUObject/Private/UObject/Class.cpp
+++ b/Engine/Source/Runtime/CoreUObject/Private/UObject/Class.cpp
@@ -4186,66 +4186,14 @@ UObject* UClass::CreateDefaultObject()
 	return ClassDefaultObject;
 }
 
-/**
- * Feedback context implementation for windows.
- */
-class FFeedbackContextImportDefaults : public FFeedbackContext
+class FFeedbackContextImportDefaults final : public FFeedbackContext
 {
 	/** Context information for warning and error messages */
 	FContextSupplier*	Context;
 
 public:
-
-	// Constructor.
-	FFeedbackContextImportDefaults()
-		: Context( NULL )
-	{
-		TreatWarningsAsErrors = true;
-	}
-	void Serialize( const TCHAR* V, ELogVerbosity::Type Verbosity, const class FName& Category )
-	{
-		if( Verbosity==ELogVerbosity::Error || Verbosity==ELogVerbosity::Warning )
-		{
-			if( TreatWarningsAsErrors && Verbosity==ELogVerbosity::Warning )
-			{
-				Verbosity = ELogVerbosity::Error;
-			}
-
-			FString Prefix;
-			if( Context )
-			{
-				Prefix = Context->GetContext() + TEXT(" : ");
-			}
-			FString Format = Prefix + FOutputDeviceHelper::FormatLogLine(Verbosity, Category, V);
-
-			if(Verbosity == ELogVerbosity::Error)
-			{
-				AddError(Format);
-			}
-			else
-			{
-				AddWarning(Format);
-			}
-		}
-
-		if (GLogConsole)
-		{
-			GLogConsole->Serialize(V, Verbosity, Category);
-		}
-		if (!GLog->IsRedirectingTo(this))
-		{
-			GLog->Serialize(V, Verbosity, Category);
-		}
-	}
-
-	FContextSupplier* GetContext() const
-	{
-		return Context;
-	}
-	void SetContext( FContextSupplier* InSupplier )
-	{
-		Context = InSupplier;
-	}
+	FContextSupplier* GetContext() const override { return Context; }
+	void SetContext(FContextSupplier* InContext) override { Context = InContext; }
 };
 
 FFeedbackContext& UClass::GetDefaultPropertiesFeedbackContext()
