Added an experimental option to allow auto-save to save in-place rather than to a backup location

#rnx

diff --git a/Engine/Source/Editor/UnrealEd/Classes/Settings/EditorLoadingSavingSettings.h b/Engine/Source/Editor/UnrealEd/Classes/Settings/EditorLoadingSavingSettings.h
--- a/Engine/Source/Editor/UnrealEd/Classes/Settings/EditorLoadingSavingSettings.h
+++ b/Engine/Source/Editor/UnrealEd/Classes/Settings/EditorLoadingSavingSettings.h
@@ -20,6 +20,16 @@ namespace ELoadLevelAtStartup
 }
 
 
+UENUM()
+enum class EAutoSaveMethod : uint8
+{
+	/** Autosave to a backup location and offer to restore after an editor crash */
+	BackupAndRestore,
+	/** Autosave in-place, overwriting your existing content after backing up the current file (BETA) */
+	BackupAndOverwrite,
+};
+
+
 /** A filter used by the auto reimport manager to explicitly include/exclude files matching the specified wildcard */
 USTRUCT()
 struct FAutoReimportWildcard
@@ -150,6 +160,10 @@ public:
 	UPROPERTY(EditAnywhere, config, Category=AutoSave, meta=(DisplayName="Save Content"))
 	uint32 bAutoSaveContent:1;
 
+	/** What method should be used when performing an autosave? */
+	UPROPERTY(EditAnywhere, config, Category=AutoSave, meta=(DisplayName="Save Method"), AdvancedDisplay)
+	EAutoSaveMethod AutoSaveMethod = EAutoSaveMethod::BackupAndRestore;
+
 	/** The time interval after which to auto save */
 	UPROPERTY(EditAnywhere, config, Category=AutoSave, meta=(DisplayName="Frequency in Minutes", ClampMin = "1"))
 	int32 AutoSaveTimeMinutes;
diff --git a/Engine/Source/Editor/UnrealEd/Private/FileHelpers.cpp b/Engine/Source/Editor/UnrealEd/Private/FileHelpers.cpp
--- a/Engine/Source/Editor/UnrealEd/Private/FileHelpers.cpp
+++ b/Engine/Source/Editor/UnrealEd/Private/FileHelpers.cpp
@@ -998,7 +998,7 @@ static bool SaveWorld(UWorld* World,
 	return bSuccess;
 }
 
-FString GetAutoSaveFilename(UPackage* const Package, const FString& AutoSavePathRoot, const int32 AutoSaveIndex, const FString& PackageExt)
+FString FEditorFileUtils::GetAutoSaveFilename(UPackage* const Package, const FString& AbsoluteAutosaveDir, const int32 AutoSaveIndex, const FString& PackageExt)
 {
 	// Come up with a meaningful name for the auto-save file
 	const FString PackagePathName = Package->GetPathName();
@@ -1010,11 +1010,11 @@ FString GetAutoSaveFilename(UPackage* const Package, const FString& AutoSavePath
 	const bool bStripRootLeadingSlash = true;
 	if(FPackageName::SplitLongPackageName(PackagePathName, PackageRoot, PackagePath, PackageName, bStripRootLeadingSlash))
 	{
-		AutoSavePath = AutoSavePathRoot / PackageRoot / PackagePath;
+		AutoSavePath = AbsoluteAutosaveDir / PackageRoot / PackagePath;
 	}
 	else
 	{
-		AutoSavePath = AutoSavePathRoot;
+		AutoSavePath = AbsoluteAutosaveDir;
 		PackageName = FPaths::GetBaseFilename(PackagePathName);
 	}
 
@@ -3009,7 +3009,7 @@ EAutosaveContentPackagesResult::Type FEditorFileUtils::AutosaveMapEx(const FStri
 	if ( WorldsArray.Num() > 0 )
 	{
 		FString FinalFilename;
-		for ( int32 WorldIndex = 0 ; WorldIndex < WorldsArray.Num() && FUnrealEdMisc::Get().GetAutosaveState() != FUnrealEdMisc::EAutosaveState::Cancelled ; ++WorldIndex )
+		for ( int32 WorldIndex = 0 ; WorldIndex < WorldsArray.Num(); ++WorldIndex )
 		{
 			UWorld* World = WorldsArray[ WorldIndex ];
 			UPackage* Package = Cast<UPackage>( World->GetOuter() );
@@ -3029,7 +3029,7 @@ EAutosaveContentPackagesResult::Type FEditorFileUtils::AutosaveMapEx(const FStri
 				// Remark the package as being dirty, as saving will have undiritied the package.
 				Package->MarkPackageDirty();
 
-				if( bLevelWasSaved == false && FUnrealEdMisc::Get().GetAutosaveState() != FUnrealEdMisc::EAutosaveState::Cancelled )
+				if( bLevelWasSaved == false )
 				{
 					UE_LOG(LogFileHelpers, Log, TEXT("Editor autosave (incl. sublevels) failed for file '%s' which belongs to world '%s'. Aborting autosave."), *FinalFilename, *EditorContext.World()->GetOutermost()->GetName() );
 					return EAutosaveContentPackagesResult::Failure;
diff --git a/Engine/Source/Editor/UnrealEd/Private/PackageAutoSaver.cpp b/Engine/Source/Editor/UnrealEd/Private/PackageAutoSaver.cpp
--- a/Engine/Source/Editor/UnrealEd/Private/PackageAutoSaver.cpp
+++ b/Engine/Source/Editor/UnrealEd/Private/PackageAutoSaver.cpp
@@ -74,6 +74,12 @@ namespace PackageAutoSaverJson
 	/** @return whether the auto-save restore should be enabled (you can force this to true when testing with a debugger attached) */
 	bool IsRestoreEnabled()
 	{
+		// Restore is disabled unless using the BackupAndRestore auto-save method
+		if (GetDefault<UEditorLoadingSavingSettings>()->AutoSaveMethod != EAutoSaveMethod::BackupAndRestore)
+		{
+			return false;
+		}
+
 		// Note: Restore is disabled when running under the debugger, as programmers
 		// like to just kill applications and we don't want this to count as a crash
 		return !FPlatformMisc::IsDebuggerPresent();
@@ -226,13 +232,10 @@ void FPackageAutoSaver::AttemptAutoSave()
 			FScopedSlowTask SlowTask(100.f, NSLOCTEXT("AutoSaveNotify", "PerformingAutoSave_Caption", "Auto-saving out of date packages..."));
 			SlowTask.MakeDialog();
 
-			bAutosaveHandled = true;
-
-			bIsAutoSaving = true;
-			UnrealEdMisc.SetAutosaveState(FUnrealEdMisc::EAutosaveState::Saving);
-
 			GUnrealEd->SaveConfig();
 
+			bAutosaveHandled = true;
+
 			// Make sure the auto-save directory exists before attempting to write the file
 			const FString AutoSaveDir = AutoSaveUtils::GetAutoSaveDir();
 			IFileManager::Get().MakeDirectory(*AutoSaveDir, true);
@@ -241,31 +244,95 @@ void FPackageAutoSaver::AttemptAutoSave()
 			// Auto-save maps and/or content packages based on user settings.
 			const int32 NewAutoSaveIndex = (AutoSaveIndex + 1) % AutoSaveMaxBackups;
 
-			bool bLevelSaved = false;
-			auto MapsSaveResults = EAutosaveContentPackagesResult::NothingToDo;
-			auto AssetsSaveResults = EAutosaveContentPackagesResult::NothingToDo;
-
-			SlowTask.EnterProgressFrame(50);
+			EAutosaveContentPackagesResult::Type MapsSaveResults = EAutosaveContentPackagesResult::NothingToDo;
+			EAutosaveContentPackagesResult::Type AssetsSaveResults = EAutosaveContentPackagesResult::NothingToDo;
 
-			if (LoadingSavingSettings->bAutoSaveMaps)
+			if (LoadingSavingSettings->AutoSaveMethod == EAutoSaveMethod::BackupAndRestore)
 			{
-				MapsSaveResults = FEditorFileUtils::AutosaveMapEx(AutoSaveDir, NewAutoSaveIndex, false, DirtyMapsForAutoSave);
-				if (MapsSaveResults == EAutosaveContentPackagesResult::Success)
+				bIsAutoSaving = true;
+
+				SlowTask.EnterProgressFrame(50);
+
+				if (LoadingSavingSettings->bAutoSaveMaps)
 				{
-					DirtyMapsForAutoSave.Empty();
+					MapsSaveResults = FEditorFileUtils::AutosaveMapEx(AutoSaveDir, NewAutoSaveIndex, false, DirtyMapsForAutoSave);
+					if (MapsSaveResults == EAutosaveContentPackagesResult::Success)
+					{
+						DirtyMapsForAutoSave.Empty();
+					}
 				}
-			}
 
-			SlowTask.EnterProgressFrame(50);
+				SlowTask.EnterProgressFrame(50);
 
-			if (LoadingSavingSettings->bAutoSaveContent && UnrealEdMisc.GetAutosaveState() != FUnrealEdMisc::EAutosaveState::Cancelled)
+				if (LoadingSavingSettings->bAutoSaveContent)
+				{
+					AssetsSaveResults = FEditorFileUtils::AutosaveContentPackagesEx(AutoSaveDir, NewAutoSaveIndex, false, DirtyContentForAutoSave);
+					if (AssetsSaveResults == EAutosaveContentPackagesResult::Success)
+					{
+						DirtyContentForAutoSave.Empty();
+					}
+				}
+			}
+			else if (LoadingSavingSettings->AutoSaveMethod == EAutoSaveMethod::BackupAndOverwrite)
 			{
-				AssetsSaveResults = FEditorFileUtils::AutosaveContentPackagesEx(AutoSaveDir, NewAutoSaveIndex, false, DirtyContentForAutoSave);
-				if (AssetsSaveResults == EAutosaveContentPackagesResult::Success)
+				// Make a backup copy of any packages we may be about to overwrite
+				{
+					auto BackupExistingPackages = [&AutoSaveDir, NewAutoSaveIndex](const TSet<TWeakObjectPtr<UPackage>, TWeakObjectPtrSetKeyFuncs<TWeakObjectPtr<UPackage>>>& PackagesToBackup)
+					{
+						FString PackageFilename;
+						for (const TWeakObjectPtr<UPackage>& PackageToBackup : PackagesToBackup)
+						{
+							if (UPackage* Pkg = PackageToBackup.Get())
+							{
+								PackageFilename.Reset();
+								if (FPackageName::DoesPackageExist(PackageToBackup->GetPathName(), &PackageFilename))
+								{
+									const FString PackageAutoSaveFilename = FEditorFileUtils::GetAutoSaveFilename(Pkg, AutoSaveDir, NewAutoSaveIndex, FPaths::GetExtension(PackageFilename, /*bIncludeDot*/true));
+									IFileManager::Get().Copy(*PackageAutoSaveFilename, *PackageFilename, /*bReplace*/true);
+								}
+							}
+						}
+					};
+				
+					if (LoadingSavingSettings->bAutoSaveMaps)
+					{
+						BackupExistingPackages(DirtyMapsForAutoSave);
+					}
+					if (LoadingSavingSettings->bAutoSaveContent)
+					{
+						BackupExistingPackages(DirtyContentForAutoSave);
+					}
+				}
+
+				// Note: The in-place save does a regular save of dirty packages, so it doesn't set the bIsAutoSaving flag since it functions like a user-initiated save
+				const bool bAutoSaveMaps = LoadingSavingSettings->bAutoSaveMaps && DirtyMapsForAutoSave.Num() > 0;
+				const bool bAutoSaveContent = LoadingSavingSettings->bAutoSaveContent && DirtyContentForAutoSave.Num() > 0;
+
+				const bool bPromptUserForSave = false;
+				const bool bFastSave = false;
+				const bool bNotifyNoPackagesSaved = false;
+				const bool bCanBeDeclined = false;
+
+				bool bPackagesNeededSaving = false;
+				const bool bSuccess = FEditorFileUtils::SaveDirtyPackages(bPromptUserForSave, bAutoSaveMaps, bAutoSaveContent, bFastSave, bNotifyNoPackagesSaved, bCanBeDeclined, &bPackagesNeededSaving);
+				if (bPackagesNeededSaving)
 				{
-					DirtyContentForAutoSave.Empty();
+					if (bSuccess)
+					{
+						MapsSaveResults = EAutosaveContentPackagesResult::Success;
+						AssetsSaveResults = EAutosaveContentPackagesResult::Success;
+					}
+					else
+					{
+						MapsSaveResults = EAutosaveContentPackagesResult::Failure;
+						AssetsSaveResults = EAutosaveContentPackagesResult::Failure;
+					}
 				}
 			}
+			else
+			{
+				checkf(false, TEXT("Unknown AutoSaveMethod!"));
+			}
 
 			const bool bNothingToDo = (MapsSaveResults == EAutosaveContentPackagesResult::NothingToDo && AssetsSaveResults == EAutosaveContentPackagesResult::NothingToDo);
 			const bool bSuccess = (MapsSaveResults != EAutosaveContentPackagesResult::Failure && AssetsSaveResults != EAutosaveContentPackagesResult::Failure && !bNothingToDo);
@@ -290,13 +357,7 @@ void FPackageAutoSaver::AttemptAutoSave()
 			ResetAutoSaveTimer();
 			bDelayingDueToFailedSave = false;
 
-			if (UnrealEdMisc.GetAutosaveState() == FUnrealEdMisc::EAutosaveState::Cancelled)
-			{
-				UE_LOG(PackageAutoSaver, Warning, TEXT("Autosave was cancelled."));
-			}
-
 			bIsAutoSaving = false;
-			UnrealEdMisc.SetAutosaveState(FUnrealEdMisc::EAutosaveState::Inactive);
 		}
 		else
 		{
diff --git a/Engine/Source/Editor/UnrealEd/Private/UnrealEdMisc.cpp b/Engine/Source/Editor/UnrealEd/Private/UnrealEdMisc.cpp
--- a/Engine/Source/Editor/UnrealEd/Private/UnrealEdMisc.cpp
+++ b/Engine/Source/Editor/UnrealEd/Private/UnrealEdMisc.cpp
@@ -209,7 +209,6 @@ namespace UnrealEdMiscDefs
 }
 
 FUnrealEdMisc::FUnrealEdMisc() :
-	AutosaveState( EAutosaveState::Inactive ), 
 	bCancelBuild( false ),
 	bInitialized( false ),
 	bSaveLayoutOnClose( true ),
diff --git a/Engine/Source/Editor/UnrealEd/Public/FileHelpers.h b/Engine/Source/Editor/UnrealEd/Public/FileHelpers.h
--- a/Engine/Source/Editor/UnrealEd/Public/FileHelpers.h
+++ b/Engine/Source/Editor/UnrealEd/Public/FileHelpers.h
@@ -307,6 +307,16 @@ public:
 	 */
 	UNREALED_API static bool SaveLevelAs(ULevel* Level, FString* OutSavedFilename = nullptr);
 
+	/**
+	 * Get the autosave filename for the given package.
+	 * 
+	 * @param	Package						Package to get the autosave filename for.
+	 * @param	AbsoluteAutosaveDir			Autosave directory.
+	 * @param	AutosaveIndex				Integer prepended to autosave filenames.
+	 * @param	PackageExt					Extension to use for the given package. Note: This must include the dot.
+	 */
+	static FString GetAutoSaveFilename(UPackage* const Package, const FString& AbsoluteAutosaveDir, const int32 AutoSaveIndex, const FString& PackageExt);
+
 	/**
 	 * Saves all levels to the specified directory.
 	 *
diff --git a/Engine/Source/Editor/UnrealEd/Public/UnrealEdMisc.h b/Engine/Source/Editor/UnrealEd/Public/UnrealEdMisc.h
--- a/Engine/Source/Editor/UnrealEd/Public/UnrealEdMisc.h
+++ b/Engine/Source/Editor/UnrealEd/Public/UnrealEdMisc.h
@@ -33,17 +33,6 @@ class UNREALED_API FUnrealEdMisc
 {
 public:
 
-	/** The various autosave states that exist */
-	struct EAutosaveState
-	{
-		enum Type
-		{
-			Inactive,
-			Saving,
-			Cancelled
-		};
-	};
-
 	/** Singleton accessor */
 	static FUnrealEdMisc& Get();
 
@@ -62,26 +51,6 @@ public:
 	/** Performs any required cleanup in the case of a fatal error. */
 	virtual void ShutdownAfterError();
 
-	/**
-	 *	Fetches the current state of the autosave in progress
-	 *
-	 *	@return	enum		the current state of the autosave
-	 */
-	EAutosaveState::Type GetAutosaveState() const
-	{
-		return AutosaveState;
-	}
-
-	/**
-	 * Sets the new state for the autosave
-	 *
-	 * @param InState		New state for the autosave
-	 */
-	void SetAutosaveState( const EAutosaveState::Type InState )
-	{
-		AutosaveState = InState;
-	}
-
 	/**
 	 *	Whether or not the map build in progressed was cancelled by the user. 
 	 *
@@ -316,9 +285,6 @@ private:
 
 	void PreSaveWorld(class UWorld* World, FObjectPreSaveContext ObjectSaveContext);
 
-	/** The current state of the autosave */
-	EAutosaveState::Type AutosaveState;
-
 	/** Stores whether or not the current map build was cancelled. */
 	bool bCancelBuild;
 
