Cooker: Optimize GetNeverCookPackageFileNames by using the AssetRegistry rather than FPackageName::IteratePackagesInDirectory
#rnx

diff --git a/Engine/Source/Editor/UnrealEd/Private/CookOnTheFlyServer.cpp b/Engine/Source/Editor/UnrealEd/Private/CookOnTheFlyServer.cpp
--- a/Engine/Source/Editor/UnrealEd/Private/CookOnTheFlyServer.cpp
+++ b/Engine/Source/Editor/UnrealEd/Private/CookOnTheFlyServer.cpp
@@ -10169,7 +10169,61 @@ TArray<FName> UCookOnTheFlyServer::GetNeverCookPackageFileNames(TArrayView<const
 	}
 
 	TArray<FString> NeverCookPackagesPaths;
-	FPackageName::FindPackagesInDirectories(NeverCookPackagesPaths, NeverCookDirectories);
+	if (AssetRegistry->IsSearchAllAssets() && !AssetRegistry->IsLoadingAssets())
+	{
+		TArray<FString> NeverCookDirectoryLongPackageNames;
+		NeverCookDirectoryLongPackageNames.Reserve(NeverCookDirectories.Num());
+		for (const FString& LocalDirectory : NeverCookDirectories)
+		{
+			FString PackagePath;
+			if (FPackageName::TryConvertFilenameToLongPackageName(LocalDirectory, PackagePath))
+			{
+				NeverCookDirectoryLongPackageNames.Add(PackagePath);
+			}
+		}
+
+		FString PackageNameStr;
+		AssetRegistry->EnumerateAllPackages(
+			[&NeverCookPackagesPaths, &NeverCookDirectoryLongPackageNames, &PackageNameStr](FName PackageName, const FAssetPackageData& PackageData)
+			{
+				PackageName.ToString(PackageNameStr);
+				bool bIsInNeverCook = false;
+				for (const FString& NeverCookLongPackageName : NeverCookDirectoryLongPackageNames)
+				{
+					if (FPathViews::IsParentPathOf(NeverCookLongPackageName, PackageNameStr))
+					{
+						bIsInNeverCook = true;
+						break;
+					}
+				}
+
+				if (bIsInNeverCook)
+				{
+					FString LocalFileName;
+					if (PackageData.Extension != EPackageExtension::Unspecified && PackageData.Extension != EPackageExtension::Custom)
+					{
+						FString Extension = LexToString(PackageData.Extension);
+						FPackageName::TryConvertLongPackageNameToFilename(PackageNameStr, LocalFileName, Extension);
+					}
+					else
+					{
+						FPackageName::DoesPackageExist(PackageNameStr, &LocalFileName);
+					}
+					if (!LocalFileName.IsEmpty())
+					{
+						NeverCookPackagesPaths.Add(LocalFileName);
+					}
+				}
+			});
+	}
+	else
+	{
+		// CookOnTheFly in editor calls this function at editorstartup, before the AssetRegistry has loaded.
+		// Rather than blocking on the AssetRegistry now, fallback to scanning the directories on disk
+		// TODO: Change CookOnTheFlyStartup in the editor to delay most of its startup until the AssetRegistry has
+		// finished loading so we can block on the AssetRegistry before calling this function.
+		FPackageName::FindPackagesInDirectories(NeverCookPackagesPaths, NeverCookDirectories);
+	}
 
 	TArray<FName> NeverCookNormalizedFileNames;
 	for (const FString& NeverCookPackagePath : NeverCookPackagesPaths)
