PR #9947: Prevent crash on opening blueprint with non-existent shortname bitmask enum. (Contributed by [REDACTED_PERSON])

diff --git a/Engine/Source/Editor/UnrealEd/Private/Kismet2/BlueprintEditorUtils.cpp b/Engine/Source/Editor/UnrealEd/Private/Kismet2/BlueprintEditorUtils.cpp
--- a/Engine/Source/Editor/UnrealEd/Private/Kismet2/BlueprintEditorUtils.cpp
+++ b/Engine/Source/Editor/UnrealEd/Private/Kismet2/BlueprintEditorUtils.cpp
@@ -6282,8 +6282,15 @@ void FBlueprintEditorUtils::FixupVariableDescription(UBlueprint* Blueprint, FBPV
 			if (FPackageName::IsShortPackageName(BitmaskEnumTypePath))
 			{
 				BitflagsEnum = FindFirstObject<UEnum>(GetData(BitmaskEnumTypePath));
-				BitmaskEnumTypePath = FTopLevelAssetPath(BitflagsEnum->GetPackage()->GetFName(), BitflagsEnum->GetFName()).ToString();
-				VarDesc.SetMetaData(FBlueprintMetadata::MD_BitmaskEnum, BitmaskEnumTypePath);
+				if (BitflagsEnum != nullptr)
+				{
+					BitmaskEnumTypePath = FTopLevelAssetPath(BitflagsEnum->GetPackage()->GetFName(), BitflagsEnum->GetFName()).ToString();
+					VarDesc.SetMetaData(FBlueprintMetadata::MD_BitmaskEnum, BitmaskEnumTypePath);
+				}
+				else
+				{
+					UE_LOG(LogBlueprint, Error, TEXT("Enum %s cannot be loaded"), *BitmaskEnumTypePath);
+				}
 			}
 			else
 			{
