World Partition: fix bad IsValidChecked usages.

diff --git a/Engine/Source/Editor/UnrealEd/Private/Commandlets/WorldPartitionConvertCommandlet.cpp b/Engine/Source/Editor/UnrealEd/Private/Commandlets/WorldPartitionConvertCommandlet.cpp
--- a/Engine/Source/Editor/UnrealEd/Private/Commandlets/WorldPartitionConvertCommandlet.cpp
+++ b/Engine/Source/Editor/UnrealEd/Private/Commandlets/WorldPartitionConvertCommandlet.cpp
@@ -537,7 +537,7 @@ bool UWorldPartitionConvertCommandlet::DetachDependantLevelPackages(ULevel* Leve
 
 	for (AActor* Actor: Level->Actors)
 	{
-		if (Actor && IsValidChecked(Actor) && Actor->IsA<ALODActor>())
+		if (IsValid(Actor) && Actor->IsA<ALODActor>())
 		{
 			Level->GetWorld()->DestroyActor(Actor);
 		}
@@ -1070,7 +1070,7 @@ int32 UWorldPartitionConvertCommandlet::Main(const FString& Params)
 		{
 			AActor* Actor = *Iter;
 
-			if (Actor && IsValidChecked(Actor))
+			if (IsValid(Actor))
 			{
 				check(Actor->GetLevel() == Level);
 
@@ -1360,7 +1360,7 @@ int32 UWorldPartitionConvertCommandlet::Main(const FString& Params)
 
 		for(AActor* Actor: ActorsToConvert)
 		{
-			if(Actor && IsValidChecked(Actor))
+			if(IsValid(Actor))
 			{
 				check(Actor->GetOuter() == SubLevel);
 				check(!ShouldDeleteActor(Actor, false));
@@ -1464,7 +1464,7 @@ int32 UWorldPartitionConvertCommandlet::Main(const FString& Params)
 	// Move child actors at the end of the list
 	for (AActor* Actor: MainLevel->Actors)
 	{
-		if (Actor && IsValidChecked(Actor))
+		if (IsValid(Actor))
 		{
 			check(Actor->GetLevel() == MainLevel);
 			check(Actor->GetActorGuid().IsValid());
@@ -1496,7 +1496,7 @@ int32 UWorldPartitionConvertCommandlet::Main(const FString& Params)
 		TSet<FGuid> ActorGuids;
 		for(AActor* Actor: ActorList)
 		{
-			if (!Actor || !IsValidChecked(Actor) || !Actor->SupportsExternalPackaging())
+			if (!IsValid(Actor) || !Actor->SupportsExternalPackaging())
 			{
 				continue;
 			}
diff --git a/Engine/Source/Runtime/Engine/Private/WorldPartition/WorldPartitionConverter.cpp b/Engine/Source/Runtime/Engine/Private/WorldPartition/WorldPartitionConverter.cpp
--- a/Engine/Source/Runtime/Engine/Private/WorldPartition/WorldPartitionConverter.cpp
+++ b/Engine/Source/Runtime/Engine/Private/WorldPartition/WorldPartitionConverter.cpp
@@ -76,7 +76,7 @@ bool FWorldPartitionConverter::Convert()
 	TArray<AActor*> Actors = MainLevel->Actors;
 	for (AActor* Actor : Actors)
 	{
-		if (Actor && IsValidChecked(Actor) && ShouldDeleteActor(Actor, /*bIsMainLevel*/ true))
+		if (IsValid(Actor) && ShouldDeleteActor(Actor, /*bIsMainLevel*/ true))
 		{
 			MainWorld->EditorDestroyActor(Actor, /*bShouldModifyLevel*/ false);
 		}
@@ -115,7 +115,7 @@ bool FWorldPartitionConverter::Convert()
 			TArray<AActor*> ActorsToConvert = DuplicatedLevel->Actors;
 			for (AActor* Actor : ActorsToConvert)
 			{
-				if (Actor && IsValidChecked(Actor) && !ShouldDeleteActor(Actor, /*bIsMainLevel*/ false))
+				if (IsValid(Actor) && !ShouldDeleteActor(Actor, /*bIsMainLevel*/ false))
 				{
 					// Before changing outer, bIsCookedForEditor flag to actor package 
 					if (SourceLevelPackage->bIsCookedForEditor)
