Fixed a localization asset gather crash when using a cooked asset registry

This can report assets that aren't currently mounted, so the conversion to a filename would fail (and assert). It now handles that failure gracefully and just removes the asset from the list to be gathered.

#rnx

diff --git a/Engine/Source/Editor/UnrealEd/Private/Commandlets/GatherTextFromAssetsCommandlet.cpp b/Engine/Source/Editor/UnrealEd/Private/Commandlets/GatherTextFromAssetsCommandlet.cpp
--- a/Engine/Source/Editor/UnrealEd/Private/Commandlets/GatherTextFromAssetsCommandlet.cpp
+++ b/Engine/Source/Editor/UnrealEd/Private/Commandlets/GatherTextFromAssetsCommandlet.cpp
@@ -646,13 +646,22 @@ int32 UGatherTextFromAssetsCommandlet::Main(const FString& Params)
 	const FFuzzyPathMatcher FuzzyPathMatcher = FFuzzyPathMatcher(IncludePathFilters, ExcludePathFilters);
 	AssetDataArray.RemoveAll([&](const FAssetData& PartiallyFilteredAssetData) -> bool
 	{
-		FString PackageFilePath;
-		if (!FPackageName::FindPackageFileWithoutExtension(FPackageName::LongPackageNameToFilename(PartiallyFilteredAssetData.PackageName.ToString()), PackageFilePath))
+		FString PackageFilePathWithoutExtension;
+		if (!FPackageName::TryConvertLongPackageNameToFilename(PartiallyFilteredAssetData.PackageName.ToString(), PackageFilePathWithoutExtension))
 		{
+			// This means the asset data is for content that isn't mounted - this can happen when using a cooked asset registry
 			return true;
 		}
-		PackageFilePath = FPaths::ConvertRelativePathToFull(PackageFilePath);
-		const FString PackageFileName = FPaths::GetCleanFilename(PackageFilePath);
+
+		FString PackageFilePathWithExtension;
+		if (!FPackageName::FindPackageFileWithoutExtension(PackageFilePathWithoutExtension, PackageFilePathWithExtension))
+		{
+			// This means the package file doesn't exist on disk, which means we cannot gather it
+			return true;
+		}
+
+		PackageFilePathWithExtension = FPaths::ConvertRelativePathToFull(PackageFilePathWithExtension);
+		const FString PackageFileName = FPaths::GetCleanFilename(PackageFilePathWithExtension);
 
 		// Filter out assets whose package file names DO NOT match any of the package file name filters.
 		{
@@ -672,7 +681,7 @@ int32 UGatherTextFromAssetsCommandlet::Main(const FString& Params)
 		}
 
 		// Filter out assets whose package file paths do not pass the "fuzzy path" filters.
-		if (FuzzyPathMatcher.TestPath(PackageFilePath) != FFuzzyPathMatcher::Included)
+		if (FuzzyPathMatcher.TestPath(PackageFilePathWithExtension) != FFuzzyPathMatcher::Included)
 		{
 			return true;
 		}
