Fix for missing KeepInDevOnly tags in dev AR
AssetRegistryGenerator: we need to not apply the runtime AR tag filter to the main state until after we have cloned it into the Development AR.
Creds to Matt

diff --git a/Engine/Source/Editor/UnrealEd/Private/Commandlets/AssetRegistryGenerator.cpp b/Engine/Source/Editor/UnrealEd/Private/Commandlets/AssetRegistryGenerator.cpp
--- a/Engine/Source/Editor/UnrealEd/Private/Commandlets/AssetRegistryGenerator.cpp
+++ b/Engine/Source/Editor/UnrealEd/Private/Commandlets/AssetRegistryGenerator.cpp
@@ -1561,11 +1561,9 @@ bool FAssetRegistryGenerator::SaveAssetRegistry(const FString& SandboxPath, bool
 		SaveOptions.DisableFilters();
 	}
 
-	// Filter tags from both development and saved registries, keeping only the tags applicable to the platform
-	State.FilterTags(SaveOptions);
-
-	// Then possibly apply previous AssetData and AssetPackageData for packages kept from a previous cook
+	// Possibly apply previous AssetData and AssetPackageData for packages kept from a previous cook
 	UpdateKeptPackages();
+
 	UpdateCollectionAssetData();
 
 	if (DevelopmentSaveOptions.bSerializeAssetRegistry)
@@ -1580,6 +1578,8 @@ bool FAssetRegistryGenerator::SaveAssetRegistry(const FString& SandboxPath, bool
 			FAssetRegistryState DevelopmentState;
 			DevelopmentState.InitializeFromExisting(State, DevelopmentSaveOptions);
 
+			DevelopmentState.FilterTags(DevelopmentSaveOptions);
+
 			AddCookTagsToState(MoveTemp(CookTagsToAdd), DevelopmentState);
 
 			// Create development registry data, used for DLC cooks, iterative cooks, and editor viewing
@@ -1609,6 +1609,8 @@ bool FAssetRegistryGenerator::SaveAssetRegistry(const FString& SandboxPath, bool
 		const int32 GenericChunkBucket = -1;
 		ChunkBucketNames.Add(GenericChunkBucket, FString());
 
+		State.FilterTags(SaveOptions);
+
 		// When chunk manifests have been generated (e.g. cook by the book) serialize 
 		// an asset registry for each chunk.
 		if (FinalChunkManifests.Num() > 0)
