Fix an issue where components had inconsistent behavior between gizmo and details panel when the component had absolute location turned on

diff --git a/Engine/Source/Editor/UnrealEd/Private/Elements/Component/ComponentElementEditorViewportInteractionCustomization.cpp b/Engine/Source/Editor/UnrealEd/Private/Elements/Component/ComponentElementEditorViewportInteractionCustomization.cpp
--- a/Engine/Source/Editor/UnrealEd/Private/Elements/Component/ComponentElementEditorViewportInteractionCustomization.cpp
+++ b/Engine/Source/Editor/UnrealEd/Private/Elements/Component/ComponentElementEditorViewportInteractionCustomization.cpp
@@ -20,7 +20,8 @@ bool FComponentElementEditorViewportInteractionCustomization::GetGizmoPivotLocat
 
 		// If necessary, transform the editor pivot location to be relative to the component's parent
 		const bool bIsRootComponent = SceneComponent->GetOwner()->GetRootComponent() == SceneComponent;
-		OutPivotLocation = bIsRootComponent || !SceneComponent->GetAttachParent() ? BasePivotLocation : SceneComponent->GetAttachParent()->GetComponentToWorld().Inverse().TransformPosition(BasePivotLocation);
+		const bool bIsComponentUsingAbsolutePosition = SceneComponent->IsUsingAbsoluteLocation();
+		OutPivotLocation = bIsRootComponent || bIsComponentUsingAbsolutePosition || !SceneComponent->GetAttachParent() ? BasePivotLocation : SceneComponent->GetAttachParent()->GetComponentToWorld().Inverse().TransformPosition(BasePivotLocation);
 		return true;
 	}
 
