Fix a crash if the toolkit edited asset get reinstanced into a new object by using weak pointer instead of a raw pointer in the lambda capture.

diff --git a/Engine/Source/Editor/UnrealEd/Private/Toolkits/AssetEditorToolkit.cpp b/Engine/Source/Editor/UnrealEd/Private/Toolkits/AssetEditorToolkit.cpp
--- a/Engine/Source/Editor/UnrealEd/Private/Toolkits/AssetEditorToolkit.cpp
+++ b/Engine/Source/Editor/UnrealEd/Private/Toolkits/AssetEditorToolkit.cpp
@@ -192,14 +192,20 @@ void FAssetEditorToolkit::InitAssetEditor( const EToolkitMode::Type Mode, const
 		NewTabManager->SetAllowWindowMenuBar(true);
 		this->TabManager = NewTabManager;
 
+		TArray<TWeakObjectPtr<UObject>> ObjectsToEditWeak;
+		ObjectsToEditWeak.Reserve(ObjectsToEdit.Num());
+		for (UObject* Object : ObjectsToEdit)
+		{
+			ObjectsToEditWeak.Add(Object);
+		}
 		NewMajorTab->SetContent
 		( 
 			SAssignNew( NewStandaloneHost, SStandaloneAssetEditorToolkitHost, NewTabManager, AppIdentifier )
-			.IsEnabled_Lambda([ObjectsToEdit]()
+			.IsEnabled_Lambda([ObjectsToEditWeak]()
 				{
-					for (const UObject* Object : ObjectsToEdit)
+					for (const TWeakObjectPtr<UObject> Object : ObjectsToEditWeak)
 					{
-						if (const IInterface_AsyncCompilation* AsyncAsset = Cast<IInterface_AsyncCompilation>(Object))
+						if (const IInterface_AsyncCompilation* AsyncAsset = Cast<IInterface_AsyncCompilation>(Object.Get()))
 						{
 							if (AsyncAsset->IsCompiling())
 							{
