Blueprint  thumbnail: geometry collection based  blueprint can now render thumbnail properly

diff --git a/Engine/Source/Editor/UnrealEd/Private/ThumbnailHelpers.cpp b/Engine/Source/Editor/UnrealEd/Private/ThumbnailHelpers.cpp
--- a/Engine/Source/Editor/UnrealEd/Private/ThumbnailHelpers.cpp
+++ b/Engine/Source/Editor/UnrealEd/Private/ThumbnailHelpers.cpp
@@ -1125,6 +1125,15 @@ bool FClassActorThumbnailScene::IsValidComponentForVisualization(UActorComponent
 		{
 			return true;
 		}
+
+		// we cannot include the geomety collection component in this module because of circular dependency 
+		// so we need to check using the name of the class instead 
+		const FName ClassName = Component->GetClass()->GetFName();
+		const FName GeometryCollectionClassName("GeometryCollectionComponent");
+		if (ClassName == GeometryCollectionClassName)
+		{
+			return true;
+		}
 	}
 
 	return false;
