Make the cook package tracker thread-safe for async loading thread support

#rnx

diff --git a/Engine/Source/Editor/UnrealEd/Private/CookOnTheFlyServer.cpp b/Engine/Source/Editor/UnrealEd/Private/CookOnTheFlyServer.cpp
--- a/Engine/Source/Editor/UnrealEd/Private/CookOnTheFlyServer.cpp
+++ b/Engine/Source/Editor/UnrealEd/Private/CookOnTheFlyServer.cpp
@@ -2825,16 +2825,19 @@ void UCookOnTheFlyServer::UpdatePackageFilter()
 
 	UE_SCOPED_COOKTIMER(UpdatePackageFilter);
 	const TArray<const ITargetPlatform*>& TargetPlatforms = PlatformManager->GetSessionPlatforms();
-	for (UPackage* Package : PackageTracker->LoadedPackages)
-	{
-		UE::Cook::FPackageData* PackageData = QueueDiscoveredPackage(Package,
-			UE::Cook::FInstigator(UE::Cook::EInstigator::Unspecified, GInstigatorUpdatePackageFilter));
-		if (PackageData && PackageData->IsInProgress())
+
+	PackageTracker->ForEachLoadedPackage(
+		[this, &TargetPlatforms](UPackage* Package)
 		{
-			PackageData->UpdateRequestData(TargetPlatforms, /*bIsUrgent*/ false, UE::Cook::FCompletionCallback(),
-				UE::Cook::FInstigator(UE::Cook::EInstigator::Unspecified, GInstigatorUpdatePackageFilter));
+			UE::Cook::FPackageData* PackageData = QueueDiscoveredPackage(Package,
+			UE::Cook::FInstigator(UE::Cook::EInstigator::Unspecified, GInstigatorUpdatePackageFilter));
+			if (PackageData && PackageData->IsInProgress())
+			{
+				PackageData->UpdateRequestData(TargetPlatforms, /*bIsUrgent*/ false, UE::Cook::FCompletionCallback(),
+					UE::Cook::FInstigator(UE::Cook::EInstigator::Unspecified, GInstigatorUpdatePackageFilter));
+			}
 		}
-	}
+	);
 }
 
 void UCookOnTheFlyServer::OnRemoveSessionPlatform(const ITargetPlatform* TargetPlatform)
@@ -5183,15 +5186,17 @@ void UCookOnTheFlyServer::PreGarbageCollectImpl(TArray<UPackage*>& GCKeepPackage
 		}
 
 		TSet<FName> LoadedPackages;
-		for (UPackage* Package : PackageTracker->LoadedPackages)
-		{
-			const FName& PackageName = Package->GetFName();
-			if (KeepPackages.Contains(PackageName))
+		PackageTracker->ForEachLoadedPackage(
+			[&KeepPackages, &LoadedPackages, &GCKeepPackages](UPackage* Package)
 			{
-				LoadedPackages.Add(PackageName);
-				GCKeepPackages.Add(Package);
+				const FName& PackageName = Package->GetFName();
+				if (KeepPackages.Contains(PackageName))
+				{
+					LoadedPackages.Add(PackageName);
+					GCKeepPackages.Add(Package);
+				}
 			}
-		}
+		);
 
 		FRequestQueue& RequestQueue = PackageDatas->GetRequestQueue();
 		TArray<FPackageData*> Requests;
diff --git a/Engine/Source/Editor/UnrealEd/Private/Cooker/PackageTracker.cpp b/Engine/Source/Editor/UnrealEd/Private/Cooker/PackageTracker.cpp
--- a/Engine/Source/Editor/UnrealEd/Private/Cooker/PackageTracker.cpp
+++ b/Engine/Source/Editor/UnrealEd/Private/Cooker/PackageTracker.cpp
@@ -5,6 +5,7 @@
 #include "CookOnTheFlyServerInterface.h"
 #include "CookPackageData.h"
 #include "CookPlatformManager.h"
+#include "Misc/ScopeRWLock.h"
 #include "ProfilingDebugging/CookStats.h"
 #include "UObject/Package.h"
 #include "UObject/UObjectIterator.h"
@@ -63,6 +64,8 @@ namespace UE::Cook
 		:PackageDatas(InPackageDatas)
 	{
 		LLM_SCOPE_BYTAG(Cooker);
+
+		FWriteScopeLock ScopeLock(Lock);
 		for (TObjectIterator<UPackage> It; It; ++It)
 		{
 			UPackage* Package = *It;
@@ -91,6 +94,7 @@ namespace UE::Cook
 
 	TMap<UPackage*, FInstigator> FPackageTracker::GetNewPackages()
 	{
+		FWriteScopeLock ScopeLock(Lock);
 		TMap<UPackage*, FInstigator> Result = MoveTemp(NewPackages);
 		NewPackages.Reset();
 		bHasBeenConsumed = true;
@@ -117,6 +121,7 @@ namespace UE::Cook
 					COOK_STAT(++Stats::NumInlineLoads);
 				}
 
+				FWriteScopeLock ScopeLock(Lock);
 				LoadedPackages.Add(Package);
 				NewPackages.Add(Package,
 						FInstigator(EInstigator::Unsolicited,
@@ -132,6 +137,7 @@ namespace UE::Cook
 		{
 			UPackage* Package = const_cast<UPackage*>(static_cast<const UPackage*>(Object));
 
+			FWriteScopeLock ScopeLock(Lock);
 			LoadedPackages.Remove(Package);
 			NewPackages.Remove(Package);
 		}
diff --git a/Engine/Source/Editor/UnrealEd/Private/Cooker/PackageTracker.h b/Engine/Source/Editor/UnrealEd/Private/Cooker/PackageTracker.h
--- a/Engine/Source/Editor/UnrealEd/Private/Cooker/PackageTracker.h
+++ b/Engine/Source/Editor/UnrealEd/Private/Cooker/PackageTracker.h
@@ -190,12 +190,6 @@ namespace Cook
 		/** Swap all ITargetPlatform* stored on this instance according to the mapping in @param Remap. */
 		void RemapTargetPlatforms(const TMap<ITargetPlatform*, ITargetPlatform*>& Remap);
 
-		// This is a complete list of currently loaded UPackages
-		TFastPointerSet<UPackage*> LoadedPackages;
-
-		// This list contains the UPackages loaded since last call to GetNewPackages
-		TMap<UPackage*, FInstigator> NewPackages;
-
 		/** The package currently being loaded at CookOnTheFlyServer's direct request. Used to determine which load dependencies were not preloaded. */
 		FPackageData* LoadingPackageData = nullptr;
 
@@ -210,6 +204,27 @@ namespace Cook
 		TFastPointerMap<const ITargetPlatform*, TSet<FName>> PlatformSpecificNeverCookPackages;
 
 		bool bHasBeenConsumed = false;
+
+		// Thread-safe enumeration of loaded package. 
+		// A lock is held during enumeration, keep code simple and optimal so the lock is released as fast as possible.
+		template <typename FunctionType>
+		void ForEachLoadedPackage(FunctionType Function)
+		{
+			FReadScopeLock ScopeLock(Lock);
+			for (UPackage* Package : LoadedPackages)
+			{
+				Function(Package);
+			}
+		}
+	private:
+		// Protects data for thread-safety
+		FRWLock Lock;
+
+		// This is a complete list of currently loaded UPackages
+		TFastPointerSet<UPackage*> LoadedPackages;
+
+		// This list contains the UPackages loaded since last call to GetNewPackages
+		TMap<UPackage*, FInstigator> NewPackages;
 	}; 
 }
 }
