BulkDataRegistry: Teardown during OnEnginePreExit so that we do not try to use AssetRegistry after its module has been torn down. Act as a noop rather than asserting if public functions are called after Teardown. Skip call to WriteCache from DDC Get callback if already torn down.
#rnx

diff --git a/Engine/Source/Editor/UnrealEd/Private/EditorDomain/BulkDataRegistryImpl.cpp b/Engine/Source/Editor/UnrealEd/Private/EditorDomain/BulkDataRegistryImpl.cpp
--- a/Engine/Source/Editor/UnrealEd/Private/EditorDomain/BulkDataRegistryImpl.cpp
+++ b/Engine/Source/Editor/UnrealEd/Private/EditorDomain/BulkDataRegistryImpl.cpp
@@ -17,6 +17,7 @@
 #include "Logging/LogMacros.h"
 #include "Memory/SharedBuffer.h"
 #include "Misc/AssertionMacros.h"
+#include "Misc/CoreDelegates.h"
 #include "Misc/Optional.h"
 #include "Misc/ScopeLock.h"
 #include "Misc/ScopeRWLock.h"
@@ -68,10 +69,30 @@ FBulkDataRegistryImpl::FBulkDataRegistryImpl()
 	// We piggyback on the BulkDataRegistry hook to set the pointer to tunnel in the pointer to the EditorBuildInputResolver as well
 	SetGlobalBuildInputResolver(&UE::DerivedData::FEditorBuildInputResolver::Get());
 	FCoreUObjectDelegates::OnEndLoadPackage.AddRaw(this, &FBulkDataRegistryImpl::OnEndLoadPackage);
+	FCoreDelegates::OnEnginePreExit.AddRaw(this, &FBulkDataRegistryImpl::OnEnginePreExit);
+}
+
+void FBulkDataRegistryImpl::OnEnginePreExit()
+{
+	Teardown();
 }
 
 FBulkDataRegistryImpl::~FBulkDataRegistryImpl()
 {
+	FCoreDelegates::OnEnginePreExit.RemoveAll(this);
+	Teardown();
+}
+
+void FBulkDataRegistryImpl::Teardown()
+{
+	{
+		FReadScopeLock RegistryScopeLock(RegistryLock);
+		if (!bActive)
+		{
+			return;
+		}
+	}
+
 	FCoreUObjectDelegates::OnEndLoadPackage.RemoveAll(this);
 	SetGlobalBuildInputResolver(nullptr);
 
@@ -192,7 +213,10 @@ FBulkDataRegistryImpl::TryRegister(UPackage* Owner, const UE::Serialization::FEd
 
 	{
 		FWriteScopeLock RegistryScopeLock(RegistryLock);
-		check(bActive); // Registrations should not come in after we destruct
+		if (!bActive)
+		{
+			return UE::BulkDataRegistry::ERegisterResult::Success;
+		}
 		UE::BulkDataRegistry::Private::FRegisteredBulk& RegisteredBulk = Registry.FindOrAdd(BulkData.GetIdentifier());
 		if (RegisteredBulk.bRegistered)
 		{
@@ -269,7 +293,10 @@ void FBulkDataRegistryImpl::UpdateRegistrationData(UPackage* Owner, const UE::Se
 	FName PackageName = Owner ? Owner->GetFName() : NAME_None;
 	{
 		FWriteScopeLock RegistryScopeLock(RegistryLock);
-		check(bActive); // Registrations should not come in after we destruct
+		if (!bActive)
+		{
+			return;
+		}
 		UE::BulkDataRegistry::Private::FRegisteredBulk& RegisteredBulk = Registry.FindOrAdd(BulkData.GetIdentifier());
 		// Update the Owner if we have it. This is important for Package Renames, which call UpdateRegistrationData when they save
 		// to the newly renamed package, with the renamed Owner package.
@@ -291,7 +318,10 @@ void FBulkDataRegistryImpl::Unregister(const UE::Serialization::FEditorBulkData&
 {
 	const FGuid& Key = BulkData.GetIdentifier();
 	FWriteScopeLock RegistryScopeLock(RegistryLock);
-	check(bActive); // Deregistrations should not come in after we destruct
+	if (!bActive)
+	{
+		return;
+	}
 	Registry.Remove(Key);
 }
 
@@ -299,7 +329,10 @@ void FBulkDataRegistryImpl::OnExitMemory(const UE::Serialization::FEditorBulkDat
 {
 	const FGuid& Key = BulkData.GetIdentifier();
 	FWriteScopeLock RegistryScopeLock(RegistryLock);
-	check(bActive); // Deregistrations should not come in after we destruct
+	if (!bActive)
+	{
+		return;
+	}
 	FRegisteredBulk* Existing = Registry.Find(Key);
 	if (Existing)
 	{
@@ -528,10 +561,13 @@ uint64 FBulkDataRegistryImpl::GetBulkDataResaveSize(FName PackageName)
 void FBulkDataRegistryImpl::TickCook(float DeltaTime, bool bTickComplete)
 {
 	bool bWaitForCooldown = !bTickComplete;
-	check(bActive); // Ticks should not come in after we destruct
 
 	{
 		FWriteScopeLock RegistryScopeLock(RegistryLock);
+		if (!bActive)
+		{
+			return;
+		}
 		PruneTempLoadedPayloads();
 	}
 }
@@ -544,7 +580,7 @@ void FBulkDataRegistryImpl::Tick(float DeltaTime)
 void FBulkDataRegistryImpl::AddPendingPackageBulkData(FName PackageName, UE::Serialization::FEditorBulkData&& BulkData)
 {
 	FScopeLock PendingPackageScopeLock(&PendingPackageLock);
-	check(bActive); // Registrations should not come in after we destruct, and AsyncTasks should check bActive before calling
+	check(bActive); // Callers should early exit if !bActive
 	TUniquePtr<FPendingPackage>& PendingPackage = PendingPackages.FindOrAdd(PackageName);
 	if (!PendingPackage.IsValid())
 	{
@@ -788,7 +824,8 @@ void FPendingPackage::OnBulkDataListResults(FSharedBuffer Buffer)
 		}
 	}
 
-	ReadCache();
+	bool bAbort;
+	ReadCache(bAbort);
 
 	if (PendingOperations.fetch_and(~Flag_BulkDataListResults) == Flag_BulkDataListResults)
 	{
@@ -802,7 +839,10 @@ void FPendingPackage::OnBulkDataListResults(FSharedBuffer Buffer)
 			// ThisPointer will be null, and Owner holds the UniquePtr to *this.
 			// It will release that UniquePtr only after Cancel returns, which will be after this function returns.
 		}
-		WriteCache();
+		if (!bAbort)
+		{
+			WriteCache();
+		}
 
 		// Deleting *this will destruct BulkDataListCacheRequest, which by
 		// default calls Cancel, but Cancel will block on this callback we are
@@ -813,10 +853,13 @@ void FPendingPackage::OnBulkDataListResults(FSharedBuffer Buffer)
 	}
 }
 
-void FPendingPackage::ReadCache()
+void FPendingPackage::ReadCache(bool& bOutAbort)
 {
+	bOutAbort = false;
 	if (CachedBulkDatas.Num() == 0)
 	{
+		FReadScopeLock RegistryScopeLock(Owner->RegistryLock);
+		bOutAbort = !Owner->bActive;
 		return;
 	}
 
@@ -831,6 +874,7 @@ void FPendingPackage::ReadCache()
 		FWriteScopeLock RegistryScopeLock(Owner->RegistryLock);
 		if (!Owner->bActive)
 		{
+			bOutAbort = true;
 			return;
 		}
 		for (const UE::Serialization::FEditorBulkData& BulkData : CachedBulkDatas)
diff --git a/Engine/Source/Editor/UnrealEd/Private/EditorDomain/BulkDataRegistryImpl.h b/Engine/Source/Editor/UnrealEd/Private/EditorDomain/BulkDataRegistryImpl.h
--- a/Engine/Source/Editor/UnrealEd/Private/EditorDomain/BulkDataRegistryImpl.h
+++ b/Engine/Source/Editor/UnrealEd/Private/EditorDomain/BulkDataRegistryImpl.h
@@ -104,7 +104,7 @@ public:
 
 private:
 	void OnBulkDataListResults(FSharedBuffer Buffer);
-	void ReadCache();
+	void ReadCache(bool& bOutAbort);
 
 	enum EFlags : int32
 	{
@@ -235,6 +235,8 @@ private:
 	static void WritePayloadIdToCache(FName PackageName, const UE::Serialization::FEditorBulkData& BulkData);
 	void ReadPayloadIdsFromCache(FName PackageName, TArray<TRefCountPtr<FPendingPayloadId>>&& OldPendings,
 		TArray<TRefCountPtr<FPendingPayloadId>>&& NewPendings);
+	void OnEnginePreExit();
+	void Teardown();
 
 	friend class FPendingPackage;
 	friend class FUpdatePayloadWorker;
