Cooker files: Change some whitespace and variable names to match Core coding standard.
#rnx

diff --git a/Engine/Source/Editor/UnrealEd/Private/Cooker/CookTypes.cpp b/Engine/Source/Editor/UnrealEd/Private/Cooker/CookTypes.cpp
--- a/Engine/Source/Editor/UnrealEd/Private/Cooker/CookTypes.cpp
+++ b/Engine/Source/Editor/UnrealEd/Private/Cooker/CookTypes.cpp
@@ -173,25 +173,25 @@ bool IsCookIgnoreTimeouts()
 
 }
 
-FCbWriter& operator<<(FCbWriter& Writer, const UE::Cook::FInitializeConfigSettings& OutValue)
+FCbWriter& operator<<(FCbWriter& Writer, const UE::Cook::FInitializeConfigSettings& Value)
 {
 	Writer.BeginObject();
-	Writer << "OutputDirectoryOverride" << OutValue.OutputDirectoryOverride;
-	Writer << "MaxPrecacheShaderJobs" << OutValue.MaxPrecacheShaderJobs;
-	Writer << "MaxConcurrentShaderJobs" << OutValue.MaxConcurrentShaderJobs;
-	Writer << "PackagesPerGC" << OutValue.PackagesPerGC;
-	Writer << "MemoryExpectedFreedToSpreadRatio" << OutValue.MemoryExpectedFreedToSpreadRatio;
-	Writer << "IdleTimeToGC" << OutValue.IdleTimeToGC;
-	Writer << "MemoryMaxUsedVirtual" << OutValue.MemoryMaxUsedVirtual;
-	Writer << "MemoryMaxUsedPhysical" << OutValue.MemoryMaxUsedPhysical;
-	Writer << "MemoryMinFreeVirtual" << OutValue.MemoryMinFreeVirtual;
-	Writer << "MemoryMinFreePhysical" << OutValue.MemoryMinFreePhysical;
-	Writer << "MemoryTriggerGCAtPressureLevel" << static_cast<uint8>(OutValue.MemoryTriggerGCAtPressureLevel);
-	Writer << "MinFreeUObjectIndicesBeforeGC" << OutValue.MinFreeUObjectIndicesBeforeGC;
-	Writer << "MaxNumPackagesBeforePartialGC" << OutValue.MaxNumPackagesBeforePartialGC;
-	Writer << "ConfigSettingDenyList" << OutValue.ConfigSettingDenyList;
-	Writer << "MaxAsyncCacheForType" << OutValue.MaxAsyncCacheForType;
-	Writer << "bHybridIterativeDebug" << OutValue.bHybridIterativeDebug;
+	Writer << "OutputDirectoryOverride" << Value.OutputDirectoryOverride;
+	Writer << "MaxPrecacheShaderJobs" << Value.MaxPrecacheShaderJobs;
+	Writer << "MaxConcurrentShaderJobs" << Value.MaxConcurrentShaderJobs;
+	Writer << "PackagesPerGC" << Value.PackagesPerGC;
+	Writer << "MemoryExpectedFreedToSpreadRatio" << Value.MemoryExpectedFreedToSpreadRatio;
+	Writer << "IdleTimeToGC" << Value.IdleTimeToGC;
+	Writer << "MemoryMaxUsedVirtual" << Value.MemoryMaxUsedVirtual;
+	Writer << "MemoryMaxUsedPhysical" << Value.MemoryMaxUsedPhysical;
+	Writer << "MemoryMinFreeVirtual" << Value.MemoryMinFreeVirtual;
+	Writer << "MemoryMinFreePhysical" << Value.MemoryMinFreePhysical;
+	Writer << "MemoryTriggerGCAtPressureLevel" << static_cast<uint8>(Value.MemoryTriggerGCAtPressureLevel);
+	Writer << "MinFreeUObjectIndicesBeforeGC" << Value.MinFreeUObjectIndicesBeforeGC;
+	Writer << "MaxNumPackagesBeforePartialGC" << Value.MaxNumPackagesBeforePartialGC;
+	Writer << "ConfigSettingDenyList" << Value.ConfigSettingDenyList;
+	Writer << "MaxAsyncCacheForType" << Value.MaxAsyncCacheForType;
+	Writer << "bHybridIterativeDebug" << Value.bHybridIterativeDebug;
 	// Make sure new values are added to LoadFromCompactBinary and MoveOrCopy
 	Writer.EndObject();
 	return Writer;
@@ -298,15 +298,15 @@ FString LexToString(FPlatformMemoryStats::EMemoryPressureStatus Value)
 	}
 }
 
-FCbWriter& operator<<(FCbWriter& Writer, const UE::Cook::FBeginCookConfigSettings& OutValue)
+FCbWriter& operator<<(FCbWriter& Writer, const UE::Cook::FBeginCookConfigSettings& Value)
 {
 	Writer.BeginObject();
-	Writer << "HybridIterativeEnabled" << OutValue.bHybridIterativeEnabled;
-	Writer << "CookShowInstigator" << OutValue.CookShowInstigator;
-	Writer << "NeverCookPackageList" << OutValue.NeverCookPackageList;
+	Writer << "HybridIterativeEnabled" << Value.bHybridIterativeEnabled;
+	Writer << "CookShowInstigator" << Value.CookShowInstigator;
+	Writer << "NeverCookPackageList" << Value.NeverCookPackageList;
 	
 	Writer.BeginArray("PlatformSpecificNeverCookPackages");
-	for (const TPair<const ITargetPlatform*, TSet<FName>>& Pair : OutValue.PlatformSpecificNeverCookPackages)
+	for (const TPair<const ITargetPlatform*, TSet<FName>>& Pair : Value.PlatformSpecificNeverCookPackages)
 	{
 		Writer.BeginObject();
 		Writer << "K" << Pair.Key->PlatformName();
@@ -361,29 +361,29 @@ bool LoadFromCompactBinary(FCbFieldView Field, UE::Cook::FBeginCookConfigSetting
 	return bOk;
 }
 
-FCbWriter& operator<<(FCbWriter& Writer, const UE::Cook::FCookByTheBookOptions& OutValue)
+FCbWriter& operator<<(FCbWriter& Writer, const UE::Cook::FCookByTheBookOptions& Value)
 {
 	Writer.BeginObject();
 	// StartupPackages and SessionStartupObjects are process-specific
 
-	Writer << "DlcName" << OutValue.DlcName;
-	Writer << "CreateReleaseVersion" << OutValue.CreateReleaseVersion;
-	Writer << "BasedOnReleaseCookedPackages" << OutValue.BasedOnReleaseCookedPackages;
-	Writer << "SourceToLocalizedPackageVariants" << OutValue.SourceToLocalizedPackageVariants;
-	Writer << "AllCulturesToCook" << OutValue.AllCulturesToCook;
+	Writer << "DlcName" << Value.DlcName;
+	Writer << "CreateReleaseVersion" << Value.CreateReleaseVersion;
+	Writer << "BasedOnReleaseCookedPackages" << Value.BasedOnReleaseCookedPackages;
+	Writer << "SourceToLocalizedPackageVariants" << Value.SourceToLocalizedPackageVariants;
+	Writer << "AllCulturesToCook" << Value.AllCulturesToCook;
 
 	// CookTime is process-specific
 	// CookStartTime is process-specific
 
-	Writer << "StartupOptions" << static_cast<int32>(OutValue.StartupOptions);
-	Writer << "GenerateStreamingInstallManifests" << OutValue.bGenerateStreamingInstallManifests;
-	Writer << "ErrorOnEngineContentUse" << OutValue.bErrorOnEngineContentUse;
-	Writer << "AllowUncookedAssetReferences" << OutValue.bAllowUncookedAssetReferences;
-	Writer << "SkipHardReferences" << OutValue.bSkipHardReferences;
-	Writer << "SkipSoftReferences" << OutValue.bSkipSoftReferences;
-	Writer << "FullLoadAndSave" << OutValue.bFullLoadAndSave;
-	Writer << "CookAgainstFixedBase" << OutValue.bCookAgainstFixedBase;
-	Writer << "DlcLoadMainAssetRegistry" << OutValue.bDlcLoadMainAssetRegistry;
+	Writer << "StartupOptions" << static_cast<int32>(Value.StartupOptions);
+	Writer << "GenerateStreamingInstallManifests" << Value.bGenerateStreamingInstallManifests;
+	Writer << "ErrorOnEngineContentUse" << Value.bErrorOnEngineContentUse;
+	Writer << "AllowUncookedAssetReferences" << Value.bAllowUncookedAssetReferences;
+	Writer << "SkipHardReferences" << Value.bSkipHardReferences;
+	Writer << "SkipSoftReferences" << Value.bSkipSoftReferences;
+	Writer << "FullLoadAndSave" << Value.bFullLoadAndSave;
+	Writer << "CookAgainstFixedBase" << Value.bCookAgainstFixedBase;
+	Writer << "DlcLoadMainAssetRegistry" << Value.bDlcLoadMainAssetRegistry;
 	Writer.EndObject();
 	return Writer;
 }
@@ -418,11 +418,11 @@ bool LoadFromCompactBinary(FCbFieldView Field, UE::Cook::FCookByTheBookOptions&
 	return bOk;
 }
 
-FCbWriter& operator<<(FCbWriter& Writer, const UE::Cook::FCookOnTheFlyOptions& OutValue)
+FCbWriter& operator<<(FCbWriter& Writer, const UE::Cook::FCookOnTheFlyOptions& Value)
 {
 	Writer.BeginObject();
-	Writer << "BindAnyPort" << OutValue.bBindAnyPort;
-	Writer << "PlatformProtocol" << OutValue.bPlatformProtocol;
+	Writer << "BindAnyPort" << Value.bBindAnyPort;
+	Writer << "PlatformProtocol" << Value.bPlatformProtocol;
 	Writer.EndObject();
 	return Writer;
 }
@@ -450,19 +450,19 @@ FCbWriter& operator<<(FCbWriter& Writer, const FBeginCookContextForWorkerPlatfor
 	return Writer;
 }
 
-bool LoadFromCompactBinary(FCbFieldView Field, FBeginCookContextForWorkerPlatform& Value)
+bool LoadFromCompactBinary(FCbFieldView Field, FBeginCookContextForWorkerPlatform& OutValue)
 {
 	bool bOk = true;
 	FString PlatformName;
 	LoadFromCompactBinary(Field["Platform"], PlatformName);
-	Value.TargetPlatform = nullptr;
+	OutValue.TargetPlatform = nullptr;
 	if (!PlatformName.IsEmpty())
 	{
 		ITargetPlatformManagerModule& TPM = GetTargetPlatformManagerRef();
-		Value.TargetPlatform = TPM.FindTargetPlatform(*PlatformName);
-		bOk = (Value.TargetPlatform != nullptr) & bOk;
+		OutValue.TargetPlatform = TPM.FindTargetPlatform(*PlatformName);
+		bOk = (OutValue.TargetPlatform != nullptr) & bOk;
 	}
-	bOk = LoadFromCompactBinary(Field["FullBuild"], Value.bFullBuild) & bOk;
+	bOk = LoadFromCompactBinary(Field["FullBuild"], OutValue.bFullBuild) & bOk;
 	return bOk;
 }
 
@@ -481,7 +481,7 @@ FCbWriter& operator<<(FCbWriter& Writer, const FBeginCookContextForWorker& Value
 	return Writer;
 }
 
-bool LoadFromCompactBinary(FCbFieldView Field, FBeginCookContextForWorker& Value)
+bool LoadFromCompactBinary(FCbFieldView Field, FBeginCookContextForWorker& OutValue)
 {
-	return LoadFromCompactBinary(Field, Value.PlatformContexts);
+	return LoadFromCompactBinary(Field, OutValue.PlatformContexts);
 }
diff --git a/Engine/Source/Editor/UnrealEd/Private/Cooker/CookTypes.h b/Engine/Source/Editor/UnrealEd/Private/Cooker/CookTypes.h
--- a/Engine/Source/Editor/UnrealEd/Private/Cooker/CookTypes.h
+++ b/Engine/Source/Editor/UnrealEd/Private/Cooker/CookTypes.h
@@ -28,13 +28,13 @@ namespace UE::Cook { struct FCookOnTheFlyOptions; }
 namespace UE::Cook { struct FInitializeConfigSettings; }
 
 FCbWriter& operator<<(FCbWriter& Writer, const UE::Cook::FBeginCookConfigSettings& Value);
-bool LoadFromCompactBinary(FCbFieldView Field, UE::Cook::FBeginCookConfigSettings& Value);
+bool LoadFromCompactBinary(FCbFieldView Field, UE::Cook::FBeginCookConfigSettings& OutValue);
 FCbWriter& operator<<(FCbWriter& Writer, const UE::Cook::FCookByTheBookOptions& Value);
-bool LoadFromCompactBinary(FCbFieldView Field, UE::Cook::FCookByTheBookOptions& Value);
+bool LoadFromCompactBinary(FCbFieldView Field, UE::Cook::FCookByTheBookOptions& OutValue);
 FCbWriter& operator<<(FCbWriter& Writer, const UE::Cook::FCookOnTheFlyOptions& Value);
-bool LoadFromCompactBinary(FCbFieldView Field, UE::Cook::FCookOnTheFlyOptions& Value);
+bool LoadFromCompactBinary(FCbFieldView Field, UE::Cook::FCookOnTheFlyOptions& OutValue);
 FCbWriter& operator<<(FCbWriter& Writer, const UE::Cook::FInitializeConfigSettings& Value);
-bool LoadFromCompactBinary(FCbFieldView Field, UE::Cook::FInitializeConfigSettings& Value);
+bool LoadFromCompactBinary(FCbFieldView Field, UE::Cook::FInitializeConfigSettings& OutValue);
 
 #define COOK_CHECKSLOW_PACKAGEDATA 0
 #define DEBUG_COOKONTHEFLY 0
diff --git a/Engine/Source/Editor/UnrealEd/Private/Cooker/PackageTracker.cpp b/Engine/Source/Editor/UnrealEd/Private/Cooker/PackageTracker.cpp
--- a/Engine/Source/Editor/UnrealEd/Private/Cooker/PackageTracker.cpp
+++ b/Engine/Source/Editor/UnrealEd/Private/Cooker/PackageTracker.cpp
@@ -14,145 +14,146 @@ namespace UE::Cook
 {
 
 #if ENABLE_COOK_STATS
-	namespace Stats
+namespace Stats
+{
+
+// Stats tracked through FAutoRegisterCallback
+static uint32 NumInlineLoads = 0;
+static FCookStatsManager::FAutoRegisterCallback RegisterCookStats([](FCookStatsManager::AddStatFuncRef AddStat)
 	{
-		// Stats tracked through FAutoRegisterCallback
-		static uint32 NumInlineLoads = 0;
-		static FCookStatsManager::FAutoRegisterCallback RegisterCookStats([](FCookStatsManager::AddStatFuncRef AddStat)
-			{
-				AddStat(TEXT("Package.Load"), FCookStatsManager::CreateKeyValueArray(TEXT("NumInlineLoads"), NumInlineLoads));
-			});
-	}
+		AddStat(TEXT("Package.Load"), FCookStatsManager::CreateKeyValueArray(TEXT("NumInlineLoads"), NumInlineLoads));
+	});
+
+}
 #endif
 
-	void FThreadSafeUnsolicitedPackagesList::AddCookedPackage(const FFilePlatformRequest& PlatformRequest)
-	{
-		FScopeLock S(&SyncObject);
-		CookedPackages.Add(PlatformRequest);
-	}
+void FThreadSafeUnsolicitedPackagesList::AddCookedPackage(const FFilePlatformRequest& PlatformRequest)
+{
+	FScopeLock S(&SyncObject);
+	CookedPackages.Add(PlatformRequest);
+}
+
+void FThreadSafeUnsolicitedPackagesList::GetPackagesForPlatformAndRemove(const ITargetPlatform* Platform, TArray<FName>& PackageNames)
+{
+	FScopeLock _(&SyncObject);
 
-	void FThreadSafeUnsolicitedPackagesList::GetPackagesForPlatformAndRemove(const ITargetPlatform* Platform, TArray<FName>& PackageNames)
+	for (int I = CookedPackages.Num() - 1; I >= 0; --I)
 	{
-		FScopeLock _(&SyncObject);
+		FFilePlatformRequest& Request = CookedPackages[I];
 
-		for (int I = CookedPackages.Num() - 1; I >= 0; --I)
+		if (Request.GetPlatforms().Contains(Platform))
 		{
-			FFilePlatformRequest& Request = CookedPackages[I];
+			// remove the platform
+			Request.RemovePlatform(Platform);
+			PackageNames.Emplace(Request.GetFilename());
 
-			if (Request.GetPlatforms().Contains(Platform))
+			if (Request.GetPlatforms().Num() == 0)
 			{
-				// remove the platform
-				Request.RemovePlatform(Platform);
-				PackageNames.Emplace(Request.GetFilename());
-
-				if (Request.GetPlatforms().Num() == 0)
-				{
-					CookedPackages.RemoveAt(I);
-				}
+				CookedPackages.RemoveAt(I);
 			}
 		}
 	}
+}
 
-	void FThreadSafeUnsolicitedPackagesList::Empty()
-	{
-		FScopeLock _(&SyncObject);
-		CookedPackages.Empty();
-	}
-
+void FThreadSafeUnsolicitedPackagesList::Empty()
+{
+	FScopeLock _(&SyncObject);
+	CookedPackages.Empty();
+}
 
-	FPackageTracker::FPackageTracker(FPackageDatas& InPackageDatas)
-		:PackageDatas(InPackageDatas)
-	{
-		LLM_SCOPE_BYTAG(Cooker);
 
-		FWriteScopeLock ScopeLock(Lock);
-		for (TObjectIterator<UPackage> It; It; ++It)
-		{
-			UPackage* Package = *It;
+FPackageTracker::FPackageTracker(FPackageDatas& InPackageDatas)
+	:PackageDatas(InPackageDatas)
+{
+	LLM_SCOPE_BYTAG(Cooker);
 
-			if (Package->GetOuter() == nullptr)
-			{
-				LoadedPackages.Add(Package);
-			}
-		}
+	FWriteScopeLock ScopeLock(Lock);
+	for (TObjectIterator<UPackage> It; It; ++It)
+	{
+		UPackage* Package = *It;
 
-		NewPackages.Reserve(LoadedPackages.Num());
-		for (UPackage* Package : LoadedPackages)
+		if (Package->GetOuter() == nullptr)
 		{
-			NewPackages.Add(Package, FInstigator(EInstigator::StartupPackage));
+			LoadedPackages.Add(Package);
 		}
-
-		GUObjectArray.AddUObjectDeleteListener(this);
-		GUObjectArray.AddUObjectCreateListener(this);
 	}
 
-	FPackageTracker::~FPackageTracker()
+	NewPackages.Reserve(LoadedPackages.Num());
+	for (UPackage* Package : LoadedPackages)
 	{
-		GUObjectArray.RemoveUObjectDeleteListener(this);
-		GUObjectArray.RemoveUObjectCreateListener(this);
+		NewPackages.Add(Package, FInstigator(EInstigator::StartupPackage));
 	}
 
-	TMap<UPackage*, FInstigator> FPackageTracker::GetNewPackages()
-	{
-		FWriteScopeLock ScopeLock(Lock);
-		TMap<UPackage*, FInstigator> Result = MoveTemp(NewPackages);
-		NewPackages.Reset();
-		bHasBeenConsumed = true;
-		return Result;
-	}
+	GUObjectArray.AddUObjectDeleteListener(this);
+	GUObjectArray.AddUObjectCreateListener(this);
+}
 
-	bool FPackageTracker::HasBeenConsumed() const
-	{
-		return bHasBeenConsumed;
-	}
+FPackageTracker::~FPackageTracker()
+{
+	GUObjectArray.RemoveUObjectDeleteListener(this);
+	GUObjectArray.RemoveUObjectCreateListener(this);
+}
+
+TMap<UPackage*, FInstigator> FPackageTracker::GetNewPackages()
+{
+	FWriteScopeLock ScopeLock(Lock);
+	TMap<UPackage*, FInstigator> Result = MoveTemp(NewPackages);
+	NewPackages.Reset();
+	bHasBeenConsumed = true;
+	return Result;
+}
 
-	void FPackageTracker::NotifyUObjectCreated(const class UObjectBase* Object, int32 Index)
+bool FPackageTracker::HasBeenConsumed() const
+{
+	return bHasBeenConsumed;
+}
+
+void FPackageTracker::NotifyUObjectCreated(const class UObjectBase* Object, int32 Index)
+{
+	if (Object->GetClass() == UPackage::StaticClass())
 	{
-		if (Object->GetClass() == UPackage::StaticClass())
+		LLM_SCOPE_BYTAG(Cooker);
+		auto Package = const_cast<UPackage*>(static_cast<const UPackage*>(Object));
+
+		if (Package->GetOuter() == nullptr)
 		{
 			LLM_SCOPE_BYTAG(Cooker);
-			auto Package = const_cast<UPackage*>(static_cast<const UPackage*>(Object));
-
-			if (Package->GetOuter() == nullptr)
+			if (LoadingPackageData && Package->GetFName() != LoadingPackageData->GetPackageName())
 			{
-				LLM_SCOPE_BYTAG(Cooker);
-				if (LoadingPackageData && Package->GetFName() != LoadingPackageData->GetPackageName())
-				{
-					COOK_STAT(++Stats::NumInlineLoads);
-				}
-
-				FWriteScopeLock ScopeLock(Lock);
-				LoadedPackages.Add(Package);
-				NewPackages.Add(Package,
-						FInstigator(EInstigator::Unsolicited,
-							LoadingPackageData ? LoadingPackageData->GetPackageName() : NAME_None)
-					);
+				COOK_STAT(++Stats::NumInlineLoads);
 			}
-		}
-	}
-
-	void FPackageTracker::NotifyUObjectDeleted(const class UObjectBase* Object, int32 Index)
-	{
-		if (Object->GetClass() == UPackage::StaticClass())
-		{
-			UPackage* Package = const_cast<UPackage*>(static_cast<const UPackage*>(Object));
 
 			FWriteScopeLock ScopeLock(Lock);
-			LoadedPackages.Remove(Package);
-			NewPackages.Remove(Package);
+			LoadedPackages.Add(Package);
+			NewPackages.Add(Package,
+					FInstigator(EInstigator::Unsolicited,
+						LoadingPackageData ? LoadingPackageData->GetPackageName() : NAME_None)
+				);
 		}
 	}
+}
 
-	void FPackageTracker::OnUObjectArrayShutdown()
+void FPackageTracker::NotifyUObjectDeleted(const class UObjectBase* Object, int32 Index)
+{
+	if (Object->GetClass() == UPackage::StaticClass())
 	{
-		GUObjectArray.RemoveUObjectDeleteListener(this);
-		GUObjectArray.RemoveUObjectCreateListener(this);
-	}
+		UPackage* Package = const_cast<UPackage*>(static_cast<const UPackage*>(Object));
 
-	void FPackageTracker::RemapTargetPlatforms(const TMap<ITargetPlatform*, ITargetPlatform*>& Remap)
-	{
-		RemapMapKeys(PlatformSpecificNeverCookPackages, Remap);
+		FWriteScopeLock ScopeLock(Lock);
+		LoadedPackages.Remove(Package);
+		NewPackages.Remove(Package);
 	}
+}
 
+void FPackageTracker::OnUObjectArrayShutdown()
+{
+	GUObjectArray.RemoveUObjectDeleteListener(this);
+	GUObjectArray.RemoveUObjectCreateListener(this);
+}
 
+void FPackageTracker::RemapTargetPlatforms(const TMap<ITargetPlatform*, ITargetPlatform*>& Remap)
+{
+	RemapMapKeys(PlatformSpecificNeverCookPackages, Remap);
 }
+
+} // namespace UE::Cook
diff --git a/Engine/Source/Editor/UnrealEd/Private/Cooker/PackageTracker.h b/Engine/Source/Editor/UnrealEd/Private/Cooker/PackageTracker.h
--- a/Engine/Source/Editor/UnrealEd/Private/Cooker/PackageTracker.h
+++ b/Engine/Source/Editor/UnrealEd/Private/Cooker/PackageTracker.h
@@ -21,210 +21,209 @@ namespace UE::Cook { struct FInstigator; }
 namespace UE::Cook { struct FPackageData; }
 namespace UE::Cook { struct FRecompileShaderRequest; }
 
-namespace UE
+namespace UE::Cook
 {
-namespace Cook
-{
-	struct FPackageDatas;
 
-	template<typename Type>
-	struct FThreadSafeQueue
+struct FPackageDatas;
+
+template<typename Type>
+struct FThreadSafeQueue
+{
+private:
+	mutable FCriticalSection SynchronizationObject; // made this mutable so this class can have const functions and still be thread safe
+	TRingBuffer<Type> Items;
+public:
+	void Enqueue(const Type& Item)
 	{
-	private:
-		mutable FCriticalSection SynchronizationObject; // made this mutable so this class can have const functions and still be thread safe
-		TRingBuffer<Type> Items;
-	public:
-		void Enqueue(const Type& Item)
-		{
-			FScopeLock ScopeLock(&SynchronizationObject);
-			Items.Add(Item);
-		}
+		FScopeLock ScopeLock(&SynchronizationObject);
+		Items.Add(Item);
+	}
 		
-		void Enqueue(Type&& Item)
-		{
-			FScopeLock ScopeLock(&SynchronizationObject);
-			Items.Add(MoveTempIfPossible(Item));
-		}
+	void Enqueue(Type&& Item)
+	{
+		FScopeLock ScopeLock(&SynchronizationObject);
+		Items.Add(MoveTempIfPossible(Item));
+	}
 
-		void EnqueueUnique(const Type& Item)
+	void EnqueueUnique(const Type& Item)
+	{
+		FScopeLock ScopeLock(&SynchronizationObject);
+		for (const Type& Existing : Items)
 		{
-			FScopeLock ScopeLock(&SynchronizationObject);
-			for (const Type& Existing : Items)
+			if (Existing == Item)
 			{
-				if (Existing == Item)
-				{
-					return;
-				}
+				return;
 			}
-			Items.PushBack(Item);
 		}
+		Items.PushBack(Item);
+	}
 
-		bool Dequeue(Type* Result)
+	bool Dequeue(Type* Result)
+	{
+		FScopeLock ScopeLock(&SynchronizationObject);
+		if (Items.Num())
 		{
-			FScopeLock ScopeLock(&SynchronizationObject);
-			if (Items.Num())
-			{
-				*Result = Items.PopFrontValue();
-				return true;
-			}
-			return false;
+			*Result = Items.PopFrontValue();
+			return true;
 		}
+		return false;
+	}
 
-		void DequeueAll(TArray<Type>& Results)
-		{
-			FScopeLock ScopeLock(&SynchronizationObject);
-			Results.Reserve(Results.Num() + Items.Num());
-			while (!Items.IsEmpty())
-			{
-				Results.Add(Items.PopFrontValue());
-			}
-		}
-
-		bool HasItems() const
+	void DequeueAll(TArray<Type>& Results)
+	{
+		FScopeLock ScopeLock(&SynchronizationObject);
+		Results.Reserve(Results.Num() + Items.Num());
+		while (!Items.IsEmpty())
 		{
-			FScopeLock ScopeLock(&SynchronizationObject);
-			return Items.Num() > 0;
+			Results.Add(Items.PopFrontValue());
 		}
+	}
 
-		void Remove(const Type& Item)
-		{
-			FScopeLock ScopeLock(&SynchronizationObject);
-			Items.Remove(Item);
-		}
+	bool HasItems() const
+	{
+		FScopeLock ScopeLock(&SynchronizationObject);
+		return Items.Num() > 0;
+	}
 
-		void CopyItems(const TArray<Type>& InItems) const
-		{
-			FScopeLock ScopeLock(&SynchronizationObject);
-			Items.Empty(InItems.Num());
-			for (const Type& Item : InItems)
-			{
-				Items.PushBack(Item);
-			}
-		}
+	void Remove(const Type& Item)
+	{
+		FScopeLock ScopeLock(&SynchronizationObject);
+		Items.Remove(Item);
+	}
 
-		int Num() const
+	void CopyItems(const TArray<Type>& InItems) const
+	{
+		FScopeLock ScopeLock(&SynchronizationObject);
+		Items.Empty(InItems.Num());
+		for (const Type& Item : InItems)
 		{
-			FScopeLock ScopeLock(&SynchronizationObject);
-			return Items.Num();
+			Items.PushBack(Item);
 		}
+	}
 
-		void Empty()
-		{
-			FScopeLock ScopeLock(&SynchronizationObject);
-			Items.Empty();
-		}
-	};
+	int Num() const
+	{
+		FScopeLock ScopeLock(&SynchronizationObject);
+		return Items.Num();
+	}
 
-	/** Simple thread safe proxy for TSet<FName> */
-	template <typename T>
-	class FThreadSafeSet
+	void Empty()
+	{
+		FScopeLock ScopeLock(&SynchronizationObject);
+		Items.Empty();
+	}
+};
+
+/** Simple thread safe proxy for TSet<FName> */
+template <typename T>
+class FThreadSafeSet
+{
+	TSet<T> InnerSet;
+	FCriticalSection SetCritical;
+public:
+	void Add(T InValue)
+	{
+		FScopeLock SetLock(&SetCritical);
+		InnerSet.Add(InValue);
+	}
+	bool AddUnique(T InValue)
 	{
-		TSet<T> InnerSet;
-		FCriticalSection SetCritical;
-	public:
-		void Add(T InValue)
+		FScopeLock SetLock(&SetCritical);
+		if (!InnerSet.Contains(InValue))
 		{
-			FScopeLock SetLock(&SetCritical);
 			InnerSet.Add(InValue);
+			return true;
 		}
-		bool AddUnique(T InValue)
-		{
-			FScopeLock SetLock(&SetCritical);
-			if (!InnerSet.Contains(InValue))
-			{
-				InnerSet.Add(InValue);
-				return true;
-			}
-			return false;
-		}
-		bool Contains(T InValue)
-		{
-			FScopeLock SetLock(&SetCritical);
-			return InnerSet.Contains(InValue);
-		}
-		void Remove(T InValue)
-		{
-			FScopeLock SetLock(&SetCritical);
-			InnerSet.Remove(InValue);
-		}
-		void Empty()
-		{
-			FScopeLock SetLock(&SetCritical);
-			InnerSet.Empty();
-		}
-
-		void GetValues(TSet<T>& OutSet)
-		{
-			FScopeLock SetLock(&SetCritical);
-			OutSet.Append(InnerSet);
-		}
-	};
+		return false;
+	}
+	bool Contains(T InValue)
+	{
+		FScopeLock SetLock(&SetCritical);
+		return InnerSet.Contains(InValue);
+	}
+	void Remove(T InValue)
+	{
+		FScopeLock SetLock(&SetCritical);
+		InnerSet.Remove(InValue);
+	}
+	void Empty()
+	{
+		FScopeLock SetLock(&SetCritical);
+		InnerSet.Empty();
+	}
 
-	struct FThreadSafeUnsolicitedPackagesList
+	void GetValues(TSet<T>& OutSet)
 	{
-		void AddCookedPackage(const FFilePlatformRequest& PlatformRequest);
-		void GetPackagesForPlatformAndRemove(const ITargetPlatform* Platform, TArray<FName>& PackageNames);
-		void Empty();
+		FScopeLock SetLock(&SetCritical);
+		OutSet.Append(InnerSet);
+	}
+};
 
-	private:
-		FCriticalSection				SyncObject;
-		TArray<FFilePlatformRequest>	CookedPackages;
-	};
+struct FThreadSafeUnsolicitedPackagesList
+{
+	void AddCookedPackage(const FFilePlatformRequest& PlatformRequest);
+	void GetPackagesForPlatformAndRemove(const ITargetPlatform* Platform, TArray<FName>& PackageNames);
+	void Empty();
 
-	struct FPackageTracker : public FUObjectArray::FUObjectCreateListener, public FUObjectArray::FUObjectDeleteListener
-	{
-	public:
-		FPackageTracker(FPackageDatas& InPackageDatas);
-		~FPackageTracker();
+private:
+	FCriticalSection				SyncObject;
+	TArray<FFilePlatformRequest>	CookedPackages;
+};
 
-		/** Returns all packages that have been loaded since the last time GetNewPackages was called */
-		TMap<UPackage*, FInstigator> GetNewPackages();
+struct FPackageTracker : public FUObjectArray::FUObjectCreateListener, public FUObjectArray::FUObjectDeleteListener
+{
+public:
+	FPackageTracker(FPackageDatas& InPackageDatas);
+	~FPackageTracker();
 
-		/** Report whether GetNewPackages has ever been called */
-		bool HasBeenConsumed() const;
+	/** Returns all packages that have been loaded since the last time GetNewPackages was called */
+	TMap<UPackage*, FInstigator> GetNewPackages();
 
-		virtual void NotifyUObjectCreated(const class UObjectBase* Object, int32 Index) override;
-		virtual void NotifyUObjectDeleted(const class UObjectBase* Object, int32 Index) override;
-		virtual void OnUObjectArrayShutdown() override;
+	/** Report whether GetNewPackages has ever been called */
+	bool HasBeenConsumed() const;
 
-		/** Swap all ITargetPlatform* stored on this instance according to the mapping in @param Remap. */
-		void RemapTargetPlatforms(const TMap<ITargetPlatform*, ITargetPlatform*>& Remap);
+	virtual void NotifyUObjectCreated(const class UObjectBase* Object, int32 Index) override;
+	virtual void NotifyUObjectDeleted(const class UObjectBase* Object, int32 Index) override;
+	virtual void OnUObjectArrayShutdown() override;
 
-		/** The package currently being loaded at CookOnTheFlyServer's direct request. Used to determine which load dependencies were not preloaded. */
-		FPackageData* LoadingPackageData = nullptr;
+	/** Swap all ITargetPlatform* stored on this instance according to the mapping in @param Remap. */
+	void RemapTargetPlatforms(const TMap<ITargetPlatform*, ITargetPlatform*>& Remap);
 
-		FPackageDatas& PackageDatas;
+	/** The package currently being loaded at CookOnTheFlyServer's direct request. Used to determine which load dependencies were not preloaded. */
+	FPackageData* LoadingPackageData = nullptr;
 
-		FThreadSafeUnsolicitedPackagesList UnsolicitedCookedPackages;
-		FThreadSafeQueue<FRecompileShaderRequest> RecompileRequests;
+	FPackageDatas& PackageDatas;
 
-		/** Packages to never cook - entries are localpaths in FPaths::MakeStandardFilename format. */
-		FThreadSafeSet<FName> NeverCookPackageList;
-		FThreadSafeSet<FName> UncookedEditorOnlyPackages; // set of packages that have been rejected due to being referenced by editor-only properties
-		TFastPointerMap<const ITargetPlatform*, TSet<FName>> PlatformSpecificNeverCookPackages;
+	FThreadSafeUnsolicitedPackagesList UnsolicitedCookedPackages;
+	FThreadSafeQueue<FRecompileShaderRequest> RecompileRequests;
 
-		bool bHasBeenConsumed = false;
+	/** Packages to never cook - entries are localpaths in FPaths::MakeStandardFilename format. */
+	FThreadSafeSet<FName> NeverCookPackageList;
+	FThreadSafeSet<FName> UncookedEditorOnlyPackages; // set of packages that have been rejected due to being referenced by editor-only properties
+	TFastPointerMap<const ITargetPlatform*, TSet<FName>> PlatformSpecificNeverCookPackages;
 
-		// Thread-safe enumeration of loaded package. 
-		// A lock is held during enumeration, keep code simple and optimal so the lock is released as fast as possible.
-		template <typename FunctionType>
-		void ForEachLoadedPackage(FunctionType Function)
+	bool bHasBeenConsumed = false;
+
+	// Thread-safe enumeration of loaded package. 
+	// A lock is held during enumeration, keep code simple and optimal so the lock is released as fast as possible.
+	template <typename FunctionType>
+	void ForEachLoadedPackage(FunctionType Function)
+	{
+		FReadScopeLock ScopeLock(Lock);
+		for (UPackage* Package : LoadedPackages)
 		{
-			FReadScopeLock ScopeLock(Lock);
-			for (UPackage* Package : LoadedPackages)
-			{
-				Function(Package);
-			}
+			Function(Package);
 		}
-	private:
-		// Protects data for thread-safety
-		FRWLock Lock;
-
-		// This is a complete list of currently loaded UPackages
-		TFastPointerSet<UPackage*> LoadedPackages;
-
-		// This list contains the UPackages loaded since last call to GetNewPackages
-		TMap<UPackage*, FInstigator> NewPackages;
-	}; 
-}
-}
+	}
+private:
+	// Protects data for thread-safety
+	FRWLock Lock;
+
+	// This is a complete list of currently loaded UPackages
+	TFastPointerSet<UPackage*> LoadedPackages;
+
+	// This list contains the UPackages loaded since last call to GetNewPackages
+	TMap<UPackage*, FInstigator> NewPackages;
+};
+
+} // namespace UE::Cook
