Cooker files: Change some whitespace and variable names to match Core coding standard.
#rnx

diff --git a/Engine/Source/Editor/UnrealEd/Private/Cooker/PackageTracker.h b/Engine/Source/Editor/UnrealEd/Private/Cooker/PackageTracker.h
--- a/Engine/Source/Editor/UnrealEd/Private/Cooker/PackageTracker.h
+++ b/Engine/Source/Editor/UnrealEd/Private/Cooker/PackageTracker.h
@@ -189,6 +189,12 @@ public:
 	/** Swap all ITargetPlatform* stored on this instance according to the mapping in @param Remap. */
 	void RemapTargetPlatforms(const TMap<ITargetPlatform*, ITargetPlatform*>& Remap);
 
+	// This is a complete list of currently loaded UPackages
+	TFastPointerSet<UPackage*> LoadedPackages;
+
+	// This list contains the UPackages loaded since last call to GetNewPackages
+	TMap<UPackage*, FInstigator> NewPackages;
+
 	/** The package currently being loaded at CookOnTheFlyServer's direct request. Used to determine which load dependencies were not preloaded. */
 	FPackageData* LoadingPackageData = nullptr;
 
