Update LevelAnnotations when we have removed a level.

diff --git a/Engine/Source/Editor/UnrealEd/Private/PackageTools.cpp b/Engine/Source/Editor/UnrealEd/Private/PackageTools.cpp
--- a/Engine/Source/Editor/UnrealEd/Private/PackageTools.cpp
+++ b/Engine/Source/Editor/UnrealEd/Private/PackageTools.cpp
@@ -850,6 +850,7 @@ UPackageTools::UPackageTools(const FObjectInitializer& ObjectInitializer)
 							{
 								CurrentWorldPtr->RemoveFromWorld(Level);
 								StreamingLevel->RemoveLevelFromCollectionForReload();
+								ULevelStreaming::RemoveLevelAnnotation(Level);
 								RemovedStreamingLevels.Add(StreamingLevel);
 								break;
 							}
@@ -959,6 +960,9 @@ UPackageTools::UPackageTools(const FObjectInitializer& ObjectInitializer)
 				for (ULevelStreaming* StreamingLevel : RemovedStreamingLevels)
 				{
 					ULevel* NewLevel = StreamingLevel->GetLoadedLevel();
+					ULevelStreaming::LevelAnnotations.AddAnnotation(
+						NewLevel,
+						ULevelStreaming::FLevelAnnotation(StreamingLevel));
 					CurrentWorldPtr->AddToWorld(NewLevel, StreamingLevel->LevelTransform, false);
 					StreamingLevel->AddLevelToCollectionAfterReload();
 				}
