Fixed UWorldPartitionRuntimeVirtualTextureBuilder : Rebuild scene primitives after shader compilation completes at start of RVT builder commandlet.
This forces RVT mesh commands to rebuild with correct material resources.
Without this the mesh commands can be built and registerd before shader compilation, and the stale comands are used when rendering the RVT streaming mips.

diff --git a/Engine/Source/Editor/UnrealEd/Private/WorldPartition/WorldPartitionRuntimeVirtualTextureBuilder.cpp b/Engine/Source/Editor/UnrealEd/Private/WorldPartition/WorldPartitionRuntimeVirtualTextureBuilder.cpp
--- a/Engine/Source/Editor/UnrealEd/Private/WorldPartition/WorldPartitionRuntimeVirtualTextureBuilder.cpp
+++ b/Engine/Source/Editor/UnrealEd/Private/WorldPartition/WorldPartitionRuntimeVirtualTextureBuilder.cpp
@@ -12,6 +12,7 @@
 #include "AssetCompilingManager.h"
 #include "Modules/ModuleManager.h"
 #include "UObject/UObjectIterator.h"
+#include "ComponentRecreateRenderStateContext.h"
 
 DEFINE_LOG_CATEGORY_STATIC(LogWorldPartitionRuntimeVirtualTextureBuilder, All, All);
 
@@ -49,7 +50,12 @@ bool UWorldPartitionRuntimeVirtualTextureBuilder::RunInternal(UWorld* World, con
 	// Load required actors
 	FWorldPartitionHelpers::FForEachActorWithLoadingResult ForEachActorWithLoadingResult;
 	LoadRuntimeVirtualTextureActors(WorldPartition, ForEachActorWithLoadingResult);
-		
+	
+	// Recreate render state after shader compilation complete
+	{
+		FGlobalComponentRecreateRenderStateContext Context;
+	}
+
 	TSet<UPackage*> ModifiedPackages;
 	for (TObjectIterator<URuntimeVirtualTextureComponent> It(RF_ClassDefaultObject, false, EInternalObjectFlags::Garbage); It; ++It)
 	{
