AssetDefinition - Adding an overload for a function to avoid having to upgrade all codepaths that still use it.

diff --git a/Engine/Source/Editor/UnrealEd/Private/AnimationEditorUtils.cpp b/Engine/Source/Editor/UnrealEd/Private/AnimationEditorUtils.cpp
--- a/Engine/Source/Editor/UnrealEd/Private/AnimationEditorUtils.cpp
+++ b/Engine/Source/Editor/UnrealEd/Private/AnimationEditorUtils.cpp
@@ -464,17 +464,19 @@ namespace AnimationEditorUtils
 	void CreateAnimationAssets(const TArray<TSoftObjectPtr<UObject>>& SkeletonsOrSkeletalMeshes, TSubclassOf<UAnimationAsset> AssetClass, const FString& InPrefix, FAnimAssetCreated AssetCreated, UObject* NameBaseObject /*= nullptr*/, bool bDoNotShowNameDialog /*= false*/, bool bAllowReplaceExisting /*= false*/)
 	{
 		TArray<UObject*> ObjectsToSync;
-		for(auto SkelIt = SkeletonsOrSkeletalMeshes.CreateConstIterator(); SkelIt; ++SkelIt)
+		for (auto SkelIt = SkeletonsOrSkeletalMeshes.CreateConstIterator(); SkelIt; ++SkelIt)
 		{
+			UObject* SkeletonOrSkeletalMeshObject = SkelIt->LoadSynchronous();
+			
 			USkeletalMesh* SkeletalMesh = nullptr;
-			USkeleton* Skeleton = Cast<USkeleton>(SkelIt->LoadSynchronous());
+			USkeleton* Skeleton = Cast<USkeleton>(SkeletonOrSkeletalMeshObject);
 			if (Skeleton == nullptr)
 			{
-				SkeletalMesh = CastChecked<USkeletalMesh>(SkelIt->LoadSynchronous());
+				SkeletalMesh = CastChecked<USkeletalMesh>(SkeletonOrSkeletalMeshObject);
 				Skeleton = SkeletalMesh->GetSkeleton();
 			}
 
-			if(Skeleton)
+			if (Skeleton)
 			{
 				FString Name;
 				FString PackageName;
@@ -567,15 +569,31 @@ namespace AnimationEditorUtils
 
 	void CreateNewAnimBlueprint(TArray<TSoftObjectPtr<UObject>> SkeletonsOrSkeletalMeshes, FAnimAssetCreated AssetCreated, bool bInContentBrowser)
 	{
+		TArray<TWeakObjectPtr<UObject>> SkeletonsOrSkeletalMeshesLoaded;
+		for (TSoftObjectPtr<UObject>& SkeletonsOrSkeletalMesh : SkeletonsOrSkeletalMeshes)
+		{
+			if (UObject* SkeletonOrSkeletalMeshObject = SkeletonsOrSkeletalMesh.LoadSynchronous())
+			{
+				SkeletonsOrSkeletalMeshesLoaded.Add(SkeletonOrSkeletalMeshObject);
+			}
+		}
+
+		CreateNewAnimBlueprint(SkeletonsOrSkeletalMeshesLoaded, AssetCreated, bInContentBrowser);
+	}
+	
+	void CreateNewAnimBlueprint(TArray<TWeakObjectPtr<UObject>> SkeletonsOrSkeletalMeshes, FAnimAssetCreated AssetCreated, bool bInContentBrowser)
+    {
 		const FString DefaultSuffix = TEXT("_AnimBlueprint");
 
 		if (SkeletonsOrSkeletalMeshes.Num() == 1)
 		{
+			UObject* SkeletonOrSkeletalMeshObject = SkeletonsOrSkeletalMeshes[0].Get();
+			
 			USkeletalMesh* SkeletalMesh = nullptr;
-			USkeleton* Skeleton = Cast<USkeleton>(SkeletonsOrSkeletalMeshes[0].LoadSynchronous());
+			USkeleton* Skeleton = Cast<USkeleton>(SkeletonOrSkeletalMeshObject);
 			if (Skeleton == nullptr)
 			{
-				SkeletalMesh = CastChecked<USkeletalMesh>(SkeletonsOrSkeletalMeshes[0].LoadSynchronous());
+				SkeletalMesh = CastChecked<USkeletalMesh>(SkeletonOrSkeletalMeshObject);
 				Skeleton = SkeletalMesh->GetSkeleton();
 			}
 
@@ -625,10 +643,10 @@ namespace AnimationEditorUtils
 			for (auto ObjIt = SkeletonsOrSkeletalMeshes.CreateConstIterator(); ObjIt; ++ObjIt)
 			{
 				USkeletalMesh* SkeletalMesh = nullptr;
-				USkeleton* Skeleton = Cast<USkeleton>(ObjIt->LoadSynchronous());
+				USkeleton* Skeleton = Cast<USkeleton>(ObjIt->Get());
 				if (Skeleton == nullptr)
 				{
-					SkeletalMesh = CastChecked<USkeletalMesh>(ObjIt->LoadSynchronous());
+					SkeletalMesh = CastChecked<USkeletalMesh>(ObjIt->Get());
 					Skeleton = SkeletalMesh->GetSkeleton();
 				}
 
diff --git a/Engine/Source/Editor/UnrealEd/Public/AnimationEditorUtils.h b/Engine/Source/Editor/UnrealEd/Public/AnimationEditorUtils.h
--- a/Engine/Source/Editor/UnrealEd/Public/AnimationEditorUtils.h
+++ b/Engine/Source/Editor/UnrealEd/Public/AnimationEditorUtils.h
@@ -141,6 +141,7 @@ namespace AnimationEditorUtils
 
 	UNREALED_API void CreateAnimationAssets(const TArray<TSoftObjectPtr<UObject>>& SkeletonsOrSkeletalMeshes, TSubclassOf<UAnimationAsset> AssetClass, const FString& InPrefix, FAnimAssetCreated AssetCreated, UObject* NameBaseObject = nullptr, bool bDoNotShowNameDialog = false, bool bAllowReplaceExisting = false);
 	
+	UNREALED_API void CreateNewAnimBlueprint(TArray<TWeakObjectPtr<UObject>> SkeletonsOrSkeletalMeshes, FAnimAssetCreated AssetCreated, bool bInContentBrowser);
 	UNREALED_API void CreateNewAnimBlueprint(TArray<TSoftObjectPtr<UObject>> SkeletonsOrSkeletalMeshes, FAnimAssetCreated AssetCreated, bool bInContentBrowser);
 	UNREALED_API void FillCreateAssetMenu(FMenuBuilder& MenuBuilder, const TArray<TSoftObjectPtr<UObject>>& SkeletonsOrSkeletalMeshes, FAnimAssetCreated AssetCreated, bool bInContentBrowser=true);
 	UNREALED_API void CreateUniqueAssetName(const FString& InBasePackageName, const FString& InSuffix, FString& OutPackageName, FString& OutAssetName);
