Address the case when PreviewPlatformName is NAME_None in MaterialShaderQualitySettings as this is valid

diff --git a/Engine/Source/Editor/UnrealEd/Private/SEditorViewport.cpp b/Engine/Source/Editor/UnrealEd/Private/SEditorViewport.cpp
--- a/Engine/Source/Editor/UnrealEd/Private/SEditorViewport.cpp
+++ b/Engine/Source/Editor/UnrealEd/Private/SEditorViewport.cpp
@@ -802,7 +802,7 @@ EShaderPlatform SEditorViewport::GetShaderPlatformHelper(const ERHIFeatureLevel:
 	UMaterialShaderQualitySettings* MaterialShaderQualitySettings = UMaterialShaderQualitySettings::Get();
 	const FName& PreviewPlatform = MaterialShaderQualitySettings->GetPreviewPlatform();
 
-	EShaderPlatform ShaderPlatform = ShaderFormatToLegacyShaderPlatform(PreviewPlatform);
+	EShaderPlatform ShaderPlatform = PreviewPlatform != NAME_None ? ShaderFormatToLegacyShaderPlatform(PreviewPlatform) : SP_NumPlatforms;
 	if (ShaderPlatform == SP_NumPlatforms)
 	{
 		ShaderPlatform = GetFeatureLevelShaderPlatform(FeatureLevel);
