CookOnTheFlyServer: Remove deprecated API elements from 4.26 and earlier.
#rnx

diff --git a/Engine/Config/BaseEditor.ini b/Engine/Config/BaseEditor.ini
--- a/Engine/Config/BaseEditor.ini
+++ b/Engine/Config/BaseEditor.ini
@@ -340,9 +340,6 @@ MaxConcurrentShaderJobs=8192
 MaxDiffsToLog=5
 IgnoreHeaderDiffs=True
 
-bUseFullGCAssetClassNames=true
-+FullGCAssetClassNames=/Script/Engine.World
-
 ; CookProcessCount sets whether we cook singleprocess or multiprocess and how many processes there should be.
 ; CookProcessCount=(1 or less) is singleprocess. CookProcessCount=(N>1) is 1 director and N-1 cookworkers.
 CookProcessCount=1
diff --git a/Engine/Source/Editor/UnrealEd/Classes/Commandlets/CookCommandlet.h b/Engine/Source/Editor/UnrealEd/Classes/Commandlets/CookCommandlet.h
--- a/Engine/Source/Editor/UnrealEd/Classes/Commandlets/CookCommandlet.h
+++ b/Engine/Source/Editor/UnrealEd/Classes/Commandlets/CookCommandlet.h
@@ -26,8 +26,9 @@ class UCookCommandlet
 {
 	GENERATED_UCLASS_BODY()
 
+protected:
 	/** List of asset types that will force GC after loading them during cook */
-	UPROPERTY(config)
+	UE_DEPRECATED(5.2, "No longer used")
 	TArray<FString> FullGCAssetClassNames;
 
 	/** If true, iterative cooking is being done */
@@ -48,8 +49,6 @@ class UCookCommandlet
 	bool bGenerateStreamingInstallManifests;
 	/** Error if we access engine content (useful for dlc) */
 	bool bErrorOnEngineContentUse;
-	/** Use historical serialization system for generating package dependencies (use for historical reasons only this method has been depricated, only affects cooked manifests) */
-	bool bUseSerializationForGeneratingPackageDependencies;
 	/** Only cook packages specified on commandline options (for debugging)*/
 	bool bCookSinglePackage;
 	/** Modification to bCookSinglePackage - cook transitive hard references in addition to the packages on the commandline */
diff --git a/Engine/Source/Editor/UnrealEd/Classes/CookOnTheSide/CookOnTheFlyServer.h b/Engine/Source/Editor/UnrealEd/Classes/CookOnTheSide/CookOnTheFlyServer.h
--- a/Engine/Source/Editor/UnrealEd/Classes/CookOnTheSide/CookOnTheFlyServer.h
+++ b/Engine/Source/Editor/UnrealEd/Classes/CookOnTheSide/CookOnTheFlyServer.h
@@ -51,7 +51,7 @@ enum class ECookInitializationFlags
 	AsyncSave =									0x00000020, // save packages async
 	// unused =									0x00000040,
 	IncludeServerMaps =							0x00000080, // should we include the server maps when cooking
-	UseSerializationForPackageDependencies =	0x00000100, // should we use the serialization code path for generating package dependencies (old method will be deprecated)
+	UseSerializationForPackageDependencies UE_DEPRECATED(4.26, "No longer used.") = 0x00000100,
 	BuildDDCInBackground =						0x00000200, // build ddc content in background while the editor is running (only valid for modes which are in editor IsCookingInEditor())
 	// unused =									0x00000400,
 	OutputVerboseCookerWarnings =				0x00000800, // output additional cooker warnings about content issues
@@ -90,7 +90,6 @@ enum class ECookByTheBookOptions
 
 	// Deprecated flags
 	NoSlatePackages UE_DEPRECATED(5.0, "The [UI]ContentDirectories is deprecated. You may use DirectoriesToAlwaysCook in your project settings instead.") = 0x00000400, // don't include slate content
-	DisableUnsolicitedPackages UE_DEPRECATED(4.26, "Use SkipSoftReferences and/or SkipHardReferences instead") = SkipSoftReferences | SkipHardReferences,
 };
 ENUM_CLASS_FLAGS(ECookByTheBookOptions);
 
@@ -663,9 +662,6 @@ public:
 	/** Is the local CookOnTheFlyServer in a cook session in CookByTheBook Mode? */
 	bool IsCookByTheBookRunning() const;
 
-	UE_DEPRECATED(4.26, "Unsolicited packages are now added directly to the savequeue and are not marked as unsolicited")
-	TArray<UPackage*> GetUnsolicitedPackages(const TArray<const ITargetPlatform*>& TargetPlatforms) const;
-
 	/** Execute class-specific special case cook postloads and reference discovery on a given package. */
 	void PostLoadPackageFixup(UE::Cook::FPackageData& PackageData, UPackage* Package);
 
@@ -703,10 +699,6 @@ public:
 	 */
 	void ClearCachedCookedPlatformDataForPlatform(const ITargetPlatform* TargetPlatform);
 
-	UE_DEPRECATED(4.25, "Use version that takes const ITargetPlatform* instead")
-	void ClearCachedCookedPlatformDataForPlatform(const FName& PlatformName);
-
-
 	/**
 	 * Clear all the previously cooked data for the platform passed in 
 	 * 
@@ -714,9 +706,6 @@ public:
 	 */
 	void ClearPlatformCookedData(const ITargetPlatform* TargetPlatform);
 
-	UE_DEPRECATED(4.25, "Use version that takes const ITargetPlatform* instead")
-	void ClearPlatformCookedData(const FString& PlatformName);
-
 	/**
 	 * Clear platforms' explored flags for all PackageDatas and optionally clear the cookresult flags.
 	 *
@@ -733,10 +722,6 @@ public:
 	 */
 	bool RecompileChangedShaders(const TArray<const ITargetPlatform*>& TargetPlatforms);
 
-	UE_DEPRECATED(4.25, "Use version that takes const ITargetPlatform* instead")
-	bool RecompileChangedShaders(const TArray<FName>& TargetPlatformNames);
-
-
 	/**
 	 * Force stop whatever pending cook requests are going on and clear all the cooked data
 	 * Note CookOnTheFly clients may not be able to recover from this if they are waiting on a cook request to complete
@@ -815,16 +800,6 @@ public:
 	/** Returns the configured amount of idle time before forcing a GC */
 	double GetIdleTimeToGC() const;
 
-	/** Returns the configured amount of memory allowed before forcing a GC */
-	UE_DEPRECATED(4.26, "Use HasExceededMaxMemory instead")
-	uint64 GetMaxMemoryAllowance() const;
-
-	/** Mark package as keep around for the cooker (don't GC) */
-	UE_DEPRECATED(4.25, "UCookOnTheFlyServer now uses FGCObject to interact with garbage collection")
-	void MarkGCPackagesToKeepForCooker()
-	{
-	}
-
 	bool HasExceededMaxMemory();
 	void EvaluateGarbageCollectionResults(bool bWasDueToOOM, bool bWasPartialGC,
 		int32 NumObjectsBeforeGC, const FPlatformMemoryStats& MemStatsBeforeGC,
@@ -842,10 +817,6 @@ public:
 	bool RequestPackage(const FName& StandardFileName, const TArrayView<const ITargetPlatform* const>& TargetPlatforms,
 		const bool bForceFrontOfQueue);
 
-	UE_DEPRECATED(4.25, "Use Version that takes TArray<const ITargetPlatform*> instead")
-	bool RequestPackage(const FName& StandardFileName, const TArrayView<const FName>& TargetPlatformNames,
-		const bool bForceFrontOfQueue);
-
 	/**
 	 * RequestPackage to be cooked
 	 * This function can only be called while the cooker is in cook by the book mode
diff --git a/Engine/Source/Editor/UnrealEd/Private/Commandlets/CookCommandlet.cpp b/Engine/Source/Editor/UnrealEd/Private/Commandlets/CookCommandlet.cpp
--- a/Engine/Source/Editor/UnrealEd/Private/Commandlets/CookCommandlet.cpp
+++ b/Engine/Source/Editor/UnrealEd/Private/Commandlets/CookCommandlet.cpp
@@ -175,7 +175,6 @@ int32 UCookCommandlet::Main(const FString& CmdLineParams)
 	bIterativeCooking = Switches.Contains(TEXT("ITERATE"));
 	bSkipEditorContent = Switches.Contains(TEXT("SKIPEDITORCONTENT")); // This won't save out any packages in Engine/Content/Editor*
 	bErrorOnEngineContentUse = Switches.Contains(TEXT("ERRORONENGINECONTENTUSE"));
-	bUseSerializationForGeneratingPackageDependencies = Switches.Contains(TEXT("UseSerializationForGeneratingPackageDependencies"));
 	bCookSinglePackage = Switches.Contains(TEXT("cooksinglepackagenorefs"));
 	bKeepSinglePackageRefs = Switches.Contains(TEXT("cooksinglepackage")); // This is a legacy parameter; it's a minor misnomer since singlepackage implies norefs, but we want to avoiding changing the behavior
 	bCookSinglePackage = bCookSinglePackage || bKeepSinglePackageRefs;
@@ -276,7 +275,6 @@ bool UCookCommandlet::CookByTheBook( const TArray<ITargetPlatform*>& Platforms)
 	ECookInitializationFlags CookFlags = ECookInitializationFlags::IncludeServerMaps;
 	CookFlags |= bIterativeCooking ? IterateFlags : ECookInitializationFlags::None;
 	CookFlags |= bSkipEditorContent ? ECookInitializationFlags::SkipEditorContent : ECookInitializationFlags::None;	
-	CookFlags |= bUseSerializationForGeneratingPackageDependencies ? ECookInitializationFlags::UseSerializationForPackageDependencies : ECookInitializationFlags::None;
 	CookFlags |= bUnversioned ? ECookInitializationFlags::Unversioned : ECookInitializationFlags::None;
 	CookFlags |= bCookEditorOptional ? ECookInitializationFlags::CookEditorOptional : ECookInitializationFlags::None;
 	CookFlags |= bVerboseCookerWarnings ? ECookInitializationFlags::OutputVerboseCookerWarnings : ECookInitializationFlags::None;
diff --git a/Engine/Source/Editor/UnrealEd/Private/CookOnTheFlyServer.cpp b/Engine/Source/Editor/UnrealEd/Private/CookOnTheFlyServer.cpp
--- a/Engine/Source/Editor/UnrealEd/Private/CookOnTheFlyServer.cpp
+++ b/Engine/Source/Editor/UnrealEd/Private/CookOnTheFlyServer.cpp
@@ -1243,21 +1243,6 @@ bool UCookOnTheFlyServer::RequestPackage(const FName& StandardFileName, const TA
 	return true;
 }
 
-bool UCookOnTheFlyServer::RequestPackage(const FName& StandardFileName, const TArrayView<const FName>& TargetPlatformNames, const bool bForceFrontOfQueue)
-{
-	TArray<const ITargetPlatform*> TargetPlatforms;
-	ITargetPlatformManagerModule& TPM = GetTargetPlatformManagerRef();
-	for (const FName& TargetPlatformName : TargetPlatformNames)
-	{
-		const ITargetPlatform* TargetPlatform = TPM.FindTargetPlatform(TargetPlatformName.ToString());
-		if (TargetPlatform)
-		{
-			TargetPlatforms.Add(TargetPlatform);
-		}
-	}
-	return RequestPackage(StandardFileName, TargetPlatforms, bForceFrontOfQueue);
-}
-
 bool UCookOnTheFlyServer::RequestPackage(const FName& StandardPackageFName, const bool bForceFrontOfQueue)
 {
 	check(!IsCookOnTheFlyMode()); // Invalid to call RequestPackage without a list of TargetPlatforms if we are in CookOnTheFly
@@ -4488,12 +4473,12 @@ bool UCookOnTheFlyServer::HasExceededMaxMemory()
 		// check validity of trigger :
 		// if the MaxUsed config exceeds the system memory, it can never be triggered and will prevent any GC :
 		uint64 MaxMaxUsed = FMath::Max(MemoryMaxUsedVirtual,MemoryMaxUsedPhysical);
-		if ( MaxMaxUsed >= MemStats.TotalPhysical )
+		if (MaxMaxUsed >= MemStats.TotalPhysical)
 		{
-			UE_CALL_ONCE( [&](){
+			UE_CALL_ONCE([&]() {
 				UE_LOG(LogCook, Warning, TEXT("Warning MemoryMaxUsed condition is larger than total memory (%dMiB >= %dMiB).  System does not have enough memory to cook this project."),
-			static_cast<uint32>(MaxMaxUsed / 1024 / 1024), static_cast<uint32>(MemStats.TotalPhysical / 1024 / 1024));			
-			} );
+				static_cast<uint32>(MaxMaxUsed / 1024 / 1024), static_cast<uint32>(MemStats.TotalPhysical / 1024 / 1024));
+				});
 		}
 
 		if (MemoryMaxUsedVirtual > 0 && MemStats.UsedVirtual >= MemoryMaxUsedVirtual)
@@ -4518,7 +4503,7 @@ bool UCookOnTheFlyServer::HasExceededMaxMemory()
 		{
 			UE_CALL_ONCE([&]() {
 				UE_LOG(LogCook, Error,
-					TEXT("MemoryPressureStatus is not available from the operating system. We may run out of memory due to lack of knowledge of when to collect garbage."));
+				TEXT("MemoryPressureStatus is not available from the operating system. We may run out of memory due to lack of knowledge of when to collect garbage."));
 				});
 		}
 		else
@@ -4797,12 +4782,6 @@ void UCookOnTheFlyServer::GetDirectAndTransitiveResourceSize(FResourceSizeEx& Ou
 	}
 }
 
-TArray<UPackage*> UCookOnTheFlyServer::GetUnsolicitedPackages(const TArray<const ITargetPlatform*>& TargetPlatforms) const
-{
-	// No longer supported
-	return TArray<UPackage*>();
-}
-
 void UCookOnTheFlyServer::OnObjectModified( UObject *ObjectMoving )
 {
 	if (IsGarbageCollecting())
@@ -4988,11 +4967,6 @@ double UCookOnTheFlyServer::GetIdleTimeToGC() const
 	}
 }
 
-uint64 UCookOnTheFlyServer::GetMaxMemoryAllowance() const
-{
-	return MemoryMaxUsedPhysical;
-}
-
 const TArray<FName>& UCookOnTheFlyServer::GetFullPackageDependencies(const FName& PackageName ) const
 {
 	TArray<FName>* PackageDependencies = CachedFullPackageDependencies.Find(PackageName);
@@ -8924,11 +8898,6 @@ void UCookOnTheFlyServer::ResetCook(TConstArrayView<TPair<const ITargetPlatform*
 	}
 }
 
-void UCookOnTheFlyServer::ClearPlatformCookedData(const FString& PlatformName)
-{
-	ClearPlatformCookedData(GetTargetPlatformManagerRef().FindTargetPlatform(PlatformName));
-}
-
 void UCookOnTheFlyServer::ClearCachedCookedPlatformDataForPlatform(const ITargetPlatform* TargetPlatform)
 {
 	if (TargetPlatform)
@@ -8940,13 +8909,6 @@ void UCookOnTheFlyServer::ClearCachedCookedPlatformDataForPlatform(const ITarget
 	}
 }
 
-void UCookOnTheFlyServer::ClearCachedCookedPlatformDataForPlatform(const FName& PlatformName)
-{
-	ITargetPlatformManagerModule& TPM = GetTargetPlatformManagerRef();
-	const ITargetPlatform* TargetPlatform = TPM.FindTargetPlatform(PlatformName.ToString());
-	return ClearCachedCookedPlatformDataForPlatform(TargetPlatform);
-}
-
 void UCookOnTheFlyServer::OnTargetPlatformChangedSupportedFormats(const ITargetPlatform* TargetPlatform)
 {
 	for (TObjectIterator<UObject> It; It; ++It)
@@ -10248,16 +10210,6 @@ bool UCookOnTheFlyServer::RecompileChangedShaders(const TArray<const ITargetPlat
 	return bShadersRecompiled;
 }
 
-bool UCookOnTheFlyServer::RecompileChangedShaders(const TArray<FName>& TargetPlatformNames)
-{
-	bool bShadersRecompiled = false;
-	for (const FName& TargetPlatformName : TargetPlatformNames)
-	{
-		bShadersRecompiled |= RecompileChangedShadersForPlatform(TargetPlatformName.ToString());
-	}
-	return bShadersRecompiled;
-}
-
 /* UCookOnTheFlyServer callbacks
  *****************************************************************************/
 
