Cooker - prevent crash when cooking with no project. We now always require a non-empty ShaderLibrary name.
#rnx

diff --git a/Engine/Source/Editor/UnrealEd/Private/CookOnTheFlyServer.cpp b/Engine/Source/Editor/UnrealEd/Private/CookOnTheFlyServer.cpp
--- a/Engine/Source/Editor/UnrealEd/Private/CookOnTheFlyServer.cpp
+++ b/Engine/Source/Editor/UnrealEd/Private/CookOnTheFlyServer.cpp
@@ -8094,7 +8094,19 @@ void UCookOnTheFlyServer::RegisterShaderChunkDataGenerator()
 
 FString UCookOnTheFlyServer::GetProjectShaderLibraryName() const
 {
-	return !IsCookingDLC() ? FApp::GetProjectName() : CookByTheBookOptions->DlcName;
+	if (!IsCookingDLC())
+	{
+		FString Result = FApp::GetProjectName();
+		if (Result.IsEmpty())
+		{
+			Result = TEXT("UnrealGame");
+		}
+		return Result;
+	}
+	else
+	{
+		return CookByTheBookOptions->DlcName;
+	}
 }
 
 static FString GenerateShaderCodeLibraryName(FString const& Name, bool bIsIterateSharedBuild)
