Fix an issue where pitch would never lock in the level editor when rotating the perspective camera

diff --git a/Engine/Source/Editor/UnrealEd/Private/LevelEditorViewport.cpp b/Engine/Source/Editor/UnrealEd/Private/LevelEditorViewport.cpp
--- a/Engine/Source/Editor/UnrealEd/Private/LevelEditorViewport.cpp
+++ b/Engine/Source/Editor/UnrealEd/Private/LevelEditorViewport.cpp
@@ -2137,12 +2137,12 @@ bool FLevelEditorViewportClient::ShouldLockPitch() const
 {
 	// If we have somehow gotten out of the locked rotation
 	if ((GetViewRotation().Pitch < -90.f + KINDA_SMALL_NUMBER || GetViewRotation().Pitch > 90.f - KINDA_SMALL_NUMBER)
-		|| FMath::Abs(GetViewRotation().Roll) > (90.f - KINDA_SMALL_NUMBER))
+		&& FMath::Abs(GetViewRotation().Roll) > (90.f - KINDA_SMALL_NUMBER))
 	{
 		return false;
 	}
 	// Else use the standard rules
-	return FEditorViewportClient::ShouldLockPitch() ;
+	return FEditorViewportClient::ShouldLockPitch();
 }
 
 void FLevelEditorViewportClient::BeginCameraMovement(bool bHasMovement)
