MPCook: Disallow retracting Generator packages in addition to Generated packages; retracting either type will cause the same error.
#rnx

diff --git a/Engine/Source/Editor/UnrealEd/Private/CookOnTheFlyServer.cpp b/Engine/Source/Editor/UnrealEd/Private/CookOnTheFlyServer.cpp
--- a/Engine/Source/Editor/UnrealEd/Private/CookOnTheFlyServer.cpp
+++ b/Engine/Source/Editor/UnrealEd/Private/CookOnTheFlyServer.cpp
@@ -9584,7 +9584,8 @@ void UCookOnTheFlyServer::GetPackagesToRetract(int32 NumToRetract, TArray<FName>
 			return true;
 		}
 
-		if (PackageData->GetWorkerAssignmentConstraint().IsValid() || PackageData->IsGenerated())
+		if (PackageData->GetWorkerAssignmentConstraint().IsValid() || PackageData->IsGenerated() ||
+			PackageData->GetGeneratorPackage())
 		{
 			// Don't send back Packages that are constrained to this worker. Doing so will just
 			// cause the CookDirector to send it back to us, and this can cause the cooker to crash
diff --git a/Engine/Source/Editor/UnrealEd/Private/Cooker/CookPackageData.cpp b/Engine/Source/Editor/UnrealEd/Private/Cooker/CookPackageData.cpp
--- a/Engine/Source/Editor/UnrealEd/Private/Cooker/CookPackageData.cpp
+++ b/Engine/Source/Editor/UnrealEd/Private/Cooker/CookPackageData.cpp
@@ -1662,9 +1662,10 @@ void FGeneratorPackage::ResetSaveState(FCookGenerationInfo& Info, UPackage* Pack
 		GetCookPackageSplitterInstance()->UseInternalReferenceToAvoidGarbageCollect() &&
 		(ReleaseSaveReason == EReleaseSaveReason::Demoted || ReleaseSaveReason == EReleaseSaveReason::RecreateObjectCache))
 	{
-		UE_LOG(LogCook, Error, TEXT("CookPackageSplitter failure: We are demoting a generated package from save and removing our references that keep its objects loaded.\n")
+		UE_LOG(LogCook, Error, TEXT("CookPackageSplitter failure: We are demoting a %s package from save and removing our references that keep its objects loaded.\n")
 			TEXT("This will allow the objects to be garbage collected and cause failures in the splitter which expects them to remain loaded.\n")
 			TEXT("Package=%s, Splitter=%s, ReleaseSaveReason=%s"),
+			Info.IsGenerator() ? TEXT("generator") : TEXT("generated"),
 			Info.PackageData ? *Info.PackageData->GetPackageName().ToString() : *Info.RelativePath,
 			*GetSplitDataObjectName().ToString(), LexToString(ReleaseSaveReason));
 	}
