Add support for multiple thumbnail render frequencies.
Preparation changing default thumbnail render frequency.

diff --git a/Engine/Plugins/Editor/ContentBrowser/ContentBrowserAssetDataSource/Source/ContentBrowserAssetDataSource/Private/ContentBrowserAssetDataSource.cpp b/Engine/Plugins/Editor/ContentBrowser/ContentBrowserAssetDataSource/Source/ContentBrowserAssetDataSource/Private/ContentBrowserAssetDataSource.cpp
--- a/Engine/Plugins/Editor/ContentBrowser/ContentBrowserAssetDataSource/Source/ContentBrowserAssetDataSource/Private/ContentBrowserAssetDataSource.cpp
+++ b/Engine/Plugins/Editor/ContentBrowser/ContentBrowserAssetDataSource/Source/ContentBrowserAssetDataSource/Private/ContentBrowserAssetDataSource.cpp
@@ -10,6 +10,7 @@
 #include "AssetViewUtils.h"
 #include "ContentBrowserItemPath.h"
 #include "UObject/GCObjectScopeGuard.h"
+#include "UObject/ObjectSaveContext.h"
 #include "Factories/Factory.h"
 #include "HAL/FileManager.h"
 #include "Editor.h"
@@ -65,6 +66,9 @@ void UContentBrowserAssetDataSource::Initialize(const bool InAutoRegister)
 
 	// Listen for when assets are loaded or changed
 	FCoreUObjectDelegates::OnObjectPropertyChanged.AddUObject(this, &UContentBrowserAssetDataSource::OnObjectPropertyChanged);
+	
+	// Listen for when assets are saved, listerns are notified in time despite presave because we queue updates for later processing 
+	FCoreUObjectDelegates::OnObjectPreSave.AddUObject(this, &UContentBrowserAssetDataSource::OnObjectPreSave);
 
 	// Listen for new mount roots
 	FPackageName::OnContentPathMounted().AddUObject(this, &UContentBrowserAssetDataSource::OnContentPathMounted);
@@ -190,6 +194,7 @@ void UContentBrowserAssetDataSource::Shutdown()
 	}
 
 	FCoreUObjectDelegates::OnObjectPropertyChanged.RemoveAll(this);
+	FCoreUObjectDelegates::OnObjectPreSave.RemoveAll(this);
 
 	AssetViewUtils::OnAlwaysShowPath().RemoveAll(this);
 
@@ -1795,6 +1800,15 @@ void UContentBrowserAssetDataSource::OnObjectPropertyChanged(UObject* InObject,
 	}
 }
 
+void UContentBrowserAssetDataSource::OnObjectPreSave(UObject* InObject, FObjectPreSaveContext InObjectPreSaveContext)
+{
+	if (InObject && InObject->IsAsset() && ContentBrowserAssetData::IsPrimaryAsset(InObject))
+	{
+		FAssetData AssetData(InObject);
+		QueueItemDataUpdate(FContentBrowserItemDataUpdate::MakeItemModifiedUpdate(CreateAssetFileItem(AssetData)));
+	}
+}
+
 void UContentBrowserAssetDataSource::OnPathAdded(const FString& InPath)
 {
 	// New paths are considered empty until assets are added inside them
diff --git a/Engine/Plugins/Editor/ContentBrowser/ContentBrowserAssetDataSource/Source/ContentBrowserAssetDataSource/Public/ContentBrowserAssetDataSource.h b/Engine/Plugins/Editor/ContentBrowser/ContentBrowserAssetDataSource/Source/ContentBrowserAssetDataSource/Public/ContentBrowserAssetDataSource.h
--- a/Engine/Plugins/Editor/ContentBrowser/ContentBrowserAssetDataSource/Source/ContentBrowserAssetDataSource/Public/ContentBrowserAssetDataSource.h
+++ b/Engine/Plugins/Editor/ContentBrowser/ContentBrowserAssetDataSource/Source/ContentBrowserAssetDataSource/Public/ContentBrowserAssetDataSource.h
@@ -296,6 +296,8 @@ private:
 
 	void OnObjectPropertyChanged(UObject* InObject, FPropertyChangedEvent& InPropertyChangedEvent);
 
+	void OnObjectPreSave(UObject* InObject, class FObjectPreSaveContext InObjectPreSaveContext);
+
 	void OnPathAdded(const FString& InPath);
 
 	void OnPathRemoved(const FString& InPath);
diff --git a/Engine/Plugins/Runtime/GeometryCache/Source/GeometryCacheEd/Classes/GeometryCacheThumbnailRenderer.h b/Engine/Plugins/Runtime/GeometryCache/Source/GeometryCacheEd/Classes/GeometryCacheThumbnailRenderer.h
--- a/Engine/Plugins/Runtime/GeometryCache/Source/GeometryCacheEd/Classes/GeometryCacheThumbnailRenderer.h
+++ b/Engine/Plugins/Runtime/GeometryCache/Source/GeometryCacheEd/Classes/GeometryCacheThumbnailRenderer.h
@@ -31,6 +31,6 @@ class UGeometryCacheThumbnailRenderer : public UDefaultSizedThumbnailRenderer
 private:
 	FGeometryCacheThumbnailScene* ThumbnailScene;
 public:
-	virtual bool AllowsRealtimeThumbnails(UObject* Object) const override;
+	virtual EThumbnailRenderFrequency GetThumbnailRenderFrequency(UObject* Object) const override;
 
 };
diff --git a/Engine/Plugins/Runtime/GeometryCache/Source/GeometryCacheEd/Private/GeometryCacheThumbnailRenderer.cpp b/Engine/Plugins/Runtime/GeometryCache/Source/GeometryCacheEd/Private/GeometryCacheThumbnailRenderer.cpp
--- a/Engine/Plugins/Runtime/GeometryCache/Source/GeometryCacheEd/Private/GeometryCacheThumbnailRenderer.cpp
+++ b/Engine/Plugins/Runtime/GeometryCache/Source/GeometryCacheEd/Private/GeometryCacheThumbnailRenderer.cpp
@@ -53,8 +53,8 @@ void UGeometryCacheThumbnailRenderer::BeginDestroy()
 	Super::BeginDestroy();
 }
 
-bool UGeometryCacheThumbnailRenderer::AllowsRealtimeThumbnails(UObject* Object) const
+EThumbnailRenderFrequency UGeometryCacheThumbnailRenderer::GetThumbnailRenderFrequency(UObject* Object) const
 {
-	return false;
+	return EThumbnailRenderFrequency::OnPropertyChange;
 }
 
diff --git a/Engine/Source/Editor/UnrealEd/Classes/ThumbnailRendering/SkeletalMeshThumbnailRenderer.h b/Engine/Source/Editor/UnrealEd/Classes/ThumbnailRendering/SkeletalMeshThumbnailRenderer.h
--- a/Engine/Source/Editor/UnrealEd/Classes/ThumbnailRendering/SkeletalMeshThumbnailRenderer.h
+++ b/Engine/Source/Editor/UnrealEd/Classes/ThumbnailRendering/SkeletalMeshThumbnailRenderer.h
@@ -24,7 +24,7 @@ class UNREALED_API USkeletalMeshThumbnailRenderer : public UDefaultSizedThumbnai
 
 	// Begin UThumbnailRenderer Object
 	virtual void Draw(UObject* Object, int32 X, int32 Y, uint32 Width, uint32 Height, FRenderTarget*, FCanvas* Canvas, bool bAdditionalViewFamily) override;
-	virtual bool AllowsRealtimeThumbnails(UObject* Object) const override;
+	virtual EThumbnailRenderFrequency GetThumbnailRenderFrequency(UObject* Object) const override;
 	virtual bool CanVisualizeAsset(UObject* Object) override;
 	// End UThumbnailRenderer Object
 
diff --git a/Engine/Source/Editor/UnrealEd/Classes/ThumbnailRendering/SkeletonThumbnailRenderer.h b/Engine/Source/Editor/UnrealEd/Classes/ThumbnailRendering/SkeletonThumbnailRenderer.h
--- a/Engine/Source/Editor/UnrealEd/Classes/ThumbnailRendering/SkeletonThumbnailRenderer.h
+++ b/Engine/Source/Editor/UnrealEd/Classes/ThumbnailRendering/SkeletonThumbnailRenderer.h
@@ -24,7 +24,7 @@ class UNREALED_API USkeletonThumbnailRenderer : public UDefaultSizedThumbnailRen
 
 	// Begin UThumbnailRenderer Object
 	virtual void Draw(UObject* Object, int32 X, int32 Y, uint32 Width, uint32 Height, FRenderTarget*, FCanvas* Canvas, bool bAdditionalViewFamily) override;
-	virtual bool AllowsRealtimeThumbnails(UObject* Object) const override;
+	virtual EThumbnailRenderFrequency GetThumbnailRenderFrequency(UObject* Object) const override;
 	// End UThumbnailRenderer Object
 
 	// UObject implementation
diff --git a/Engine/Source/Editor/UnrealEd/Classes/ThumbnailRendering/SoundWaveThumbnailRenderer.h b/Engine/Source/Editor/UnrealEd/Classes/ThumbnailRendering/SoundWaveThumbnailRenderer.h
--- a/Engine/Source/Editor/UnrealEd/Classes/ThumbnailRendering/SoundWaveThumbnailRenderer.h
+++ b/Engine/Source/Editor/UnrealEd/Classes/ThumbnailRendering/SoundWaveThumbnailRenderer.h
@@ -23,7 +23,7 @@ class USoundWaveThumbnailRenderer : public UDefaultSizedThumbnailRenderer
 	// Begin UThumbnailRenderer Object
 	virtual bool CanVisualizeAsset(UObject* Object);
 	virtual void Draw(UObject* Object, int32 X, int32 Y, uint32 Width, uint32 Height, FRenderTarget*, FCanvas* Canvas, bool bAdditionalViewFamily) override;
-	virtual bool AllowsRealtimeThumbnails(UObject* Object) const override;
+	virtual EThumbnailRenderFrequency GetThumbnailRenderFrequency(UObject* Object) const override;
 	// End UThumbnailRenderer Object
 };
 
diff --git a/Engine/Source/Editor/UnrealEd/Classes/ThumbnailRendering/ThumbnailRenderer.h b/Engine/Source/Editor/UnrealEd/Classes/ThumbnailRendering/ThumbnailRenderer.h
--- a/Engine/Source/Editor/UnrealEd/Classes/ThumbnailRendering/ThumbnailRenderer.h
+++ b/Engine/Source/Editor/UnrealEd/Classes/ThumbnailRendering/ThumbnailRenderer.h
@@ -1,12 +1,5 @@
 // Copyright Epic Games, Inc. All Rights Reserved.
 
-/**
- *
- * This is an abstract base class that is used to define the interface that
- * UnrealEd will use when rendering a given object's thumbnail. The editor
- * only calls the virtual rendering function.
- */
-
 #pragma once
 
 #include "CoreMinimal.h"
@@ -17,6 +10,24 @@
 class FCanvas;
 class FRenderTarget;
 
+/** Refresh frequency for thumbnail rendering, listed from most to least CPU demanding / frequent */
+enum class EThumbnailRenderFrequency : uint8
+{
+	/** Always render when requested, used for assets needing live animated thumbnails like materials */
+	Realtime,
+	/** Render whenever a property has changed on request */
+	OnPropertyChange,
+	/** Render only on asset save */
+	OnAssetSave,
+	/** Render thumbnail only once */
+	Once,
+};
+
+/**
+ * This is an abstract base class that is used to define the interface that
+ * UnrealEd will use when rendering a given object's thumbnail. The editor
+ * only calls the virtual rendering function.
+ */
 UCLASS(abstract, MinimalAPI)
 class UThumbnailRenderer : public UObject
 {
@@ -78,9 +89,21 @@ public:
 	 *
 	 * @return True if the thumbnail needs to always be redrawn, false if it can be just drawn once and then reused.
 	 */
+	UE_DEPRECATED(5.2, "Please override GetThumbnailRenderFrequency instead, AllowsRealtimeThumbnails is not used.")
 	virtual bool AllowsRealtimeThumbnails(UObject* Object) const { return true; }
 
-
+	/**
+	 * Override this method to control the render frequency for thumbnails in your editor
+	 * Generally speaking you should use Realtime if you want an animated thumbnail, but if not it is expensive to use (per-frame draws)
+	 * If not using realtime, the next best default is OnAssetSave, which updates the thumbnail only save.
+	 * Remaining types are not suggested unless you have specific asset thumbnail update needs.
+	 * 
+	 * @TODO: The current default of Realtime maintains legacy behavior, but is overally agressive / for frequent than needed for most.
+	 *
+	 * @return ThumbnailRenderer-specific frequency for thumbnail updates
+	 */
+	virtual EThumbnailRenderFrequency GetThumbnailRenderFrequency(UObject* Object) const { return EThumbnailRenderFrequency::Realtime; }
+	 
 protected:
 	/** Renders the thumbnail's view family. */
 	UNREALED_API static void RenderViewFamily(FCanvas* Canvas, class FSceneViewFamily* ViewFamily, class FSceneView* View);
diff --git a/Engine/Source/Editor/UnrealEd/Private/AssetThumbnail.cpp b/Engine/Source/Editor/UnrealEd/Private/AssetThumbnail.cpp
--- a/Engine/Source/Editor/UnrealEd/Private/AssetThumbnail.cpp
+++ b/Engine/Source/Editor/UnrealEd/Private/AssetThumbnail.cpp
@@ -1180,19 +1180,20 @@ bool FAssetThumbnailPool::LoadThumbnail(TSharedRef<FThumbnailInfo> ThumbnailInfo
 			if (RenderInfo != nullptr && RenderInfo->Renderer != nullptr)
 			{
 				FThumbnailInfo_RenderThread ThumbInfo = ThumbnailInfo.Get();
-				ENQUEUE_RENDER_COMMAND(SyncSlateTextureCommand)(
-					[ThumbInfo](FRHICommandListImmediate& RHICmdList)
-				{
-					if (ThumbInfo.ThumbnailTexture->GetTypedResource() != ThumbInfo.ThumbnailRenderTarget->GetTextureRHI())
-					{
-						ThumbInfo.ThumbnailTexture->ClearTextureData();
-						ThumbInfo.ThumbnailTexture->ReleaseDynamicRHI();
-						ThumbInfo.ThumbnailTexture->SetRHIRef(ThumbInfo.ThumbnailRenderTarget->GetTextureRHI(), ThumbInfo.Width, ThumbInfo.Height);
-					}
-				});
 
-				if (ThumbnailInfo->LastUpdateTime <= 0.0f || RenderInfo->Renderer->AllowsRealtimeThumbnails(Asset))
+				auto EnqueueThumbnailRender = [Asset, ThumbInfo, ThumbnailInfo]()
 				{
+					ENQUEUE_RENDER_COMMAND(SyncSlateTextureCommand)(
+						[ThumbInfo](FRHICommandListImmediate& RHICmdList)
+						{
+							if (ThumbInfo.ThumbnailTexture->GetTypedResource() != ThumbInfo.ThumbnailRenderTarget->GetTextureRHI())
+							{
+								ThumbInfo.ThumbnailTexture->ClearTextureData();
+								ThumbInfo.ThumbnailTexture->ReleaseDynamicRHI();
+								ThumbInfo.ThumbnailTexture->SetRHIRef(ThumbInfo.ThumbnailRenderTarget->GetTextureRHI(), ThumbInfo.Width, ThumbInfo.Height);
+							}
+						});
+
 					//@todo: this should be done on the GPU only but it is not supported by thumbnail tools yet
 					ThumbnailTools::RenderThumbnail(
 						Asset,
@@ -1201,9 +1202,41 @@ bool FAssetThumbnailPool::LoadThumbnail(TSharedRef<FThumbnailInfo> ThumbnailInfo
 						ThumbnailTools::EThumbnailTextureFlushMode::NeverFlush,
 						ThumbnailInfo->ThumbnailRenderTarget
 					);
-				}
+				};
 
-				return true;
+				EThumbnailRenderFrequency ThumbnailRenderFrequency = RenderInfo->Renderer->GetThumbnailRenderFrequency(Asset);
+
+				switch (ThumbnailRenderFrequency)
+				{
+					case EThumbnailRenderFrequency::Realtime:
+					{
+						EnqueueThumbnailRender();
+						return true;
+					}
+					case EThumbnailRenderFrequency::OnPropertyChange:
+					{
+						if (ThumbnailInfo->LastUpdateTime <= 0.0f)
+						{
+							EnqueueThumbnailRender();
+							return true;
+						}
+						break;
+					}
+					case EThumbnailRenderFrequency::OnAssetSave:
+					{
+						// OnAssetSave is default behavior below, so nohing to do
+						break;
+					}
+					case EThumbnailRenderFrequency::Once:
+					{
+						// Eagerly return if we aren't interested in cached thumbnails
+						return true;
+					}
+					default:
+					{
+						break;
+					}
+				}
 			}
 		}
 	}
diff --git a/Engine/Source/Editor/UnrealEd/Private/ThumbnailRendering/SkeletalMeshThumbnailRenderer.cpp b/Engine/Source/Editor/UnrealEd/Private/ThumbnailRendering/SkeletalMeshThumbnailRenderer.cpp
--- a/Engine/Source/Editor/UnrealEd/Private/ThumbnailRendering/SkeletalMeshThumbnailRenderer.cpp
+++ b/Engine/Source/Editor/UnrealEd/Private/ThumbnailRendering/SkeletalMeshThumbnailRenderer.cpp
@@ -35,14 +35,12 @@ void USkeletalMeshThumbnailRenderer::Draw(UObject* Object, int32 X, int32 Y, uin
 	ThumbnailScene->SetSkeletalMesh(nullptr);
 }
 
-bool USkeletalMeshThumbnailRenderer::AllowsRealtimeThumbnails(UObject* Object) const
+EThumbnailRenderFrequency USkeletalMeshThumbnailRenderer::GetThumbnailRenderFrequency(UObject* Object) const
 {
-	if (!Super::AllowsRealtimeThumbnails(Object))
-	{
-		return false;
-	}
 	USkeletalMesh* SkeletalMesh = Cast<USkeletalMesh>(Object);
-	return SkeletalMesh && SkeletalMesh->GetResourceForRendering();
+	return SkeletalMesh && SkeletalMesh->GetResourceForRendering() 
+		? EThumbnailRenderFrequency::Realtime
+		: EThumbnailRenderFrequency::OnPropertyChange;
 }
 
 bool USkeletalMeshThumbnailRenderer::CanVisualizeAsset(UObject* Object)
diff --git a/Engine/Source/Editor/UnrealEd/Private/ThumbnailRendering/SkeletonThumbnailRenderer.cpp b/Engine/Source/Editor/UnrealEd/Private/ThumbnailRendering/SkeletonThumbnailRenderer.cpp
--- a/Engine/Source/Editor/UnrealEd/Private/ThumbnailRendering/SkeletonThumbnailRenderer.cpp
+++ b/Engine/Source/Editor/UnrealEd/Private/ThumbnailRendering/SkeletonThumbnailRenderer.cpp
@@ -46,21 +46,19 @@ void USkeletonThumbnailRenderer::Draw(UObject* Object, int32 X, int32 Y, uint32
 	ThumbnailScene->SetSkeletalMesh(nullptr);
 }
 
-bool USkeletonThumbnailRenderer::AllowsRealtimeThumbnails(UObject* Object) const
+EThumbnailRenderFrequency USkeletonThumbnailRenderer::GetThumbnailRenderFrequency(UObject* Object) const
 {
-	if (!Super::AllowsRealtimeThumbnails(Object))
-	{
-		return false;
-	}
 	USkeleton* Skeleton = Cast<USkeleton>(Object);
 	if(Skeleton)
 	{
 		if(USkeletalMesh* SkeletalMesh = Skeleton->GetPreviewMesh())
 		{
-			return SkeletalMesh->GetResourceForRendering() != nullptr;
+			return SkeletalMesh->GetResourceForRendering() != nullptr 
+				? EThumbnailRenderFrequency::Realtime
+				: EThumbnailRenderFrequency::OnPropertyChange;
 		}
 	}
-	return false;
+	return EThumbnailRenderFrequency::OnPropertyChange;
 }
 
 void USkeletonThumbnailRenderer::BeginDestroy()
diff --git a/Engine/Source/Editor/UnrealEd/Private/ThumbnailRendering/SoundWaveThumbnailRenderer.cpp b/Engine/Source/Editor/UnrealEd/Private/ThumbnailRendering/SoundWaveThumbnailRenderer.cpp
--- a/Engine/Source/Editor/UnrealEd/Private/ThumbnailRendering/SoundWaveThumbnailRenderer.cpp
+++ b/Engine/Source/Editor/UnrealEd/Private/ThumbnailRendering/SoundWaveThumbnailRenderer.cpp
@@ -175,12 +175,14 @@ void USoundWaveThumbnailRenderer::Draw(UObject* Object, int32 X, int32 Y, uint32
 
 }
 
-bool USoundWaveThumbnailRenderer::AllowsRealtimeThumbnails(UObject* Object) const
+EThumbnailRenderFrequency USoundWaveThumbnailRenderer::GetThumbnailRenderFrequency(UObject* Object) const
 {
 	if (USoundWave* SoundWave = Cast<USoundWave>(Object))
 	{
-		return SoundWave->CanVisualizeAsset();
+		return SoundWave->CanVisualizeAsset()
+			? EThumbnailRenderFrequency::Realtime
+			: EThumbnailRenderFrequency::OnPropertyChange;
 	}
 
-	return false;
+	return EThumbnailRenderFrequency::OnPropertyChange;
 }
