CompileShadersTestBed commandlet update: added argument to exclude global, default and material shaders separately.

diff --git a/Engine/Source/Editor/UnrealEd/Private/Commandlets/CompileShadersTestBedCommandlet.cpp b/Engine/Source/Editor/UnrealEd/Private/Commandlets/CompileShadersTestBedCommandlet.cpp
--- a/Engine/Source/Editor/UnrealEd/Private/Commandlets/CompileShadersTestBedCommandlet.cpp
+++ b/Engine/Source/Editor/UnrealEd/Private/Commandlets/CompileShadersTestBedCommandlet.cpp
@@ -41,6 +41,9 @@ int32 UCompileShadersTestBedCommandlet::Main(const FString& Params)
 		UE_LOG(LogCompileShadersTestBedCommandlet, Log, TEXT(" Optional: -collection=<name>                (You can also specify a collection of assets to narrow down the results e.g. if you maintain a collection that represents the actually used in-game assets)."));
 		UE_LOG(LogCompileShadersTestBedCommandlet, Log, TEXT(" Optional: -materials=<path1>+<path2>        (You can also specify a list of material asset paths separated by a '+' to narrow down the results.)"));
 		UE_LOG(LogCompileShadersTestBedCommandlet, Log, TEXT(" Optional: -all                              (You can specify -all to compile all global/default shaders as well as shaders for all materials/material instances found in a project.)"))
+		UE_LOG(LogCompileShadersTestBedCommandlet, Log, TEXT(" Optional: -ExcludeGlobalShaders             (Skip the compilation of global shaders.)"))
+		UE_LOG(LogCompileShadersTestBedCommandlet, Log, TEXT(" Optional: -ExcludeDefaultMaterials          (Skip the compilation of default material shaders.)"))
+		UE_LOG(LogCompileShadersTestBedCommandlet, Log, TEXT(" Optional: -ExcludeMaterials                 (Skip the compilation of non default material shaders.)"))
 		return 0;
 	}
 
@@ -118,6 +121,7 @@ int32 UCompileShadersTestBedCommandlet::Main(const FString& Params)
 	for (ITargetPlatform* Platform : Platforms)
 	{
 		// Compile default materials
+		if (!Switches.Contains(TEXT("ExcludeDefaultMaterials")))
 		{
 			TRACE_CPUPROFILER_EVENT_SCOPE(DefaultMaterials);
 
@@ -137,6 +141,7 @@ int32 UCompileShadersTestBedCommandlet::Main(const FString& Params)
 		}
 
 		// Compile global shaders
+		if (!Switches.Contains(TEXT("ExcludeGlobalShaders")))
 		{
 			TRACE_CPUPROFILER_EVENT_SCOPE(GlobalShaders);
 
@@ -157,6 +162,7 @@ int32 UCompileShadersTestBedCommandlet::Main(const FString& Params)
 		TSet<UMaterialInterface*> MaterialsToCompile;
 
 		// Begin Material Compiles
+		if (!Switches.Contains(TEXT("ExcludeMaterials")))
 		{
 			TRACE_CPUPROFILER_EVENT_SCOPE(BeginCacheForCookedPlatformData);
 
