Fix custom thumbnails disappearing after GC

diff --git a/Engine/Source/Editor/UnrealEd/Private/AssetThumbnail.cpp b/Engine/Source/Editor/UnrealEd/Private/AssetThumbnail.cpp
--- a/Engine/Source/Editor/UnrealEd/Private/AssetThumbnail.cpp
+++ b/Engine/Source/Editor/UnrealEd/Private/AssetThumbnail.cpp
@@ -1127,15 +1127,27 @@ void FAssetThumbnailPool::Tick( float DeltaTime )
 					bool bLoadedThumbnail = LoadThumbnail(InfoRef, bIsAssetStillCompiling, CustomThumbnailAsset);
 
 					// If we failed to load a custom thumbnail, then load the custom thumbnail's asset and try again
-					if (!bLoadedThumbnail && CustomThumbnailAsset.IsValid())
+					if (!bLoadedThumbnail && !bIsAssetStillCompiling && CustomThumbnailAsset.IsValid())
 					{
-						// Only load the custom thumbnail if the original asset is also loaded
-						//if (InfoRef->AssetData.IsAssetLoaded())
+						if (UPackage* CustomThumbnailPackage = CustomThumbnailAsset.GetPackage())
 						{
-							if (UObject* CustomThumbnail = CustomThumbnailAsset.GetAsset())
+							UObject* CustomThumbnail = FindObjectFast<UObject>(CustomThumbnailPackage, CustomThumbnailAsset.AssetName);
+							if (!CustomThumbnail)
+							{
+								// Because the custom thumbnail asset can be GCed (RF_Standalone flag cleared), 
+								// its package might need to be reloaded
+								CustomThumbnailPackage = LoadPackage(NULL, *CustomThumbnailAsset.PackageName.ToString(), LOAD_None);
+								CustomThumbnail = FindObjectFast<UObject>(CustomThumbnailPackage, CustomThumbnailAsset.AssetName);
+							}
+
+							if (CustomThumbnail)
 							{
 								bLoadedThumbnail = LoadThumbnail(InfoRef, bIsAssetStillCompiling, CustomThumbnailAsset);
-								CustomThumbnail->ClearFlags(RF_Standalone);
+								if (!bIsAssetStillCompiling)
+								{
+									// Clear RF_Standalone flag on the loaded custom thumbnail asset so it gets GCed eventually
+									CustomThumbnail->ClearFlags(RF_Standalone);
+								}
 							}
 						}
 					}
@@ -1150,7 +1162,7 @@ void FAssetThumbnailPool::Tick( float DeltaTime )
 					}
 					// Do not send a failure event for this asset yet if shaders are still compiling or the asset itself is compiling.
 					// The failure event will disable the rendering of this asset for good and we need to have a chance to 
-					// rerender it when everything settles down.
+					// re-render it when everything settles down.
 					else if (!bAreShadersCompiling && !bIsAssetStillCompiling)
 					{
 						// Notify listeners that a thumbnail render has failed
@@ -1162,13 +1174,13 @@ void FAssetThumbnailPool::Tick( float DeltaTime )
 	}
 }
 
-bool FAssetThumbnailPool::LoadThumbnail(TSharedRef<FThumbnailInfo> ThumbnailInfo, bool &bIsAssetStillCompiling, const FAssetData& CustomAssetToRender)
+bool FAssetThumbnailPool::LoadThumbnail(TSharedRef<FThumbnailInfo> ThumbnailInfo, bool& bIsAssetStillCompiling, const FAssetData& CustomAssetToRender)
 {
 	const FAssetData& AssetData = CustomAssetToRender.IsValid() ? CustomAssetToRender : ThumbnailInfo->AssetData;
-	UObject* Asset = AssetData.IsAssetLoaded() ? AssetData.GetAsset() : nullptr;
+	UObject* Asset = AssetData.FastGetAsset();
+	ensure(!Asset || IsValidChecked(Asset));
 
 	const bool bAreShadersCompiling = (GShaderCompilingManager && GShaderCompilingManager->IsCompiling());
-	
 	if (Asset && !bAreShadersCompiling)
 	{
 		//Avoid rendering the thumbnail of an asset that is currently edited asynchronously
@@ -1241,23 +1253,21 @@ bool FAssetThumbnailPool::LoadThumbnail(TSharedRef<FThumbnailInfo> ThumbnailInfo
 		}
 	}
 
-	FThumbnailMap ThumbnailMap;
 	// If we could not render a fresh thumbnail, see if we already have a cached one to load
 	const FObjectThumbnail* FoundThumbnail = ThumbnailTools::FindCachedThumbnail(AssetData.GetFullName());
+
+	// If we don't have a cached thumbnail, try to find it on disk
 	if (!FoundThumbnail)
 	{
-		// If we don't have a cached thumbnail, try to find it on disk
 		FString PackageFilename;
 		if (FPackageName::DoesPackageExist(AssetData.PackageName.ToString(), &PackageFilename))
 		{
+			const FName ObjectFullName = FName(*AssetData.GetFullName());
 			TSet<FName> ObjectFullNames;
-
-
-			FName ObjectFullName = FName(*AssetData.GetFullName());
 			ObjectFullNames.Add(ObjectFullName);
 
+			FThumbnailMap ThumbnailMap;
 			ThumbnailTools::LoadThumbnailsFromPackage(PackageFilename, ObjectFullNames, ThumbnailMap);
-
 			FoundThumbnail = ThumbnailMap.Find(ObjectFullName);
 		}
 	}
@@ -1598,4 +1608,4 @@ void FAssetThumbnailPool::OnAssetLoaded( UObject* Asset )
 void FAssetThumbnailPool::OnThumbnailDirtied( const FSoftObjectPath& ObjectPath )
 {
 	RefreshThumbnailsFor( ObjectPath );
-}
\ No newline at end of file
+}
