Cooker: Implement SoftGC - trigger a partial garbage collect when available physical memory reaches a schedule of targets. This significantly reduces the peak memory load of the cook commandlet for large projects, while adding only a small cost due to reloaded packages. In testing, the reload cost is so small that it is outweighed by time savings in various functions due to having a smaller number of UObjects loaded.
#rn Cooking, Minor

diff --git a/Engine/Config/BaseEditor.ini b/Engine/Config/BaseEditor.ini
--- a/Engine/Config/BaseEditor.ini
+++ b/Engine/Config/BaseEditor.ini
@@ -340,6 +340,12 @@ MaxConcurrentShaderJobs=8192
 MaxDiffsToLog=5
 IgnoreHeaderDiffs=True
 
+; SoftGC triggers a partial garbage collect when available physical memory on the machine reaches configured milestones
+; e.g. 5/10, 4/10, 3/10, 2/10, 1/10.
+bUseSoftGC=true
+SoftGCStartNumerator=5
+SoftGCDenominator=10
+
 ; CookProcessCount sets whether we cook singleprocess or multiprocess and how many processes there should be.
 ; CookProcessCount=(1 or less) is singleprocess. CookProcessCount=(N>1) is 1 director and N-1 cookworkers.
 CookProcessCount=1
diff --git a/Engine/Source/Editor/UnrealEd/Classes/CookOnTheSide/CookOnTheFlyServer.h b/Engine/Source/Editor/UnrealEd/Classes/CookOnTheSide/CookOnTheFlyServer.h
--- a/Engine/Source/Editor/UnrealEd/Classes/CookOnTheSide/CookOnTheFlyServer.h
+++ b/Engine/Source/Editor/UnrealEd/Classes/CookOnTheSide/CookOnTheFlyServer.h
@@ -317,8 +317,16 @@ private:
 	int32 MaxConcurrentShaderJobs;
 	/** Min number of free UObject indices before the cooker should partial gc */
 	int32 MinFreeUObjectIndicesBeforeGC;
-	/** Next time at which we are allowed to declare out of memory and trigger a garbage collect. */
-	double ExceededMaxMemoryCooldownEndTimeSeconds = 0;
+	double LastGCTime = 0.;
+	double LastFullGCTime = 0.;
+	double LastSoftGCTime = 0.;
+	int64 SoftGCNextAvailablePhysicalTarget = -1;
+	int32 SoftGCStartNumerator = 5;
+	int32 SoftGCDenominator = 10;
+	bool bUseSoftGC = false;
+	bool bWarnedExceededMaxMemoryWithinGCCooldown = false;
+	bool bGarbageCollectTypeSoft = false;
+
 	/**
 	 * The maximum number of packages that should be preloaded at once. Once this is full, packages in LoadPrepare will
 	 * remain unpreloaded in LoadPrepare until the existing preloaded packages exit {LoadPrepare,LoadReady} state.
@@ -557,6 +565,7 @@ public:
 		COSR_RequiresGC_PackageCount= 0x00000100,
 		COSR_RequiresGC_IdleTimer	= 0x00000200,
 		COSR_YieldTick				= 0x00000400,
+		COSR_RequiresGC_Soft_OOM	= 0x00000800,
 	};
 
 	struct FCookByTheBookStartupOptions
@@ -797,8 +806,12 @@ public:
 	/** Returns the configured amount of idle time before forcing a GC */
 	double GetIdleTimeToGC() const;
 
-	bool HasExceededMaxMemory();
-	void EvaluateGarbageCollectionResults(bool bWasDueToOOM, bool bWasPartialGC,
+	UE_DEPRECATED(5.2, "UCookOnTheFLyServer now uses a more complicated private GC scheme; HasExceededMaxMemory is no longer used and returns false")
+	bool HasExceededMaxMemory() { return false; }
+	void SetGarbageCollectType(uint32 ResultFlagsFromTick);
+	void ClearGarbageCollectType();
+
+	void EvaluateGarbageCollectionResults(bool bWasDueToOOM, bool bWasPartialGC, uint32 ResultFlags,
 		int32 NumObjectsBeforeGC, const FPlatformMemoryStats& MemStatsBeforeGC,
 		const FGenericMemoryStats& AllocatorStatsBeforeGC,
 		int32 NumObjectsAfterGC, const FPlatformMemoryStats& MemStatsAfterGC,
@@ -872,6 +885,8 @@ protected:
 
 private:
 
+	bool PumpHasExceededMaxMemory(uint32& OutResultFlags);
+
 	/**
 	 * Is the local CookOnTheFlyServer initialized to run in CookOnTheFly AND using the legacy scheduling method
 	 * that predates COTF2?
@@ -1385,6 +1400,9 @@ private:
 	 * as referenced in CookerAddReferencedObjects.
 	 */
 	TArray<UObject*> GCKeepObjects;
+	/** Packages that were expected to be freed by the last Soft GC and we expect not to load again. */
+	TArray<FName> ExpectedFreedPackageNames;
+
 	UE::Cook::FPackageData* SavingPackageData = nullptr;
 	/** Helper struct for running cooking in diagnostic modes */
 	TUniquePtr<FDiffModeCookServerUtils> DiffModeHelper;
diff --git a/Engine/Source/Editor/UnrealEd/Private/Commandlets/CookCommandlet.cpp b/Engine/Source/Editor/UnrealEd/Private/Commandlets/CookCommandlet.cpp
--- a/Engine/Source/Editor/UnrealEd/Private/Commandlets/CookCommandlet.cpp
+++ b/Engine/Source/Editor/UnrealEd/Private/Commandlets/CookCommandlet.cpp
@@ -562,10 +562,16 @@ void UCookCommandlet::ConditionalCollectGarbage(uint32 TickResults, UCookOnTheFl
 	}
 
 	FString GCReason;
+	FString GCType = bPartialGC ? TEXT(" partial gc") : TEXT("");
 	if ((TickResults & UCookOnTheFlyServer::COSR_RequiresGC_PackageCount) != 0)
 	{
 		GCReason = TEXT("Exceeded packages per GC");
 	}
+	else if ((TickResults & UCookOnTheFlyServer::COSR_RequiresGC_Soft_OOM) != 0)
+	{
+		GCType = TEXT("");
+		GCReason = TEXT("Soft GC");
+	}
 	else if ((TickResults & UCookOnTheFlyServer::COSR_RequiresGC_OOM) != 0)
 	{
 		// this can cause thrashing if the cooker loads the same stuff into memory next tick
@@ -632,10 +638,12 @@ void UCookCommandlet::ConditionalCollectGarbage(uint32 TickResults, UCookOnTheFl
 #if ENABLE_LOW_LEVEL_MEM_TRACKER
 	FLowLevelMemTracker::Get().UpdateStatsPerFrame();
 #endif
+	COTFS.SetGarbageCollectType(TickResults);
 
-	UE_LOG(LogCookCommandlet, Display, TEXT("GarbageCollection...%s (%s)"), (bPartialGC ? TEXT(" partial gc") : TEXT("")), *GCReason);
+	UE_LOG(LogCookCommandlet, Display, TEXT("GarbageCollection...%s (%s)"), *GCType, *GCReason);
 	CollectGarbage(RF_NoFlags);
 
+	COTFS.ClearGarbageCollectType();
 	FPlatformMemoryStats MemStatsAfterGC = FPlatformMemory::GetStats();
 	int32 NumObjectsAfterGC = GUObjectArray.GetObjectArrayNumMinusAvailable();
 	FGenericMemoryStats AllocatorStatsAfterGC;
@@ -645,7 +653,7 @@ void UCookCommandlet::ConditionalCollectGarbage(uint32 TickResults, UCookOnTheFl
 #endif
 
 	bool bWasDueToOOM = (TickResults & UCookOnTheFlyServer::COSR_RequiresGC_OOM) != 0;
-	COTFS.EvaluateGarbageCollectionResults(bWasDueToOOM, bPartialGC,
+	COTFS.EvaluateGarbageCollectionResults(bWasDueToOOM, bPartialGC, TickResults,
 		NumObjectsBeforeGC, MemStatsBeforeGC, AllocatorStatsBeforeGC,
 		NumObjectsAfterGC, MemStatsAfterGC, AllocatorStatsAfterGC);
 }
diff --git a/Engine/Source/Editor/UnrealEd/Private/CookOnTheFlyServer.cpp b/Engine/Source/Editor/UnrealEd/Private/CookOnTheFlyServer.cpp
--- a/Engine/Source/Editor/UnrealEd/Private/CookOnTheFlyServer.cpp
+++ b/Engine/Source/Editor/UnrealEd/Private/CookOnTheFlyServer.cpp
@@ -2540,9 +2540,8 @@ void UCookOnTheFlyServer::PumpLoads(UE::Cook::FTickStackData& StackData, uint32
 		OutNumPushed += NumPushed;
 		ProcessUnsolicitedPackages(); // May add new packages into the LoadQueue
 
-		if (HasExceededMaxMemory())
+		if (PumpHasExceededMaxMemory(StackData.ResultFlags))
 		{
-			StackData.ResultFlags |= COSR_RequiresGC | COSR_RequiresGC_OOM | COSR_YieldTick;
 			return;
 		}
 	}
@@ -4009,10 +4008,7 @@ void UCookOnTheFlyServer::PumpSaves(UE::Cook::FTickStackData& StackData, uint32
 	check(IsInGameThread());
 	ON_SCOPE_EXIT
 	{
-		if (HasExceededMaxMemory())
-		{
-			StackData.ResultFlags |= COSR_RequiresGC | COSR_RequiresGC_OOM | COSR_YieldTick;
-		}
+		PumpHasExceededMaxMemory(StackData.ResultFlags);
 	};
 
 	// save as many packages as we can during our time slice
@@ -4407,18 +4403,8 @@ void UCookOnTheFlyServer::TickPrecacheObjectsForPlatforms(const float TimeSlice,
 
 }
 
-bool UCookOnTheFlyServer::HasExceededMaxMemory()
+bool UCookOnTheFlyServer::PumpHasExceededMaxMemory(uint32& OutResultFlags)
 {
-	// If we are in a cooldown period, return false.
-	if (ExceededMaxMemoryCooldownEndTimeSeconds > 0)
-	{
-		if (ExceededMaxMemoryCooldownEndTimeSeconds > FPlatformTime::Seconds())
-		{
-			return false;
-		}
-		ExceededMaxMemoryCooldownEndTimeSeconds = 0.; // Return to 0 to avoid needing to call FPlatformTime::Seconds() on every call
-	}
-
 #if UE_GC_TRACK_OBJ_AVAILABLE
 	if (GUObjectArray.GetObjectArrayEstimatedAvailable() < MinFreeUObjectIndicesBeforeGC)
 	{
@@ -4430,6 +4416,7 @@ bool UCookOnTheFlyServer::HasExceededMaxMemory()
 			GEngine->Exec(nullptr, TEXT("OBJ LIST -COUNTSORT -SKIPMEMORYSIZE"));
 			bPerformedObjListWhenNearMaxObjects = true;
 		}
+		OutResultFlags |= COSR_RequiresGC | COSR_RequiresGC_OOM | COSR_YieldTick;
 		return true;
 	}
 #endif // UE_GC_TRACK_OBJ_AVAILABLE
@@ -4536,15 +4523,85 @@ bool UCookOnTheFlyServer::HasExceededMaxMemory()
 		bTriggerGC = true;
 	}
 
-	if (bTriggerGC)
+	// If a normal GC was not triggered, check the SoftGC trigger conditions
+	double CurrentTime = 0.;
+	bool bIsSoftGC = false;
+	if (!bTriggerGC && bUseSoftGC && IsDirectorCookByTheBook() && !IsCookingInEditor())
 	{
-		UE_LOG(LogCook, Display, TEXT("Garbage collection triggered by conditions:%s"), TriggerMessages.ToString());
-		return true;
+		if (SoftGCNextAvailablePhysicalTarget == -1) // Uninitialized
+		{
+			int32 StartNumerator = FMath::Max(SoftGCStartNumerator, 1);
+			int32 Denominator = FMath::Max(SoftGCDenominator, 1);
+			// e.g. Start the target at 5/10, and decrease it by 1/10 each time the target is reached
+			SoftGCNextAvailablePhysicalTarget = (static_cast<int64>(MemStats.TotalPhysical)*StartNumerator)
+				/Denominator;
+		}
+
+		if (SoftGCNextAvailablePhysicalTarget < -1)
+		{
+			// No further targets, no further SoftGC
+		}
+		else if (static_cast<int64>(MemStats.AvailablePhysical) <= SoftGCNextAvailablePhysicalTarget)
+		{
+			constexpr float SoftGCInstigateCooldown = 5 * 60.f;
+			CurrentTime = FPlatformTime::Seconds();
+			if (LastSoftGCTime + SoftGCInstigateCooldown <= CurrentTime)
+			{
+				TriggerMessages.Appendf(TEXT("\n  CookSettings.bUseSoftGC: Available physical memory %dMiB is less than the current target for SoftGC %dMiB."),
+					static_cast<uint32>(MemStats.AvailablePhysical / 1024 / 1024), static_cast<uint32>(SoftGCNextAvailablePhysicalTarget / 1024 / 1024));
+				bTriggerGC = true;
+				bIsSoftGC = true;
+			}
+		}
 	}
-	else
+
+	if (!bTriggerGC)
+	{
+		return false;
+	}
+
+	if (CurrentTime == 0.)
 	{
+		CurrentTime = FPlatformTime::Seconds();
+	}
+
+	// Don't allow a second GC (soft or normal) within the GC cooldown period after a full GC
+	constexpr float GCCooldown = 60.f;
+	if (LastFullGCTime + GCCooldown > CurrentTime)
+	{
+		if (!bIsSoftGC && !bWarnedExceededMaxMemoryWithinGCCooldown)
+		{
+			bWarnedExceededMaxMemoryWithinGCCooldown = true;
+			// If we are in a cooldown period, return false.
+			UE_LOG(LogCook, Display, TEXT("Garbage collection triggers ignored: Out of memory condition has been detected, but is only %.0fs after the last GC. ")
+				TEXT("It will be prevented until %.0f seconds have passed and we may run out of memory.\n")
+				TEXT("Garbage collection triggered by conditions: %s"),
+				static_cast<float>(CurrentTime - LastFullGCTime), GCCooldown, TriggerMessages.ToString());
+		}
 		return false;
 	}
+
+	const TCHAR* TypeMessage = bIsSoftGC ? TEXT("Soft") : (IsCookFlagSet(ECookInitializationFlags::EnablePartialGC) ? TEXT("Partial") : TEXT("Full"));
+
+	UE_LOG(LogCook, Display, TEXT("Garbage collection triggered (%s). Triggered by conditions:%s"),
+		TypeMessage, TriggerMessages.ToString());
+	OutResultFlags |= COSR_RequiresGC | COSR_YieldTick;
+	OutResultFlags |= COSR_RequiresGC_OOM;
+	if (bIsSoftGC)
+	{
+		OutResultFlags |= COSR_RequiresGC_Soft_OOM;
+	}
+	return true;
+}
+
+void UCookOnTheFlyServer::SetGarbageCollectType(uint32 ResultFlagsFromTick)
+{
+	bGarbageCollectTypeSoft = (ResultFlagsFromTick & COSR_RequiresGC_Soft_OOM);
+}
+
+void UCookOnTheFlyServer::ClearGarbageCollectType()
+{
+	bGarbageCollectTypeSoft = false;
 }
 
 void UCookOnTheFlyServer::PreGarbageCollect()
@@ -4637,69 +4694,87 @@ void UCookOnTheFlyServer::PreGarbageCollectImpl(TArray<UPackage*>& GCKeepPackage
 	}
 
 	const bool bPartialGC = IsCookFlagSet(ECookInitializationFlags::EnablePartialGC);
-	if (bPartialGC)
+	if (bGarbageCollectTypeSoft || bPartialGC)
 	{
-		LocalGCKeepObjects.Empty(1000);
-
-		// Keep all inprogress packages (including packages that have only made it to the request list) that have been partially loaded
-		// Additionally, keep all partially loaded packages that are transitively dependended on by any inprogress packages
-		// Keep all UObjects that have been loaded so far under these packages
-		TMap<const FPackageData*, int32> DependenciesCount;
-
-		TSet<FName> KeepPackages;
-		for (FPackageData* PackageData : *PackageDatas)
+		// Keep referenced all packages in requestqueue, loadqueue, and savequeue, and any packages they depend on
+		TArray<FName> Queue;
+		TSet<FName> Visited;
+		auto AddPackageName = [&Visited, &Queue](FName PackageName)
 		{
-			if (!PackageData->IsInProgress())
+			bool bAlreadyExists;
+			Visited.Add(PackageName, &bAlreadyExists);
+			if (!bAlreadyExists)
 			{
-				continue;
+				Queue.Add(PackageName);
 			}
-			const TArray<FName>& NeededPackages = GetFullPackageDependencies(PackageData->GetPackageName());
-			DependenciesCount.Add(PackageData, NeededPackages.Num());
-			KeepPackages.Append(NeededPackages);
-			GCKeepPackageDatas.Add(PackageData);
+		};
+		for (FPackageData* PackageData : PackageDatas->GetRequestQueue().GetReadyRequestsUrgent())
+		{
+			AddPackageName(PackageData->GetPackageName());
+		}
+		for (FPackageData* PackageData : PackageDatas->GetRequestQueue().GetReadyRequestsNormal())
+		{
+			AddPackageName(PackageData->GetPackageName());
+		}
+		for (FPackageData* PackageData : PackageDatas->GetLoadPrepareQueue().EntryQueue)
+		{
+			AddPackageName(PackageData->GetPackageName());
+		}
+		for (FPackageData* PackageData : PackageDatas->GetLoadPrepareQueue().PreloadingQueue)
+		{
+			AddPackageName(PackageData->GetPackageName());
+		}
+		for (FPackageData* PackageData : PackageDatas->GetLoadReadyQueue())
+		{
+			AddPackageName(PackageData->GetPackageName());
+		}
+		for (FPackageData* PackageData : PackageDatas->GetSaveQueue())
+		{
+			AddPackageName(PackageData->GetPackageName());
 		}
 
-		TSet<FName> LoadedPackages;
-		PackageTracker->ForEachLoadedPackage(
-			[&KeepPackages, &LoadedPackages, &GCKeepPackages](UPackage* Package)
+		TArray<FName> Dependencies;
+		while (!Queue.IsEmpty())
+		{
+			FName PackageName = Queue.Pop();
+			Dependencies.Reset();
+			AssetRegistry->GetDependencies(PackageName, Dependencies, UE::AssetRegistry::EDependencyCategory::Package,
+				UE::AssetRegistry::EDependencyQuery::Hard);
+			for (FName DependencyName : Dependencies)
 			{
-				const FName& PackageName = Package->GetFName();
-				if (KeepPackages.Contains(PackageName))
+				AddPackageName(DependencyName);
+			};
+		}
+
+		TSet<UPackage*> GCKeepPackagesSet;
+		GCKeepPackagesSet.Append(GCKeepPackages);
+		for (FName PackageName : Visited)
+		{
+			UPackage* Package = FindPackage(nullptr, *WriteToString<256>(PackageName));
+			if (Package)
+			{
+				bool bAlreadyInSet;
+				GCKeepPackagesSet.Add(Package, &bAlreadyInSet);
+				if (!bAlreadyInSet)
 				{
-					LoadedPackages.Add(PackageName);
 					GCKeepPackages.Add(Package);
+					FPackageData* PackageData = PackageDatas->FindPackageDataByPackageName(Package->GetFName());
+					if (PackageData)
+					{
+						GCKeepPackageDatas.Add(PackageData);
+					}
 				}
 			}
-		);
-
-		FRequestQueue& RequestQueue = PackageDatas->GetRequestQueue();
-		TArray<FPackageData*> Requests;
-		Requests.Reserve(RequestQueue.ReadyRequestsNum());
-		while (!RequestQueue.IsReadyRequestsEmpty())
-		{
-			Requests.Add(RequestQueue.PopReadyRequest());
-		}
-		// We are not looking at UnclusteredRequests or RequestClusters from the RequestQueue. These are supposed to 
-		// be quickly processed and should usually be empty, so failing to account for them will only rarely cause
-		// a performance issue (and will not cause a behavior issue)
-		
-		// Sort the cook requests by the packages which are loaded first
-		// then sort by the number of dependencies which are referenced by the package
-		// we want to process the packages with the highest dependencies so that they can
-		// be evicted from memory and are likely to be able to be released on next GC pass
-		Algo::Sort(Requests, [&DependenciesCount, &LoadedPackages](const FPackageData* A, const FPackageData* B)
-			{
-				int32 ADependencies = DependenciesCount.FindChecked(A);
-				int32 BDependencies = DependenciesCount.FindChecked(B);
-				bool ALoaded = LoadedPackages.Contains(A->GetPackageName());
-				bool BLoaded = LoadedPackages.Contains(B->GetPackageName());
-				return (ALoaded == BLoaded) ? (ADependencies > BDependencies) : ALoaded > BLoaded;
-			}
-		);
-		for (FPackageData* Request : Requests)
-		{
-			RequestQueue.AddRequest(Request); // Urgent requests will still be moved to the front of the RequestQueue by AddRequest
 		}
+		ExpectedFreedPackageNames.Reset();
+		PackageTracker->ForEachLoadedPackage(
+			[this, &GCKeepPackagesSet](UPackage* Package)
+			{
+				if (!GCKeepPackagesSet.Contains(Package))
+				{
+					ExpectedFreedPackageNames.Add(Package->GetFName());
+				}
+			});
 	}
 
 	// Add packages and all RF_Public objects outered to them to LocalGCKeepObjects
@@ -4711,7 +4786,7 @@ void UCookOnTheFlyServer::PreGarbageCollectImpl(TArray<UPackage*>& GCKeepPackage
 		GetObjectsWithOuter(Package, ObjectsWithOuter);
 		for (UObject* Obj : ObjectsWithOuter)
 		{
-			if (Obj->HasAnyFlags(RF_Public))
+			if (IsValidChecked(Obj) && Obj->HasAnyFlags(RF_Public))
 			{
 				LocalGCKeepObjects.Add(Obj);
 			}
@@ -4782,7 +4857,7 @@ void UCookOnTheFlyServer::PostGarbageCollectImpl(TArray<UObject*>& LocalGCKeepOb
 	CookedPackageCountSinceLastGC = 0;
 }
 
-void UCookOnTheFlyServer::EvaluateGarbageCollectionResults(bool bWasDueToOOM, bool bWasPartialGC,
+void UCookOnTheFlyServer::EvaluateGarbageCollectionResults(bool bWasDueToOOM, bool bWasPartialGC, uint32 ResultFlags,
 	int32 NumObjectsBeforeGC, const FPlatformMemoryStats& MemStatsBeforeGC,
 	const FGenericMemoryStats& AllocatorStatsBeforeGC,
 	int32 NumObjectsAfterGC, const FPlatformMemoryStats& MemStatsAfterGC,
@@ -4790,6 +4865,35 @@ void UCookOnTheFlyServer::EvaluateGarbageCollectionResults(bool bWasDueToOOM, bo
 {
 	using namespace UE::Cook;
 
+	bWarnedExceededMaxMemoryWithinGCCooldown = false;
+	LastGCTime = FPlatformTime::Seconds();
+	bool bWasSoftGC = ResultFlags & COSR_RequiresGC_Soft_OOM;
+	if (bWasSoftGC)
+	{
+		LastSoftGCTime = LastGCTime;
+		int32 StartNumerator = FMath::Max(SoftGCStartNumerator, 1);
+		int32 Denominator = FMath::Max(SoftGCDenominator, 1);
+		// Calculate the new SoftGCNextAvailablePhysicalTarget. Use the floor of NewAvailableMemory/Denominator,
+		// unless we are already 50% of the way through that level, in which case use the next value below that
+		int64 PhysicalMemoryQuantum = static_cast<int64>(MemStatsAfterGC.TotalPhysical) / Denominator;
+		int32 NextTarget =
+			static_cast<int64>(MemStatsAfterGC.AvailablePhysical - PhysicalMemoryQuantum/2) / PhysicalMemoryQuantum;
+		NextTarget = FMath::Max(NextTarget, StartNumerator);
+		if (NextTarget <= 0)
+		{
+			SoftGCNextAvailablePhysicalTarget = -2; // disabled, no further targets
+		}
+		else
+		{
+			SoftGCNextAvailablePhysicalTarget = static_cast<int64>((MemStatsAfterGC.TotalPhysical) * NextTarget)
+				/ Denominator;
+		}
+	}
+	else
+	{
+		LastFullGCTime = LastGCTime;
+	}
+
 	if (IsCookingInEditor())
 	{
 		return;
@@ -4820,13 +4924,8 @@ void UCookOnTheFlyServer::EvaluateGarbageCollectionResults(bool bWasDueToOOM, bo
 #else
 	const bool bAlwaysShowAnalysis = bCookMemoryAnalysis;
 #endif		
-
-	bool bWasImpactful =
-		(NumObjectsFreed >= ExpectedObjectsFreed || NumObjectsBeforeGC - NumObjectsMin < ExpectedObjectsFreed) &&
-		(VirtualMemFreed >= ExpectedMemFreed || VirtualMemBeforeGC - VirtualMemMin <= ExpectedMemFreed);
-
 	constexpr int32 BytesPerMeg = 1000000;
-	if ((!bWasDueToOOM || bWasImpactful) && !bAlwaysShowAnalysis)
+	auto DisplaySimpleSummary = [&]()
 	{
 		UE_LOG(LogCook, Display, TEXT("GarbageCollection Results:\n")
 			TEXT("\tType: %s\n")
@@ -4839,113 +4938,131 @@ void UCookOnTheFlyServer::EvaluateGarbageCollectionResults(bool bWasDueToOOM, bo
 			TEXT("\t\tBefore GC:        %10" INT64_FMT " MB\n")
 			TEXT("\t\tAfter GC:         %10" INT64_FMT " MB\n")
 			TEXT("\t\tFreed by GC:      %10" INT64_FMT " MB\n"),
-			(bWasPartialGC ? TEXT("Partial") : TEXT("Full")),
+			(bWasSoftGC ? TEXT("Soft") : (bWasPartialGC ? TEXT("Partial") : TEXT("Full"))),
 			NumObjectsCapacity, (int64)NumObjectsBeforeGC, (int64)NumObjectsAfterGC, NumObjectsFreed,
 			VirtualMemBeforeGC / BytesPerMeg, VirtualMemAfterGC / BytesPerMeg, VirtualMemFreed / BytesPerMeg
 		);
-		return;
-	}
+	};
 
-	// When the garbage collection is not impactful, we add a cooldown. In some cases it really was impactful and we just could not tell
-	// that it was impactful, because e.g. GMalloc is holding on to unused memory but allowing the operating system to page that unused
-	// memory out. In the case where it really was not impactful, it is better to crash on out of memory or to run slowly due to many page
-	// faults than it is to continuously collect garbage for no purpose.
-	constexpr float ExceededMaxMemoryCooldownEndTimeDuration = 60.0f;
-	FString CooldownMessage;
-	if (bWasDueToOOM && !bWasImpactful)
+	if (!bWasSoftGC)
 	{
-		ExceededMaxMemoryCooldownEndTimeSeconds = FPlatformTime::Seconds() + ExceededMaxMemoryCooldownEndTimeDuration;
-		UE_LOG(LogCook, Display, TEXT("GarbageCollection Results: Garbage Collection was not very impactful."));
-		CooldownMessage = FString::Printf(
-			TEXT("\t\tCooldown: Further OOM GarbageCollection will be prevented for %.0f seconds and we may run out of memory.\n"),
-			ExceededMaxMemoryCooldownEndTimeDuration);
+		bool bWasImpactful =
+			(NumObjectsFreed >= ExpectedObjectsFreed || NumObjectsBeforeGC - NumObjectsMin < ExpectedObjectsFreed) &&
+			(VirtualMemFreed >= ExpectedMemFreed || VirtualMemBeforeGC - VirtualMemMin <= ExpectedMemFreed);
+
+		if ((!bWasDueToOOM || bWasImpactful) && !bAlwaysShowAnalysis)
+		{
+			DisplaySimpleSummary();
+			return;
+		}
+
+		if (bWasDueToOOM && !bWasImpactful)
+		{
+			UE_LOG(LogCook, Display, TEXT("GarbageCollection Results: Garbage Collection was not very impactful."));
+		}
+		else
+		{
+			UE_LOG(LogCook, Display, TEXT("GarbageCollection Results:"));
+		}
+		UE_LOG(LogCook, Display, TEXT("\tMemoryAnalysis: General:\n")
+			TEXT("\t\tType: %s"),
+			(bWasSoftGC ? TEXT("Soft") : (bWasPartialGC ? TEXT("Partial") : TEXT("Full"))));
+		UE_LOG(LogCook, Display, TEXT("\tMemoryAnalysis: NumObjects:\n")
+			TEXT("\t\tCapacity:         %10" INT64_FMT "\n")
+			TEXT("\t\tProcess Min:      %10" INT64_FMT "\n")
+			TEXT("\t\tProcess Max:      %10" INT64_FMT "\n")
+			TEXT("\t\tProcess Spread:   %10" INT64_FMT "\n")
+			TEXT("\t\tBefore GC:        %10" INT64_FMT "\n")
+			TEXT("\t\tAfter GC:         %10" INT64_FMT "\n")
+			TEXT("\t\tFreed by GC:      %10" INT64_FMT ""),
+			NumObjectsCapacity, NumObjectsMin, NumObjectsMax, NumObjectsSpread,
+			(int64)NumObjectsBeforeGC, (int64)NumObjectsAfterGC, NumObjectsFreed);
+		UE_LOG(LogCook, Display, TEXT("\tMemoryAnalysis: Virtual Memory:\n")
+			TEXT("\t\tProcess Min:      %10" INT64_FMT " MB\n")
+			TEXT("\t\tProcess Max:      %10" INT64_FMT " MB\n")
+			TEXT("\t\tProcess Spread:   %10" INT64_FMT " MB\n")
+			TEXT("\t\tBefore GC:        %10" INT64_FMT " MB\n")
+			TEXT("\t\tAfter GC:         %10" INT64_FMT " MB\n")
+			TEXT("\t\tFreed by GC:      %10" INT64_FMT " MB"),
+			VirtualMemMin / BytesPerMeg, VirtualMemMax / BytesPerMeg, VirtualMemSpread / BytesPerMeg,
+			VirtualMemBeforeGC / BytesPerMeg, VirtualMemAfterGC / BytesPerMeg, VirtualMemFreed / BytesPerMeg);
+		auto AllocatorStatsToString = [](const FGenericMemoryStats& AllocatorStats)
+		{
+			TStringBuilder<256> Writer;
+			for (const TPair<FString, SIZE_T>& Item : AllocatorStats.Data)
+			{
+				Writer << TEXT("\n\t\tItem ") << Item.Key << TEXT(" ") << (uint64)Item.Value;
+			}
+			return FString(*Writer);
+		};
+		UE_LOG(LogCook, Display, TEXT("\tMemoryAnalysis: Allocator Stats Before:%s"),
+			*AllocatorStatsToString(AllocatorStatsBeforeGC));
+		UE_LOG(LogCook, Display, TEXT("\tMemoryAnalysis: Allocator Stats After:%s"),
+			*AllocatorStatsToString(AllocatorStatsAfterGC));
+
+
+		UE_LOG(LogCook, Display, TEXT("See log for memory use information for UObject classes and LLM tags."));
+		UE::Cook::DumpObjClassList(CookByTheBookOptions->SessionStartupObjects);
+		GLog->Logf(TEXT("Memory Analysis: LLM Tags:"));
+#if ENABLE_LOW_LEVEL_MEM_TRACKER
+		if (FLowLevelMemTracker::Get().IsEnabled())
+		{
+			FLowLevelMemTracker::Get().DumpToLog();
+		}
+		else
+#endif
+		{
+			GLog->Logf(TEXT("LLM Tags are not displayed because llm is disabled. Run with -llm or -trace=memtag to see llm tags."));
+		}
+		GLog->Flush();
 	}
 	else
 	{
-		UE_LOG(LogCook, Display, TEXT("GarbageCollection Results:"));
-	}
-	UE_LOG(LogCook, Display, TEXT("\tMemoryAnalysis: General:\n")
-		TEXT("%s")
-		TEXT("\t\tType: %s"),
-		*CooldownMessage,
-		(bWasPartialGC ? TEXT("Partial") : TEXT("Full")));
-	UE_LOG(LogCook, Display, TEXT("\tMemoryAnalysis: NumObjects:\n")
-		TEXT("\t\tCapacity:         %10" INT64_FMT "\n")
-		TEXT("\t\tProcess Min:      %10" INT64_FMT "\n")
-		TEXT("\t\tProcess Max:      %10" INT64_FMT "\n")
-		TEXT("\t\tProcess Spread:   %10" INT64_FMT "\n")
-		TEXT("\t\tBefore GC:        %10" INT64_FMT "\n")
-		TEXT("\t\tAfter GC:         %10" INT64_FMT "\n")
-		TEXT("\t\tFreed by GC:      %10" INT64_FMT ""),
-		NumObjectsCapacity, NumObjectsMin, NumObjectsMax, NumObjectsSpread,
-		(int64)NumObjectsBeforeGC, (int64)NumObjectsAfterGC, NumObjectsFreed);
-	UE_LOG(LogCook, Display, TEXT("\tMemoryAnalysis: Virtual Memory:\n")
-		TEXT("\t\tProcess Min:      %10" INT64_FMT " MB\n")
-		TEXT("\t\tProcess Max:      %10" INT64_FMT " MB\n")
-		TEXT("\t\tProcess Spread:   %10" INT64_FMT " MB\n")
-		TEXT("\t\tBefore GC:        %10" INT64_FMT " MB\n")
-		TEXT("\t\tAfter GC:         %10" INT64_FMT " MB\n")
-		TEXT("\t\tFreed by GC:      %10" INT64_FMT " MB"),
-		VirtualMemMin / BytesPerMeg, VirtualMemMax / BytesPerMeg, VirtualMemSpread / BytesPerMeg,
-		VirtualMemBeforeGC / BytesPerMeg, VirtualMemAfterGC / BytesPerMeg, VirtualMemFreed / BytesPerMeg);
-	auto AllocatorStatsToString = [](const FGenericMemoryStats& AllocatorStats)
-	{
-		TStringBuilder<256> Writer;
-		for (const TPair<FString, SIZE_T>& Item : AllocatorStats.Data)
-		{
-			Writer << TEXT("\n\t\tItem ") << Item.Key << TEXT(" ") << (uint64)Item.Value;
-		}
-		return FString(*Writer);
-	};
-	UE_LOG(LogCook, Display, TEXT("\tMemoryAnalysis: Allocator Stats Before:%s"),
-		*AllocatorStatsToString(AllocatorStatsBeforeGC));
-	UE_LOG(LogCook, Display, TEXT("\tMemoryAnalysis: Allocator Stats After:%s"),
-		*AllocatorStatsToString(AllocatorStatsAfterGC));
+		DisplaySimpleSummary();
 
-	TArray<UPackage*> GCKeepPackages;
-	TArray<FPackageData*> GCKeepPackageDatas;
-	TArray<UObject*> LocalGCKeepObjects;
-	PreGarbageCollectImpl(GCKeepPackages, GCKeepPackageDatas, LocalGCKeepObjects);
-	PostGarbageCollectImpl(LocalGCKeepObjects);
-
-	TSet<UPackage*> DirectPackages(GCKeepPackages);
-	DirectPackages.Append(GCKeepPackages);
-	// Some Objects can be in KeepObjects without their Package being added, because they are PollPendingCookedPlatformDatas
-	// Add their packages, since we do the transitive search based on package dependencies
-	for (UObject* Object : LocalGCKeepObjects)
-	{
-		DirectPackages.Add(Object->GetPackage());
-	}
-
-	FResourceSizeEx DirectResourceSize;
-	FResourceSizeEx TransitiveResourceSize;
-	int64 NumDirectPackages;
-	int64 NumTransitivePackages;
-	GetDirectAndTransitiveResourceSize(DirectResourceSize, TransitiveResourceSize,
-		NumDirectPackages, NumTransitivePackages, MoveTemp(DirectPackages));
-	UE_LOG(LogCook, Display, TEXT("\tMemoryAnalysis: Cooker References:\n")
-		TEXT("\t\tCooker direct packages:         %10" INT64_FMT "\n")
-		TEXT("\t\tCooker transitive packages:     %10" INT64_FMT "\n")
-		TEXT("\t\tCooker direct package size:     %10" INT64_FMT " MB\n")
-		TEXT("\t\tCooker transitive package size: %10" INT64_FMT " MB\n"),
-		NumDirectPackages, NumTransitivePackages,
-		DirectResourceSize.GetTotalMemoryBytes() / BytesPerMeg,
-		TransitiveResourceSize.GetTotalMemoryBytes() / BytesPerMeg);
-
-	UE_LOG(LogCook, Display, TEXT("See log for memory use information for UObject classes and LLM tags."));
-	UE::Cook::DumpObjClassList(CookByTheBookOptions->SessionStartupObjects);
-	GLog->Logf(TEXT("Memory Analysis: LLM Tags:"));
+		// Mark the packages we freed so we can give a warning to diagnose why they are still referenced if they
+		// get loaded again.
+		PackageTracker->AddExpectedNeverLoadPackages(ExpectedFreedPackageNames);
+
+		// If some packages we expected to be freed were not freed, show the reference chains for why
+		// they were not freed.
+		TArray<UPackage*> PackagesReferencedOutsideOfCooker;
+		for (FWeakObjectPtr& WeakPtr : CookByTheBookOptions->SessionStartupObjects)
+		{
+			UObject* Object = WeakPtr.Get();
+			if (!Object)
+			{
+				continue;
+			}
+			UPackage* Package = Object->GetPackage();
+			ExpectedFreedPackageNames.Remove(Package->GetFName());
+		}
+		PackageTracker->ForEachLoadedPackage(
+			[this, &PackagesReferencedOutsideOfCooker](UPackage* Package)
+			{
+				if (ExpectedFreedPackageNames.Contains(Package->GetFName()))
+				{
+					PackagesReferencedOutsideOfCooker.Add(Package);
+				}
+			});
+		if (PackagesReferencedOutsideOfCooker.Num() > 0)
+		{
+			// Only show diagnostics if LLM is on. We could add a separate setting for this, but it's more convenient to combine it with the LLM enabled setting
 #if ENABLE_LOW_LEVEL_MEM_TRACKER
-	if (FLowLevelMemTracker::Get().IsEnabled())
-	{
-		FLowLevelMemTracker::Get().DumpToLog();
-	}
-	else
+			bool bShowDiagnostics = FLowLevelMemTracker::Get().IsEnabled();
+#else
+			constexpr bool bShowDiagnostics = false;
 #endif
-	{
-		GLog->Logf(TEXT("LLM Tags are not displayed because llm is disabled. Run with -llm or -trace=memtag to see llm tags."));
+			if (bShowDiagnostics)
+			{
+				UE::Cook::DumpPackageReferencers(PackagesReferencedOutsideOfCooker);
+			}
+			else
+			{
+				GLog->Logf(TEXT("Soft Garbage Collect diagnostics are not displayed because llm is disabled. Run with -llm or -trace=memtag to see diagnostics."));
+			}
+		}
 	}
-	GLog->Flush();
 }
 
 void UCookOnTheFlyServer::GetDirectAndTransitiveResourceSize(FResourceSizeEx& OutDirectSize,
@@ -5957,6 +6074,10 @@ void FInitializeConfigSettings::LoadLocal(const FString& InOutputDirectoryOverri
 	MemoryMaxUsedPhysical = 0;
 	MemoryMinFreeVirtual = 0;
 	MemoryMinFreePhysical = 0;
+	bUseSoftGC = false;
+	SoftGCStartNumerator = 5;
+	SoftGCDenominator = 10;
+
 	ReadMemorySetting(TEXT("MemoryMaxUsedVirtual"), MemoryMaxUsedVirtual);
 	ReadMemorySetting(TEXT("MemoryMaxUsedPhysical"), MemoryMaxUsedPhysical);
 	ReadMemorySetting(TEXT("MemoryMinFreeVirtual"), MemoryMinFreeVirtual);
@@ -5968,6 +6089,9 @@ void FInitializeConfigSettings::LoadLocal(const FString& InOutputDirectoryOverri
 		UE_LOG(LogCook, Error, TEXT("Unrecognized value \"%s\" for MemoryTriggerGCAtPressureLevel. Expected None or Critical."),
 			*ConfigText);
 	}
+	GConfig->GetBool(TEXT("CookSettings"), TEXT("bUseSoftGC"), bUseSoftGC, GEditorIni);
+	GConfig->GetInt(TEXT("CookSettings"), TEXT("SoftGCStartNumerator"), SoftGCStartNumerator, GEditorIni);
+	GConfig->GetInt(TEXT("CookSettings"), TEXT("SoftGCDenominator"), SoftGCDenominator, GEditorIni);
 
 	MemoryExpectedFreedToSpreadRatio = 0.10f;
 	GConfig->GetFloat(TEXT("CookSettings"), TEXT("MemoryExpectedFreedToSpreadRatio"),
@@ -5982,10 +6106,17 @@ void FInitializeConfigSettings::LoadLocal(const FString& InOutputDirectoryOverri
 	
 	GConfig->GetArray(TEXT("CookSettings"), TEXT("CookOnTheFlyConfigSettingDenyList"), ConfigSettingDenyList, GEditorIni);
 
-	UE_LOG(LogCook, Display, TEXT("CookSettings for Memory: MemoryMaxUsedVirtual %dMiB, MemoryMaxUsedPhysical %dMiB, ")
-		TEXT("MemoryMinFreeVirtual %dMiB, MemoryMinFreePhysical %dMiB, MemoryTriggerGCAtPressureLevel %s"),
+	UE_LOG(LogCook, Display, TEXT("CookSettings for Memory:")
+		TEXT("\n\tMemoryMaxUsedVirtual %dMiB")
+		TEXT("\n\tMemoryMaxUsedPhysical %dMiB")
+		TEXT("\n\tMemoryMinFreeVirtual %dMiB")
+		TEXT("\n\tMemoryMinFreePhysical %dMiB")
+		TEXT("\n\tMemoryTriggerGCAtPressureLevel %s")
+		TEXT("\n\tUseSoftGC %s%s"),
 		MemoryMaxUsedVirtual / 1024 / 1024, MemoryMaxUsedPhysical / 1024 / 1024, MemoryMinFreeVirtual / 1024 / 1024, MemoryMinFreePhysical / 1024 / 1024,
-		*LexToString(MemoryTriggerGCAtPressureLevel));
+		*LexToString(MemoryTriggerGCAtPressureLevel),
+		bUseSoftGC ? TEXT("true") : TEXT("false"),
+		bUseSoftGC ? *FString::Printf(TEXT(" (%d/%d)"), SoftGCStartNumerator, SoftGCDenominator) : TEXT(""));
 
 	const FConfigSection* CacheSettings = GConfig->GetSectionPrivate(TEXT("CookPlatformDataCacheSettings"), false, true, GEditorIni);
 	if (CacheSettings)
diff --git a/Engine/Source/Editor/UnrealEd/Private/Cooker/CookDirector.cpp b/Engine/Source/Editor/UnrealEd/Private/Cooker/CookDirector.cpp
--- a/Engine/Source/Editor/UnrealEd/Private/Cooker/CookDirector.cpp
+++ b/Engine/Source/Editor/UnrealEd/Private/Cooker/CookDirector.cpp
@@ -17,6 +17,7 @@
 #include "Math/NumericLimits.h"
 #include "Misc/CommandLine.h"
 #include "Misc/ConfigCacheIni.h"
+#include "PackageTracker.h"
 #include "Serialization/CompactBinary.h"
 #include "Serialization/CompactBinaryWriter.h"
 #include "ShaderCompiler.h"
@@ -920,19 +921,26 @@ void FCookDirector::ActivateMachineResourceReduction()
 	}
 	bHasReducedMachineResources = true;
 
-	// When running a multiprocess cook, we remove the Memory triggers and trigger GC based solely on PressureLevel
+	// When running a multiprocess cook, we remove the Memory triggers and trigger GC based solely on PressureLevel. But keep the Soft GC settings
 	COTFS.MemoryMaxUsedPhysical = 0;
 	COTFS.MemoryMaxUsedVirtual = 0;
 	COTFS.MemoryMinFreeVirtual = 0;
 	COTFS.MemoryMinFreePhysical = 0;
 	COTFS.MemoryTriggerGCAtPressureLevel = FGenericPlatformMemoryStats::EMemoryPressureStatus::Critical;
 
-	UE_LOG(LogCook, Display, TEXT("CookMultiprocess changed CookSettings for Memory: MemoryMaxUsedVirtual %dMiB, MemoryMaxUsedPhysical %dMiB,")
-		TEXT("MemoryMinFreeVirtual % dMiB, MemoryMinFreePhysical % dMiB, MemoryTriggerGCAtPressureLevel %s"),
+	UE_LOG(LogCook, Display, TEXT("CookMultiprocess changed CookSettings for Memory:")
+		TEXT("\n\tMemoryMaxUsedVirtual %dMiB")
+		TEXT("\n\tMemoryMaxUsedPhysical %dMiB")
+		TEXT("\n\tMemoryMinFreeVirtual %dMiB")
+		TEXT("\n\tMemoryMinFreePhysical %dMiB")
+		TEXT("\n\tMemoryTriggerGCAtPressureLevel %s")
+		TEXT("\n\tUseSoftGC %s%s"),
 		COTFS.MemoryMaxUsedVirtual / 1024 / 1024, COTFS.MemoryMaxUsedPhysical / 1024 / 1024,
 		COTFS.MemoryMinFreeVirtual / 1024 / 1024, COTFS.MemoryMinFreePhysical / 1024 / 1024,
-		*LexToString(COTFS.MemoryTriggerGCAtPressureLevel)
-		);
+		*LexToString(COTFS.MemoryTriggerGCAtPressureLevel),
+		COTFS.bUseSoftGC ? TEXT("true") : TEXT("false"),
+		COTFS.bUseSoftGC ? *FString::Printf(TEXT(" (%d/%d)"), COTFS.SoftGCStartNumerator, COTFS.SoftGCDenominator) : TEXT("")
+	);
 
 	// Set CoreLimit for updating workerthreads in this process and passing to the commandline for workers
 	int32 NumProcesses = RequestedCookWorkerCount + 1;
@@ -1549,6 +1557,7 @@ void FCookDirector::FRetractionHandler::ReassignPackages(const FWorkerId& FromWo
 	TArray<FWorkerId> Assignments;
 	Director.AssignRequests(MoveTemp(WorkersToSplitOver), LocalRemoteWorkers, AssignmentPackages, Assignments, MoveTemp(RequestGraph));
 	FRequestQueue& RequestQueue = Director.COTFS.PackageDatas->GetRequestQueue();
+	bool bAssignedToLocal = false;
 	for (int32 Index = 0; Index < AssignmentPackages.Num(); ++Index)
 	{
 		FPackageData* PackageData = AssignmentPackages[Index];
@@ -1560,6 +1569,7 @@ void FCookDirector::FRetractionHandler::ReassignPackages(const FWorkerId& FromWo
 		else if (Assignment.IsLocal())
 		{
 			RequestQueue.AddReadyRequest(PackageData);
+			bAssignedToLocal = true;
 		}
 		else
 		{
@@ -1568,6 +1578,12 @@ void FCookDirector::FRetractionHandler::ReassignPackages(const FWorkerId& FromWo
 		}
 	}
 	Director.DisplayRemainingPackages();
+	if (bAssignedToLocal)
+	{
+		// Clear the SoftGC diagnostic ExpectedNeverLoadPackages because we have new assigned packages
+		// that we didn't consider during SoftGC
+		Director.COTFS.PackageTracker->ClearExpectedNeverLoadPackages();
+	}
 }
 
 TArray<FWorkerId> FCookDirector::FRetractionHandler::CalculateWorkersToSplitOver(int32 NumPackages, const FWorkerId& FromWorker,
diff --git a/Engine/Source/Editor/UnrealEd/Private/Cooker/CookPackageData.h b/Engine/Source/Editor/UnrealEd/Private/Cooker/CookPackageData.h
--- a/Engine/Source/Editor/UnrealEd/Private/Cooker/CookPackageData.h
+++ b/Engine/Source/Editor/UnrealEd/Private/Cooker/CookPackageData.h
@@ -956,6 +956,8 @@ public:
 
 	FPackageDataSet& GetUnclusteredRequests() { return UnclusteredRequests; }
 	TRingBuffer<FRequestCluster>& GetRequestClusters() { return RequestClusters; }
+	FPackageDataSet& GetReadyRequestsUrgent() { return UrgentRequests; }
+	FPackageDataSet& GetReadyRequestsNormal() { return NormalRequests; }
 private:
 	FPackageDataSet UnclusteredRequests;
 	TRingBuffer<FRequestCluster> RequestClusters;
diff --git a/Engine/Source/Editor/UnrealEd/Private/Cooker/CookProfiling.cpp b/Engine/Source/Editor/UnrealEd/Private/Cooker/CookProfiling.cpp
--- a/Engine/Source/Editor/UnrealEd/Private/Cooker/CookProfiling.cpp
+++ b/Engine/Source/Editor/UnrealEd/Private/Cooker/CookProfiling.cpp
@@ -266,6 +266,8 @@ enum class EObjectReferencerType : uint8
 	Referenced,
 };
 
+struct FObjectGraphProfileData;
+
 /**
  * Data for how an object is referenced in the DumpObjClassList graph search,
  * including the type of reference and the vertex of the referencer.
@@ -303,53 +305,92 @@ struct FObjectReferencer
 		VertexArgument = InVertexArgument;
 		LinkType = InLinkType;
 	}
-	void ToString(FStringBuilderBase& Builder, TConstArrayView<UObject*> VertexToObject)
+	void ToString(FStringBuilderBase& Builder, FObjectGraphProfileData& ProfileData);
+
+private:
+	Algo::Graph::FVertex VertexArgument = Algo::Graph::InvalidVertex;
+	EObjectReferencerType LinkType = EObjectReferencerType::Unknown;
+};
+
+struct FObjectGraphProfileData
+{
+	/** The list of UObjects found from a TObectIterator */
+	TArray<UObject*> AllObjects;
+	/** We assign FVertex N <-> AllObjects[N]; this field records the reverse map. */
+	TMap<UObject*, Algo::Graph::FVertex> VertexOfObject;
+	/** Element N records whether AllObjects[N] is not one of InitialObjects */
+	TBitArray<> IsNew;
+	/** The first reason found that AllObjects[n] is still referenced. */
+	TArray<FObjectReferencer> AliveReason;
+	/** The first rooted vertex found that has a reference chain to AllObjects[n]. */
+	TArray<Algo::Graph::FVertex> RootOfVertex;
+	/** The referencenames reported by FGCObject::GGCObjectReferencer for why it refers to objects. */
+	TArray<FString> AllGCObjectNames;
+	/** We assign (FVertex NumObjects+N) <-> AllGCObjectNames[N]; this field records the reverse map. */
+	TMap<FString, Algo::Graph::FVertex> GCObjectNameToVertex;
+	/** Buffer of edges used for ObjectGraph */
+	TArray64<Algo::Graph::FVertex> ObjectGraphBuffer;
+	/** ObjectGraph constructed from the edges between vertices defined by serialization references between objects. */
+	TArray<TConstArrayView<Algo::Graph::FVertex>> ObjectGraph;
+	/** Total number of vertices, both Objects and GCObjectNames */
+	int32 NumVertices;
+	/** Number of object vertices. The first GCObjectName vertex starts after this number. */
+	int32 NumObjectVertices;
+	/** Number of GCObjectName vertices. */
+	int32 NumGCObjectNameVertices;
+	/** The vertex that is assigned to FGCObject::GGCObjectReferencer. */
+	Algo::Graph::FVertex GCObjectReferencerVertex;
+
+	void AppendVertexName(Algo::Graph::FVertex Vertex, FStringBuilderBase& Builder)
 	{
-		switch (GetLinkType())
+		if (Vertex < 0)
 		{
-		case EObjectReferencerType::Unknown:
-			Builder << TEXT("<Unknown>");
-			break;
-		case EObjectReferencerType::Rooted:
-			Builder << TEXT("<Rooted>");
-			break;
-		case EObjectReferencerType::GCObjectRef:
+			Builder << TEXT("InvalidVertex");
+		}
+		else if (Vertex < NumObjectVertices)
 		{
-			check(VertexArgument != Algo::Graph::InvalidVertex);
-			UObject* Object = VertexToObject[VertexArgument];
-			FString ReferencerName;
-			if (!Object || !FGCObject::GGCObjectReferencer->GetReferencerName(Object, ReferencerName))
-			{
-				ReferencerName = TEXT("<Unknown>");
-			}
-			Builder << TEXT("FGCObject ") << ReferencerName;
-			break;
+			AllObjects[Vertex]->GetPathName(nullptr, Builder);
 		}
-		case EObjectReferencerType::Referenced:
+		else if (Vertex - NumObjectVertices < NumGCObjectNameVertices)
 		{
-			check(VertexArgument != Algo::Graph::InvalidVertex);
-			UObject* Object = VertexToObject[VertexArgument];
-			if (Object)
-			{
-				Object->GetPathName(nullptr, Builder);
-			}
-			else
-			{
-				Builder << TEXT("<UnknownObject>");
-			}
-			break;
+			Builder << TEXT("FGCObject ") << AllGCObjectNames[Vertex - NumObjectVertices];
 		}
-		default:
-			checkNoEntry();
-			break;
+		else
+		{
+			Builder << TEXT("InvalidVertex");
 		}
 	}
-
-private:
-	Algo::Graph::FVertex VertexArgument = Algo::Graph::InvalidVertex;
-	EObjectReferencerType LinkType = EObjectReferencerType::Unknown;
 };
 
+void FObjectReferencer::ToString(FStringBuilderBase& Builder, FObjectGraphProfileData& ProfileData)
+{
+	switch (GetLinkType())
+	{
+	case EObjectReferencerType::Unknown:
+		Builder << TEXT("<Unknown>");
+		break;
+	case EObjectReferencerType::Rooted:
+		Builder << TEXT("<Rooted>");
+		break;
+	case EObjectReferencerType::GCObjectRef:
+	{
+		check(VertexArgument != Algo::Graph::InvalidVertex);
+		check(ProfileData.NumObjectVertices <= VertexArgument && VertexArgument < ProfileData.NumObjectVertices + ProfileData.NumGCObjectNameVertices);
+		ProfileData.AppendVertexName(VertexArgument, Builder);
+		break;
+	}
+	case EObjectReferencerType::Referenced:
+	{
+		check(VertexArgument != Algo::Graph::InvalidVertex);
+		check(VertexArgument < ProfileData.NumObjectVertices);
+		ProfileData.AppendVertexName(VertexArgument, Builder);
+		break;
+	}
+	default:
+		checkNoEntry();
+		break;
+	}
+}
 /** An ObjectReferenceCollector to pass to Object->Serialize to collect references into an array. */
 class FArchiveGetReferences : public FArchiveUObject
 {
@@ -376,6 +417,49 @@ private:
 	TArray<UObject*>& ReferencedObjects;
 };
 
+/**
+ *  A ReferenceFinder used only when serializing FGCObject::GGCObjectReferencer.
+ * It captures the referencerName from GGCObjectReferencer for each UObject passed to it.
+ */
+class FGCObjectReferencerFinder : public FReferenceFinder
+{
+public:
+
+	FGCObjectReferencerFinder(TArray<UObject*>& InObjectArray, TMap<UObject*, FString>& InObjectReferencerNames)
+		: FReferenceFinder(InObjectArray)
+		, ObjectReferencerNames(InObjectReferencerNames)
+	{
+	}
+
+	virtual void HandleObjectReference(UObject*& InObject, const UObject* InReferencingObject, const FProperty* InReferencingProperty) override
+	{
+		// Avoid duplicate entries.
+		if (InObject != NULL)
+		{
+			// Many places that use FReferenceFinder expect the object to not be const.
+			UObject* Object = const_cast<UObject*>(InObject);
+			// do not add or recursively serialize objects that have already been added
+			bool bAlreadyExists;
+			ObjectSet.Add(Object, &bAlreadyExists);
+			if (!bAlreadyExists)
+			{
+				check(Object->IsValidLowLevel());
+				ObjectArray.Add(Object);
+				FString ReferencerName;
+				FGCObject::GGCObjectReferencer->GetReferencerName(Object, ReferencerName, true /* bOnlyIfAddingReferenced */);
+				if (!ReferencerName.IsEmpty())
+				{
+					ObjectReferencerNames.Add(Object, MoveTemp(ReferencerName));
+				}
+			}
+		}
+	}
+
+private:
+	TMap<UObject*, FString>& ObjectReferencerNames;
+	FGCObject* CurrentlySerializingObject;
+};
+
 /**
  * Given the list of AllObjects from e.g. a TObjectIterator, use serialization and other methods from Garbage Collection
  * to find all the dependencies of each Object.
@@ -384,7 +468,7 @@ private:
  */
 void ConstructObjectGraph(TConstArrayView<UObject*> AllObjects,
 	const TMap<UObject*, Algo::Graph::FVertex>& ObjectToVertex, TArray64<Algo::Graph::FVertex>& OutGraphBuffer,
-	TArray<TConstArrayView<Algo::Graph::FVertex>>& OutGraph)
+	TArray<TConstArrayView<Algo::Graph::FVertex>>& OutGraph, TMap<UObject*, FString>& OutGCObjectReferencerNames)
 {
 	using namespace Algo::Graph;
 
@@ -393,19 +477,21 @@ void ConstructObjectGraph(TConstArrayView<UObject*> AllObjects,
 	LooseEdges.SetNum(NumVertices);
 	TArray<UObject*> TargetObjects;
 	int32 NumEdges = 0;
+	OutGCObjectReferencerNames.Reset();
 
 	for (FVertex SourceVertex = 0; SourceVertex < NumVertices; ++SourceVertex)
 	{
 		UObject* SourceObject = AllObjects[SourceVertex];
 		TargetObjects.Reset();
 		{
-			FReferenceFinder Collector(TargetObjects);
 			if (SourceObject == FGCObject::GGCObjectReferencer)
 			{
+				FGCObjectReferencerFinder Collector(TargetObjects, OutGCObjectReferencerNames);
 				UGCObjectReferencer::AddReferencedObjects(FGCObject::GGCObjectReferencer, Collector);
 			}
 			else
 			{
+				FReferenceFinder Collector(TargetObjects);
 				FArchiveGetReferences Ar(SourceObject, TargetObjects);
 				if (SourceObject->GetClass())
 				{
@@ -420,6 +506,7 @@ void ConstructObjectGraph(TConstArrayView<UObject*> AllObjects,
 				}
 			}
 		}
+
 		if (TargetObjects.Num())
 		{
 			Algo::Sort(TargetObjects);
@@ -448,14 +535,11 @@ void ConstructObjectGraph(TConstArrayView<UObject*> AllObjects,
 	}
 }
 
-void DumpObjClassList(TConstArrayView<FWeakObjectPtr> InitialObjects)
+void ConstructObjectGraphProfileData(TConstArrayView<FWeakObjectPtr> InitialObjects, FObjectGraphProfileData& OutProfileData)
 {
 	using namespace Algo::Graph;
-
-	FOutputDevice& LogAr = *(GLog);
-
 	// Get the list of Objects
-	TArray<UObject*> AllObjects;
+	OutProfileData.AllObjects.Reset();
 	for (FThreadSafeObjectIterator Iter; Iter; ++Iter)
 	{
 		UObject* Object = *Iter;
@@ -463,98 +547,133 @@ void DumpObjClassList(TConstArrayView<FWeakObjectPtr> InitialObjects)
 		{
 			continue;
 		}
-		AllObjects.Add(Object);
+		OutProfileData.AllObjects.Add(Object);
 	}
 
 	// Convert Objects to Algo::Graph::FVertex to reduce graph search memory
-	int32 NumVertices = AllObjects.Num();
-	TMap<UObject*, FVertex> VertexOfObject;
-	for (FVertex Vertex = 0; Vertex < NumVertices; ++Vertex)
+	OutProfileData.NumObjectVertices = OutProfileData.AllObjects.Num();
+	OutProfileData.NumVertices = OutProfileData.NumObjectVertices;
+	OutProfileData.VertexOfObject.Reset();
+	for (FVertex Vertex = 0; Vertex < OutProfileData.NumObjectVertices; ++Vertex)
 	{
-		VertexOfObject.Add(AllObjects[Vertex], Vertex);
+		OutProfileData.VertexOfObject.Add(OutProfileData.AllObjects[Vertex], Vertex);
 	}
 
 	// Store for each vertex whether the vertex is new - not in InitialObjects
-	TBitArray<> IsNew(true, NumVertices);
+	OutProfileData.IsNew.Init(true, OutProfileData.NumObjectVertices);
 	for (const FWeakObjectPtr& InitialObjectWeak : InitialObjects)
 	{
 		UObject* InitialObject = InitialObjectWeak.Get();
 		if (InitialObject)
 		{
-			FVertex* Vertex = VertexOfObject.Find(InitialObject);
+			FVertex* Vertex = OutProfileData.VertexOfObject.Find(InitialObject);
 			if (Vertex)
 			{
-				IsNew[*Vertex] = false;
+				OutProfileData.IsNew[*Vertex] = false;
 			}
 		}
 	}
 
 	// Serialize objects to get dependencies and use them to create the ObjectGraph
-	TArray64<FVertex> ObjectGraphBuffer;
-	TArray<TConstArrayView<FVertex>> ObjectGraph;
-	ConstructObjectGraph(AllObjects, VertexOfObject, ObjectGraphBuffer, ObjectGraph);
+	TMap<UObject*, FString> GCObjectReferencerNames;
+	ConstructObjectGraph(OutProfileData.AllObjects, OutProfileData.VertexOfObject,
+		OutProfileData.ObjectGraphBuffer, OutProfileData.ObjectGraph,
+		GCObjectReferencerNames);
+
+	OutProfileData.GCObjectReferencerVertex = InvalidVertex;
+	OutProfileData.AliveReason.SetNum(OutProfileData.NumObjectVertices);
+	OutProfileData.RootOfVertex.SetNumUninitialized(OutProfileData.NumObjectVertices);
+	for (FVertex& Root : OutProfileData.RootOfVertex)
+	{
+		Root = InvalidVertex;
+	}
 
-	// Mark the objects that are rooted by IsRooted, and find any special vertices
-	FVertex GCObjectReferencerVertex = InvalidVertex;
-	TArray<FObjectReferencer> AliveReason;
-	AliveReason.SetNum(NumVertices);
-	for (FVertex Vertex = 0; Vertex < NumVertices; ++Vertex)
+	// Mark the objects that are rooted by IsRooted, and find the special GCObjectReferencerVertex
+	for (FVertex Vertex = 0; Vertex < OutProfileData.NumObjectVertices; ++Vertex)
 	{
-		UObject* Object = AllObjects[Vertex];
+		UObject* Object = OutProfileData.AllObjects[Vertex];
 		if (Object->IsRooted())
 		{
-			AliveReason[Vertex].Set(EObjectReferencerType::Rooted);
+			OutProfileData.AliveReason[Vertex].Set(EObjectReferencerType::Rooted);
+			OutProfileData.RootOfVertex[Vertex] = Vertex;
 		}
 		if (Object == FGCObject::GGCObjectReferencer)
 		{
-			GCObjectReferencerVertex = Vertex;
+			OutProfileData.GCObjectReferencerVertex = Vertex;
 		}
 	}
-
-	// Mark the objects that are rooted by GCObjectReferencerVertex
-	for (FVertex Vertex : ObjectGraph[GCObjectReferencerVertex])
+	check(OutProfileData.GCObjectReferencerVertex != InvalidVertex);
+
+	// Mark the objects that are rooted by GCObjectReferencerVertex, and construct a synthetic vertex
+	// for each of the referencer names reported by GCObjectReferencerVertex.
+	OutProfileData.GCObjectNameToVertex.Reset();
+	OutProfileData.AllGCObjectNames.Reset();
+	FString UnknownReferencer(TEXT("<Unknown>"));
+	for (FVertex Vertex : OutProfileData.ObjectGraph[OutProfileData.GCObjectReferencerVertex])
 	{
-		if (AliveReason[Vertex].GetLinkType() == EObjectReferencerType::Unknown)
+		if (OutProfileData.AliveReason[Vertex].GetLinkType() == EObjectReferencerType::Unknown)
 		{
-			AliveReason[Vertex].Set(EObjectReferencerType::GCObjectRef, Vertex);
+			UObject* Object = OutProfileData.AllObjects[Vertex];
+			FString* ReferencerName = &UnknownReferencer;
+			if (Object)
+			{
+				ReferencerName = GCObjectReferencerNames.Find(Object);
+				if (!ReferencerName)
+				{
+					ReferencerName = &UnknownReferencer;
+				}
+			}
+			FVertex& ReferencerVertex = OutProfileData.GCObjectNameToVertex.FindOrAdd(*ReferencerName);
+			if (ReferencerVertex == (FVertex)0) // Having value 0 means it was newly added by FindOrAdd
+			{
+				ReferencerVertex = OutProfileData.NumVertices++;
+				OutProfileData.AllGCObjectNames.Add(*ReferencerName);
+				check(OutProfileData.NumVertices == OutProfileData.AllObjects.Num() + OutProfileData.AllGCObjectNames.Num());
+			}
+			OutProfileData.AliveReason[Vertex].Set(EObjectReferencerType::GCObjectRef, ReferencerVertex);
+			OutProfileData.RootOfVertex[Vertex] = ReferencerVertex;
 		}
 	}
-	check(GCObjectReferencerVertex != InvalidVertex);
+	OutProfileData.NumObjectVertices = OutProfileData.AllObjects.Num();
+	OutProfileData.NumGCObjectNameVertices = OutProfileData.AllGCObjectNames.Num();
 
 	// Do a DFS to mark the referencer and root of all non-rooted objects
-	TArray<FVertex> RootOfVertex;
-	RootOfVertex.SetNumUninitialized(NumVertices);
-	for (FVertex& Root : RootOfVertex)
-	{
-		Root = InvalidVertex;
-	}
-
 	TArray<FVertex> Stack;
-	for (FVertex RootedVertex = 0; RootedVertex < NumVertices; ++RootedVertex)
+	for (FVertex PotentialRoot = 0; PotentialRoot < OutProfileData.NumObjectVertices; ++PotentialRoot)
 	{
-		if (AliveReason[RootedVertex].GetLinkType() == EObjectReferencerType::Unknown ||
-			RootedVertex == GCObjectReferencerVertex)
+		if (PotentialRoot == OutProfileData.GCObjectReferencerVertex ||
+			(OutProfileData.AliveReason[PotentialRoot].GetLinkType() != EObjectReferencerType::Rooted &&
+				OutProfileData.AliveReason[PotentialRoot].GetLinkType() != EObjectReferencerType::GCObjectRef))
 		{
 			continue;
 		}
+		FVertex RootVertex = OutProfileData.RootOfVertex[PotentialRoot];
 
-		RootOfVertex[RootedVertex] = RootedVertex;
 		Stack.Reset();
-		Stack.Add(RootedVertex);
+		Stack.Add(PotentialRoot);
 		while (!Stack.IsEmpty())
 		{
 			FVertex SourceVertex = Stack.Pop(false /* bAllowShrinking */);
-			for (FVertex TargetVertex : ObjectGraph[SourceVertex])
+			for (FVertex TargetVertex : OutProfileData.ObjectGraph[SourceVertex])
 			{
-				if (AliveReason[TargetVertex].GetLinkType() == EObjectReferencerType::Unknown)
+				if (OutProfileData.AliveReason[TargetVertex].GetLinkType() == EObjectReferencerType::Unknown)
 				{
-					AliveReason[TargetVertex].Set(EObjectReferencerType::Referenced, SourceVertex);
-					RootOfVertex[TargetVertex] = RootedVertex;
+					OutProfileData.AliveReason[TargetVertex].Set(EObjectReferencerType::Referenced, SourceVertex);
+					OutProfileData.RootOfVertex[TargetVertex] = RootVertex;
 					Stack.Add(TargetVertex);
 				}
 			}
 		}
 	}
+}
+
+void DumpObjClassList(TConstArrayView<FWeakObjectPtr> InitialObjects)
+{
+	using namespace Algo::Graph;
+
+	FOutputDevice& LogAr = *(GLog);
+	FObjectGraphProfileData ProfileData;
+	ConstructObjectGraphProfileData(InitialObjects, ProfileData);
 
 	// Count how many new objects of each class there are, and store all root objects that keep them in memory
 	struct FClassInfo
@@ -564,14 +683,14 @@ void DumpObjClassList(TConstArrayView<FWeakObjectPtr> InitialObjects)
 		UClass* Class = nullptr;
 	};
 	TMap<UClass*, FClassInfo> ClassInfos;
-	for (FVertex Vertex = 0; Vertex < NumVertices; ++Vertex)
+	for (FVertex Vertex = 0; Vertex < ProfileData.NumObjectVertices; ++Vertex)
 	{
 		// Ignore non-new objects
-		if (!IsNew[Vertex] || Vertex == GCObjectReferencerVertex)
+		if (!ProfileData.IsNew[Vertex] || Vertex == ProfileData.GCObjectReferencerVertex)
 		{
 			continue;
 		}
-		FObjectReferencer Link = AliveReason[Vertex];
+		FObjectReferencer Link = ProfileData.AliveReason[Vertex];
 		EObjectReferencerType LinkType = Link.GetLinkType();
 		// Ignore objects that have AliveReason unknown. This can occur if the objects were rooted during garbage
 		// collection but then asynchronous work RemovedThemFromRoot in between GC finishing and our call to IsRooted.
@@ -579,14 +698,14 @@ void DumpObjClassList(TConstArrayView<FWeakObjectPtr> InitialObjects)
 		{
 			continue;
 		}
-		UClass* Class = AllObjects[Vertex]->GetClass();
+		UClass* Class = ProfileData.AllObjects[Vertex]->GetClass();
 		if (!Class || !Class->IsNative())
 		{
 			continue;
 		}
 		FClassInfo& ClassInfo = ClassInfos.FindOrAdd(Class);
 		ClassInfo.Class = Class;
-		ClassInfo.Roots.FindOrAdd(RootOfVertex[Vertex], 0)++;
+		ClassInfo.Roots.FindOrAdd(ProfileData.RootOfVertex[Vertex], 0)++;
 		ClassInfo.Count++;
 	}
 
@@ -631,21 +750,106 @@ void DumpObjClassList(TConstArrayView<FWeakObjectPtr> InitialObjects)
 		{
 			RootObjectString.Reset();
 			RootObjectString.Appendf(TEXT("\t\t%6d: "), RootPair.Value);
-			AllObjects[RootPair.Key]->GetFullName(RootObjectString);
-			FObjectReferencer Link = AliveReason[RootPair.Key];
-			RootObjectString << TEXT(" <- ");
-			Link.ToString(RootObjectString, AllObjects);
-			while (Link.GetLinkType() == EObjectReferencerType::Referenced)
+			ProfileData.AppendVertexName(RootPair.Key, RootObjectString);
+			if (RootPair.Key < ProfileData.NumObjectVertices)
 			{
-				Link = AliveReason[Link.GetVertexArgument()];
+				FObjectReferencer Link = ProfileData.AliveReason[RootPair.Key];
+				while (Link.GetLinkType() == EObjectReferencerType::Referenced)
+				{
+					RootObjectString << TEXT(" <- ");
+					Link.ToString(RootObjectString, ProfileData);
+					Link = ProfileData.AliveReason[Link.GetVertexArgument()];
+				}
 				RootObjectString << TEXT(" <- ");
-				Link.ToString(RootObjectString, AllObjects);
+				Link.ToString(RootObjectString, ProfileData);
 			}
 			LogAr.Logf(TEXT("%s"), *RootObjectString);
 		}
 	}
 }
 
+void DumpPackageReferencers(TConstArrayView<UPackage*> Packages)
+{
+	using namespace Algo::Graph;
+
+	FOutputDevice& LogAr = *(GLog);
+	FObjectGraphProfileData ProfileData;
+	ConstructObjectGraphProfileData(TConstArrayView<FWeakObjectPtr>(), ProfileData);
+
+	// List all roots that cause any of the Packages to remain alive, and count how many packages each one causes
+	TMap<FVertex, int32> Roots;
+	TMap<FVertex, FVertex> RootExamples;
+	int32 Unexpected = 0;
+	for (UPackage* Package : Packages)
+	{
+		FVertex* FoundVertex = ProfileData.VertexOfObject.Find(Package);
+		if (!FoundVertex)
+		{
+			++Unexpected;
+			continue;
+		}
+		FVertex PackageVertex = *FoundVertex;
+		FVertex RootOfThisVertex = ProfileData.RootOfVertex[PackageVertex];
+		if (RootOfThisVertex == InvalidVertex)
+		{
+			++Unexpected;
+			continue;
+		}
+		int32& RootCount = Roots.FindOrAdd(RootOfThisVertex);
+		if (RootCount == 0)
+		{
+			RootExamples.Add(RootOfThisVertex, PackageVertex);
+		}
+		++RootCount;
+	}
+
+	LogAr.Logf(TEXT("Memory Analysis: Referencers of SoftGCPackages:"));
+	Roots.ValueSort([](int32 A, int32 B) { return A > B; });
+	for (TPair<FVertex, int32>& Pair : Roots)
+	{
+		TStringBuilder<256> ReferencerName;
+		ProfileData.AppendVertexName(Pair.Key, ReferencerName);
+		LogAr.Logf(TEXT("\t%5d: %s"), Pair.Value, *ReferencerName);
+		FVertex* ExampleVertexPtr = RootExamples.Find(Pair.Key);
+		if (ExampleVertexPtr)
+		{
+			TStringBuilder<256> Chain;
+			FObjectReferencer Link = ProfileData.AliveReason[*ExampleVertexPtr];
+			ProfileData.AllObjects[*ExampleVertexPtr]->GetFullName(Chain);
+			while (Link.GetLinkType() == EObjectReferencerType::Referenced)
+			{
+				Chain << TEXT(" <- ");
+				Link.ToString(Chain, ProfileData);
+				Link = ProfileData.AliveReason[Link.GetVertexArgument()];
+			}
+			Chain << TEXT(" <- ");
+			Link.ToString(Chain, ProfileData);
+			LogAr.Logf(TEXT("\t\t     Ex: %s"), *Chain);
+		}
+	}
+
+	if (Unexpected > 0)
+	{
+		LogAr.Logf(TEXT("Memory Analysis: Unknown referenced SoftGCPackages:"));
+		for (UPackage* Package : Packages)
+		{
+			FVertex* FoundVertex = ProfileData.VertexOfObject.Find(Package);
+			if (!FoundVertex)
+			{
+				LogAr.Logf(TEXT("%s: unknown, we did not create a vertex for the package"), *Package->GetName());
+				continue;
+			}
+			FVertex PackageVertex = *FoundVertex;
+
+			if (ProfileData.AliveReason[PackageVertex].GetLinkType() == EObjectReferencerType::Unknown)
+			{
+				LogAr.Logf(TEXT("%s: no reference found"), *Package->GetName());
+				continue;
+			}
+		}
+	}
+}
+
 } // namespace UE::Cook
 
 #if ENABLE_COOK_STATS
diff --git a/Engine/Source/Editor/UnrealEd/Private/Cooker/CookProfiling.h b/Engine/Source/Editor/UnrealEd/Private/Cooker/CookProfiling.h
--- a/Engine/Source/Editor/UnrealEd/Private/Cooker/CookProfiling.h
+++ b/Engine/Source/Editor/UnrealEd/Private/Cooker/CookProfiling.h
@@ -192,4 +192,8 @@ namespace UE::Cook
  */
 void DumpObjClassList(TConstArrayView<FWeakObjectPtr> SessionStartupObjects);
 
+/** Used for profiling after a soft garbage collect. Find the referencers keeping each package in memory. */
+void DumpPackageReferencers(TConstArrayView<UPackage*> Packages);
+
+
 }
\ No newline at end of file
diff --git a/Engine/Source/Editor/UnrealEd/Private/Cooker/CookTypes.cpp b/Engine/Source/Editor/UnrealEd/Private/Cooker/CookTypes.cpp
--- a/Engine/Source/Editor/UnrealEd/Private/Cooker/CookTypes.cpp
+++ b/Engine/Source/Editor/UnrealEd/Private/Cooker/CookTypes.cpp
@@ -187,6 +187,9 @@ FCbWriter& operator<<(FCbWriter& Writer, const UE::Cook::FInitializeConfigSettin
 	Writer << "MemoryMinFreeVirtual" << Value.MemoryMinFreeVirtual;
 	Writer << "MemoryMinFreePhysical" << Value.MemoryMinFreePhysical;
 	Writer << "MemoryTriggerGCAtPressureLevel" << static_cast<uint8>(Value.MemoryTriggerGCAtPressureLevel);
+	Writer << "bUseSoftGC" << Value.bUseSoftGC;
+	Writer << "SoftGCStartNumerator" << Value.SoftGCStartNumerator;
+	Writer << "SoftGCDenominator" << Value.SoftGCDenominator;
 	Writer << "MinFreeUObjectIndicesBeforeGC" << Value.MinFreeUObjectIndicesBeforeGC;
 	Writer << "MaxNumPackagesBeforePartialGC" << Value.MaxNumPackagesBeforePartialGC;
 	Writer << "ConfigSettingDenyList" << Value.ConfigSettingDenyList;
@@ -220,6 +223,9 @@ bool LoadFromCompactBinary(FCbFieldView Field, UE::Cook::FInitializeConfigSettin
 		OutValue.MemoryTriggerGCAtPressureLevel = FGenericPlatformMemoryStats::EMemoryPressureStatus::Unknown;
 		bOk = false;
 	}
+	bOk = LoadFromCompactBinary(Field["bUseSoftGC"], OutValue.bUseSoftGC) & bOk;
+	bOk = LoadFromCompactBinary(Field["SoftGCStartNumerator"], OutValue.SoftGCStartNumerator) & bOk;
+	bOk = LoadFromCompactBinary(Field["SoftGCDenominator"], OutValue.SoftGCDenominator) & bOk;
 	bOk = LoadFromCompactBinary(Field["MinFreeUObjectIndicesBeforeGC"], OutValue.MinFreeUObjectIndicesBeforeGC) & bOk;
 	bOk = LoadFromCompactBinary(Field["MaxNumPackagesBeforePartialGC"], OutValue.MaxNumPackagesBeforePartialGC) & bOk;
 	bOk = LoadFromCompactBinary(Field["ConfigSettingDenyList"], OutValue.ConfigSettingDenyList) & bOk;
@@ -246,6 +252,9 @@ void FInitializeConfigSettings::MoveOrCopy(SourceType&& Source, TargetType&& Tar
 	Target.MemoryMinFreeVirtual = Source.MemoryMinFreeVirtual;
 	Target.MemoryMinFreePhysical = Source.MemoryMinFreePhysical;
 	Target.MemoryTriggerGCAtPressureLevel = Source.MemoryTriggerGCAtPressureLevel;
+	Target.bUseSoftGC = Source.bUseSoftGC;
+	Target.SoftGCStartNumerator = Source.SoftGCStartNumerator;
+	Target.SoftGCDenominator = Source.SoftGCDenominator;
 	Target.MinFreeUObjectIndicesBeforeGC = Source.MinFreeUObjectIndicesBeforeGC;
 	Target.MaxNumPackagesBeforePartialGC = Source.MaxNumPackagesBeforePartialGC;
 	Target.ConfigSettingDenyList = MoveTempIfPossible(Source.ConfigSettingDenyList);
diff --git a/Engine/Source/Editor/UnrealEd/Private/Cooker/CookTypes.h b/Engine/Source/Editor/UnrealEd/Private/Cooker/CookTypes.h
--- a/Engine/Source/Editor/UnrealEd/Private/Cooker/CookTypes.h
+++ b/Engine/Source/Editor/UnrealEd/Private/Cooker/CookTypes.h
@@ -298,9 +298,12 @@ namespace UE::Cook
 		FGenericPlatformMemoryStats::EMemoryPressureStatus MemoryTriggerGCAtPressureLevel;
 		int32 MinFreeUObjectIndicesBeforeGC;
 		int32 MaxNumPackagesBeforePartialGC;
+		int32 SoftGCStartNumerator;
+		int32 SoftGCDenominator;
 		TArray<FString> ConfigSettingDenyList;
 		TMap<FName, int32> MaxAsyncCacheForType; // max number of objects of a specific type which are allowed to async cache at once
 		bool bHybridIterativeDebug = false;
+		bool bUseSoftGC = false;
 
 		friend FCbWriter& ::operator<<(FCbWriter& Writer, const UE::Cook::FInitializeConfigSettings& Value);
 		friend bool ::LoadFromCompactBinary(FCbFieldView Field, UE::Cook::FInitializeConfigSettings& Value);
diff --git a/Engine/Source/Editor/UnrealEd/Private/Cooker/CookWorkerClient.cpp b/Engine/Source/Editor/UnrealEd/Private/Cooker/CookWorkerClient.cpp
--- a/Engine/Source/Editor/UnrealEd/Private/Cooker/CookWorkerClient.cpp
+++ b/Engine/Source/Editor/UnrealEd/Private/Cooker/CookWorkerClient.cpp
@@ -12,6 +12,7 @@
 #include "CoreGlobals.h"
 #include "HAL/PlatformTime.h"
 #include "PackageResultsMessage.h"
+#include "PackageTracker.h"
 #include "Sockets.h"
 #include "SocketSubsystem.h"
 #include "WorkerRequestsRemote.h"
@@ -605,6 +606,13 @@ void FCookWorkerClient::AssignPackages(FAssignPackagesMessage& Message)
 			FInstigator(AssignData.Instigator));
 		PackageData.SendToState(EPackageState::Request, ESendFlags::QueueAddAndRemove);
 	}
+
+	if (!Message.PackageDatas.IsEmpty())
+	{
+		// Clear the SoftGC diagnostic ExpectedNeverLoadPackages because we have new assigned packages
+		// that we didn't consider during SoftGC
+		COTFS.PackageTracker->ClearExpectedNeverLoadPackages();
+	}
 }
 
 void FCookWorkerClient::Register(IMPCollector* Collector)
diff --git a/Engine/Source/Editor/UnrealEd/Private/Cooker/PackageTracker.cpp b/Engine/Source/Editor/UnrealEd/Private/Cooker/PackageTracker.cpp
--- a/Engine/Source/Editor/UnrealEd/Private/Cooker/PackageTracker.cpp
+++ b/Engine/Source/Editor/UnrealEd/Private/Cooker/PackageTracker.cpp
@@ -123,6 +123,10 @@ void FPackageTracker::NotifyUObjectCreated(const class UObjectBase* Object, int3
 				COOK_STAT(++Stats::NumInlineLoads);
 			}
 
+			if (ExpectedNeverLoadPackages.Contains(Package->GetFName()))
+			{
+				UE_LOG(LogCook, Verbose, TEXT("SoftGC PoorPerformance: Reloaded package %s."), *WriteToString<256>(Package->GetFName()));
+			}
 			FWriteScopeLock ScopeLock(Lock);
 			LoadedPackages.Add(Package);
 			NewPackages.Add(Package,
diff --git a/Engine/Source/Editor/UnrealEd/Private/Cooker/PackageTracker.h b/Engine/Source/Editor/UnrealEd/Private/Cooker/PackageTracker.h
--- a/Engine/Source/Editor/UnrealEd/Private/Cooker/PackageTracker.h
+++ b/Engine/Source/Editor/UnrealEd/Private/Cooker/PackageTracker.h
@@ -215,12 +215,23 @@ public:
 			Function(Package);
 		}
 	}
+	void AddExpectedNeverLoadPackages(TArrayView<FName> PackageNames)
+	{
+		FWriteScopeLock ScopeLock(Lock);
+		ExpectedNeverLoadPackages.Append(PackageNames);
+	}
+	void ClearExpectedNeverLoadPackages()
+	{
+		FWriteScopeLock ScopeLock(Lock);
+		ExpectedNeverLoadPackages.Empty();
+	}
 private:
 	// Protects data for thread-safety
 	FRWLock Lock;
 
 	// This is a complete list of currently loaded UPackages
 	TFastPointerSet<UPackage*> LoadedPackages;
+	TSet<FName> ExpectedNeverLoadPackages;
 
 	// This list contains the UPackages loaded since last call to GetNewPackages
 	TMap<UPackage*, FInstigator> NewPackages;
