MPCook: Fix error Package %s can only be cooked by a now-disconnected CookWorker. We have added a heartbeat message, and we now wait for a round-trip heartbeat from all cookworkers (which will ensure collection of any final discovered packages) before shutting down CookWorkers.
#rnx

diff --git a/Engine/Source/Editor/UnrealEd/Private/Cooker/CookDirector.cpp b/Engine/Source/Editor/UnrealEd/Private/Cooker/CookDirector.cpp
--- a/Engine/Source/Editor/UnrealEd/Private/Cooker/CookDirector.cpp
+++ b/Engine/Source/Editor/UnrealEd/Private/Cooker/CookDirector.cpp
@@ -124,6 +124,12 @@ FCookDirector::FCookDirector(UCookOnTheFlyServer& InCOTFS, int32 CookProcessCoun
 		{
 			HandleRetractionMessage(Context, bReadSuccessful, MoveTemp(Message));
 		}, TEXT("HandleRetractionMessage")));
+	Register(new IMPCollectorCbServerMessage<FHeartbeatMessage>([this]
+	(FMPCollectorServerMessageContext& Context, bool bReadSuccessful, FHeartbeatMessage&& Message)
+		{
+			HandleHeartbeatMessage(Context, bReadSuccessful, MoveTemp(Message));
+		}, TEXT("HandleHeartbeatMessage")));
+
 	LastTickTimeSeconds = FPlatformTime::Seconds();
 
 	FCookStatsManager::CookStatsCallbacks.AddRaw(this, &FCookDirector::LogCookStats);
@@ -462,6 +468,10 @@ void FCookDirector::TickFromSchedulerThread()
 	bool bAnyIdle = bLocalWorkerIdle;
 	LastTickTimeSeconds = CurrentTime;
 	LocalWorkerProfileData->UpdateIdle(bLocalWorkerIdle, DeltaTime);
+	
+	bool bSendHeartbeat;
+	int32 LocalHeartbeatNumber;
+	TickHeartbeat(false /* bForceHeartbeat */, CurrentTime, bSendHeartbeat, LocalHeartbeatNumber);
 
 	TArray<TRefCountPtr<FCookWorkerServer>, TInlineAllocator<16>> WorkersWithMessage;
 	bool bAllWorkersConnected = false;
@@ -483,6 +493,10 @@ void FCookDirector::TickFromSchedulerThread()
 			{
 				WorkersWithMessage.Add(RemoteWorker);
 			}
+			if (bSendHeartbeat)
+			{
+				RemoteWorker->SignalHeartbeat(ECookDirectorThread::SchedulerThread, LocalHeartbeatNumber);
+			}
 		}
 		for (TPair<FCookWorkerServer*, TRefCountPtr<FCookWorkerServer>>& Pair : ShuttingDownWorkers)
 		{
@@ -630,16 +644,40 @@ void FCookDirector::PumpCookComplete(bool& bCompleted)
 			for (TPair<int32, TRefCountPtr<FCookWorkerServer>>& Pair : RemoteWorkers)
 			{
 				FCookWorkerServer& RemoteWorker = *Pair.Value;
-				// MPCOOKTODO: Messages sent from the CookWorkers about discovered packages might come in after
-				// the last message sent about completed saves. These discovered packages can cause the cook to fall
-				// back to incomplete. Don't send CookComplete messages to the CookWorkers until all CookWorkers have
-				// reported they are idle and have no further messages to send.
 				if (RemoteWorker.NumAssignments() > 0)
 				{
 					bAllIdle = false;
 					break;
 				}
 			}
+
+			if (bAllIdle && FinalIdleHeartbeatFence == -1)
+			{
+				bool bSendHeartbeat;
+				double CurrentTime = FPlatformTime::Seconds();
+				TickHeartbeat(true /* bForceHeartbeat */, CurrentTime, bSendHeartbeat, FinalIdleHeartbeatFence);
+
+				for (TPair<int32, TRefCountPtr<FCookWorkerServer>>& Pair : RemoteWorkers)
+				{
+					FCookWorkerServer& RemoteWorker = *Pair.Value;
+					RemoteWorker.SignalHeartbeat(ECookDirectorThread::SchedulerThread, FinalIdleHeartbeatFence);
+				}
+				bAllIdle = false;
+			}
+			else
+			{
+				for (TPair<int32, TRefCountPtr<FCookWorkerServer>>& Pair : RemoteWorkers)
+				{
+					FCookWorkerServer& RemoteWorker = *Pair.Value;
+					if (RemoteWorker.GetLastReceivedHeartbeatNumber() < FinalIdleHeartbeatFence)
+					{
+						bAllIdle = false;
+						SetWorkersStalled(true);
+						break;
+					}
+				}
+			}
+
 			if (bAllIdle)
 			{
 				for (TPair<int32, TRefCountPtr<FCookWorkerServer>>& Pair : RemoteWorkers)
@@ -717,6 +755,9 @@ void FCookDirector::ShutdownCookSession()
 	bIsFirstAssignment = true;
 	bCookCompleteSent = false;
 	bWorkersActive = false;
+	HeartbeatNumber = 0;
+	NextHeartbeatTimeSeconds = 0.;
+	FinalIdleHeartbeatFence = -1;
 }
 
 void FCookDirector::Register(IMPCollector* Collector)
@@ -772,6 +813,53 @@ void FCookDirector::SetWorkersStalled(bool bInWorkersStalled)
 	}
 }
 
+void FCookDirector::TickHeartbeat(bool bForceHeartbeat, double CurrentTimeSeconds, bool& bOutSendHeartbeat,
+	int32& OutHeartbeatNumber)
+{
+	constexpr float HeartbeatPeriodSeconds = 30.f;
+
+	bOutSendHeartbeat = false;
+	OutHeartbeatNumber = HeartbeatNumber;
+	if (bForceHeartbeat)
+	{
+		bOutSendHeartbeat = true;
+	}
+	else if (NextHeartbeatTimeSeconds == 0.)
+	{
+		NextHeartbeatTimeSeconds = CurrentTimeSeconds + HeartbeatPeriodSeconds;
+	}
+	else if (CurrentTimeSeconds >= NextHeartbeatTimeSeconds)
+	{
+		bOutSendHeartbeat = true;
+	}
+
+	if (bOutSendHeartbeat)
+	{
+		checkf(HeartbeatNumber < MAX_int32, TEXT("Overflow"));
+		HeartbeatNumber++;
+		NextHeartbeatTimeSeconds = CurrentTimeSeconds + HeartbeatPeriodSeconds;
+	}
+}
+
+void FCookDirector::ResetFinalIdleHeartbeatFence()
+{
+	FinalIdleHeartbeatFence = -1;
+}
+
+void FCookDirector::HandleHeartbeatMessage(FMPCollectorServerMessageContext& Context, bool bReadSuccessful,
+	FHeartbeatMessage&& Message)
+{
+	if (!bReadSuccessful)
+	{
+		UE_LOG(LogCook, Error, TEXT("Corrupt HeartbeatMessage received from CookWorker %d. It will be ignored."),
+			Context.GetProfileId());
+		return;
+	}
+
+	Context.GetCookWorkerServer()->SetLastReceivedHeartbeatNumberInLock(Message.HeartbeatNumber);
+}
+
+
 FCookDirector::FPendingConnection::FPendingConnection(FPendingConnection&& Other)
 {
 	Swap(Socket, Other.Socket);
diff --git a/Engine/Source/Editor/UnrealEd/Private/Cooker/CookDirector.h b/Engine/Source/Editor/UnrealEd/Private/Cooker/CookDirector.h
--- a/Engine/Source/Editor/UnrealEd/Private/Cooker/CookDirector.h
+++ b/Engine/Source/Editor/UnrealEd/Private/Cooker/CookDirector.h
@@ -30,6 +30,7 @@ namespace UE::Cook { class FCookWorkerServer; }
 namespace UE::Cook { class FMPCollectorServerMessageContext; }
 namespace UE::Cook { class IMPCollector; }
 namespace UE::Cook { struct FCookWorkerProfileData; }
+namespace UE::Cook { struct FHeartbeatMessage; }
 namespace UE::Cook { struct FInitialConfigMessage; }
 namespace UE::Cook { struct FPackageData; }
 namespace UE::Cook { struct FRetractionResultsMessage; }
@@ -176,6 +177,14 @@ private:
 	/** Calls the configured LoadBalanceAlgorithm. Input Requests have been sorted by leaf to root load order. */
 	void LoadBalance(TConstArrayView<FWorkerId> SortedWorkers, TArrayView<FPackageData*> Requests,
 		TMap<FPackageData*, TArray<FPackageData*>>&& RequestGraph, TArray<FWorkerId>& OutAssignments);
+	/** Report whether it is time for a heartbeat message and update the timer data. */
+	void TickHeartbeat(bool bForceHeartbeat, double CurrentTimeSeconds, bool& bOutSendHeartbeat,
+		int32& OutHeartbeatNumber);
+	/** Reset the IdleHeartbeatFence when new idle-breaking data comes in. */
+	void ResetFinalIdleHeartbeatFence();
+	/** Log the occurrence of a heartbeat message from a CookWorker. */
+	void HandleHeartbeatMessage(FMPCollectorServerMessageContext& Context, bool bReadSuccessful,
+		FHeartbeatMessage&& Message);
 
 	/** Move the given worker from active workers to the list of workers shutting down. */
 	void AbortWorker(FWorkerId WorkerId, ECookDirectorThread TickThread);
@@ -213,6 +222,9 @@ private:
 	double WorkersStalledStartTimeSeconds = 0.;
 	double WorkersStalledWarnTimeSeconds = 0.;
 	double LastTickTimeSeconds = 0.;
+	double NextHeartbeatTimeSeconds = 0.;
+	int32 HeartbeatNumber = 0;
+	int32 FinalIdleHeartbeatFence = -1;
 	bool bWorkersInitialized = false;
 	bool bHasReducedMachineResources = false;
 	bool bIsFirstAssignment = true;
diff --git a/Engine/Source/Editor/UnrealEd/Private/Cooker/CookMPCollector.h b/Engine/Source/Editor/UnrealEd/Private/Cooker/CookMPCollector.h
--- a/Engine/Source/Editor/UnrealEd/Private/Cooker/CookMPCollector.h
+++ b/Engine/Source/Editor/UnrealEd/Private/Cooker/CookMPCollector.h
@@ -64,11 +64,13 @@ public:
 	const ITargetPlatform* IntToPlatform(uint8 PlatformAsInt) const;
 	FWorkerId GetWorkerId() const { return WorkerId; }
 	int32 GetProfileId() const { return ProfileId; }
+	FCookWorkerServer* GetCookWorkerServer() { return Server; }
 
 private:
 	TConstArrayView<const ITargetPlatform*> Platforms;
 	FWorkerId WorkerId;
 	int32 ProfileId;
+	FCookWorkerServer* Server = nullptr;
 
 	friend class FCookWorkerServer;
 };
diff --git a/Engine/Source/Editor/UnrealEd/Private/Cooker/CookWorkerClient.cpp b/Engine/Source/Editor/UnrealEd/Private/Cooker/CookWorkerClient.cpp
--- a/Engine/Source/Editor/UnrealEd/Private/Cooker/CookWorkerClient.cpp
+++ b/Engine/Source/Editor/UnrealEd/Private/Cooker/CookWorkerClient.cpp
@@ -41,6 +41,11 @@ FCookWorkerClient::FCookWorkerClient(UCookOnTheFlyServer& InCOTFS)
 		{
 			HandleAbortPackagesMessage(Context, bReadSuccessful, MoveTemp(Message));
 		}, TEXT("HandleAbortPackagesMessage")));
+	Register(new IMPCollectorCbClientMessage<FHeartbeatMessage>([this]
+	(FMPCollectorClientMessageContext& Context, bool bReadSuccessful, FHeartbeatMessage&& Message)
+		{
+			HandleHeartbeatMessage(Context, bReadSuccessful, MoveTemp(Message));
+		}, TEXT("HandleHeartbeatMessage")));
 }
 
 FCookWorkerClient::~FCookWorkerClient()
@@ -718,4 +723,16 @@ void FCookWorkerClient::HandleRetractionMessage(FMPCollectorClientMessageContext
 	SendMessage(ResultsMessage);
 }
 
+void FCookWorkerClient::HandleHeartbeatMessage(FMPCollectorClientMessageContext& Context, bool bReadSuccessful,
+	FHeartbeatMessage&& Message)
+{
+	if (!bReadSuccessful)
+	{
+		LogInvalidMessage(TEXT("HeartbeatMessage"));
+		return;
+	}
+
+	SendMessage(FHeartbeatMessage(Message.HeartbeatNumber));
+}
+
 } // namespace UE::Cook
\ No newline at end of file
diff --git a/Engine/Source/Editor/UnrealEd/Private/Cooker/CookWorkerClient.h b/Engine/Source/Editor/UnrealEd/Private/Cooker/CookWorkerClient.h
--- a/Engine/Source/Editor/UnrealEd/Private/Cooker/CookWorkerClient.h
+++ b/Engine/Source/Editor/UnrealEd/Private/Cooker/CookWorkerClient.h
@@ -15,6 +15,7 @@ namespace UE::Cook { struct FAbortPackagesMessage; }
 namespace UE::Cook { struct FAssignPackagesMessage; }
 namespace UE::Cook { struct FDirectorConnectionInfo; }
 namespace UE::Cook { struct FDiscoveredPackage; }
+namespace UE::Cook { struct FHeartbeatMessage; }
 namespace UE::Cook { struct FInitialConfigMessage; }
 namespace UE::Cook { struct FPackageRemoteResult; }
 namespace UE::Cook { struct FRetractionRequestMessage; }
@@ -112,6 +113,8 @@ private:
 		FAbortPackagesMessage&& Message);
 	void HandleRetractionMessage(FMPCollectorClientMessageContext& Context, bool bReadSuccessful,
 		FRetractionRequestMessage&& Message);
+	void HandleHeartbeatMessage(FMPCollectorClientMessageContext& Context, bool bReadSuccessful,
+		FHeartbeatMessage&& Message);
 
 private:
 	TSharedPtr<FInternetAddr> DirectorAddr;
diff --git a/Engine/Source/Editor/UnrealEd/Private/Cooker/CookWorkerServer.cpp b/Engine/Source/Editor/UnrealEd/Private/Cooker/CookWorkerServer.cpp
--- a/Engine/Source/Editor/UnrealEd/Private/Cooker/CookWorkerServer.cpp
+++ b/Engine/Source/Editor/UnrealEd/Private/Cooker/CookWorkerServer.cpp
@@ -213,6 +213,18 @@ bool FCookWorkerServer::HasMessages() const
 	return !ReceiveMessages.IsEmpty();
 }
 
+int32 FCookWorkerServer::GetLastReceivedHeartbeatNumber() const
+{
+	FScopeLock CommunicationScopeLock(&CommunicationLock);
+	return LastReceivedHeartbeatNumber;
+
+}
+void FCookWorkerServer::SetLastReceivedHeartbeatNumberInLock(int32 InHeartbeatNumber)
+{
+	LastReceivedHeartbeatNumber = InHeartbeatNumber;
+}
+
+
 bool FCookWorkerServer::TryHandleConnectMessage(FWorkerConnectMessage& Message, FSocket* InSocket, TArray<UE::CompactBinaryTCP::FMarshalledMessage>&& OtherPacketMessages, ECookDirectorThread TickThread)
 {
 	FCommunicationScopeLock ScopeLock(this, TickThread, ETickAction::Tick);
@@ -299,6 +311,20 @@ void FCookWorkerServer::TickCommunication(ECookDirectorThread TickThread)
 	}
 }
 
+void FCookWorkerServer::SignalHeartbeat(ECookDirectorThread TickThread, int32 HeartbeatNumber)
+{
+	FCommunicationScopeLock ScopeLock(this, TickThread, ETickAction::Tick);
+
+	switch (ConnectStatus)
+	{
+	case EConnectStatus::Connected:
+		SendMessageInLock(FHeartbeatMessage(HeartbeatNumber));
+		break;
+	default:
+		break;
+	}
+}
+
 void FCookWorkerServer::SignalCookComplete(ECookDirectorThread TickThread)
 {
 	FCommunicationScopeLock ScopeLock(this, TickThread, ETickAction::Tick);
@@ -522,6 +548,7 @@ void FCookWorkerServer::HandleReceiveMessagesInternal()
 			{
 				check(*Collector);
 				FMPCollectorServerMessageContext Context;
+				Context.Server = this;
 				Context.Platforms = OrderedSessionPlatforms;
 				Context.WorkerId = WorkerId;
 				Context.ProfileId = ProfileId;
@@ -629,6 +656,7 @@ void FCookWorkerServer::RecordResults(FPackageResultsMessage& Message)
 			COTFS.DemoteToIdle(*PackageData, ESendFlags::QueueAddAndRemove, Result.SuppressCookReason);
 		}
 	}
+	Director.ResetFinalIdleHeartbeatFence();
 }
 
 void FCookWorkerServer::LogInvalidMessage(const TCHAR* MessageTypeName)
@@ -654,6 +682,7 @@ void FCookWorkerServer::AddDiscoveredPackage(FDiscoveredPackage&& DiscoveredPack
 		PackageData.SetGenerated(true);
 		PackageData.SetWorkerAssignmentConstraint(GetWorkerId());
 	}
+	Director.ResetFinalIdleHeartbeatFence();
 	COTFS.QueueDiscoveredPackageData(PackageData, MoveTemp(DiscoveredPackage.Instigator));
 }
 
@@ -1001,4 +1030,21 @@ void FLogMessagesMessageHandler::Serialize(const TCHAR* V, ELogVerbosity::Type V
 	Serialize(V, Verbosity, Category);
 }
 
+FGuid FHeartbeatMessage::MessageType(TEXT("C08FFAF07BF34DD3A2FFB8A287CDDE83"));
+
+FHeartbeatMessage::FHeartbeatMessage(int32 InHeartbeatNumber)
+	: HeartbeatNumber(InHeartbeatNumber)
+{
+}
+
+void FHeartbeatMessage::Write(FCbWriter& Writer) const
+{
+	Writer << "H" << HeartbeatNumber;
+}
+
+bool FHeartbeatMessage::TryRead(FCbObjectView Object)
+{
+	return LoadFromCompactBinary(Object["H"], HeartbeatNumber);
+}
+
 }
\ No newline at end of file
diff --git a/Engine/Source/Editor/UnrealEd/Private/Cooker/CookWorkerServer.h b/Engine/Source/Editor/UnrealEd/Private/Cooker/CookWorkerServer.h
--- a/Engine/Source/Editor/UnrealEd/Private/Cooker/CookWorkerServer.h
+++ b/Engine/Source/Editor/UnrealEd/Private/Cooker/CookWorkerServer.h
@@ -71,6 +71,8 @@ public:
 
 	/** Periodic Tick function to send and receive messages to the Client. */
 	void TickCommunication(ECookDirectorThread TickThread);
+	/** Called when the COTFS wants to send a heartbeat message to the Client. */
+	void SignalHeartbeat(ECookDirectorThread TickThread, int32 HeartbeatNumber);
 	/** Called when the COTFS Server has detected all packages are complete. Tell the CookWorker to flush messages and exit. */
 	void SignalCookComplete(ECookDirectorThread TickThread);
 	/**
@@ -92,6 +94,11 @@ public:
 	/** Does this Server have any ReceivedMessages that need to be processed by the Scheduler thread? */
 	bool HasMessages() const;
 
+	/** Get the LastReceivedHeartbeatNumber. */
+	int32 GetLastReceivedHeartbeatNumber() const;
+	/** Set the LastReceivedHeartbeatNumber. Assumes lock is already entered; can only be called from with a HandleReceivedMessages callback */
+	void SetLastReceivedHeartbeatNumberInLock(int32 InHeartbeatNumber);
+
 private:
 	enum class EConnectStatus
 	{
@@ -177,6 +184,7 @@ private:
 	FTickState TickState;
 	uint32 CookWorkerProcessId = 0;
 	int32 ProfileId = 0;
+	int32 LastReceivedHeartbeatNumber = 0;
 	double ConnectStartTimeSeconds = 0.;
 	double ConnectTestStartTimeSeconds = 0.;
 	FWorkerId WorkerId = FWorkerId::Invalid();
@@ -354,4 +362,23 @@ private:
 	static FGuid MessageType;
 };
 
+/**
+ * Message from Director to CookWorker or CookWorker to Director that reports a heartbeat number, in addition to reporting the
+ * machine is still alive just by the presence of the message.
+ * The Director intiates a heartbeat message; the CookWorker always responds to a heartbeat message with its own heartbeat message
+ * in reply, with the same number.
+ */
+struct FHeartbeatMessage : public UE::CompactBinaryTCP::IMessage
+{
+public:
+	FHeartbeatMessage(int32 InHeartbeatNumber=-1);
+	virtual void Write(FCbWriter& Writer) const override;
+	virtual bool TryRead(FCbObjectView Object) override;
+	virtual FGuid GetMessageType() const override { return MessageType; }
+
+public:
+	int32 HeartbeatNumber;
+	static FGuid MessageType;
+};
+
 }
\ No newline at end of file
