MPCook: Fix bad robomerge; FPackageTracker.Lock is not yet present in this branch.
#rnx

diff --git a/Engine/Source/Editor/UnrealEd/Private/Cooker/PackageTracker.h b/Engine/Source/Editor/UnrealEd/Private/Cooker/PackageTracker.h
--- a/Engine/Source/Editor/UnrealEd/Private/Cooker/PackageTracker.h
+++ b/Engine/Source/Editor/UnrealEd/Private/Cooker/PackageTracker.h
@@ -217,12 +217,10 @@ public:
 	}
 	void AddExpectedNeverLoadPackages(TArrayView<FName> PackageNames)
 	{
-		FWriteScopeLock ScopeLock(Lock);
 		ExpectedNeverLoadPackages.Append(PackageNames);
 	}
 	void ClearExpectedNeverLoadPackages()
 	{
-		FWriteScopeLock ScopeLock(Lock);
 		ExpectedNeverLoadPackages.Empty();
 	}
 private:
