MPCook: Fix bad robomerge (2); removal of FPackageTracker.Lock should from other branch have been suppressed from reaching this branch.
#rnx

diff --git a/Engine/Source/Editor/UnrealEd/Private/Cooker/PackageTracker.h b/Engine/Source/Editor/UnrealEd/Private/Cooker/PackageTracker.h
--- a/Engine/Source/Editor/UnrealEd/Private/Cooker/PackageTracker.h
+++ b/Engine/Source/Editor/UnrealEd/Private/Cooker/PackageTracker.h
@@ -217,10 +217,12 @@ public:
 	}
 	void AddExpectedNeverLoadPackages(TArrayView<FName> PackageNames)
 	{
+		FWriteScopeLock ScopeLock(Lock);
 		ExpectedNeverLoadPackages.Append(PackageNames);
 	}
 	void ClearExpectedNeverLoadPackages()
 	{
+		FWriteScopeLock ScopeLock(Lock);
 		ExpectedNeverLoadPackages.Empty();
 	}
 private:
