Fix wrong IsValidChecked usage.

diff --git a/Engine/Source/Editor/UnrealEd/Private/FileHelpers.cpp b/Engine/Source/Editor/UnrealEd/Private/FileHelpers.cpp
--- a/Engine/Source/Editor/UnrealEd/Private/FileHelpers.cpp
+++ b/Engine/Source/Editor/UnrealEd/Private/FileHelpers.cpp
@@ -3052,7 +3052,7 @@ EAutosaveContentPackagesResult::Type FEditorFileUtils::AutosaveMapEx(const FStri
 							// @todo_ow: Find better way
 							if (Object->IsA<AActor>() || Object->IsA<UActorFolder>())
 							{
-								if (IsValidChecked(Object))
+								if (IsValid(Object))
 								{
 									ExternalPackagesToSave.Add(ExternalPackage);
 									return false;
@@ -4818,18 +4818,18 @@ void FEditorFileUtils::GetDirtyWorldPackages(TArray<UPackage*>& OutDirtyPackages
 							{
 								bActorPackageNeedsToSave = false;
 								ForEachObjectWithPackage(ExternalPackage, [&bActorPackageNeedsToSave](UObject* Object)
+								{
+									// @todo_ow: Find better way
+									if (Object->IsA<AActor>() || Object->IsA<UActorFolder>())
 									{
-										// @todo_ow: Find better way
-										if (Object->IsA<AActor>() || Object->IsA<UActorFolder>())
+										if (IsValid(Object))
 										{
-											if (IsValidChecked(Object))
-											{
-												bActorPackageNeedsToSave = true;
-												return false;
-											}
+											bActorPackageNeedsToSave = true;
+											return false;
 										}
-										return true;
-									}, false);
+									}
+									return true;
+								}, false);
 							}
 
 							// Filter out Actors that might be unsaved (/Temp folder)
