Cooker memory: Remove some private functions that are no longer referenced after the SoftGC change. Change debug name of COTFS's AddReferencedObjects from PackageDatas to CookOnTheFlyServer.
#rnx

diff --git a/Engine/Source/Editor/UnrealEd/Classes/CookOnTheSide/CookOnTheFlyServer.h b/Engine/Source/Editor/UnrealEd/Classes/CookOnTheSide/CookOnTheFlyServer.h
--- a/Engine/Source/Editor/UnrealEd/Classes/CookOnTheSide/CookOnTheFlyServer.h
+++ b/Engine/Source/Editor/UnrealEd/Classes/CookOnTheSide/CookOnTheFlyServer.h
@@ -392,10 +392,6 @@ private:
 	double LastProgressDisplayTime = 0;
 	double LastDiagnosticsDisplayTime = 0;
 
-	/** Get dependencies for package */
-	const TArray<FName>& GetFullPackageDependencies(const FName& PackageName) const;
-	mutable TMap<FName, TArray<FName>> CachedFullPackageDependencies;
-
 	/** Cached copy of asset registry */
 	IAssetRegistry* AssetRegistry = nullptr;
 
@@ -898,11 +894,6 @@ private:
 	void BeginCookPackageWriters(FBeginCookContext& BeginContext);
 	void BeginCookDirector(FBeginCookContext& BeginContext);
 	void InitializeSession();
-	void PreGarbageCollectImpl(TArray<UPackage*>& GCKeepPackages, TArray<UE::Cook::FPackageData*>& GCKeepPackageDatas,
-		TArray<UObject*>& LocalGCKeepObjects);
-	void PostGarbageCollectImpl(TArray<UObject*>& LocalGCKeepObjects);
-	void GetDirectAndTransitiveResourceSize(FResourceSizeEx& OutDirectSize, FResourceSizeEx& OutTransitiveSize,
-		int64& OutNumDirectPackages, int64& OutNumTransitivePackages, TSet<UPackage*>&& DirectPackages);
 
 	//////////////////////////////////////////////////////////////////////////
 	// cook by the book specific functions
diff --git a/Engine/Source/Editor/UnrealEd/Private/CookOnTheFlyServer.cpp b/Engine/Source/Editor/UnrealEd/Private/CookOnTheFlyServer.cpp
--- a/Engine/Source/Editor/UnrealEd/Private/CookOnTheFlyServer.cpp
+++ b/Engine/Source/Editor/UnrealEd/Private/CookOnTheFlyServer.cpp
@@ -4606,6 +4606,8 @@ void UCookOnTheFlyServer::ClearGarbageCollectType()
 
 void UCookOnTheFlyServer::PreGarbageCollect()
 {
+	using namespace UE::Cook;
+
 	if (!IsInSession())
 	{
 		return;
@@ -4615,13 +4617,6 @@ void UCookOnTheFlyServer::PreGarbageCollect()
 	VirtualMemoryHistory.AddInstance(FPlatformMemory::GetStats().UsedVirtual);
 	TArray<UPackage*> GCKeepPackages;
 	TArray<UE::Cook::FPackageData*> GCKeepPackageDatas;
-	PreGarbageCollectImpl(GCKeepPackages, GCKeepPackageDatas, GCKeepObjects);
-}
-
-void UCookOnTheFlyServer::PreGarbageCollectImpl(TArray<UPackage*>& GCKeepPackages,
-	TArray<UE::Cook::FPackageData*>& GCKeepPackageDatas, TArray<UObject*>& LocalGCKeepObjects)
-{
-	using namespace UE::Cook;
 
 #if COOK_CHECKSLOW_PACKAGEDATA
 	// Verify that only packages in the save state have pointers to objects
@@ -4633,7 +4628,7 @@ void UCookOnTheFlyServer::PreGarbageCollectImpl(TArray<UPackage*>& GCKeepPackage
 	if (SavingPackageData)
 	{
 		check(SavingPackageData->GetPackage());
-		LocalGCKeepObjects.Add(SavingPackageData->GetPackage());
+		GCKeepObjects.Add(SavingPackageData->GetPackage());
 	}
 
 
@@ -4650,7 +4645,7 @@ void UCookOnTheFlyServer::PreGarbageCollectImpl(TArray<UPackage*>& GCKeepPackage
 		if (Info)
 		{
 			bool bShouldDemote;
-			Generator->PreGarbageCollect(*Info, LocalGCKeepObjects, GCKeepPackages, GCKeepPackageDatas, bShouldDemote);
+			Generator->PreGarbageCollect(*Info, GCKeepObjects, GCKeepPackages, GCKeepPackageDatas, bShouldDemote);
 			if (bShouldDemote)
 			{
 				ReleaseCookedPlatformData(*PackageData, UE::Cook::EReleaseSaveReason::Demoted);
@@ -4689,7 +4684,7 @@ void UCookOnTheFlyServer::PreGarbageCollectImpl(TArray<UPackage*>& GCKeepPackage
 		{
 			UObject* Object = Pending.Object.Get();
 			check(Object); // Otherwise PollIsComplete would have returned true
-			LocalGCKeepObjects.Add(Object);
+			GCKeepObjects.Add(Object);
 		}
 	}
 
@@ -4777,18 +4772,18 @@ void UCookOnTheFlyServer::PreGarbageCollectImpl(TArray<UPackage*>& GCKeepPackage
 			});
 	}
 
-	// Add packages and all RF_Public objects outered to them to LocalGCKeepObjects
+	// Add packages and all RF_Public objects outered to them to GCKeepObjects
 	TArray<UObject*> ObjectsWithOuter;
 	for (UPackage* Package : GCKeepPackages)
 	{
-		LocalGCKeepObjects.Add(Package);
+		GCKeepObjects.Add(Package);
 		ObjectsWithOuter.Reset();
 		GetObjectsWithOuter(Package, ObjectsWithOuter);
 		for (UObject* Obj : ObjectsWithOuter)
 		{
 			if (IsValidChecked(Obj) && Obj->HasAnyFlags(RF_Public))
 			{
-				LocalGCKeepObjects.Add(Obj);
+				GCKeepObjects.Add(Obj);
 			}
 		}
 	}
@@ -4808,14 +4803,10 @@ void UCookOnTheFlyServer::CookerAddReferencedObjects(FReferenceCollector& Collec
 
 void UCookOnTheFlyServer::PostGarbageCollect()
 {
+	using namespace UE::Cook;
+
 	NumObjectsHistory.AddInstance(GUObjectArray.GetObjectArrayNumMinusAvailable());
 	VirtualMemoryHistory.AddInstance(FPlatformMemory::GetStats().UsedVirtual);
-	PostGarbageCollectImpl(GCKeepObjects);
-}
-
-void UCookOnTheFlyServer::PostGarbageCollectImpl(TArray<UObject*>& LocalGCKeepObjects)
-{
-	using namespace UE::Cook;
 
 	// If any PackageDatas with ObjectPointers had any of their object pointers deleted out from under them, demote them back to request
 	TArray<FPackageData*> Demotes;
@@ -4840,7 +4831,7 @@ void UCookOnTheFlyServer::PostGarbageCollectImpl(TArray<UObject*>& LocalGCKeepOb
 	// However, if garbage collection deleted the package WHILE WE WERE SAVING IT, then we have problems.
 	check(!SavingPackageData || SavingPackageData->GetPackage() != nullptr);
 
-	LocalGCKeepObjects.Empty();
+	GCKeepObjects.Empty();
 
 	for (FPackageData* PackageData : *PackageDatas.Get())
 	{
@@ -5065,75 +5056,6 @@ void UCookOnTheFlyServer::EvaluateGarbageCollectionResults(bool bWasDueToOOM, bo
 	}
 }
 
-void UCookOnTheFlyServer::GetDirectAndTransitiveResourceSize(FResourceSizeEx& OutDirectSize,
-	FResourceSizeEx& OutTransitiveSize, int64& OutNumDirectPackages, int64& OutNumTransitivePackages,
-	TSet<UPackage*>&& DirectPackages)
-{
-	OutDirectSize = FResourceSizeEx(EResourceSizeMode::Exclusive);
-
-	TArray<UObject*> ObjectsInPackage;
-	auto AddObjectSize = [](UObject* Object, FResourceSizeEx& InResourceSize)
-	{
-		FArchiveCountMem MemoryCount(Object, true);
-		InResourceSize.AddDedicatedSystemMemoryBytes(MemoryCount.GetMax());
-		Object->GetResourceSizeEx(InResourceSize);
-	};
-	auto AddPackageSize = [&ObjectsInPackage, &AddObjectSize](UPackage* InPackage, FResourceSizeEx& InResourceSize)
-	{
-		InPackage->GetResourceSizeEx(InResourceSize);
-		AddObjectSize(InPackage, InResourceSize);
-
-		ObjectsInPackage.Reset();
-		GetObjectsWithPackage(InPackage, ObjectsInPackage);
-		for (UObject* Object : ObjectsInPackage)
-		{
-			AddObjectSize(Object, InResourceSize);
-		}
-	};
-
-	OutNumTransitivePackages = 0;
-	OutNumDirectPackages = 0;
-	TRingBuffer<UPackage*> NextPackages;
-	NextPackages.Reserve(DirectPackages.Num());
-	for (UPackage* Package : DirectPackages)
-	{
-		check(Package);
-		if (!Package->IsRooted())
-		{
-			AddPackageSize(Package, OutDirectSize);
-			NextPackages.Add(Package);
-			++OutNumDirectPackages;
-			++OutNumTransitivePackages;
-		}
-	}
-
-	OutTransitiveSize = OutDirectSize;
-	TSet<UPackage*>& TransitivePackages = DirectPackages;
-	TArray<FName> Dependencies;
-	while (!NextPackages.IsEmpty())
-	{
-		UPackage* NextPackage = NextPackages.PopFrontValue();
-		Dependencies.Reset();
-		AssetRegistry->GetDependencies(NextPackage->GetFName(), Dependencies, UE::AssetRegistry::EDependencyCategory::Package, UE::AssetRegistry::EDependencyQuery::Hard);
-		for (FName DependencyName : Dependencies)
-		{
-			UPackage* DependencyPackage = FindObject<UPackage>(FTopLevelAssetPath(DependencyName, NAME_None));
-			if (!DependencyPackage)
-			{
-				continue;
-			}
-			bool bAlreadyInSet;
-			TransitivePackages.Add(DependencyPackage, &bAlreadyInSet);
-			if (!bAlreadyInSet && !DependencyPackage->IsRooted())
-			{
-				++OutNumTransitivePackages;
-				NextPackages.Add(DependencyPackage);
-				AddPackageSize(DependencyPackage, OutTransitiveSize);
-			}
-		}
-	}
-}
-
 void UCookOnTheFlyServer::OnObjectModified( UObject *ObjectMoving )
 {
 	if (IsGarbageCollecting())
@@ -5319,86 +5241,6 @@ double UCookOnTheFlyServer::GetIdleTimeToGC() const
 	}
 }
 
-const TArray<FName>& UCookOnTheFlyServer::GetFullPackageDependencies(const FName& PackageName ) const
-{
-	TArray<FName>* PackageDependencies = CachedFullPackageDependencies.Find(PackageName);
-	if ( !PackageDependencies )
-	{
-		static const FName NAME_CircularReference(TEXT("CircularReference"));
-		static int32 UniqueArrayCounter = 0;
-		++UniqueArrayCounter;
-		FName CircularReferenceArrayName = FName(NAME_CircularReference,UniqueArrayCounter);
-		{
-			// can't initialize the PackageDependencies array here because we call GetFullPackageDependencies below and that could recurse and resize CachedFullPackageDependencies
-			TArray<FName>& TempPackageDependencies = CachedFullPackageDependencies.Add(PackageName); // IMPORTANT READ ABOVE COMMENT
-			// initialize TempPackageDependencies to a dummy dependency so that we can detect circular references
-			TempPackageDependencies.Add(CircularReferenceArrayName);
-			// when someone finds the circular reference name they look for this array name in the CachedFullPackageDependencies map
-			// and add their own package name to it, so that they can get fixed up 
-			CachedFullPackageDependencies.Add(CircularReferenceArrayName);
-		}
-
-		TArray<FName> ChildDependencies;
-		if ( AssetRegistry->GetDependencies(PackageName, ChildDependencies, UE::AssetRegistry::EDependencyCategory::Package) )
-		{
-			TArray<FName> Dependencies = ChildDependencies;
-			Dependencies.AddUnique(PackageName);
-			for ( const FName& ChildDependency : ChildDependencies)
-			{
-				const TArray<FName>& ChildPackageDependencies = GetFullPackageDependencies(ChildDependency);
-				for ( const FName& ChildPackageDependency : ChildPackageDependencies )
-				{
-					if ( ChildPackageDependency == CircularReferenceArrayName )
-					{
-						continue;
-					}
-
-					if ( ChildPackageDependency.GetComparisonIndex() == NAME_CircularReference.GetComparisonIndex() )
-					{
-						// add our self to the package which we are circular referencing
-						TArray<FName>& TempCircularReference = CachedFullPackageDependencies.FindChecked(ChildPackageDependency);
-						TempCircularReference.AddUnique(PackageName); // add this package name so that it's dependencies get fixed up when the outer loop returns
-					}
-
-					Dependencies.AddUnique(ChildPackageDependency);
-				}
-			}
-
-			// all these packages referenced us apparently so fix them all up
-			const TArray<FName>& PackagesForFixup = CachedFullPackageDependencies.FindChecked(CircularReferenceArrayName);
-			for ( const FName& FixupPackage : PackagesForFixup )
-			{
-				TArray<FName> &FixupList = CachedFullPackageDependencies.FindChecked(FixupPackage);
-				// check( FixupList.Contains( CircularReferenceArrayName) );
-				ensure( FixupList.Remove(CircularReferenceArrayName) == 1 );
-				for( const FName& AdditionalDependency : Dependencies )
-				{
-					FixupList.AddUnique(AdditionalDependency);
-					if ( AdditionalDependency.GetComparisonIndex() == NAME_CircularReference.GetComparisonIndex() )
-					{
-						// add our self to the package which we are circular referencing
-						TArray<FName>& TempCircularReference = CachedFullPackageDependencies.FindChecked(AdditionalDependency);
-						TempCircularReference.AddUnique(FixupPackage); // add this package name so that it's dependencies get fixed up when the outer loop returns
-					}
-				}
-			}
-			CachedFullPackageDependencies.Remove(CircularReferenceArrayName);
-
-			PackageDependencies = CachedFullPackageDependencies.Find(PackageName);
-			check(PackageDependencies);
-
-			Swap(*PackageDependencies, Dependencies);
-		}
-		else
-		{
-			PackageDependencies = CachedFullPackageDependencies.Find(PackageName);
-			PackageDependencies->Add(PackageName);
-		}
-	}
-
-	return *PackageDependencies;
-}
-
 void UCookOnTheFlyServer::BeginDestroy()
 {
 	ShutdownCookOnTheFly();
diff --git a/Engine/Source/Editor/UnrealEd/Private/Cooker/CookPackageData.cpp b/Engine/Source/Editor/UnrealEd/Private/Cooker/CookPackageData.cpp
--- a/Engine/Source/Editor/UnrealEd/Private/Cooker/CookPackageData.cpp
+++ b/Engine/Source/Editor/UnrealEd/Private/Cooker/CookPackageData.cpp
@@ -2150,7 +2150,7 @@ void FPackageDatas::OnAssetRegistryGenerated(IAssetRegistry& InAssetRegistry)
 
 FString FPackageDatas::GetReferencerName() const
 {
-	return TEXT("FPackageDatas");
+	return TEXT("CookOnTheFlyServer");
 }
 
 void FPackageDatas::AddReferencedObjects(FReferenceCollector& Collector)
