Fixed stat commands in editor not executing because the viewport is nullptr: Revert Exec_Editor change on UnrealEdEngine since we need to go through it before engine exec.

diff --git a/Engine/Source/Editor/UnrealEd/Classes/Editor/UnrealEdEngine.h b/Engine/Source/Editor/UnrealEd/Classes/Editor/UnrealEdEngine.h
--- a/Engine/Source/Editor/UnrealEd/Classes/Editor/UnrealEdEngine.h
+++ b/Engine/Source/Editor/UnrealEd/Classes/Editor/UnrealEdEngine.h
@@ -198,11 +198,9 @@ public:
 	//~ End UEditorEngine Interface 
 	
 	//~ Begin FExec Interface
-protected:
-	virtual bool Exec_Editor( UWorld* InWorld, const TCHAR* Cmd, FOutputDevice& Ar=*GLog ) override;
+	virtual bool Exec( UWorld* InWorld, const TCHAR* Cmd, FOutputDevice& Ar=*GLog ) override;
 	//~ End FExec Interface
 
-public:
 
 	//~ Begin UEngine Interface.
 	virtual void Init(IEngineLoop* InEngineLoop) override;
diff --git a/Engine/Source/Editor/UnrealEd/Private/UnrealEdSrv.cpp b/Engine/Source/Editor/UnrealEd/Private/UnrealEdSrv.cpp
--- a/Engine/Source/Editor/UnrealEd/Private/UnrealEdSrv.cpp
+++ b/Engine/Source/Editor/UnrealEd/Private/UnrealEdSrv.cpp
@@ -629,7 +629,7 @@ bool UUnrealEdEngine::HandleDisasmScriptCommand(const TCHAR* Str, FOutputDevice&
 	return true;
 }
 
-bool UUnrealEdEngine::Exec_Editor( UWorld* InWorld, const TCHAR* Stream, FOutputDevice& Ar )
+bool UUnrealEdEngine::Exec( UWorld* InWorld, const TCHAR* Stream, FOutputDevice& Ar )
 {
 	const TCHAR* Str = Stream;
 	// disallow set commands in the editor as that modifies the default object, affecting object serialization
@@ -658,7 +658,7 @@ bool UUnrealEdEngine::Exec_Editor( UWorld* InWorld, const TCHAR* Stream, FOutput
 		GStatProcessingViewportClient = GLastKeyLevelEditingViewportClient ? GLastKeyLevelEditingViewportClient : GCurrentLevelEditingViewportClient;
 	}
 
-	bool bExecSucceeded = UEditorEngine::Exec_Editor( InWorld, Stream, Ar );
+	bool bExecSucceeded = UEditorEngine::Exec( InWorld, Stream, Ar );
 
 	GStatProcessingViewportClient = NULL;
 
