MPCook: Fix crash due to writing PackageDatas from communication thread without thread synchronization.
Downgrade message about a crashed CookWorker from error to warning.
#rnx

diff --git a/Engine/Source/Editor/UnrealEd/Private/Cooker/CookDirector.cpp b/Engine/Source/Editor/UnrealEd/Private/Cooker/CookDirector.cpp
--- a/Engine/Source/Editor/UnrealEd/Private/Cooker/CookDirector.cpp
+++ b/Engine/Source/Editor/UnrealEd/Private/Cooker/CookDirector.cpp
@@ -505,6 +505,7 @@ void FCookDirector::TickFromSchedulerThread()
 				WorkersWithMessage.Add(Pair.Value);
 			}
 		}
+		ReassignAbortedPackages(DeferredPackagesToReassign);
 	}
 	if (bAllWorkersConnected && (RetractionHandler.IsValid() || bAnyIdle))
 	{
@@ -1295,24 +1296,40 @@ void FCookDirector::AbortWorker(FWorkerId WorkerId, ECookDirectorThread TickThre
 			return;
 		}
 	}
-	TSet<FPackageData*> PackagesToReassign;
-	RemoteWorker->AbortAllAssignments(PackagesToReassign, TickThread);
+	TSet<FPackageData*> PackagesToReassignSet;
+	RemoteWorker->AbortAllAssignments(PackagesToReassignSet, TickThread);
 	if (!RemoteWorker->IsShuttingDown())
 	{
-		RemoteWorker->AbortWorker(PackagesToReassign, TickThread);
+		RemoteWorker->AbortWorker(PackagesToReassignSet, TickThread);
 	}
-	for (FPackageData* PackageData : PackagesToReassign)
+	TArray<FPackageData*> PackagesToReassign = PackagesToReassignSet.Array();
+
+	if (TickThread == ECookDirectorThread::SchedulerThread)
 	{
-		check(PackageData->IsInProgress()); // Packages that were assigned to a worker should be in the AssignedToWorker state
-		PackageData->SetWorkerAssignment(FWorkerId::Invalid());
-		PackageData->SendToState(UE::Cook::EPackageState::Request, ESendFlags::QueueAddAndRemove);
+		ReassignAbortedPackages(PackagesToReassign);
 	}
 	{
 		FScopeLock RemoteWorkersScopeLock(&CommunicationLock);
 		TRefCountPtr<FCookWorkerServer>& Existing = ShuttingDownWorkers.FindOrAdd(RemoteWorker.GetReference());
 		check(!Existing); // We should not be able to abort a worker twice because we removed it from RemoteWorkers above
 		Existing = MoveTemp(RemoteWorker);
+
+		if (TickThread != ECookDirectorThread::SchedulerThread)
+		{
+			DeferredPackagesToReassign.Append(PackagesToReassign);
+		}
+	}
+}
+
+void FCookDirector::ReassignAbortedPackages(TArray<FPackageData*>& PackagesToReassign)
+{
+	for (FPackageData* PackageData : PackagesToReassign)
+	{
+		check(PackageData->IsInProgress()); // Packages that were assigned to a worker should be in the AssignedToWorker state
+		PackageData->SetWorkerAssignment(FWorkerId::Invalid());
+		PackageData->SendToState(UE::Cook::EPackageState::Request, ESendFlags::QueueAddAndRemove);
 	}
+	PackagesToReassign.Empty();
 }
 
 void FCookDirector::ConstructReadonlyThreadVariables()
diff --git a/Engine/Source/Editor/UnrealEd/Private/Cooker/CookDirector.h b/Engine/Source/Editor/UnrealEd/Private/Cooker/CookDirector.h
--- a/Engine/Source/Editor/UnrealEd/Private/Cooker/CookDirector.h
+++ b/Engine/Source/Editor/UnrealEd/Private/Cooker/CookDirector.h
@@ -188,6 +188,9 @@ private:
 
 	/** Move the given worker from active workers to the list of workers shutting down. */
 	void AbortWorker(FWorkerId WorkerId, ECookDirectorThread TickThread);
+	/** Send the given packages from an aborted worker back to the CookOnTheFlyServer for reassignment. */
+	void ReassignAbortedPackages(TArray<FPackageData*>& PackagesToReassign);
+
 	/**
 	 * Periodically update whether (1) local server is done and (2) no results from cookworkers have come in.
 	 * Send warning when it goes on too long.
@@ -250,6 +253,7 @@ private:
 	// Data shared between SchedulerThread and CommunicationThread that can only be accessed inside CommunicationLock
 	TMap<int32, TRefCountPtr<FCookWorkerServer>> RemoteWorkers;
 	TMap<FCookWorkerServer*, TRefCountPtr<FCookWorkerServer>> ShuttingDownWorkers;
+	TArray<FPackageData*> DeferredPackagesToReassign;
 	TUniquePtr<FRetractionHandler> RetractionHandler;
 	bool bWorkersActive = false;
 
diff --git a/Engine/Source/Editor/UnrealEd/Private/Cooker/CookWorkerServer.cpp b/Engine/Source/Editor/UnrealEd/Private/Cooker/CookWorkerServer.cpp
--- a/Engine/Source/Editor/UnrealEd/Private/Cooker/CookWorkerServer.cpp
+++ b/Engine/Source/Editor/UnrealEd/Private/Cooker/CookWorkerServer.cpp
@@ -470,7 +470,7 @@ void FCookWorkerServer::PumpReceiveMessages()
 	EConnectionStatus SocketStatus = TryReadPacket(Socket, ReceiveBuffer, Messages);
 	if (SocketStatus != EConnectionStatus::Okay && SocketStatus != EConnectionStatus::Incomplete)
 	{
-		UE_LOG(LogCook, Error, TEXT("CookWorkerServer %d failed to read from socket, we will shutdown the remote process. Assigned packages will be returned to the director."),
+		UE_LOG(LogCook, Warning, TEXT("CookWorkerServer %d failed to read from socket, we will shutdown the remote process. Assigned packages will be returned to the director."),
 			ProfileId);
 		SendToState(EConnectStatus::WaitForDisconnect);
 		bTerminateImmediately = true;
