Replaces blocking dialog when copy/pasting landscape actor by a slowtask.
Introduces slow task dialog when importing/exporting heavy actors.

diff --git a/Engine/Plugins/Enterprise/DataprepEditor/Source/DataprepEditor/Private/DataprepSnapshot.cpp b/Engine/Plugins/Enterprise/DataprepEditor/Source/DataprepEditor/Private/DataprepSnapshot.cpp
--- a/Engine/Plugins/Enterprise/DataprepEditor/Source/DataprepEditor/Private/DataprepSnapshot.cpp
+++ b/Engine/Plugins/Enterprise/DataprepEditor/Source/DataprepEditor/Private/DataprepSnapshot.cpp
@@ -528,6 +528,8 @@ public:
 		return Actor && SelectedActors.Contains(Actor);
 	}
 
+	virtual int32 GetObjectNumber() const override { return SelectedActors.Num(); }
+
 	/** Set of actors marked as selected so they get included in the copy */
 	TSet<const AActor*> SelectedActors;
 };
diff --git a/Engine/Source/Editor/UnrealEd/Private/EditorActor.cpp b/Engine/Source/Editor/UnrealEd/Private/EditorActor.cpp
--- a/Engine/Source/Editor/UnrealEd/Private/EditorActor.cpp
+++ b/Engine/Source/Editor/UnrealEd/Private/EditorActor.cpp
@@ -162,14 +162,18 @@ void UUnrealEdEngine::CopyActors(const TArray<AActor*>& InActorsToCopy, UWorld*
 	// Export the actors.
 	FStringOutputDevice Ar;
 	const FSelectedActorExportObjectInnerContext Context(ActorsToCopy);
-	UExporter::ExportToOutputDevice(&Context, InWorld, NULL, Ar, TEXT("copy"), 0, PPF_DeepCompareInstances | PPF_ExportsNotFullyQualified);
-	if (DestinationData)
-	{
-		*DestinationData = MoveTemp(Ar);
-	}
-	else
+	const bool bExportSucceeded = UExporter::ExportToOutputDevice(&Context, InWorld, NULL, Ar, TEXT("copy"), 0, PPF_DeepCompareInstances | PPF_ExportsNotFullyQualified);
+
+	if (bExportSucceeded)
 	{
-		FPlatformApplicationMisc::ClipboardCopy(*Ar);
+		if (DestinationData)
+		{
+			*DestinationData = MoveTemp(Ar);
+		}
+		else
+		{
+			FPlatformApplicationMisc::ClipboardCopy(*Ar);
+		}
 	}
 }
 
diff --git a/Engine/Source/Editor/UnrealEd/Private/EditorExporters.cpp b/Engine/Source/Editor/UnrealEd/Private/EditorExporters.cpp
--- a/Engine/Source/Editor/UnrealEd/Private/EditorExporters.cpp
+++ b/Engine/Source/Editor/UnrealEd/Private/EditorExporters.cpp
@@ -14,6 +14,7 @@
 #include "Misc/OutputDeviceFile.h"
 #include "UObject/Object.h"
 #include "UObject/UObjectIterator.h"
+#include "Misc/ScopedSlowTask.h"
 #include "Misc/TextBuffer.h"
 #include "UObject/Package.h"
 #include "Engine/EngineTypes.h"
@@ -521,9 +522,18 @@ bool ULevelExporterT3D::ExportText( const FExportObjectInnerContext* Context, UO
 
 	TextIndent += 3;
 
+	const int32 ActorNumberToCopy = bAllActors ? Level->Actors.Num() : Context->GetObjectNumber();
+	FScopedSlowTask SlowTask(ActorNumberToCopy, NSLOCTEXT("UnrealEd", "ExportingActors", "Exporting Actors"));
+	const bool bShowCancelButton = true;
+	SlowTask.MakeDialogDelayed(1.f, bShowCancelButton);
+
 	// loop through all of the actors just in this level
 	for (AActor* Actor : Level->Actors)
 	{
+		if (SlowTask.ShouldCancel())
+		{
+			break;
+		}
 		// Don't export the default physics volume, as it doesn't have a UModel associated with it
 		// and thus will not import properly.
 		if ( Actor == DefaultPhysicsVolume )
@@ -569,6 +579,8 @@ bool ULevelExporterT3D::ExportText( const FExportObjectInnerContext* Context, UO
 			{
 				GEditor->GetSelectedActors()->Deselect(Actor);
 			}
+
+			SlowTask.EnterProgressFrame();
 		}
 	}
 
@@ -606,7 +618,7 @@ bool ULevelExporterT3D::ExportText( const FExportObjectInnerContext* Context, UO
 	Ar.Logf( TEXT("%sEnd Map\r\n"), FCString::Spc(TextIndent) );
 
 
-	return 1;
+	return !SlowTask.ShouldCancel();
 }
 
 void ULevelExporterT3D::ExportComponentExtra(const FExportObjectInnerContext* Context, const TArray<UActorComponent*>& Components, FOutputDevice& Ar, uint32 PortFlags)
diff --git a/Engine/Source/Editor/UnrealEd/Private/EditorServer.cpp b/Engine/Source/Editor/UnrealEd/Private/EditorServer.cpp
--- a/Engine/Source/Editor/UnrealEd/Private/EditorServer.cpp
+++ b/Engine/Source/Editor/UnrealEd/Private/EditorServer.cpp
@@ -3602,12 +3602,6 @@ void UEditorEngine::CopySelectedActorsToClipboard( UWorld* InWorld, bool bShould
 					delete Job;
 				}
 
-				// Clean-up flag for Landscape Proxy cases...
-				for( TActorIterator<ALandscapeProxy> ProxyIt(InWorld); ProxyIt; ++ProxyIt )
-				{
-					ProxyIt->bIsMovingToLevel = false;
-				}
-
 				BufferLevel->ClearLevelComponents();
 				InWorld->RemoveLevel( BufferLevel );
 				BufferLevel->OwningWorld = NULL;
diff --git a/Engine/Source/Editor/UnrealEd/Private/Factories/EditorFactories.cpp b/Engine/Source/Editor/UnrealEd/Private/Factories/EditorFactories.cpp
--- a/Engine/Source/Editor/UnrealEd/Private/Factories/EditorFactories.cpp
+++ b/Engine/Source/Editor/UnrealEd/Private/Factories/EditorFactories.cpp
@@ -1228,89 +1228,95 @@ UObject* ULevelFactory::FactoryCreateText
 
 	// Pass 1: Sort out all the properties on the individual actors
 	bool bIsMoveToStreamingLevel =(FCString::Stricmp(Type, TEXT("move")) == 0);
-	for (auto& ActorMapElement : NewActorMap)
 	{
-		AActor* Actor = ActorMapElement.Key;
+		FScopedSlowTask SlowTask(NewActorMap.Num(), LOCTEXT("Importing Actors", "Importing Actors"));
+		SlowTask.MakeDialogDelayed(1.f);
 
-		// Import properties if the new actor is 
-		bool		bActorChanged = false;
-		const FString&	PropText = ActorMapElement.Value;
-		if ( Actor->ShouldImport(FStringView(PropText), bIsMoveToStreamingLevel) )
+		for (auto& ActorMapElement : NewActorMap)
 		{
-			Actor->PreEditChange(nullptr);
-			ImportObjectProperties( (uint8*)Actor, *PropText, Actor->GetClass(), Actor, Actor, Warn, 0, INDEX_NONE, NULL, &ExistingToNewMap );
-			bActorChanged = true;
+			AActor* Actor = ActorMapElement.Key;
 
-			if (GWorld == World)
+			// Import properties if the new actor is 
+			bool		bActorChanged = false;
+			const FString&	PropText = ActorMapElement.Value;
+			if ( Actor->ShouldImport(FStringView(PropText), bIsMoveToStreamingLevel) )
 			{
-				GEditor->SelectActor( Actor, true, false, true );
-			}
-		}
-		else // This actor is new, but rejected to import its properties, so just delete...
-		{
-			Actor->Destroy();
-			continue;
-		}
+				Actor->PreEditChange(nullptr);
+				ImportObjectProperties( (uint8*)Actor, *PropText, Actor->GetClass(), Actor, Actor, Warn, 0, INDEX_NONE, NULL, &ExistingToNewMap );
+				bActorChanged = true;
 
-		// If this is a newly imported brush, validate it.  If it's a newly imported dynamic brush, rebuild it first.
-		// Previously, this just called bspValidateBrush.  However, that caused the dynamic brushes which require a valid BSP tree
-		// to be built to break after being duplicated.  Calling RebuildBrush will rebuild the BSP tree from the imported polygons.
-		ABrush* Brush = Cast<ABrush>(Actor);
-		if( bActorChanged && Brush && Brush->Brush )
-		{
-			const bool bIsStaticBrush = Brush->IsStaticBrush();
-			if( !bIsStaticBrush )
+				if (GWorld == World)
+				{
+					GEditor->SelectActor( Actor, true, false, true );
+				}
+			}
+			else // This actor is new, but rejected to import its properties, so just delete...
 			{
-				FBSPOps::RebuildBrush( Brush->Brush );
+				Actor->Destroy();
+				continue;
 			}
 
-			FBSPOps::bspValidateBrush(Brush->Brush, true, false);
-		}
-
-		// Copy brushes' model pointers over to their BrushComponent, to keep compatibility with old T3Ds.
-		if( Brush && bActorChanged )
-		{
-			if( Brush->GetBrushComponent() ) // Should always be the case, but not asserting so that old broken content won't crash.
+			// If this is a newly imported brush, validate it.  If it's a newly imported dynamic brush, rebuild it first.
+			// Previously, this just called bspValidateBrush.  However, that caused the dynamic brushes which require a valid BSP tree
+			// to be built to break after being duplicated.  Calling RebuildBrush will rebuild the BSP tree from the imported polygons.
+			ABrush* Brush = Cast<ABrush>(Actor);
+			if( bActorChanged && Brush && Brush->Brush )
 			{
-				Brush->GetBrushComponent()->Brush = Brush->Brush;
+				const bool bIsStaticBrush = Brush->IsStaticBrush();
+				if( !bIsStaticBrush )
+				{
+					FBSPOps::RebuildBrush( Brush->Brush );
+				}
 
-				// We need to avoid duplicating default/ builder brushes. This is done by destroying all brushes that are CSG_Active and are not
-				// the default brush in their respective levels.
-				if( Brush->IsStaticBrush() && Brush->BrushType==Brush_Default )
+				FBSPOps::bspValidateBrush(Brush->Brush, true, false);
+			}
+
+			// Copy brushes' model pointers over to their BrushComponent, to keep compatibility with old T3Ds.
+			if( Brush && bActorChanged )
+			{
+				if( Brush->GetBrushComponent() ) // Should always be the case, but not asserting so that old broken content won't crash.
 				{
-					bool bIsDefaultBrush = false;
-					
-					// Iterate over all levels and compare current actor to the level's default brush.
-					for( int32 LevelIndex=0; LevelIndex<World->GetNumLevels(); LevelIndex++ )
+					Brush->GetBrushComponent()->Brush = Brush->Brush;
+
+					// We need to avoid duplicating default/ builder brushes. This is done by destroying all brushes that are CSG_Active and are not
+					// the default brush in their respective levels.
+					if( Brush->IsStaticBrush() && Brush->BrushType==Brush_Default )
 					{
-						ULevel* Level = World->GetLevel(LevelIndex);
-						if(Level->GetDefaultBrush() == Brush)
+						bool bIsDefaultBrush = false;
+						
+						// Iterate over all levels and compare current actor to the level's default brush.
+						for( int32 LevelIndex=0; LevelIndex<World->GetNumLevels(); LevelIndex++ )
 						{
-							bIsDefaultBrush = true;
-							break;
+							ULevel* Level = World->GetLevel(LevelIndex);
+							if(Level->GetDefaultBrush() == Brush)
+							{
+								bIsDefaultBrush = true;
+								break;
+							}
 						}
-					}
 
-					// Destroy actor if it's a builder brush but not the default brush in any of the currently loaded levels.
-					if( !bIsDefaultBrush )
-					{
-						World->DestroyActor( Brush );
+						// Destroy actor if it's a builder brush but not the default brush in any of the currently loaded levels.
+						if( !bIsDefaultBrush )
+						{
+							World->DestroyActor( Brush );
 
-						// Since the actor has been destroyed, skip the rest of this iteration of the loop.
-						continue;
+							// Since the actor has been destroyed, skip the rest of this iteration of the loop.
+							continue;
+						}
 					}
 				}
 			}
-		}
-		
-		// If the actor was imported . . .
-		if( bActorChanged )
-		{
-			// Let the actor deal with having been imported, if desired.
-			Actor->PostEditImport();
+			
+			// If the actor was imported . . .
+			if( bActorChanged )
+			{
+				// Let the actor deal with having been imported, if desired.
+				Actor->PostEditImport();
 
-			// Notify actor its properties have changed.
-			Actor->PostEditChange();
+				// Notify actor its properties have changed.
+				Actor->PostEditChange();
+			}
+			SlowTask.EnterProgressFrame();
 		}
 	}
 
diff --git a/Engine/Source/Runtime/Engine/Classes/Exporters/Exporter.h b/Engine/Source/Runtime/Engine/Classes/Exporters/Exporter.h
--- a/Engine/Source/Runtime/Engine/Classes/Exporters/Exporter.h
+++ b/Engine/Source/Runtime/Engine/Classes/Exporters/Exporter.h
@@ -226,8 +226,9 @@ class UExporter : public UObject
 	 * @param PortFlags - Flags controlling export behavior
 	 * @param bInSelectedOnly - true to export selected only
 	 * @param	ExportRootScope	The scope to create relative paths from, if the PPF_ExportsNotFullyQualified flag is passed in.  If NULL, the package containing the object will be used instead.
+	 * @return true if the export went successfully.
 	 */
-	ENGINE_API static void ExportToOutputDevice( const class FExportObjectInnerContext* Context, UObject* Object, UExporter* InExporter, FOutputDevice& Out, const TCHAR* FileType, int32 Indent, uint32 PortFlags=0, bool bInSelectedOnly=false, UObject* ExportRootScope = NULL);
+	ENGINE_API static bool ExportToOutputDevice( const class FExportObjectInnerContext* Context, UObject* Object, UExporter* InExporter, FOutputDevice& Out, const TCHAR* FileType, int32 Indent, uint32 PortFlags=0, bool bInSelectedOnly=false, UObject* ExportRootScope = NULL);
 
 	struct FExportToFileParams
 	{
diff --git a/Engine/Source/Runtime/Engine/Private/UnrealExporter.cpp b/Engine/Source/Runtime/Engine/Private/UnrealExporter.cpp
--- a/Engine/Source/Runtime/Engine/Private/UnrealExporter.cpp
+++ b/Engine/Source/Runtime/Engine/Private/UnrealExporter.cpp
@@ -163,7 +163,7 @@ bool UExporter::ExportToArchive( UObject* Object, UExporter* InExporter, FArchiv
 }
 
 
-void UExporter::ExportToOutputDevice(const FExportObjectInnerContext* Context, UObject* Object, UExporter* InExporter, FOutputDevice& Out, const TCHAR* FileType, int32 Indent, uint32 PortFlags, bool bInSelectedOnly, UObject* ExportRootScope)
+bool UExporter::ExportToOutputDevice(const FExportObjectInnerContext* Context, UObject* Object, UExporter* InExporter, FOutputDevice& Out, const TCHAR* FileType, int32 Indent, uint32 PortFlags, bool bInSelectedOnly, UObject* ExportRootScope)
 {
 	check(Object);
 	UExporter* Exporter = InExporter;
@@ -174,7 +174,7 @@ void UExporter::ExportToOutputDevice(const FExportObjectInnerContext* Context, U
 	if( !Exporter )
 	{
 		UE_LOG(LogExporter, Warning, TEXT("No %s exporter found for %s"), FileType, *Object->GetFullName() );
-		return;
+		return false;
 	}
 	check(Object->IsA(Exporter->SupportedClass));
 	int32 SavedIndent = Exporter->TextIndent;
@@ -196,8 +196,10 @@ void UExporter::ExportToOutputDevice(const FExportObjectInnerContext* Context, U
 		PortFlags |= PPF_Copy;
 	}
 
-	Exporter->ExportText( Context, Object, FileType, Out, GWarn, PortFlags );
+	const bool bSuccess = Exporter->ExportText( Context, Object, FileType, Out, GWarn, PortFlags );
 	Exporter->TextIndent = SavedIndent;
+
+	return bSuccess;
 }
 
 
diff --git a/Engine/Source/Runtime/Engine/Public/UnrealExporter.h b/Engine/Source/Runtime/Engine/Public/UnrealExporter.h
--- a/Engine/Source/Runtime/Engine/Public/UnrealExporter.h
+++ b/Engine/Source/Runtime/Engine/Public/UnrealExporter.h
@@ -38,6 +38,7 @@ public:
 	FExportObjectInnerContext(const bool bIgnoredValue) {};
 
 	virtual ~FExportObjectInnerContext() = default;
+	virtual int32 GetObjectNumber() const { return 0; }
 
 protected:
 	friend class UExporter;
@@ -73,6 +74,7 @@ public:
 	explicit FSelectedActorExportObjectInnerContext(const TArray<AActor*> InSelectedActors);
 
 	virtual bool IsObjectSelected(const UObject* InObj) const override;
+	virtual int32 GetObjectNumber() const override { return SelectedActors.Num(); }
 
 private:
 	void AddSelectedActor(const AActor* InActor);
diff --git a/Engine/Source/Runtime/Landscape/Classes/LandscapeProxy.h b/Engine/Source/Runtime/Landscape/Classes/LandscapeProxy.h
--- a/Engine/Source/Runtime/Landscape/Classes/LandscapeProxy.h
+++ b/Engine/Source/Runtime/Landscape/Classes/LandscapeProxy.h
@@ -646,11 +646,6 @@ public:
 	UPROPERTY(EditAnywhere, AdvancedDisplay, BlueprintReadOnly, Category=Rendering, meta = (DisplayName = "Desired Max Draw Distance"))
 	float LDMaxDrawDistance;
 
-#if WITH_EDITORONLY_DATA
-	UPROPERTY(transient)
-	uint32 bIsMovingToLevel:1;    // Check for the Move to Current Level case
-#endif // WITH_EDITORONLY_DATA
-
 	/** The Lightmass settings for this object. */
 	UPROPERTY(EditAnywhere, Category=Lightmass)
 	FLightmassPrimitiveSettings LightmassSettings;
@@ -837,8 +832,6 @@ public:
 	virtual void EditorApplyScale(const FVector& DeltaScale, const FVector* PivotLocation, bool bAltDown, bool bShiftDown, bool bCtrlDown) override;
 	virtual void EditorApplyMirror(const FVector& MirrorScale, const FVector& PivotLocation) override;
 	virtual void PostEditMove(bool bFinished) override;
-	virtual bool ShouldImport(FStringView ActorPropString, bool IsMovingLevel) override;
-	virtual bool ShouldExport() override;
 	//~ End AActor Interface
 	virtual uint32 GetDefaultGridSize(UWorld* InWorld) const override { return 1; }
 	virtual FGuid GetGridGuid() const override { return LandscapeGuid; }
diff --git a/Engine/Source/Runtime/Landscape/Private/Landscape.cpp b/Engine/Source/Runtime/Landscape/Private/Landscape.cpp
--- a/Engine/Source/Runtime/Landscape/Private/Landscape.cpp
+++ b/Engine/Source/Runtime/Landscape/Private/Landscape.cpp
@@ -1292,7 +1292,6 @@ ALandscapeProxy::ALandscapeProxy(const FObjectInitializer& ObjectInitializer)
 	bUseDynamicMaterialInstance = false;
 #if WITH_EDITORONLY_DATA
 	bLockLocation = true;
-	bIsMovingToLevel = false;
 #endif // WITH_EDITORONLY_DATA
 	ComponentScreenSizeToUseSubSections = 0.65f;
 	LOD0ScreenSize = 0.5f;
diff --git a/Engine/Source/Runtime/Landscape/Private/LandscapeEdit.cpp b/Engine/Source/Runtime/Landscape/Private/LandscapeEdit.cpp
--- a/Engine/Source/Runtime/Landscape/Private/LandscapeEdit.cpp
+++ b/Engine/Source/Runtime/Landscape/Private/LandscapeEdit.cpp
@@ -6824,45 +6824,7 @@ void ULandscapeComponent::InitWeightmapData(TArray<ULandscapeLayerInfoObject*>&
 	//  TODO: need to update layer system?
 }
 
-#define MAX_LANDSCAPE_EXPORT_COMPONENTS_NUM		16
-#define MAX_LANDSCAPE_PROP_TEXT_LENGTH			1024*1024*16
 
-bool ALandscapeProxy::ShouldExport()
-{
-	if (!bIsMovingToLevel && LandscapeComponents.Num() > MAX_LANDSCAPE_EXPORT_COMPONENTS_NUM)
-	{
-		// Prompt to save startup packages
-		if (EAppReturnType::Yes == FMessageDialog::Open(EAppMsgType::YesNo, FText::Format(
-			NSLOCTEXT("UnrealEd", "LandscapeExport_Warning", "Landscape has large number({0}) of components, so it will use large amount memory to copy it to the clipboard. Do you want to proceed?"), FText::AsNumber(LandscapeComponents.Num()))))
-		{
-			return true;
-		}
-		else
-		{
-			return false;
-		}
-	}
-	return true;
-}
-
-bool ALandscapeProxy::ShouldImport(FStringView ActorPropString, bool IsMovingToLevel)
-{
-	bIsMovingToLevel = IsMovingToLevel;
-	if (!bIsMovingToLevel && ActorPropString.Len() > MAX_LANDSCAPE_PROP_TEXT_LENGTH)
-	{
-		// Prompt to save startup packages
-		if (EAppReturnType::Yes == FMessageDialog::Open(EAppMsgType::YesNo, FText::Format(
-			NSLOCTEXT("UnrealEd", "LandscapeImport_Warning", "Landscape is about to import large amount memory ({0}MB) from the clipboard, which will take some time. Do you want to proceed?"), FText::AsNumber(ActorPropString.Len() >> 20))))
-		{
-			return true;
-		}
-		else
-		{
-			return false;
-		}
-	}
-	return true;
-}
 
 bool ALandscapeStreamingProxy::IsValidLandscapeActor(ALandscape* Landscape)
 {
