DebugSkelMeshComponent: Do not clear the preview instance (or mark it for garbage collection) when clearing anim script instance.

This is a followup fix to CL 23407330. The original bug was caused by calling ClearAnimScriptInstance (which would mark the AnimScriptInstance for garbage collection) when the PreviewInstance is pointing to it. The original fix was to set the PreviewInstance to nullptr after this happened. However, this broke a few other things, namely RBAN simulation and PhysicsAssetEditor preview animation stopped working because the PreviewInstance was expected to exist.

Upon review, thomas.sarkanen suggested that the MarkAsGarbage was completely unnecessary, so just not doing that.

Tested original bug  (UE-171723) was still fixed, PhysicsAssetEditor preview animation works, and RBAN preview simulation still works.

diff --git a/Engine/Source/Editor/UnrealEd/Classes/Animation/DebugSkelMeshComponent.h b/Engine/Source/Editor/UnrealEd/Classes/Animation/DebugSkelMeshComponent.h
--- a/Engine/Source/Editor/UnrealEd/Classes/Animation/DebugSkelMeshComponent.h
+++ b/Engine/Source/Editor/UnrealEd/Classes/Animation/DebugSkelMeshComponent.h
@@ -366,7 +366,6 @@ class UNREALED_API UDebugSkelMeshComponent : public USkeletalMeshComponent
 	virtual void InitAnim(bool bForceReinit) override;
 	virtual bool IsWindEnabled() const override { return true; }
 	virtual void SetAnimClass(class UClass* NewClass) override;
-	virtual void OnClearAnimScriptInstance() override;
 	//~ End SkeletalMeshComponent Interface
 
 	// return true if currently preview animation asset is on
diff --git a/Engine/Source/Editor/UnrealEd/Private/Animation/DebugSkelMeshComponent.cpp b/Engine/Source/Editor/UnrealEd/Private/Animation/DebugSkelMeshComponent.cpp
--- a/Engine/Source/Editor/UnrealEd/Private/Animation/DebugSkelMeshComponent.cpp
+++ b/Engine/Source/Editor/UnrealEd/Private/Animation/DebugSkelMeshComponent.cpp
@@ -570,14 +570,6 @@ void UDebugSkelMeshComponent::SetAnimClass(class UClass* NewClass)
 	UE_LOG(LogAnimation, Warning, TEXT("Attempting to destroy an animation preview actor, skipping."));
 }
 
-void UDebugSkelMeshComponent::OnClearAnimScriptInstance()
-{
-	// call to super not strictly necessary (since it is empty)
-	Super::OnClearAnimScriptInstance();
-	
-	PreviewInstance = nullptr;
-}
-
 void UDebugSkelMeshComponent::EnablePreview(bool bEnable, UAnimationAsset* PreviewAsset)
 {
 	if (PreviewInstance)
diff --git a/Engine/Source/Runtime/Engine/Classes/Components/SkeletalMeshComponent.h b/Engine/Source/Runtime/Engine/Classes/Components/SkeletalMeshComponent.h
--- a/Engine/Source/Runtime/Engine/Classes/Components/SkeletalMeshComponent.h
+++ b/Engine/Source/Runtime/Engine/Classes/Components/SkeletalMeshComponent.h
@@ -2318,8 +2318,6 @@ protected:
 
 private:
 
-	virtual void OnClearAnimScriptInstance() {};
-
 	friend struct FSkeletalMeshComponentEndPhysicsTickFunction;
 
 	/** Update systems after physics sim is done */
diff --git a/Engine/Source/Runtime/Engine/Private/Components/SkeletalMeshComponent.cpp b/Engine/Source/Runtime/Engine/Private/Components/SkeletalMeshComponent.cpp
--- a/Engine/Source/Runtime/Engine/Private/Components/SkeletalMeshComponent.cpp
+++ b/Engine/Source/Runtime/Engine/Private/Components/SkeletalMeshComponent.cpp
@@ -997,12 +997,10 @@ void USkeletalMeshComponent::ClearAnimScriptInstance()
 		HandleExistingParallelEvaluationTask(bBlockOnTask, bPerformPostAnimEvaluation);
 
 		AnimScriptInstance->EndNotifyStates();
-		AnimScriptInstance->MarkAsGarbage();
 	}
 	AnimScriptInstance = nullptr;
 	ResetLinkedAnimInstances();
 	ClearCachedAnimProperties();
-	OnClearAnimScriptInstance();
 }
 
 void USkeletalMeshComponent::ClearCachedAnimProperties()
