Fixed regression fetching asset thumbnails from the package

diff --git a/Engine/Source/Editor/UnrealEd/Private/AssetThumbnail.cpp b/Engine/Source/Editor/UnrealEd/Private/AssetThumbnail.cpp
--- a/Engine/Source/Editor/UnrealEd/Private/AssetThumbnail.cpp
+++ b/Engine/Source/Editor/UnrealEd/Private/AssetThumbnail.cpp
@@ -1253,6 +1253,7 @@ bool FAssetThumbnailPool::LoadThumbnail(TSharedRef<FThumbnailInfo> ThumbnailInfo
 		}
 	}
 
+	FThumbnailMap ThumbnailMap;
 	// If we could not render a fresh thumbnail, see if we already have a cached one to load
 	const FObjectThumbnail* FoundThumbnail = ThumbnailTools::FindCachedThumbnail(AssetData.GetFullName());
 
