Fix bad merge of AssetThumbnail.cpp in RES

diff --git a/Engine/Source/Editor/UnrealEd/Private/AssetThumbnail.cpp b/Engine/Source/Editor/UnrealEd/Private/AssetThumbnail.cpp
--- a/Engine/Source/Editor/UnrealEd/Private/AssetThumbnail.cpp
+++ b/Engine/Source/Editor/UnrealEd/Private/AssetThumbnail.cpp
@@ -1253,7 +1253,6 @@ bool FAssetThumbnailPool::LoadThumbnail(TSharedRef<FThumbnailInfo> ThumbnailInfo
 		}
 	}
 
-	FThumbnailMap ThumbnailMap;
 	// If we could not render a fresh thumbnail, see if we already have a cached one to load
 	const FObjectThumbnail* FoundThumbnail = ThumbnailTools::FindCachedThumbnail(AssetData.GetFullName());
 
