Debugging BP utility.
NOTE need to request icon before submitting.
Tested in PIE/Simulated, multiplayers (different server type).

diff --git a/Engine/Source/Editor/EditorStyle/Private/StarshipStyle.cpp b/Engine/Source/Editor/EditorStyle/Private/StarshipStyle.cpp
--- a/Engine/Source/Editor/EditorStyle/Private/StarshipStyle.cpp
+++ b/Engine/Source/Editor/EditorStyle/Private/StarshipStyle.cpp
@@ -5293,6 +5293,10 @@ void FStarshipEditorStyle::FStyle::SetupPersonaStyle()
 
 		Set( "PlayWorld.ShowCurrentStatement", new IMAGE_BRUSH_SVG( "Starship/MainToolbar/DebugFindNode", Icon40x40 ) );
 		Set( "PlayWorld.ShowCurrentStatement.Small", new IMAGE_BRUSH_SVG( "Starship/MainToolbar/DebugFindNode", Icon20x20 ) );
+		Set( "PlayWorld.AbortExecution", new IMAGE_BRUSH_SVG("Starship/MainToolbar/DebugAbortExecution", Icon40x40));
+		Set( "PlayWorld.AbortExecution.Small", new IMAGE_BRUSH_SVG("Starship/MainToolbar/DebugAbortExecution", Icon20x20));
+		Set( "PlayWorld.ContinueExecution", new IMAGE_BRUSH_SVG("Starship/MainToolbar/DebugContinueExecution", Icon40x40));
+		Set( "PlayWorld.ContinueExecution.Small", new IMAGE_BRUSH_SVG("Starship/MainToolbar/DebugContinueExecution", Icon20x20));
 		Set( "PlayWorld.StepOut", new IMAGE_BRUSH_SVG("Starship/MainToolbar/DebugStepOut", Icon40x40));
 		Set( "PlayWorld.StepOut.Small", new IMAGE_BRUSH_SVG("Starship/MainToolbar/DebugStepOut", Icon20x20));
 		Set( "PlayWorld.StepInto", new IMAGE_BRUSH_SVG( "Starship/MainToolbar/DebugStepInto", Icon40x40 ) );
diff --git a/Engine/Source/Editor/UnrealEd/Private/Kismet2/DebuggerCommands.cpp b/Engine/Source/Editor/UnrealEd/Private/Kismet2/DebuggerCommands.cpp
--- a/Engine/Source/Editor/UnrealEd/Private/Kismet2/DebuggerCommands.cpp
+++ b/Engine/Source/Editor/UnrealEd/Private/Kismet2/DebuggerCommands.cpp
@@ -90,10 +90,7 @@ namespace DebuggerCommands
 {
 	static bool bAllowPlayWorldFeature = true;
 	static FAutoConsoleVariableRef AllowPlayWorldFeatureCVar(TEXT("Editor.AllowPlayWorldFeature"), bAllowPlayWorldFeature, TEXT("When true play world is allowed."));
-	static bool AllowPlayWorldFeature()
-	{
-		return bAllowPlayWorldFeature;
-	}
+	static bool AllowPlayWorldFeature();
 }
 
 void SGlobalPlayWorldActions::Construct(const FArguments& InArgs)
@@ -172,6 +169,8 @@ public:
 	static void SingleFrameAdvance_Clicked();
 
 	static void ShowCurrentStatement_Clicked();
+	static void AbortExecution_Clicked();
+	static void ContinueExecution_Clicked();
 	static void StepInto_Clicked();
 	static void StepOver_Clicked();
 	static void StepOut_Clicked();
@@ -191,6 +190,7 @@ public:
 	static bool CanShowLateJoinButton();
 
 	static bool IsStoppedAtBreakpoint();
+	static bool IsStoppedAtBreakpoint_InEngineMode();
 
 	static bool CanShowNonPlayWorldOnlyActions();
 	static bool CanShowVulkanNonPlayWorldOnlyActions();
@@ -239,6 +239,13 @@ protected:
 	static bool IsReadyToLaunchOnDevice(FString DeviceId);
 };
 
+namespace DebuggerCommands
+{
+	bool AllowPlayWorldFeature()
+	{
+		return bAllowPlayWorldFeature && !FInternalPlayWorldCommandCallbacks::IsStoppedAtBreakpoint_InEngineMode();
+	}
+}
 
 /**
  * Called to leave K2 debugging mode
@@ -350,6 +357,8 @@ void FPlayWorldCommands::RegisterCommands()
 	UI_COMMAND(TogglePlayPauseOfPlaySession, "Toggle Play/Pause", "Resume playing if paused, or pause if playing", EUserInterfaceActionType::Button, FInputChord(EKeys::Pause));
 	UI_COMMAND(PossessEjectPlayer, "Possess or Eject Player", "Possesses or ejects the player from the camera", EUserInterfaceActionType::Button, FInputChord(EKeys::F8));
 	UI_COMMAND(ShowCurrentStatement, "Locate", "Locate the currently active node", EUserInterfaceActionType::Button, FInputChord());
+	UI_COMMAND(AbortExecution, "Abort", "Abort the execution", EUserInterfaceActionType::Button, FInputChord());
+	UI_COMMAND(ContinueExecution, "Continue", "Continue the execution", EUserInterfaceActionType::Button, FInputChord(EModifierKey::Alt | EModifierKey::Shift, EKeys::F10));
 	UI_COMMAND(StepInto, "Step Into", "Step Into the next node to be executed", EUserInterfaceActionType::Button, PLATFORM_MAC ? FInputChord(EModifierKey::Control, EKeys::F11) : FInputChord(EKeys::F11));
 	UI_COMMAND(StepOver, "Step Over", "Step to the next node to be executed in the current graph", EUserInterfaceActionType::Button, FInputChord(EKeys::F10));
 	UI_COMMAND(StepOut, "Step Out", "Step Out to the next node to be executed in the parent graph", EUserInterfaceActionType::Button, FInputChord(EModifierKey::Alt | EModifierKey::Shift, EKeys::F11));
@@ -522,7 +531,7 @@ void FPlayWorldCommands::BindGlobalPlayWorldCommands()
 		FExecuteAction::CreateStatic(&FInternalPlayWorldCommandCallbacks::GetMouseControlExecute),
 		FCanExecuteAction::CreateStatic(&FInternalPlayWorldCommandCallbacks::HasPlayWorld),
 		FIsActionChecked(),
-		FIsActionChecked::CreateStatic(&FInternalPlayWorldCommandCallbacks::HasPlayWorld)
+		FIsActionButtonVisible::CreateStatic(&FInternalPlayWorldCommandCallbacks::HasPlayWorld)
 	);
 
 	// Toggle PIE/SIE, Eject (PIE->SIE), and Possess (SIE->PIE)
@@ -538,28 +547,42 @@ void FPlayWorldCommands::BindGlobalPlayWorldCommands()
 		FExecuteAction::CreateStatic(&FInternalPlayWorldCommandCallbacks::ShowCurrentStatement_Clicked),
 		FCanExecuteAction::CreateStatic(&FInternalPlayWorldCommandCallbacks::IsStoppedAtBreakpoint),
 		FIsActionChecked(),
-		FIsActionChecked::CreateStatic(&FInternalPlayWorldCommandCallbacks::IsStoppedAtBreakpoint)
+		FIsActionButtonVisible::CreateStatic(&FInternalPlayWorldCommandCallbacks::IsStoppedAtBreakpoint)
 	);
 
+	ActionList.MapAction(Commands.AbortExecution,
+		FExecuteAction::CreateStatic(&FInternalPlayWorldCommandCallbacks::AbortExecution_Clicked),
+		FCanExecuteAction::CreateStatic(&FInternalPlayWorldCommandCallbacks::IsStoppedAtBreakpoint_InEngineMode),
+		FIsActionChecked(),
+		FIsActionButtonVisible::CreateStatic(&FInternalPlayWorldCommandCallbacks::IsStoppedAtBreakpoint_InEngineMode)
+	);
+
+	ActionList.MapAction(Commands.ContinueExecution,
+		FExecuteAction::CreateStatic(&FInternalPlayWorldCommandCallbacks::ContinueExecution_Clicked),
+		FCanExecuteAction::CreateStatic(&FInternalPlayWorldCommandCallbacks::IsStoppedAtBreakpoint),
+		FIsActionChecked(),
+		FIsActionButtonVisible::CreateStatic(&FInternalPlayWorldCommandCallbacks::IsStoppedAtBreakpoint)
+	);
+	
 	ActionList.MapAction(Commands.StepInto,
 		FExecuteAction::CreateStatic(&FInternalPlayWorldCommandCallbacks::StepInto_Clicked),
 		FCanExecuteAction::CreateStatic(&FInternalPlayWorldCommandCallbacks::IsStoppedAtBreakpoint),
 		FIsActionChecked(),
-		FIsActionChecked::CreateStatic(&FInternalPlayWorldCommandCallbacks::IsStoppedAtBreakpoint)
+		FIsActionButtonVisible::CreateStatic(&FInternalPlayWorldCommandCallbacks::IsStoppedAtBreakpoint)
 	);
 
 	ActionList.MapAction(Commands.StepOver,
 		FExecuteAction::CreateStatic(&FInternalPlayWorldCommandCallbacks::StepOver_Clicked),
 		FCanExecuteAction::CreateStatic(&FInternalPlayWorldCommandCallbacks::IsStoppedAtBreakpoint),
 		FIsActionChecked(),
-		FIsActionChecked::CreateStatic(&FInternalPlayWorldCommandCallbacks::IsStoppedAtBreakpoint)
+		FIsActionButtonVisible::CreateStatic(&FInternalPlayWorldCommandCallbacks::IsStoppedAtBreakpoint)
 	);
 
 	ActionList.MapAction(Commands.StepOut,
 		FExecuteAction::CreateStatic(&FInternalPlayWorldCommandCallbacks::StepOut_Clicked),
 		FCanExecuteAction::CreateStatic(&FInternalPlayWorldCommandCallbacks::IsStoppedAtBreakpoint),
 		FIsActionChecked(),
-		FIsActionChecked::CreateStatic(&FInternalPlayWorldCommandCallbacks::IsStoppedAtBreakpoint)
+		FIsActionButtonVisible::CreateStatic(&FInternalPlayWorldCommandCallbacks::IsStoppedAtBreakpoint)
 	);
 
 	AddPIEPreviewDeviceActions(Commands, ActionList);
@@ -660,6 +683,12 @@ void FPlayWorldCommands::BuildToolbar(FToolMenuSection& InSection, bool bInclude
 	FToolMenuEntry ShowCurrentStatementEntry = FToolMenuEntry::InitToolBarButton(FPlayWorldCommands::Get().ShowCurrentStatement, TAttribute<FText>(), TAttribute<FText>(), TAttribute<FSlateIcon>(), FName(TEXT("ShowCurrentStatement")));
 	ShowCurrentStatementEntry.StyleNameOverride = FName("Toolbar.BackplateLeft");
 
+	FToolMenuEntry ContinueEntry = FToolMenuEntry::InitToolBarButton(FPlayWorldCommands::Get().ContinueExecution, TAttribute<FText>(), TAttribute<FText>(), TAttribute<FSlateIcon>(), FName(TEXT("ContinueExecution")));
+	ContinueEntry.StyleNameOverride = FName("Toolbar.BackplateCenter");
+
+	FToolMenuEntry AbortEntry = FToolMenuEntry::InitToolBarButton(FPlayWorldCommands::Get().AbortExecution, TAttribute<FText>(), TAttribute<FText>(), TAttribute<FSlateIcon>(), FName(TEXT("AbortExecution")));
+	AbortEntry.StyleNameOverride = FName("Toolbar.BackplateCenter");
+	
 	FToolMenuEntry StepIntoEntry = FToolMenuEntry::InitToolBarButton(FPlayWorldCommands::Get().StepInto, TAttribute<FText>(), TAttribute<FText>(), TAttribute<FSlateIcon>(), FName(TEXT("StepInto")));
 	StepIntoEntry.StyleNameOverride = FName("Toolbar.BackplateCenter");
 
@@ -670,11 +699,10 @@ void FPlayWorldCommands::BuildToolbar(FToolMenuSection& InSection, bool bInclude
 	StepOutEntry.StyleNameOverride = FName("Toolbar.BackplateRight");
 
 	InSection.AddEntry(ShowCurrentStatementEntry);
-
-	InSection.AddEntry(StepIntoEntry);
-	
-	InSection.AddEntry(StepOverEntry);
-	
+	InSection.AddEntry(ContinueEntry);
+	InSection.AddEntry(AbortEntry);
+	InSection.AddEntry(StepIntoEntry);	
+	InSection.AddEntry(StepOverEntry);	
 	InSection.AddEntry(StepOutEntry);
 
 	if (bIncludeLaunchButtonAndOptions)
@@ -935,6 +963,11 @@ void FPlayWorldCommandCallbacks::StartPlayFromHere()
 
 void FPlayWorldCommandCallbacks::StartPlayFromHere(const TOptional<FVector>& Location, const TOptional<FRotator>& Rotation, const TSharedPtr<IAssetViewport>& ActiveLevelViewport)
 {
+	if (FInternalPlayWorldCommandCallbacks::IsStoppedAtBreakpoint_InEngineMode())
+	{
+		return;
+	}
+
 	// Is a PIE session already running?  If so we close it first
 	if (GUnrealEd->PlayWorld != NULL)
 	{
@@ -959,6 +992,11 @@ void FPlayWorldCommandCallbacks::StartPlayFromHere(const TOptional<FVector>& Loc
 
 void FPlayWorldCommandCallbacks::ResumePlaySession_Clicked()
 {
+	if (FInternalPlayWorldCommandCallbacks::IsStoppedAtBreakpoint_InEngineMode())
+	{
+		return;
+	}
+
 	if (HasPlayWorld())
 	{
 		LeaveDebuggingMode();
@@ -971,6 +1009,11 @@ void FPlayWorldCommandCallbacks::ResumePlaySession_Clicked()
 
 void FPlayWorldCommandCallbacks::PausePlaySession_Clicked()
 {
+	if (FInternalPlayWorldCommandCallbacks::IsStoppedAtBreakpoint_InEngineMode())
+	{
+		return;
+	}
+
 	if (HasPlayWorld())
 	{
 		GUnrealEd->SetPIEWorldsPaused(true);
@@ -993,6 +1036,11 @@ void FPlayWorldCommandCallbacks::PausePlaySession_Clicked()
 
 void FPlayWorldCommandCallbacks::SingleFrameAdvance_Clicked()
 {
+	if (FInternalPlayWorldCommandCallbacks::IsStoppedAtBreakpoint_InEngineMode())
+	{
+		return;
+	}
+
 	if (HasPlayWorld())
 	{
 		FInternalPlayWorldCommandCallbacks::SingleFrameAdvance_Clicked();
@@ -1093,6 +1141,11 @@ FSlateIcon FInternalPlayWorldCommandCallbacks::GetPossessEjectImage()
 
 bool FInternalPlayWorldCommandCallbacks::CanLateJoin()
 {
+	if (FInternalPlayWorldCommandCallbacks::IsStoppedAtBreakpoint_InEngineMode())
+	{
+		return false;
+	}
+
 	return HasPlayWorld();
 }
 
@@ -1104,6 +1157,11 @@ bool FInternalPlayWorldCommandCallbacks::CanShowLateJoinButton()
 
 void FInternalPlayWorldCommandCallbacks::Simulate_Clicked()
 {
+	if (IsStoppedAtBreakpoint_InEngineMode())
+	{
+		return;
+	}
+
 	// Is a simulation session already running?  If so, do nothing
 	if (HasPlayWorld() && GUnrealEd->bIsSimulatingInEditor)
 	{
@@ -1139,6 +1197,11 @@ void FInternalPlayWorldCommandCallbacks::Simulate_Clicked()
 
 bool FInternalPlayWorldCommandCallbacks::Simulate_CanExecute()
 {
+	if (IsStoppedAtBreakpoint_InEngineMode())
+	{
+		return false;
+	}
+
 	// Can't simulate while already simulating; PIE is fine as we toggle to simulate
 	return !(HasPlayWorld() && GUnrealEd->bIsSimulatingInEditor) && !GEditor->IsLightingBuildCurrentlyRunning();
 }
@@ -1366,6 +1429,11 @@ FSlateIcon FInternalPlayWorldCommandCallbacks::GetRepeatLastPlayIcon()
 
 void FInternalPlayWorldCommandCallbacks::PlayInViewport_Clicked()
 {
+	if (IsStoppedAtBreakpoint_InEngineMode())
+	{
+		return;
+	}
+
 	TRACE_BOOKMARK(TEXT("PIE"));
 
 	FLevelEditorModule& LevelEditorModule = FModuleManager::GetModuleChecked<FLevelEditorModule>(TEXT("LevelEditor"));
@@ -1413,6 +1481,11 @@ void FInternalPlayWorldCommandCallbacks::PlayInViewport_Clicked()
 
 bool FInternalPlayWorldCommandCallbacks::PlayInViewport_CanExecute()
 {
+	if (IsStoppedAtBreakpoint_InEngineMode())
+	{
+		return false;
+	}
+
 	// Disallow PIE when compiling in the editor
 	if (GEditor->bIsCompiling)
 	{
@@ -1426,6 +1499,11 @@ bool FInternalPlayWorldCommandCallbacks::PlayInViewport_CanExecute()
 
 void FInternalPlayWorldCommandCallbacks::PlayInEditorFloating_Clicked()
 {
+	if (IsStoppedAtBreakpoint_InEngineMode())
+	{
+		return;
+	}
+
 	FLevelEditorModule& LevelEditorModule = FModuleManager::GetModuleChecked<FLevelEditorModule>(TEXT("LevelEditor"));
 
 	SetLastExecutedPlayMode(PlayMode_InEditorFloating);
@@ -1467,11 +1545,21 @@ void FInternalPlayWorldCommandCallbacks::PlayInEditorFloating_Clicked()
 
 bool FInternalPlayWorldCommandCallbacks::PlayInEditorFloating_CanExecute()
 {
+	if (IsStoppedAtBreakpoint_InEngineMode())
+	{
+		return false;
+	}
+
 	return (!HasPlayWorld() || !GUnrealEd->bIsSimulatingInEditor) && !GEditor->IsLightingBuildCurrentlyRunning();
 }
 
 void FInternalPlayWorldCommandCallbacks::PlayInVR_Clicked()
 {
+	if (IsStoppedAtBreakpoint_InEngineMode())
+	{
+		return;
+	}
+
 	FLevelEditorModule& LevelEditorModule = FModuleManager::GetModuleChecked<FLevelEditorModule>(TEXT("LevelEditor"));
 
 	SetLastExecutedPlayMode(PlayMode_InVR);
@@ -1509,6 +1597,10 @@ void FInternalPlayWorldCommandCallbacks::PlayInVR_Clicked()
 
 bool FInternalPlayWorldCommandCallbacks::PlayInVR_CanExecute()
 {
+	if (IsStoppedAtBreakpoint_InEngineMode())
+	{
+		return false;
+	}
 	return (!HasPlayWorld() || !GUnrealEd->bIsSimulatingInEditor) && !GEditor->IsLightingBuildCurrentlyRunning() && GEngine && GEngine->XRSystem.IsValid();
 }
 
@@ -1532,6 +1624,11 @@ void FInternalPlayWorldCommandCallbacks::PlayInNewProcess_Clicked(EPlayModeType
 	check(PlayModeType == PlayMode_InNewProcess || PlayModeType == PlayMode_InMobilePreview
 		|| PlayModeType == PlayMode_InTargetedMobilePreview || PlayModeType == PlayMode_InVulkanPreview);
 
+	if (IsStoppedAtBreakpoint_InEngineMode())
+	{
+		return;
+	}
+
 	SetLastExecutedPlayMode(PlayModeType);
 	FRequestPlaySessionParams SessionParams;
 
@@ -1580,7 +1677,7 @@ void FInternalPlayWorldCommandCallbacks::PlayInNewProcess_Clicked(EPlayModeType
 
 bool FInternalPlayWorldCommandCallbacks::PlayInNewProcess_CanExecute()
 {
-	return true;
+	return !IsStoppedAtBreakpoint_InEngineMode();
 }
 
 
@@ -1597,6 +1694,11 @@ bool FInternalPlayWorldCommandCallbacks::PlayInModeIsChecked(EPlayModeType PlayM
 
 bool FInternalPlayWorldCommandCallbacks::PlayInLocation_CanExecute(EPlayModeLocations Location)
 {
+	if (IsStoppedAtBreakpoint_InEngineMode())
+	{
+		return false;
+	}
+
 	switch (Location)
 	{
 	case PlayLocation_CurrentCameraLocation:
@@ -1694,6 +1796,11 @@ FText FInternalPlayWorldCommandCallbacks::GetResumePlaySessionToolTip()
 
 void FInternalPlayWorldCommandCallbacks::SingleFrameAdvance_Clicked()
 {
+	if (IsStoppedAtBreakpoint_InEngineMode())
+	{
+		return;
+	}
+
 	// We want to function just like Single stepping where we will stop at a breakpoint if one is encountered but we also want to stop after 1 tick if a breakpoint is not encountered.
 	FKismetDebugUtilities::RequestSingleStepIn();
 	if (HasPlayWorld())
@@ -1707,6 +1814,11 @@ void FInternalPlayWorldCommandCallbacks::SingleFrameAdvance_Clicked()
 
 void FInternalPlayWorldCommandCallbacks::StopPlaySession_Clicked()
 {
+	if (IsStoppedAtBreakpoint_InEngineMode())
+	{
+		return;
+	}
+
 	if (HasPlayWorld())
 	{
 		GEditor->RequestEndPlayMap();
@@ -1731,35 +1843,42 @@ void FInternalPlayWorldCommandCallbacks::ShowCurrentStatement_Clicked()
 	}
 }
 
-
-void FInternalPlayWorldCommandCallbacks::StepInto_Clicked()
+void FInternalPlayWorldCommandCallbacks::ContinueExecution_Clicked()
 {
-	FKismetDebugUtilities::RequestSingleStepIn();
-	if (HasPlayWorld())
+	if (FKismetDebugUtilities::GetCurrentDebuggingWorld())
 	{
 		LeaveDebuggingMode();
 		GUnrealEd->PlaySessionSingleStepped();
 	}
+	else
+	{
+		const bool bIsResumingPlay = !FKismetDebugUtilities::IsSingleStepping();
+		FSlateApplication::Get().LeaveDebuggingMode(!bIsResumingPlay);
+	}
+}
+
+void FInternalPlayWorldCommandCallbacks::AbortExecution_Clicked()
+{
+	FKismetDebugUtilities::RequestAbortingExecution();
+	ContinueExecution_Clicked();
+}
+
+void FInternalPlayWorldCommandCallbacks::StepInto_Clicked()
+{
+	FKismetDebugUtilities::RequestSingleStepIn();
+	ContinueExecution_Clicked();
 }
 
 void FInternalPlayWorldCommandCallbacks::StepOver_Clicked()
 {
 	FKismetDebugUtilities::RequestStepOver();
-	if (HasPlayWorld())
-	{
-		LeaveDebuggingMode();
-		GUnrealEd->PlaySessionSingleStepped();
-	}
+	ContinueExecution_Clicked();
 }
 
 void FInternalPlayWorldCommandCallbacks::StepOut_Clicked()
 {
 	FKismetDebugUtilities::RequestStepOut();
-	if (HasPlayWorld())
-	{
-		LeaveDebuggingMode();
-		GUnrealEd->PlaySessionSingleStepped();
-	}
+	ContinueExecution_Clicked();
 }
 
 void FInternalPlayWorldCommandCallbacks::TogglePlayPause_Clicked()
@@ -1855,6 +1974,11 @@ bool FInternalPlayWorldCommandCallbacks::IsStoppedAtBreakpoint()
 	return GIntraFrameDebuggingGameThread;
 }
 
+bool FInternalPlayWorldCommandCallbacks::IsStoppedAtBreakpoint_InEngineMode()
+{
+	return GIntraFrameDebuggingGameThread && FKismetDebugUtilities::GetCurrentDebuggingWorld() == nullptr;
+}
+
 
 void FInternalPlayWorldCommandCallbacks::PossessEjectPlayer_Clicked()
 {
diff --git a/Engine/Source/Editor/UnrealEd/Private/Kismet2/KismetDebugUtilities.cpp b/Engine/Source/Editor/UnrealEd/Private/Kismet2/KismetDebugUtilities.cpp
--- a/Engine/Source/Editor/UnrealEd/Private/Kismet2/KismetDebugUtilities.cpp
+++ b/Engine/Source/Editor/UnrealEd/Private/Kismet2/KismetDebugUtilities.cpp
@@ -47,6 +47,13 @@
 
 #define LOCTEXT_NAMESPACE "BlueprintDebugging"
 
+enum class EKismetDebuggingMode : uint8
+{
+	None,
+	Engine, // for Blutility
+	World,  // for PIE
+};
+
 /** Per-thread data for use by FKismetDebugUtilities functions */
 class FKismetDebugUtilitiesData : public TThreadSingleton<FKismetDebugUtilitiesData>
 {
@@ -56,11 +63,13 @@ public:
 		, CurrentInstructionPointer(nullptr)
 		, MostRecentBreakpointInstructionPointer(nullptr)
 		, MostRecentStoppedNode(nullptr)
+		, CurrentDebuggingWorld(nullptr)
 		, TargetGraphStackDepth(INDEX_NONE)
 		, MostRecentBreakpointGraphStackDepth(INDEX_NONE)
 		, MostRecentBreakpointInstructionOffset(INDEX_NONE)
 		, StackFrameAtIntraframeDebugging(nullptr)
 		, TraceStackSamples(FKismetDebugUtilities::MAX_TRACE_STACK_SAMPLES)
+		, CurrentDebuggingMode(EKismetDebuggingMode::None)
 		, bIsSingleStepping(false)
 		, bIsSteppingOut(false)
 	{
@@ -71,11 +80,13 @@ public:
 		TargetGraphNodes.Empty();
 		CurrentInstructionPointer = nullptr;
 		MostRecentStoppedNode = nullptr;
+		CurrentDebuggingWorld = nullptr;
 
 		TargetGraphStackDepth = INDEX_NONE;
 		MostRecentBreakpointGraphStackDepth = INDEX_NONE;
 		MostRecentBreakpointInstructionOffset = INDEX_NONE;
 		StackFrameAtIntraframeDebugging = nullptr;
+		CurrentDebuggingMode = EKismetDebuggingMode::None;
 
 		bIsSingleStepping = false;
 		bIsSteppingOut = false;
@@ -93,6 +104,9 @@ public:
 	// The last node that we decided to break on for any reason (e.g. breakpoint, exception, or step operation):
 	TWeakObjectPtr< class UEdGraphNode > MostRecentStoppedNode;
 
+	// The PlayWorld that generated
+	TWeakObjectPtr<UWorld> CurrentDebuggingWorld;
+
 	// The target graph call stack depth. INDEX_NONE if not active
 	int32 TargetGraphStackDepth;
 
@@ -113,6 +127,9 @@ public:
 	// This data is used for the 'marching ants' display in the blueprint editor
 	TSimpleRingBuffer<FKismetTraceSample> TraceStackSamples;
 
+	// The type of current debugging 
+	EKismetDebuggingMode CurrentDebuggingMode;
+
 	// This flag controls whether we're trying to 'step in' to a function
 	bool bIsSingleStepping;
 
@@ -134,6 +151,16 @@ void FKismetDebugUtilities::EndOfScriptExecution(const FBlueprintContextTracker&
 	}
 }
 
+void FKismetDebugUtilities::RequestAbortingExecution()
+{
+	check(IsInGameThread());
+	FKismetDebugUtilitiesData& Data = FKismetDebugUtilitiesData::Get();
+	if (Data.StackFrameAtIntraframeDebugging)
+	{
+		const_cast<FFrame*>(Data.StackFrameAtIntraframeDebugging)->bAbortingExecution = true;
+	}
+}
+
 void FKismetDebugUtilities::RequestSingleStepIn()
 {
 	FKismetDebugUtilitiesData& Data = FKismetDebugUtilitiesData::Get();
@@ -658,11 +685,13 @@ void FKismetDebugUtilities::CheckBreakConditions(UEdGraphNode* NodeStoppedAt, bo
 void FKismetDebugUtilities::AttemptToBreakExecution(UBlueprint* BlueprintObj, const UObject* ActiveObject, const FFrame& StackFrame, const FBlueprintExceptionInfo& Info, UEdGraphNode* NodeStoppedAt, int32 DebugOpcodeOffset)
 {
 	checkSlow(BlueprintObj->GetObjectBeingDebugged() == ActiveObject);
+	check(IsInGameThread());
 
 	FKismetDebugUtilitiesData& Data = FKismetDebugUtilitiesData::Get();
 
 	// Cannot have re-entrancy while processing a breakpoint; return from this call stack before resuming execution!
-	check( !GIntraFrameDebuggingGameThread );
+	check(!GIntraFrameDebuggingGameThread);
+	check(Data.CurrentDebuggingMode == EKismetDebuggingMode::None);
 	
 	TGuardValue<bool> SignalGameThreadBeingDebugged(GIntraFrameDebuggingGameThread, true);
 	TGuardValue<const FFrame*> ResetStackFramePointer(Data.StackFrameAtIntraframeDebugging, &StackFrame);
@@ -731,16 +760,22 @@ void FKismetDebugUtilities::AttemptToBreakExecution(UBlueprint* BlueprintObj, co
 	// under automation and this was preventing debugging on automation test bp's.
 	if ((GUnrealEd->PlayWorld != NULL) && NodeStoppedAt)
 	{
+		Data.CurrentDebuggingMode = EKismetDebuggingMode::World;
+		Data.CurrentDebuggingWorld = GUnrealEd->PlayWorld;
+
 		// Pause the simulation
 		GUnrealEd->PlayWorld->bDebugPauseExecution = true;
 		GUnrealEd->PlayWorld->bDebugFrameStepExecution = false;
 		bShouldInStackDebug = true;
 	}
+	else if (NodeStoppedAt)
+	{
+		Data.CurrentDebuggingMode = EKismetDebuggingMode::Engine;
+	}
 	else
 	{
+		Data.CurrentDebuggingMode = EKismetDebuggingMode::None;
 		bShouldInStackDebug = false;
-		//@TODO: Determine exactly what behavior we want for breakpoints hit when not in PIE/SIE
-		//ensureMsgf(false, TEXT("Breakpoints placed in a function instead of the event graph are not supported yet"));
 	}
 
 	// Now enter within-the-frame debugging mode
@@ -753,6 +788,9 @@ void FKismetDebugUtilities::AttemptToBreakExecution(UBlueprint* BlueprintObj, co
 		CallStackViewer::UpdateDisplayedCallstack(ScriptStack);
 		FSlateApplication::Get().EnterDebuggingMode();
 	}
+
+	Data.CurrentDebuggingMode = EKismetDebuggingMode::None;
+	Data.CurrentDebuggingWorld.Reset();
 }
 
 UEdGraphNode* FKismetDebugUtilities::GetCurrentInstruction()
@@ -765,7 +803,7 @@ UEdGraphNode* FKismetDebugUtilities::GetCurrentInstruction()
 	}
 	else
 	{
-		return NULL;
+		return nullptr;
 	}
 }
 
@@ -779,7 +817,22 @@ UEdGraphNode* FKismetDebugUtilities::GetMostRecentBreakpointHit()
 	}
 	else
 	{
-		return NULL;
+		return nullptr;
+	}
+}
+
+UWorld* FKismetDebugUtilities::GetCurrentDebuggingWorld()
+{
+	// If paused at the end of the frame, or while not paused, there is no 'current instruction' to speak of
+	// It only has meaning during intraframe debugging.
+	if (GIntraFrameDebuggingGameThread)
+	{
+		const FKismetDebugUtilitiesData& Data = FKismetDebugUtilitiesData::Get();
+		return Data.CurrentDebuggingMode == EKismetDebuggingMode::World ? Data.CurrentDebuggingWorld.Get() : nullptr;
+	}
+	else
+	{
+		return nullptr;
 	}
 }
 
diff --git a/Engine/Source/Editor/UnrealEd/Public/Kismet2/DebuggerCommands.h b/Engine/Source/Editor/UnrealEd/Public/Kismet2/DebuggerCommands.h
--- a/Engine/Source/Editor/UnrealEd/Public/Kismet2/DebuggerCommands.h
+++ b/Engine/Source/Editor/UnrealEd/Public/Kismet2/DebuggerCommands.h
@@ -104,10 +104,14 @@ public:
 	TSharedPtr<FUICommandInfo> LateJoinSession;
 	TSharedPtr<FUICommandInfo> PossessEjectPlayer;
 	TSharedPtr<FUICommandInfo> ShowCurrentStatement;
+	TSharedPtr<FUICommandInfo> GetMouseControl;
+
+	/** BP Debugging controls */
+	TSharedPtr<FUICommandInfo> AbortExecution;
+	TSharedPtr<FUICommandInfo> ContinueExecution;
 	TSharedPtr<FUICommandInfo> StepInto;
 	TSharedPtr<FUICommandInfo> StepOver;
 	TSharedPtr<FUICommandInfo> StepOut;
-	TSharedPtr<FUICommandInfo> GetMouseControl;
 
 protected:
 
diff --git a/Engine/Source/Editor/UnrealEd/Public/Kismet2/KismetDebugUtilities.h b/Engine/Source/Editor/UnrealEd/Public/Kismet2/KismetDebugUtilities.h
--- a/Engine/Source/Editor/UnrealEd/Public/Kismet2/KismetDebugUtilities.h
+++ b/Engine/Source/Editor/UnrealEd/Public/Kismet2/KismetDebugUtilities.h
@@ -205,12 +205,18 @@ public:
 
 	static void OnScriptException(const UObject* ActiveObject, const FFrame& StackFrame, const FBlueprintExceptionInfo& Info);
 
-	/** Returns the current instruction; if a PIE/SIE session is started but paused; otherwise NULL */
+	/** Returns the current instruction; if a debugging session has started; otherwise none */
 	static class UEdGraphNode* GetCurrentInstruction();
 
-	/** Returns the most recent hit breakpoint; if a PIE/SIE session is started but paused; otherwise NULL */
+	/** Returns the most recent hit breakpoint; if a debugging session has started; otherwise none */
 	static class UEdGraphNode* GetMostRecentBreakpointHit();
 
+	/** Returns the most recent hit breakpoint; if a debugging session has started in PIE/SIE; otherwise none */
+	static UWorld* GetCurrentDebuggingWorld();
+
+	/** Request abort the current frame execution */
+	static void RequestAbortingExecution();
+
 	/** Request an attempt to single-step to the next node */
 	static void RequestSingleStepIn();
 	
