Logging: Extended GatherTextFromSourceCommandlet to support localized log macros

#rnx

diff --git a/Engine/Source/Editor/UnrealEd/Classes/Commandlets/GatherTextFromSourceCommandlet.h b/Engine/Source/Editor/UnrealEd/Classes/Commandlets/GatherTextFromSourceCommandlet.h
--- a/Engine/Source/Editor/UnrealEd/Classes/Commandlets/GatherTextFromSourceCommandlet.h
+++ b/Engine/Source/Editor/UnrealEd/Classes/Commandlets/GatherTextFromSourceCommandlet.h
@@ -475,6 +475,23 @@ private:
 		virtual void TryParse(const FString& Text, FSourceFileParseContext& Context) const override;
 	};
 
+	class FStructuredLogMacroDescriptor final : public FMacroDescriptor
+	{
+	public:
+		enum class EFlags : int32
+		{
+			None = 0,
+			Namespace = 1,
+		};
+
+		FStructuredLogMacroDescriptor(const TCHAR* Name, EFlags Flags);
+
+		virtual void TryParse(const FString& Text, FSourceFileParseContext& Context) const override;
+
+	private:
+		EFlags Flags;
+	};
+
 	class FIniNamespaceDescriptor : public FPreProcessorDescriptor
 	{
 	public:
diff --git a/Engine/Source/Editor/UnrealEd/Private/Commandlets/GatherTextFromSourceCommandlet.cpp b/Engine/Source/Editor/UnrealEd/Private/Commandlets/GatherTextFromSourceCommandlet.cpp
--- a/Engine/Source/Editor/UnrealEd/Private/Commandlets/GatherTextFromSourceCommandlet.cpp
+++ b/Engine/Source/Editor/UnrealEd/Private/Commandlets/GatherTextFromSourceCommandlet.cpp
@@ -273,6 +273,11 @@ int32 UGatherTextFromSourceCommandlet::Main( const FString& Params )
 
 	Parsables.Add(new FStringTableEntryMetaDataMacroDescriptor());
 
+	Parsables.Add(new FStructuredLogMacroDescriptor(TEXT("UE_SLOG_LOC"), FStructuredLogMacroDescriptor::EFlags::None));
+	Parsables.Add(new FStructuredLogMacroDescriptor(TEXT("UE_SLOG_LOC_EX"), FStructuredLogMacroDescriptor::EFlags::None));
+	Parsables.Add(new FStructuredLogMacroDescriptor(TEXT("UE_SLOG_NSLOC"), FStructuredLogMacroDescriptor::EFlags::Namespace));
+	Parsables.Add(new FStructuredLogMacroDescriptor(TEXT("UE_SLOG_NSLOC_EX"), FStructuredLogMacroDescriptor::EFlags::Namespace));
+
 	Parsables.Add(new FIniNamespaceDescriptor());
 
 	// Init a parse context to track the state of the file parsing 
@@ -2273,6 +2278,74 @@ void UGatherTextFromSourceCommandlet::FStringTableEntryMetaDataMacroDescriptor::
 	}
 }
 
+UGatherTextFromSourceCommandlet::FStructuredLogMacroDescriptor::FStructuredLogMacroDescriptor(const TCHAR* InName, EFlags InFlags)
+	: FMacroDescriptor(InName, ((int32)Flags & (int32)EFlags::Namespace) ? 5 : 4)
+	, Flags(InFlags)
+{
+}
+
+void UGatherTextFromSourceCommandlet::FStructuredLogMacroDescriptor::TryParse(const FString& Text, FSourceFileParseContext& Context) const
+{
+	// Attempt to parse something of the format
+	// UE_SLOG_LOC[_EX](CategoryName, Verbosity, Key, Format, ...)
+	// UE_SLOG_NSLOC[_EX](CategoryName, Verbosity, Namespace, Key, Format, ...)
+	int32 NextArg = 2;
+
+	if (!Context.ExcludedRegion && !Context.WithinBlockComment && !Context.WithinLineComment && !Context.WithinStringLiteral)
+	{
+		TArray<FString> Arguments;
+		if (ParseArgsFromMacro(StripCommentsFromToken(Text, Context), Arguments, Context))
+		{
+			const FString SourceLocation = FSourceLocation(Context.Filename, Context.LineNumber).ToString();
+
+			if (Arguments.Num() < GetMinNumberOfArgument())
+			{
+				UE_LOG(LogGatherTextFromSourceCommandlet, Warning, TEXT("%s: Expected at least %d arguments for %s macro, but got %d. %s"), *SourceLocation, GetMinNumberOfArgument(), *GetToken(), Arguments.Num(), *FLocTextHelper::SanitizeLogOutput(Context.LineText.TrimStartAndEnd()));
+				return;
+			}
+
+			bool bParseOk = true;
+			const auto ParseArg = [this, &SourceLocation, &bParseOk](FString&& Arg, const TCHAR* ArgName) -> FString
+			{
+				FString Out = MoveTemp(Arg).TrimStart();
+				constexpr bool bIsAutoText = true;
+				const FString Desc = FString::Printf(TEXT("%s: \"%s\" argument in %s macro"), *SourceLocation, ArgName, *GetToken());
+				bool bHasQuotes = false;
+				bParseOk &= PrepareArgument(Out, bIsAutoText, Desc, bHasQuotes);
+				return Out;
+			};
+
+			FString Namespace;
+			if (const bool bHasNamespace = ((int32)Flags & (int32)EFlags::Namespace) != 0)
+			{
+				Namespace = ParseArg(MoveTemp(Arguments[NextArg++]), TEXT("Namespace"));
+			}
+			else
+			{
+				Namespace = Context.Namespace;
+			}
+
+			if (Namespace.IsEmpty())
+			{
+				UE_LOG(LogGatherTextFromSourceCommandlet, Warning, TEXT("%s: %s macro doesn't define a namespace and no external namespace was set. An empty namspace will be used."), *SourceLocation, *GetToken());
+			}
+
+			FString Key = ParseArg(MoveTemp(Arguments[NextArg++]), TEXT("Key"));
+			FString Format = ParseArg(MoveTemp(Arguments[NextArg++]), TEXT("Format"));
+
+			if (bParseOk && !Key.IsEmpty())
+			{
+				FManifestContext MacroContext;
+				MacroContext.Key = Key;
+				MacroContext.SourceLocation = SourceLocation;
+				MacroContext.PlatformName = Context.FilePlatformName;
+
+				Context.AddManifestText(GetToken(), Namespace, Format, MacroContext);
+			}
+		}
+	}
+}
+
 void UGatherTextFromSourceCommandlet::FIniNamespaceDescriptor::TryParse(const FString& Text, FSourceFileParseContext& Context) const
 {
 	// Attempt to parse something of the format
