Fixed a compile error in FStructuredLogMacroDescriptor

#rnx

diff --git a/Engine/Source/Editor/UnrealEd/Private/Commandlets/GatherTextFromSourceCommandlet.cpp b/Engine/Source/Editor/UnrealEd/Private/Commandlets/GatherTextFromSourceCommandlet.cpp
--- a/Engine/Source/Editor/UnrealEd/Private/Commandlets/GatherTextFromSourceCommandlet.cpp
+++ b/Engine/Source/Editor/UnrealEd/Private/Commandlets/GatherTextFromSourceCommandlet.cpp
@@ -2279,7 +2279,7 @@ void UGatherTextFromSourceCommandlet::FStringTableEntryMetaDataMacroDescriptor::
 }
 
 UGatherTextFromSourceCommandlet::FStructuredLogMacroDescriptor::FStructuredLogMacroDescriptor(const TCHAR* InName, EFlags InFlags)
-	: FMacroDescriptor(InName, ((int32)Flags & (int32)EFlags::Namespace) ? 5 : 4)
+	: FMacroDescriptor(InName, ((int32)InFlags & (int32)EFlags::Namespace) ? 5 : 4)
 	, Flags(InFlags)
 {
 }
