Cooker: Optimize PollPendingPlatformDatas by calling IsCachedCookedPlatformDataLoaded at exponentially longer intervals for objects that keep returning false.
#rnx

diff --git a/Engine/Source/Editor/UnrealEd/Classes/CookOnTheSide/CookOnTheFlyServer.h b/Engine/Source/Editor/UnrealEd/Classes/CookOnTheSide/CookOnTheFlyServer.h
--- a/Engine/Source/Editor/UnrealEd/Classes/CookOnTheSide/CookOnTheFlyServer.h
+++ b/Engine/Source/Editor/UnrealEd/Classes/CookOnTheSide/CookOnTheFlyServer.h
@@ -1142,7 +1142,7 @@ private:
 	/**
 	 * Frees all the memory used to call BeginCacheForCookedPlatformData on all the objects in PackageData.
 	 * If the calls were incomplete because the PackageData's save was cancelled, handles canceling them and
-	 * leaving any required CancelManagers in GetPendingCookedPlatformDatas
+	 * leaving any required CancelManagers in PendingCookedPlatformDatas
 	 * 
 	 * @param bCompletedSave If false, data that can not be efficiently recomputed will be preserved to try
 	 *        the save again. If true, all data will be wiped.
@@ -1151,7 +1151,7 @@ private:
 	void ReleaseCookedPlatformData(UE::Cook::FPackageData& PackageData, UE::Cook::EReleaseSaveReason ReleaseSaveReason);
 
 	/**
-	 * Poll the GetPendingCookedPlatformDatas and release their resources when they are complete.
+	 * Poll the PendingCookedPlatformDatas and release their resources when they are complete.
 	 * This is done inside of PumpSaveQueue, but is also required after a cancelled cook so that references to
 	 * the pending objects will be eventually dropped.
 	 */
diff --git a/Engine/Source/Editor/UnrealEd/Private/CookOnTheFlyServer.cpp b/Engine/Source/Editor/UnrealEd/Private/CookOnTheFlyServer.cpp
--- a/Engine/Source/Editor/UnrealEd/Private/CookOnTheFlyServer.cpp
+++ b/Engine/Source/Editor/UnrealEd/Private/CookOnTheFlyServer.cpp
@@ -1431,15 +1431,16 @@ void UCookOnTheFlyServer::SetSaveBusy(bool bInBusy)
 			TArray<UObject*> ExpectedObjects;
 			TSet<UPackage*> ExpectedPackages;
 			FPackageDataQueue& SaveQueue = PackageDatas->GetSaveQueue();
-			const TArray<FPendingCookedPlatformData>& PendingCookedPlatformDatas = PackageDatas->GetPendingCookedPlatformDatas();
 			TArray<UClass*> CompilationUsers({ UMaterialInterface::StaticClass(), FindObject<UClass>(nullptr, TEXT("/Script/Niagara.NiagaraScript")) });
 
-			for (const FPendingCookedPlatformData& Data : PendingCookedPlatformDatas)
+			PackageDatas->ForEachPendingCookedPlatformData(
+			[&CompilationUsers, &ExpectedObjects, &ExpectedPackages, &NonExpectedObjects, &NonExpectedPackages]
+			(const FPendingCookedPlatformData& Data)
 			{
 				UObject* Object = Data.Object.Get();
 				if (!Object)
 				{
-					continue;
+					return;
 				}
 				bool bCompilationUser = false;
 				for (UClass* CompilationUserClass : CompilationUsers)
@@ -1460,7 +1461,7 @@ void UCookOnTheFlyServer::SetSaveBusy(bool bInBusy)
 					NonExpectedObjects.Add(Object);
 					NonExpectedPackages.Add(Object->GetPackage());
 				}
-			}
+			});
 			TArray<UPackage*> RemovePackages;
 			for (UPackage* Package : ExpectedPackages)
 			{
@@ -1506,7 +1507,8 @@ void UCookOnTheFlyServer::SetSaveBusy(bool bInBusy)
 				UE_LOG(LogCook, Display, TEXT("    <None>"));
 			}
 
-			UE_LOG(LogCook, Display, TEXT("%d objects that have not yet returned true from IsCachedCookedPlatformDataLoaded:"), PackageDatas->GetPendingCookedPlatformDatas().Num());
+			UE_LOG(LogCook, Display, TEXT("%d objects that have not yet returned true from IsCachedCookedPlatformDataLoaded:"),
+				PackageDatas->GetPendingCookedPlatformDataNum());
 			DisplayCount = 0;
 			for (TArray<UObject*>* ObjectArray : {  &NonExpectedObjects, &ExpectedObjects })
 			{
@@ -3292,7 +3294,8 @@ UE::Cook::EPollStatus UCookOnTheFlyServer::BeginCachePostMove(UE::Cook::FGenerat
 	{
 		UObject* FirstPendingObject = nullptr;
 		FString FirstPendingObjectName;
-		for (const FPendingCookedPlatformData& Pending : PackageDatas->GetPendingCookedPlatformDatas())
+		PackageDatas->ForEachPendingCookedPlatformData([&PackageData, &FirstPendingObject, &FirstPendingObjectName]
+		(const FPendingCookedPlatformData& Pending)
 		{
 			if (&Pending.PackageData == &PackageData)
 			{
@@ -3303,7 +3306,7 @@ UE::Cook::EPollStatus UCookOnTheFlyServer::BeginCachePostMove(UE::Cook::FGenerat
 					FirstPendingObjectName = MoveTemp(ObjectName);
 				}
 			}
-		}
+		});
 		UE_LOG(LogCook, Warning, TEXT("CookPackageSplitter created or moved objects during %s that are not yet ready to save. This will cause an error if garbage collection runs before the package is saved.\n")
 			TEXT("Change the splitter's %s to construct new objects and declare existing objects that will be moved from other packages.\n")
 			TEXT("SplitterObject: %s%s\n")
@@ -3623,7 +3626,7 @@ UE::Cook::EPollStatus UCookOnTheFlyServer::CallBeginCacheOnObjects(UE::Cook::FPa
 		else
 		{
 			bool bNeedsResourceRelease = CurrentAsyncCache != nullptr;
-			PackageDatas->GetPendingCookedPlatformDatas().Emplace(Obj, TargetPlatform, PackageData, bNeedsResourceRelease, *this);
+			PackageDatas->AddPendingCookedPlatformData(FPendingCookedPlatformData(Obj, TargetPlatform, PackageData, bNeedsResourceRelease, *this));
 		}
 
 		if (Timer.IsTimeUp())
@@ -3701,7 +3704,8 @@ void UCookOnTheFlyServer::ReleaseCookedPlatformData(UE::Cook::FPackageData& Pack
 	else 
 	{
 		// This is a slower but more general flow that can handle releasing whether or not we called SavePackage
-		// Note that even after we return from this function, some objects with pending IsCachedCookedPlatformDataLoaded calls may still exist for this Package in PackageDatas->GetPendingCookedPlatformDatas(),
+		// Note that even after we return from this function, some objects with pending IsCachedCookedPlatformDataLoaded
+		// calls may still exist for this Package in PendingCookedPlatformDatas
 		// and this PackageData may therefore still have GetNumPendingCookedPlatformData > 0
 		if (!IsCookingInEditor()) // ClearAllCachedCookedPlatformData calls are only used when not in editor.
 		{
@@ -3713,7 +3717,8 @@ void UCookOnTheFlyServer::ReleaseCookedPlatformData(UE::Cook::FPackageData& Pack
 
 				// Find all pending BeginCacheForCookedPlatformData for this FPackageData
 				TMap<UObject*, TArray<FPendingCookedPlatformData*>> PendingObjects;
-				for (FPendingCookedPlatformData& PendingCookedPlatformData : PackageDatas->GetPendingCookedPlatformDatas())
+				PackageDatas->ForEachPendingCookedPlatformData(
+				[&PendingObjects, &PackageData](FPendingCookedPlatformData& PendingCookedPlatformData)
 				{
 					if (&PendingCookedPlatformData.PackageData == &PackageData && !PendingCookedPlatformData.PollIsComplete())
 					{
@@ -3722,7 +3727,7 @@ void UCookOnTheFlyServer::ReleaseCookedPlatformData(UE::Cook::FPackageData& Pack
 						check(!PendingCookedPlatformData.bHasReleased); // bHasReleased should be false since PollIsComplete returned false
 						PendingObjects.FindOrAdd(Object).Add(&PendingCookedPlatformData);
 					}
-				}
+				});
 
 				TArray<UObject*> ObjectsToClear;
 				TArray<FWeakObjectPtr>* CachedObjects = &PackageData.GetCachedObjectsInOuter();
@@ -4163,7 +4168,7 @@ void UCookOnTheFlyServer::PumpSaves(UE::Cook::FTickStackData& StackData, uint32
 				{
 					// PrepareSave might block on pending CookedPlatformDatas, and it might block on resources held by other
 					// CookedPlatformDatas. Calling PollPendingCookedPlatformDatas should handle pumping all of those.
-					if (!PackageDatas->GetPendingCookedPlatformDatas().Num())
+					if (!PackageDatas->GetPendingCookedPlatformDataNum())
 					{
 						// We're waiting on something other than pendingcookedplatformdatas; this loop does not yet handle
 						// updating anything else, so break out
@@ -4659,7 +4664,8 @@ void UCookOnTheFlyServer::PreGarbageCollect()
 	// completed.
 	{
 		TMap<FPackageData*, UPackage*> UniquePendingPackages;
-		for (FPendingCookedPlatformData& PendingData : PackageDatas->GetPendingCookedPlatformDatas())
+		PackageDatas->ForEachPendingCookedPlatformData(
+		[&UniquePendingPackages](const FPendingCookedPlatformData& PendingData)
 		{
 			if (UObject* Object = PendingData.Object.Get())
 			{	
@@ -4668,7 +4674,7 @@ void UCookOnTheFlyServer::PreGarbageCollect()
 					UniquePendingPackages.Add(&PendingData.PackageData, Package);
 				}	
 			}
-		}
+		});
 
 		GCKeepPackages.Reserve(GCKeepPackages.Num() + UniquePendingPackages.Num());
 		for (const TPair<FPackageData*,UPackage*>& Pair : UniquePendingPackages)
@@ -4679,7 +4685,8 @@ void UCookOnTheFlyServer::PreGarbageCollect()
 	}
 
 	// Prevent GC of any objects on which we are still waiting for IsCachedCookedPlatformData
-	for (UE::Cook::FPendingCookedPlatformData& Pending : PackageDatas->GetPendingCookedPlatformDatas())
+	PackageDatas->ForEachPendingCookedPlatformData(
+	[this](UE::Cook::FPendingCookedPlatformData& Pending)
 	{
 		if (!Pending.PollIsComplete())
 		{
@@ -4687,7 +4694,7 @@ void UCookOnTheFlyServer::PreGarbageCollect()
 			check(Object); // Otherwise PollIsComplete would have returned true
 			GCKeepObjects.Add(Object);
 		}
-	}
+	});
 
 	const bool bPartialGC = IsCookFlagSet(ECookInitializationFlags::EnablePartialGC);
 	if (bGarbageCollectTypeSoft || bPartialGC)
diff --git a/Engine/Source/Editor/UnrealEd/Private/Cooker/CookPackageData.cpp b/Engine/Source/Editor/UnrealEd/Private/Cooker/CookPackageData.cpp
--- a/Engine/Source/Editor/UnrealEd/Private/Cooker/CookPackageData.cpp
+++ b/Engine/Source/Editor/UnrealEd/Private/Cooker/CookPackageData.cpp
@@ -1905,9 +1905,11 @@ FPendingCookedPlatformData::FPendingCookedPlatformData(UObject* InObject, const
 FPendingCookedPlatformData::FPendingCookedPlatformData(FPendingCookedPlatformData&& Other)
 	: Object(Other.Object), TargetPlatform(Other.TargetPlatform), PackageData(Other.PackageData)
 	, CookOnTheFlyServer(Other.CookOnTheFlyServer), CancelManager(Other.CancelManager), ClassName(Other.ClassName)
-	, bHasReleased(Other.bHasReleased), bNeedsResourceRelease(Other.bNeedsResourceRelease)
+	, UpdatePeriodMultiplier(Other.UpdatePeriodMultiplier), bHasReleased(Other.bHasReleased)
+	, bNeedsResourceRelease(Other.bNeedsResourceRelease)
 {
 	Other.Object = nullptr;
+	Other.bHasReleased = true;
 }
 
 FPendingCookedPlatformData::~FPendingCookedPlatformData()
@@ -2602,7 +2604,7 @@ void FPackageDatas::GetCookedPackagesForPlatform(const ITargetPlatform* Platform
 void FPackageDatas::Clear()
 {
 	FWriteScopeLock ExistenceWriteLock(ExistenceLock);
-	PendingCookedPlatformDatas.Empty(); // These destructors will dereference PackageDatas
+	PendingCookedPlatformDataLists.Empty(); // These destructors will dereference PackageDatas
 	RequestQueue.Empty();
 	SaveQueue.Empty();
 	PackageNameToPackageData.Empty();
@@ -2635,14 +2637,23 @@ void FPackageDatas::OnRemoveSessionPlatform(const ITargetPlatform* TargetPlatfor
 	}
 }
 
-TArray<FPendingCookedPlatformData>& FPackageDatas::GetPendingCookedPlatformDatas()
+constexpr int32 PendingPlatformDataReservationSize = 128;
+constexpr int32 PendingPlatformDataMaxUpdatePeriod = 16;
+
+void FPackageDatas::AddPendingCookedPlatformData(FPendingCookedPlatformData&& Data)
 {
-	return PendingCookedPlatformDatas;
+	if (PendingCookedPlatformDataLists.IsEmpty())
+	{
+		PendingCookedPlatformDataLists.Emplace();
+		PendingCookedPlatformDataLists.Last().Reserve(PendingPlatformDataReservationSize );
+	}
+	PendingCookedPlatformDataLists.First().Add(MoveTemp(Data));
+	++PendingCookedPlatformDataNum;
 }
 
 void FPackageDatas::PollPendingCookedPlatformDatas(bool bForce, double& LastCookableObjectTickTime)
 {
-	if (PendingCookedPlatformDatas.Num() == 0)
+	if (PendingCookedPlatformDataNum == 0)
 	{
 		return;
 	}
@@ -2659,6 +2670,22 @@ void FPackageDatas::PollPendingCookedPlatformDatas(bool bForce, double& LastCook
 			return;
 		}
 	}
+	LastPollAsyncTime = CurrentTime;
+
+	// PendingPlatformDataLists is a rotating list of lists of PendingPlatformDatas
+	// The first list contains all of the PendingPlatformDatas that we should poll on this tick
+	// The nth list is all of the PendingPlatformDatas that we should poll after N more ticks
+	// Each poll period we pull the front list off and all other lists move frontwards by 1.
+	// New PendingPlatformDatas are inserted into the first list, to be polled in the next poll period
+	// When a PendingPlatformData signals it is not ready after polling, we increase its poll period
+	// exponentially - we double it.
+	// A poll period of N times the default poll period means we insert it into the Nth list in
+	// PendingPlatformDataLists.
+	FPendingCookedPlatformDataContainer List = PendingCookedPlatformDataLists.PopFrontValue();
+	if (!bForce && List.IsEmpty())
+	{
+		return;
+	}
 
 	GShaderCompilingManager->ProcessAsyncResults(true /* bLimitExecutionTime */,
 		false /* bBlockOnGlobalShaderCompletion */);
@@ -2670,20 +2697,51 @@ void FPackageDatas::PollPendingCookedPlatformDatas(bool bForce, double& LastCook
 		LastCookableObjectTickTime = CurrentTime;
 	}
 
-	FPendingCookedPlatformData* Datas = PendingCookedPlatformDatas.GetData();
-	for (int Index = 0; Index < PendingCookedPlatformDatas.Num();)
+	if (!bForce)
 	{
-		if (Datas[Index].PollIsComplete())
+		for (FPendingCookedPlatformData& Data : List)
 		{
-			PendingCookedPlatformDatas.RemoveAtSwap(Index, 1 /* Count */, false /* bAllowShrinking */);
+			if (Data.PollIsComplete())
+			{
+				// We are destructing all elements of List after the for loop is done; we leave
+				// the completed Data on List to be destructed.
+				--PendingCookedPlatformDataNum;
+			}
+			else
+			{
+				Data.UpdatePeriodMultiplier = FMath::Clamp(Data.UpdatePeriodMultiplier*2, 1, PendingPlatformDataMaxUpdatePeriod);
+				int32 ContainerIndex = Data.UpdatePeriodMultiplier - 1;
+				while (PendingCookedPlatformDataLists.Num() <= ContainerIndex)
+				{
+					PendingCookedPlatformDataLists.Emplace();
+					PendingCookedPlatformDataLists.Last().Reserve(PendingPlatformDataReservationSize);
+				}
+				PendingCookedPlatformDataLists[ContainerIndex].Add(MoveTemp(Data));
+			}
 		}
-		else
+	}
+	else
+	{
+		// When called with bForce, we poll all PackageDatas in all lists, and do not update
+		// any PollPeriods.
+		PendingCookedPlatformDataLists.AddFront(MoveTemp(List));
+		for (FPendingCookedPlatformDataContainer& ForceList : PendingCookedPlatformDataLists)
 		{
-			++Index;
+			for (int32 Index = 0; Index < ForceList.Num(); )
+			{
+				FPendingCookedPlatformData& Data = ForceList[Index];
+				if (Data.PollIsComplete())
+				{
+					ForceList.RemoveAtSwap(Index, 1, false /* bAllowShrinking */);
+					--PendingCookedPlatformDataNum;
+				}
+				else
+				{
+					++Index;
+				}
+			}
 		}
 	}
-
-	LastPollAsyncTime = CurrentTime;
 }
 
 TArray<FPackageData*>::RangedForIteratorType FPackageDatas::begin()
@@ -2702,10 +2760,10 @@ void FPackageDatas::RemapTargetPlatforms(const TMap<ITargetPlatform*, ITargetPla
 	{
 		PackageData->RemapTargetPlatforms(Remap);
 	}
-	for (FPendingCookedPlatformData& CookedPlatformData : PendingCookedPlatformDatas)
+	ForEachPendingCookedPlatformData([&Remap](FPendingCookedPlatformData& CookedPlatformData)
 	{
 		CookedPlatformData.RemapTargetPlatforms(Remap);
-	}
+	});
 }
 
 void FPackageDatas::DebugInstigator(FPackageData& PackageData)
diff --git a/Engine/Source/Editor/UnrealEd/Private/Cooker/CookPackageData.h b/Engine/Source/Editor/UnrealEd/Private/Cooker/CookPackageData.h
--- a/Engine/Source/Editor/UnrealEd/Private/Cooker/CookPackageData.h
+++ b/Engine/Source/Editor/UnrealEd/Private/Cooker/CookPackageData.h
@@ -853,6 +853,8 @@ struct FPendingCookedPlatformData
 	FPendingCookedPlatformDataCancelManager* CancelManager;
 	/* Saved copy of the ClassName to use for resource releasing. */
 	FName ClassName;
+	/** Polling performance field: how many UpdatePeriods should we wait before polling again. */
+	int32 UpdatePeriodMultiplier = 1;
 	/** Flag for whether we have executed the release. */
 	bool bHasReleased;
 	/**
@@ -983,6 +985,8 @@ public:
 	FPackageDataQueue EntryQueue;
 };
 
+typedef TArray<FPendingCookedPlatformData> FPendingCookedPlatformDataContainer;
+
 /*
  * Class that manages the list of all PackageDatas for a CookOnTheFlyServer. PackageDatas is an associative
  * array for extra data about a package (e.g. the requested platforms) that is needed by the CookOnTheFlyServer.
@@ -1180,8 +1184,21 @@ public:
 	/** Remove all request data about the given platform from all PackageDatas and other memory used by *this. */
 	void OnRemoveSessionPlatform(const ITargetPlatform* TargetPlatform);
 
-	/** Return the container that tracks pending calls to BeginCacheForCookedPlatformData. */
-	TArray<FPendingCookedPlatformData>& GetPendingCookedPlatformDatas();
+	/** Enumerate PendingPlatformDatas: the list of pending calls to BeginCacheForCookedPlatformData. */
+	template <typename FunctionType>
+	void ForEachPendingCookedPlatformData(const FunctionType& Function)
+	{
+		for (FPendingCookedPlatformDataContainer& Container : PendingCookedPlatformDataLists)
+		{
+			for (FPendingCookedPlatformData& Data : Container)
+			{
+				Function(Data);
+			}
+		}
+	}
+	int32 GetPendingCookedPlatformDataNum() const { return PendingCookedPlatformDataNum; }
+	void AddPendingCookedPlatformData(FPendingCookedPlatformData&& Data);
+
 	/**
 	 * Iterate over all elements in PendingCookedPlatformDatas and check whether they have completed,
 	 * releasing their resources and pending count if so.
@@ -1217,7 +1234,8 @@ private:
 	FPackageDataMonitor Monitor;
 	TMap<FName, FPackageData*> PackageNameToPackageData;
 	TMap<FName, FPackageData*> FileNameToPackageData;
-	TArray<FPendingCookedPlatformData> PendingCookedPlatformDatas;
+	TRingBuffer<FPendingCookedPlatformDataContainer> PendingCookedPlatformDataLists;
+	int32 PendingCookedPlatformDataNum = 0;
 	FRequestQueue RequestQueue;
 	TFastPointerSet<FPackageData*> AssignedToWorkerSet;
 	FLoadPrepareQueue LoadPrepareQueue;
