Cooker SoftGC Bugfix: Next SoftGC memory target needs to be LESS than or equal to StartGC, not more than or equal, so that we only do SoftGC a few times.
#rnx

diff --git a/Engine/Source/Editor/UnrealEd/Private/CookOnTheFlyServer.cpp b/Engine/Source/Editor/UnrealEd/Private/CookOnTheFlyServer.cpp
--- a/Engine/Source/Editor/UnrealEd/Private/CookOnTheFlyServer.cpp
+++ b/Engine/Source/Editor/UnrealEd/Private/CookOnTheFlyServer.cpp
@@ -4877,7 +4877,7 @@ void UCookOnTheFlyServer::EvaluateGarbageCollectionResults(bool bWasDueToOOM, bo
 		int64 PhysicalMemoryQuantum = static_cast<int64>(MemStatsAfterGC.TotalPhysical) / Denominator;
 		int32 NextTarget =
 			static_cast<int64>(MemStatsAfterGC.AvailablePhysical - PhysicalMemoryQuantum/2) / PhysicalMemoryQuantum;
-		NextTarget = FMath::Max(NextTarget, StartNumerator);
+		NextTarget = FMath::Min(NextTarget, StartNumerator);
 		if (NextTarget <= 0)
 		{
 			SoftGCNextAvailablePhysicalTarget = -2; // disabled, no further targets
