Fix warnings when trying to delete objects from the content browsers that have referencers marked as garbage/pending kill

#rnx

diff --git a/Engine/Source/Editor/UnrealEd/Private/ObjectTools.cpp b/Engine/Source/Editor/UnrealEd/Private/ObjectTools.cpp
--- a/Engine/Source/Editor/UnrealEd/Private/ObjectTools.cpp
+++ b/Engine/Source/Editor/UnrealEd/Private/ObjectTools.cpp
@@ -264,18 +264,22 @@ namespace ObjectTools
 			LastObjectCount = InObjects.Num();
 			for (UObject* NewReferencer : FReferencerFinder::GetAllReferencers(InObjects, &InObjects, EReferencerFinderFlags::SkipWeakReferences))
 			{
-				// Stop walking the dependency chain when the transactor is the referencer
-				if (Transactor == NewReferencer)
+				// Exclude any pendingkill or garbage object from counting as referencers
+				if (IsValid(NewReferencer))
 				{
-					Roots.Add(Transactor);
-				}
-				else if (IsNonGCObject(NewReferencer))
-				{
-					Roots.Add(NewReferencer);
-				}
-				else
-				{
-					InObjects.Add(NewReferencer);
+					// Stop walking the dependency chain when the transactor is the referencer
+					if (Transactor == NewReferencer)
+					{
+						Roots.Add(Transactor);
+					}
+					else if (IsNonGCObject(NewReferencer))
+					{
+						Roots.Add(NewReferencer);
+					}
+					else
+					{
+						InObjects.Add(NewReferencer);
+					}
 				}
 			}
 		}
@@ -335,20 +339,24 @@ namespace ObjectTools
 			// Check and see whether we are referenced by any objects that won't be garbage collected (*including* the undo buffer)
 			for (UObject* Referencer : FReferencerFinder::GetAllReferencers(ObjectsToDelete, nullptr, EReferencerFinderFlags::SkipWeakReferences))
 			{
-				if (Referencer->IsIn(InObject))
+				// Exclude any pendingkill or garbage object from counting as referencers
+				if (IsValid(Referencer))
 				{
-					InternalReferences.Add(Referencer);
-				}
-				else
-				{
-					if (Transactor == Referencer)
+					if (Referencer->IsIn(InObject))
 					{
-						bOutIsReferencedInMemoryByUndo = true;
+						InternalReferences.Add(Referencer);
 					}
 					else
 					{
-						References.ExternalReferences.Emplace(Referencer);
-						bOutIsReferenced = true;
+						if (Transactor == Referencer)
+						{
+							bOutIsReferencedInMemoryByUndo = true;
+						}
+						else
+						{
+							References.ExternalReferences.Emplace(Referencer);
+							bOutIsReferenced = true;
+						}
 					}
 				}
 			}
