Level Editor: Added a Show Subcomponents flag
- When true, all components that are part of the selection set will draw their vizualizers.
- When false, only the selected components and the selected actors root components will draw their vizualizers.
- Defaults to true.

diff --git a/Engine/Source/Editor/LevelEditor/Private/LevelEditor.cpp b/Engine/Source/Editor/LevelEditor/Private/LevelEditor.cpp
--- a/Engine/Source/Editor/LevelEditor/Private/LevelEditor.cpp
+++ b/Engine/Source/Editor/LevelEditor/Private/LevelEditor.cpp
@@ -1582,6 +1582,13 @@ void FLevelEditorModule::BindGlobalLevelEditorCommands()
 		FIsActionChecked::CreateStatic( &FLevelEditorActionCallbacks::OnGetTransformWidgetVisibility )
 		);
 
+	ActionList.MapAction(
+		Commands.ShowSelectionSubcomponents,
+		FExecuteAction::CreateStatic(&FLevelEditorActionCallbacks::OnToggleShowSelectionSubcomponents),
+		FCanExecuteAction(),
+		FIsActionChecked::CreateStatic(&FLevelEditorActionCallbacks::OnGetShowSelectionSubcomponents)
+	);
+
 	ActionList.MapAction(
 		Commands.AllowTranslucentSelection,
 		FExecuteAction::CreateStatic( &FLevelEditorActionCallbacks::OnAllowTranslucentSelection ),
diff --git a/Engine/Source/Editor/LevelEditor/Private/LevelEditorActions.cpp b/Engine/Source/Editor/LevelEditor/Private/LevelEditorActions.cpp
--- a/Engine/Source/Editor/LevelEditor/Private/LevelEditorActions.cpp
+++ b/Engine/Source/Editor/LevelEditor/Private/LevelEditorActions.cpp
@@ -2590,9 +2590,23 @@ bool FLevelEditorActionCallbacks::OnGetTransformWidgetVisibility()
 	return GLevelEditorModeTools().GetShowWidget();
 }
 
+void FLevelEditorActionCallbacks::OnToggleShowSelectionSubcomponents()
+{
+	UEditorPerProjectUserSettings* Settings = GetMutableDefault<UEditorPerProjectUserSettings>();
+	Settings->bShowSelectionSubcomponents = !Settings->bShowSelectionSubcomponents;
+	Settings->PostEditChange();
+
+	GUnrealEd->RedrawAllViewports();
+}
+
+bool FLevelEditorActionCallbacks::OnGetShowSelectionSubcomponents()
+{
+	return GetDefault<UEditorPerProjectUserSettings>()->bShowSelectionSubcomponents == true;
+}
+
 void FLevelEditorActionCallbacks::OnAllowTranslucentSelection()
 {
-	auto* Settings = GetMutableDefault<UEditorPerProjectUserSettings>();
+	UEditorPerProjectUserSettings* Settings = GetMutableDefault<UEditorPerProjectUserSettings>();
 
 	// Toggle 'allow select translucent'
 	Settings->bAllowSelectTranslucent = !Settings->bAllowSelectTranslucent;
@@ -3724,6 +3738,8 @@ void FLevelEditorCommands::RegisterCommands()
 	UI_COMMAND( AllowGroupSelection, "Allow Group Selection", "Allows actor groups to be selected", EUserInterfaceActionType::ToggleButton, FInputChord(EModifierKey::Control|EModifierKey::Shift, EKeys::G) );
 	UI_COMMAND( StrictBoxSelect, "Strict Box Selection", "When enabled an object must be entirely encompassed by the selection box when marquee box selecting", EUserInterfaceActionType::ToggleButton, FInputChord() );
 	UI_COMMAND( TransparentBoxSelect, "Box Select Occluded Objects", "When enabled, marquee box select operations will also select objects that are occluded by other objects.", EUserInterfaceActionType::ToggleButton, FInputChord() );
+	UI_COMMAND( ShowSelectionSubcomponents, "Show Subcomponents", "Toggles the visibility of the subcomponents related to the current selection", EUserInterfaceActionType::ToggleButton, FInputChord() );
+	
 	UI_COMMAND( DrawBrushMarkerPolys, "Draw Brush Polys", "Draws semi-transparent polygons around a brush when selected", EUserInterfaceActionType::ToggleButton, FInputChord() );
 	UI_COMMAND( OnlyLoadVisibleInPIE, "Only Load Visible Levels in Game Preview", "If enabled, when game preview starts, only visible levels will be loaded", EUserInterfaceActionType::ToggleButton, FInputChord() );
 	UI_COMMAND( ToggleSocketSnapping, "Enable Socket Snapping", "Enables or disables snapping to sockets", EUserInterfaceActionType::ToggleButton, FInputChord() ); 
diff --git a/Engine/Source/Editor/LevelEditor/Private/LevelEditorToolBar.cpp b/Engine/Source/Editor/LevelEditor/Private/LevelEditorToolBar.cpp
--- a/Engine/Source/Editor/LevelEditor/Private/LevelEditorToolBar.cpp
+++ b/Engine/Source/Editor/LevelEditor/Private/LevelEditorToolBar.cpp
@@ -1699,6 +1699,7 @@ void FLevelEditorToolBar::RegisterQuickSettingsMenu()
 		Section.AddMenuEntry( FLevelEditorCommands::Get().StrictBoxSelect );
 		Section.AddMenuEntry( FLevelEditorCommands::Get().TransparentBoxSelect );
 		Section.AddMenuEntry( FLevelEditorCommands::Get().ShowTransformWidget );
+		Section.AddMenuEntry( FLevelEditorCommands::Get().ShowSelectionSubcomponents );
 	}
 
 	{
diff --git a/Engine/Source/Editor/LevelEditor/Public/LevelEditorActions.h b/Engine/Source/Editor/LevelEditor/Public/LevelEditorActions.h
--- a/Engine/Source/Editor/LevelEditor/Public/LevelEditorActions.h
+++ b/Engine/Source/Editor/LevelEditor/Public/LevelEditorActions.h
@@ -574,6 +574,7 @@ public:
 	TSharedPtr< FUICommandInfo > ShowTransformWidget;
 	TSharedPtr< FUICommandInfo > AllowTranslucentSelection;
 	TSharedPtr< FUICommandInfo > AllowGroupSelection;
+	TSharedPtr< FUICommandInfo > ShowSelectionSubcomponents;
 
 	TSharedPtr< FUICommandInfo > StrictBoxSelect;
 	TSharedPtr< FUICommandInfo > TransparentBoxSelect;
@@ -1187,6 +1188,8 @@ public:
 	 */ 
 	static void OnToggleTransformWidgetVisibility();
 	static bool OnGetTransformWidgetVisibility();
+	static void OnToggleShowSelectionSubcomponents();
+	static bool OnGetShowSelectionSubcomponents();
 	static void OnAllowTranslucentSelection();
 	static bool OnIsAllowTranslucentSelectionEnabled();	
 	static void OnAllowGroupSelection();
diff --git a/Engine/Source/Editor/UnrealEd/Classes/Editor/EditorPerProjectUserSettings.h b/Engine/Source/Editor/UnrealEd/Classes/Editor/EditorPerProjectUserSettings.h
--- a/Engine/Source/Editor/UnrealEd/Classes/Editor/EditorPerProjectUserSettings.h
+++ b/Engine/Source/Editor/UnrealEd/Classes/Editor/EditorPerProjectUserSettings.h
@@ -147,6 +147,10 @@ class UEditorPerProjectUserSettings : public UObject
 	UPROPERTY(config)
 	uint32 bAllowSelectTranslucent:1;
 
+	/** If this is true, all of an actors' components will be drawn when the actor or one of its component is selected */
+	UPROPERTY(config)
+	uint32 bShowSelectionSubcomponents:1;
+
 	UPROPERTY()
 	TObjectPtr<class UBlueprintPaletteFavorites> BlueprintFavorites;
 	
diff --git a/Engine/Source/Editor/UnrealEd/Classes/Editor/UnrealEdEngine.h b/Engine/Source/Editor/UnrealEd/Classes/Editor/UnrealEdEngine.h
--- a/Engine/Source/Editor/UnrealEd/Classes/Editor/UnrealEdEngine.h
+++ b/Engine/Source/Editor/UnrealEd/Classes/Editor/UnrealEdEngine.h
@@ -886,7 +886,13 @@ protected:
 	/**
 	 * The list of visualizers to draw when selection changes
 	 */
-	TArray<FCachedComponentVisualizer> VisualizersForSelection;
+	struct FComponentVisualizerForSelection
+	{
+		FCachedComponentVisualizer ComponentVisualizer;
+		TOptional<TFunction<bool(void)>> IsEnabledDelegate;
+	};
+
+	TArray<FComponentVisualizerForSelection> VisualizersForSelection;
 
 	/** Instance responsible for monitoring this editor's performance */
 	FPerformanceMonitor* PerformanceMonitor;
diff --git a/Engine/Source/Editor/UnrealEd/Private/Settings/EditorPerProjectUserSettings.cpp b/Engine/Source/Editor/UnrealEd/Private/Settings/EditorPerProjectUserSettings.cpp
--- a/Engine/Source/Editor/UnrealEd/Private/Settings/EditorPerProjectUserSettings.cpp
+++ b/Engine/Source/Editor/UnrealEd/Private/Settings/EditorPerProjectUserSettings.cpp
@@ -34,6 +34,7 @@ UEditorPerProjectUserSettings::UEditorPerProjectUserSettings(const FObjectInitia
 	bPreviewFeatureLevelActive = false;
 	bPreviewFeatureLevelWasDefault = true;
 	PreviewDeviceProfileName = NAME_None;
+	bShowSelectionSubcomponents = true;
 }
 
 void UEditorPerProjectUserSettings::PostInitProperties()
diff --git a/Engine/Source/Editor/UnrealEd/Private/UnrealEdEngine.cpp b/Engine/Source/Editor/UnrealEd/Private/UnrealEdEngine.cpp
--- a/Engine/Source/Editor/UnrealEd/Private/UnrealEdEngine.cpp
+++ b/Engine/Source/Editor/UnrealEd/Private/UnrealEdEngine.cpp
@@ -1419,7 +1419,7 @@ void UUnrealEdEngine::RegisterComponentVisualizer(FName ComponentClassName, TSha
 void UUnrealEdEngine::UnregisterComponentVisualizer(FName ComponentClassName)
 {
 	TSharedPtr<FComponentVisualizer> Visualizer = FindComponentVisualizer(ComponentClassName);
-	VisualizersForSelection.RemoveAll([&Visualizer](const auto& CachedComponentVisualizer) { return CachedComponentVisualizer.Visualizer == Visualizer; });
+	VisualizersForSelection.RemoveAll([&Visualizer](const FComponentVisualizerForSelection& VisualizerForSelection) { return VisualizerForSelection.ComponentVisualizer.Visualizer == Visualizer; });
 
 	ComponentVisualizerMap.Remove(ComponentClassName);
 }
@@ -1454,24 +1454,30 @@ TSharedPtr<class FComponentVisualizer> UUnrealEdEngine::FindComponentVisualizer(
 
 void UUnrealEdEngine::DrawComponentVisualizers(const FSceneView* View, FPrimitiveDrawInterface* PDI)
 {
-	for(FCachedComponentVisualizer& CachedVisualizer : VisualizersForSelection)
+	for(FComponentVisualizerForSelection& VisualizerForSelection : VisualizersForSelection)
 	{
-		CachedVisualizer.Visualizer->DrawVisualization(CachedVisualizer.ComponentPropertyPath.GetComponent(), View, PDI);
+		if (!VisualizerForSelection.IsEnabledDelegate.IsSet() || VisualizerForSelection.IsEnabledDelegate.GetValue()())
+		{
+			VisualizerForSelection.ComponentVisualizer.Visualizer->DrawVisualization(VisualizerForSelection.ComponentVisualizer.ComponentPropertyPath.GetComponent(), View, PDI);
+		}
 	}
 }
 
 
 void UUnrealEdEngine::DrawComponentVisualizersHUD(const FViewport* Viewport, const FSceneView* View, FCanvas* Canvas)
 {
-	for(FCachedComponentVisualizer& CachedVisualizer : VisualizersForSelection)
+	for(FComponentVisualizerForSelection& VisualizerForSelection : VisualizersForSelection)
 	{
-		CachedVisualizer.Visualizer->DrawVisualizationHUD(CachedVisualizer.ComponentPropertyPath.GetComponent(), Viewport, View, Canvas);
+		if (!VisualizerForSelection.IsEnabledDelegate.IsSet() || VisualizerForSelection.IsEnabledDelegate.GetValue()())
+		{
+			VisualizerForSelection.ComponentVisualizer.Visualizer->DrawVisualizationHUD(VisualizerForSelection.ComponentVisualizer.ComponentPropertyPath.GetComponent(), Viewport, View, Canvas);
+		}
 	}
 }
 
 void UUnrealEdEngine::OnEditorSelectionChanged(UObject* SelectionThatChanged)
 {
-	auto GetVisualizersForSelection = [&](AActor* Actor)
+	auto GetVisualizersForSelection = [this](AActor* Actor, const UActorComponent* SelectedComponent)
 	{
 		// Iterate over components of that actor (and recurse through child components)
 		TInlineComponentArray<UActorComponent*> Components;
@@ -1486,7 +1492,15 @@ void UUnrealEdEngine::OnEditorSelectionChanged(UObject* SelectionThatChanged)
 				TSharedPtr<FComponentVisualizer> Visualizer = FindComponentVisualizer(Comp->GetClass());
 				if (Visualizer.IsValid())
 				{
-					VisualizersForSelection.Add(FCachedComponentVisualizer(Comp, Visualizer));
+					FCachedComponentVisualizer CachedComponentVisualizer(Comp, Visualizer);
+					FComponentVisualizerForSelection Temp{CachedComponentVisualizer};
+
+					FComponentVisualizerForSelection& ComponentVisualizerForSelection = VisualizersForSelection.Add_GetRef(MoveTemp(Temp));
+
+					if (Comp != SelectedComponent)
+					{
+						ComponentVisualizerForSelection.IsEnabledDelegate.Emplace([](){ return GetDefault<UEditorPerProjectUserSettings>()->bShowSelectionSubcomponents == true; });
+					}
 				}
 			}
 		}
@@ -1504,7 +1518,7 @@ void UUnrealEdEngine::OnEditorSelectionChanged(UObject* SelectionThatChanged)
 			AActor* Actor = Cast<AActor>(*It);
 			if (Actor != nullptr)
 			{
-				GetVisualizersForSelection(Actor);
+				GetVisualizersForSelection(Actor, Actor->GetRootComponent());
 			}
 		}
 	}
@@ -1530,7 +1544,7 @@ void UUnrealEdEngine::OnEditorSelectionChanged(UObject* SelectionThatChanged)
 						{
 							if (!ActorsProcessed.Contains(Actor))
 							{
-								GetVisualizersForSelection(Actor);
+								GetVisualizersForSelection(Actor, Comp);
 								ActorsProcessed.Emplace(Actor);
 							}
 						}
@@ -1543,11 +1557,11 @@ void UUnrealEdEngine::OnEditorSelectionChanged(UObject* SelectionThatChanged)
 	// If there is an undo/redo operation in progress, restore the active component visualizer.
 	if (GIsTransacting)
 	{
-		for (FCachedComponentVisualizer& CachedComponentVisualizer : VisualizersForSelection)
+		for (FComponentVisualizerForSelection& VisualizerForSelection : VisualizersForSelection)
 		{
-			if (CachedComponentVisualizer.Visualizer->GetEditedComponent() != nullptr)
+			if (VisualizerForSelection.ComponentVisualizer.Visualizer->GetEditedComponent() != nullptr)
 			{
-				ComponentVisManager.SetActiveComponentVis(GCurrentLevelEditingViewportClient, CachedComponentVisualizer.Visualizer);
+				ComponentVisManager.SetActiveComponentVis(GCurrentLevelEditingViewportClient, VisualizerForSelection.ComponentVisualizer.Visualizer);
 				break;
 			}
 		}
