Use the new generic finish compilation API to only wait on assets inside packages that are going to be deleted
Also follow up on the comment by merging duplicated functions between PackageTools and ObjectTools in 5.2

#rnx

diff --git a/Engine/Source/Editor/UnrealEd/Private/ObjectTools.cpp b/Engine/Source/Editor/UnrealEd/Private/ObjectTools.cpp
--- a/Engine/Source/Editor/UnrealEd/Private/ObjectTools.cpp
+++ b/Engine/Source/Editor/UnrealEd/Private/ObjectTools.cpp
@@ -159,53 +159,6 @@ void ReloadEditorWorldForReferenceReplacementIfNecessary(TArray< TWeakObjectPtr<
 	}
 }
 
-// Note that the contents of this namespace are a copy of code from PackageTools.cpp submitted as part of a 
-// release hot fix. In UE 5.2 we should merge the two sets of code and place them in a common utility file.
-namespace UE::Private
-{
-
-/**
- * Utility function that checks each UObject inside of the given UPackage to see if it is waiting
- * on async compilation.
- *
- * @return true if the package contains at least one UObject that has compilation work running, otherwise false.
- */
-static bool IsPackageCompiling(const UPackage* Package)
-{
-	bool bIsCompiling = false;
-	ForEachObjectWithPackage(Package, [&bIsCompiling](const UObject* Object)
-		{
-			const IInterface_AsyncCompilation* AsyncCompilationIF = Cast<IInterface_AsyncCompilation>(Object);
-			if (AsyncCompilationIF != nullptr && AsyncCompilationIF->IsCompiling())
-			{
-				bIsCompiling = true;
-				return false;
-			}
-			else
-			{
-				return true;
-			}
-		});
-
-	return bIsCompiling;
-}
-
-/**
- * Utility function that checks the provided package to see if it contains
- * any assets that currently have async compilation work running.
- * If there are assets that are waiting on async compilation work then we
- * wait on all currently outstanding work to finish before returning.
- */
-static void FlushAsyncCompilation(const UPackage* Package)
-{
-	if (UE::Private::IsPackageCompiling(Package))
-	{
-		FAssetCompilingManager::Get().FinishAllCompilation();
-	}
-}
-
-} // namespace UE::Private
-
 namespace ObjectTools
 {
 	/** Returns true if the specified object can be displayed in a content browser */
@@ -4214,8 +4167,8 @@ namespace ObjectTools
 			// Look for a thumbnail for this asset before we rename it
 			FObjectThumbnail* Thumbnail = ThumbnailTools::GetThumbnailForObject(Object);
 
-			//
-			UE::Private::FlushAsyncCompilation(Object->GetPackage());
+			// Make sure there is no async compilation outstanding in the package we're going to unload
+			UPackageTools::FlushAsyncCompilation( { Object->GetPackage() });
 
 			FString OldObjectFullName = Object->GetFullName();
 			FString OldObjectPathName = Object->GetPathName();
diff --git a/Engine/Source/Editor/UnrealEd/Private/PackageTools.cpp b/Engine/Source/Editor/UnrealEd/Private/PackageTools.cpp
--- a/Engine/Source/Editor/UnrealEd/Private/PackageTools.cpp
+++ b/Engine/Source/Editor/UnrealEd/Private/PackageTools.cpp
@@ -74,19 +74,11 @@ TSet<UPackage*>* UPackageTools::PackagesBeingUnloaded = nullptr;
 TSet<UObject*> UPackageTools::ObjectsThatHadFlagsCleared;
 FDelegateHandle UPackageTools::ReachabilityCallbackHandle;
 
-namespace
-{
-
-/**
- * Utility function that gathers all async compilable objects from given packages
- * and flush them to make sure there is no remaining async work trying to load data
- * from said packages.
- */
-static void FlushAsyncCompilation(const TSet<UPackage*>& PackagesToUnload)
+void UPackageTools::FlushAsyncCompilation(TArrayView<UPackage* const> InPackages)
 {
 	TArray<UObject*> ObjectsToFinish;
 
-	for (const UPackage* Package : PackagesToUnload)
+	for (const UPackage* Package : InPackages)
 	{
 		ForEachObjectWithPackage(Package, [&ObjectsToFinish](UObject* Object)
 		{
@@ -105,8 +97,6 @@ static void FlushAsyncCompilation(const TSet<UPackage*>& PackagesToUnload)
 	}
 }
 
-} // anonymous namespace
-
 UPackageTools::UPackageTools(const FObjectInitializer& ObjectInitializer)
 	: Super(ObjectInitializer)
 {
@@ -417,7 +407,7 @@ UPackageTools::UPackageTools(const FObjectInitializer& ObjectInitializer)
 
 			// We need to make sure that there is no async compilation work running for the packages that we are about to unload
 			// so that it is safe to call ::ResetLoaders
-			FlushAsyncCompilation(PackagesToUnload);
+			FlushAsyncCompilation(PackagesToUnload.Array());
 
 			// Now try to clean up assets in all packages to unload.
 			int32 PackageIndex = 0;
diff --git a/Engine/Source/Editor/UnrealEd/Public/PackageTools.h b/Engine/Source/Editor/UnrealEd/Public/PackageTools.h
--- a/Engine/Source/Editor/UnrealEd/Public/PackageTools.h
+++ b/Engine/Source/Editor/UnrealEd/Public/PackageTools.h
@@ -210,6 +210,14 @@ public:
 	 */
 	static UPackage* FindOrCreatePackageForAssetType(const FName LongPackageName, UClass* AssetClass);
 
+	/**
+	 * Utility function that gathers all async compilable objects from given packages
+	 * and flush them to make sure there is no remaining async work trying to load data
+	 * from said packages.
+	 * @param Packages List of packages from which to gather compilable assets to flush.
+	 */
+	static void FlushAsyncCompilation(TArrayView<UPackage* const> Packages);
+
 private:
 	static void RestoreStandaloneOnReachableObjects();
 
