Error if we fail loading streamable actors in the moderation commandlet.

diff --git a/Engine/Source/Editor/UnrealEd/Classes/Commandlets/BaseIteratePackagesCommandlet.h b/Engine/Source/Editor/UnrealEd/Classes/Commandlets/BaseIteratePackagesCommandlet.h
--- a/Engine/Source/Editor/UnrealEd/Classes/Commandlets/BaseIteratePackagesCommandlet.h
+++ b/Engine/Source/Editor/UnrealEd/Classes/Commandlets/BaseIteratePackagesCommandlet.h
@@ -9,6 +9,7 @@
 #include "Engine/EngineTypes.h"
 // world partition builder includes
 #include "WorldPartition/WorldPartitionBuilder.h"
+#include "WorldPartition/WorldPartitionActorDesc.h"
 
 #include "BaseIteratePackagesCommandlet.generated.h"
 
@@ -239,10 +240,13 @@ protected:
 
 	virtual void PerformWorldBuilderAdditionalOperations(class AActor* Actor, bool& bSavePackage) { }
 	virtual void PerformWorldBuilderAdditionalOperations(class UObject* Object, bool& bSavePackage) { }
+	virtual void WorldBuilderFailedLoadingActor(const FWorldPartitionActorDesc* ActorDesc) { } 
 
 	virtual void PostPerformAdditionalOperations(class UPackage* Package) { }
 
 
+
+
 	/**
 	 * Allows the commandlet to do something after all the packages have been processed
 	 */
diff --git a/Engine/Source/Editor/UnrealEd/Private/Commandlets/BaseIteratePackagesCommandlet.cpp b/Engine/Source/Editor/UnrealEd/Private/Commandlets/BaseIteratePackagesCommandlet.cpp
--- a/Engine/Source/Editor/UnrealEd/Private/Commandlets/BaseIteratePackagesCommandlet.cpp
+++ b/Engine/Source/Editor/UnrealEd/Private/Commandlets/BaseIteratePackagesCommandlet.cpp
@@ -721,6 +721,7 @@ void UBaseIteratePackagesCommandlet::LoadAndSaveOnePackage(const FString& Filena
 
 						if (!Actor)
 						{
+							WorldBuilderFailedLoadingActor(ActorDesc);
 							return true;
 						}
 
