Fixing issue where REINST classes are available in the class reference drop down

diff --git a/Engine/Source/Editor/UnrealEd/Private/Kismet2/KismetReinstanceUtilities.cpp b/Engine/Source/Editor/UnrealEd/Private/Kismet2/KismetReinstanceUtilities.cpp
--- a/Engine/Source/Editor/UnrealEd/Private/Kismet2/KismetReinstanceUtilities.cpp
+++ b/Engine/Source/Editor/UnrealEd/Private/Kismet2/KismetReinstanceUtilities.cpp
@@ -1649,7 +1649,7 @@ void FBlueprintCompileReinstancer::GetSortedClassHierarchy(UClass* ClassToSearch
 	OutHierarchy.Sort([](UClass& A, UClass& B)->bool { return FBlueprintCompileReinstancer::ReinstancerOrderingFunction(&A, &B); });
 }
 
-void FBlueprintCompileReinstancer::MoveDependentSkelToReinst(UClass* OwnerClass, TMap<UClass*, UClass*>& OldToNewMap)
+void FBlueprintCompileReinstancer::MoveDependentSkelToReinst(UClass* const OwnerClass, TMap<UClass*, UClass*>& OldToNewMap)
 {
 	// Gather the whole class hierarchy up the native class so that we can correctly create the REINST class parented to native
 	TArray<UClass*> ClassHierarchy;
@@ -1668,14 +1668,15 @@ void FBlueprintCompileReinstancer::MoveDependentSkelToReinst(UClass* OwnerClass,
 
 		GIsDuplicatingClassForReinstancing = true;
 		// Create a REINST version of the given class
-		UObject* OldCDO = OwnerClass->ClassDefaultObject;
-		const FName ReinstanceName = MakeUniqueObjectName(GetTransientPackage(), OwnerClass->GetClass(), *(FString(TEXT("REINST_")) + *OwnerClass->GetName()));
+		UObject* OldCDO = CurClass->ClassDefaultObject;
+		const FName ReinstanceName = MakeUniqueObjectName(GetTransientPackage(), CurClass->GetClass(), *(FString(TEXT("REINST_")) + *CurClass->GetName()));
 
-		checkf(IsValid(OwnerClass), TEXT("%s is invalid - will not duplicate successfully"), *(OwnerClass->GetName()));
+		checkf(IsValid(CurClass), TEXT("%s is invalid - will not duplicate successfully"), *(CurClass->GetName()));
 		UClass* ReinstClass = CastChecked<UClass>(StaticDuplicateObject(CurClass, GetTransientPackage(), ReinstanceName, ~RF_Transactional));
 		
 		ReinstClass->RemoveFromRoot();
-		OwnerClass->ClassFlags &= ~CLASS_NewerVersionExists;
+		CurClass->ClassFlags |= CLASS_NewerVersionExists;
+		CurClass->ClassFlags &= ~CLASS_NewerVersionExists;
 		GIsDuplicatingClassForReinstancing = false;
 
 		UClass** OverridenParent = OldToNewMap.Find(ReinstClass->GetSuperClass());
@@ -1693,7 +1694,7 @@ void FBlueprintCompileReinstancer::MoveDependentSkelToReinst(UClass* OwnerClass,
 		// Actually move the old CDO reference out of the way
 		if (OldCDO)
 		{
-			OwnerClass->ClassDefaultObject = nullptr;
+			CurClass->ClassDefaultObject = nullptr;
 			OldCDO->Rename(nullptr, ReinstClass->GetOuter(), REN_DoNotDirty | REN_DontCreateRedirectors | REN_ForceNoResetLoaders | REN_NonTransactional);
 			ReinstClass->ClassDefaultObject = OldCDO;
 			OldCDO->SetClass(ReinstClass);
