Add PrepareWorldsForExplicitSave Flow

diff --git a/Engine/Source/Editor/UnrealEd/Private/FileHelpers.cpp b/Engine/Source/Editor/UnrealEd/Private/FileHelpers.cpp
--- a/Engine/Source/Editor/UnrealEd/Private/FileHelpers.cpp
+++ b/Engine/Source/Editor/UnrealEd/Private/FileHelpers.cpp
@@ -104,6 +104,8 @@ static TAutoConsoleVariable<int32> CVarSkipSourceControlCheckForEditablePackages
 
 #define LOCTEXT_NAMESPACE "FileHelpers"
 
+FEditorFileUtils::FOnPrepareWorldsForExplicitSave FEditorFileUtils::OnPrepareWorldsForExplicitSave;
+
 /** A special output device that puts save output in the message log when flushed */
 class FSaveErrorOutputDevice : public FOutputDevice
 {
@@ -3831,6 +3833,42 @@ bool FEditorFileUtils::SaveDirtyPackages(const bool bPromptUserToSave, const boo
 
 	TArray<UPackage*> PackagesToSave = InternalGetDirtyPackages(bSaveMapPackages, bSaveContentPackages, ShouldIgnorePackageFunction);
 
+	if (bPromptUserToSave)
+	{
+		// In a given set of packages it can contain at least one World Package (map) and/or at least one Actor package.
+		// If an external actor is being saved but not its world we still want to collect its owning world to pass to PrepareWorldsForExplicitSave
+		// In case there is any validation/extra steps needed for that world based on the add/edit of that Actor
+		// We use a set here to dedupe in case both the actor and its world are included in the dirty packages
+		bool bFoundActorWorld = false;
+		TSet<UWorld*> WorldsToSave;
+		for (UPackage* Package : PackagesToSave)
+		{
+			if (UWorld* WorldToSave = UWorld::FindWorldInPackage(Package))
+			{
+				WorldsToSave.Add(WorldToSave);
+			}
+			else if (!bFoundActorWorld)
+			{
+				// Currently there is only one world associated with saving actors as actors from multiple worlds can't be opened
+				// We can skip checking any further Actor packages once we grab the world off the first discovered
+				if (AActor* ActorToSave = AActor::FindActorInPackage(Package))
+				{
+					WorldsToSave.Add(ActorToSave->GetWorld());
+					bFoundActorWorld = true;
+				}
+			}
+		}
+
+		if (!WorldsToSave.IsEmpty())
+		{
+			PrepareWorldsForExplicitSave(WorldsToSave.Array());
+
+			// PrepareWorldsForExplicitSave could have dirtied further packages so refresh our list
+			// We do not go looking for further dirtied worlds however
+			PackagesToSave = InternalGetDirtyPackages(bSaveMapPackages, bSaveContentPackages, ShouldIgnorePackageFunction);
+		}
+	}
+
 	// Need to track the number of packages we're not ignoring for save.
 	int32 NumPackagesNotIgnored = 0;
 
@@ -3926,6 +3964,11 @@ bool FEditorFileUtils::SaveDirtyContentPackages(TArray<UClass*>& SaveContentClas
 	return bResult;
 }
 
+void FEditorFileUtils::PrepareWorldsForExplicitSave(TArray<UWorld*> Worlds)
+{
+	OnPrepareWorldsForExplicitSave.Broadcast(Worlds);
+}
+
 /**
  * Saves the active level, prompting the use for checkout if necessary.
  *
@@ -3945,6 +3988,14 @@ bool FEditorFileUtils::SaveCurrentLevel()
 		TArray<UPackage*> PackagesToSave;
 		
 		UPackage* LevelPackage = Level->GetPackage();
+
+		if (UWorld* World = UWorld::FindWorldInPackage(LevelPackage))
+		{
+			// Regardless of if the LevelPackage is dirty or if we are only saving external objects
+			// we still give an opportunity to validate/run extra steps on the world involved in this save before continuing
+			PrepareWorldsForExplicitSave({ World });
+		}
+
 		// Get Packages to save
 		if (!bCheckDirty || LevelPackage->IsDirty() || LevelPackage->HasAnyPackageFlags(PKG_NewlyCreated))
 		{
@@ -3998,7 +4049,7 @@ FEditorFileUtils::EPromptReturnCode InternalPromptForCheckoutAndSave(const TArra
 	PackagesToSave.Reserve(FinalSaveList.Num());
 
 	{
-		FScopedSlowTask SlowTask(static_cast<float>(FinalSaveList.Num() * 2), NSLOCTEXT("UnrealEd", "SavingPackagesE", "Saving packages..."));
+		FScopedSlowTask SlowTask(static_cast<float>(FinalSaveList.Num() * 2), NSLOCTEXT("UnrealEd", "SavingPackages", "Saving packages..."));
 		SlowTask.MakeDialog();
 
 		UWorld* ActorsWorld = nullptr;
diff --git a/Engine/Source/Editor/UnrealEd/Public/FileHelpers.h b/Engine/Source/Editor/UnrealEd/Public/FileHelpers.h
--- a/Engine/Source/Editor/UnrealEd/Public/FileHelpers.h
+++ b/Engine/Source/Editor/UnrealEd/Public/FileHelpers.h
@@ -409,6 +409,18 @@ public:
 	 */
 	UNREALED_API static void GetDirtyPackages(TArray<UPackage*>& OutDirtyPackages, const FShouldIgnorePackageFunctionRef& ShouldIgnorePackageFunction = FShouldIgnorePackage::Default);
 
+	DECLARE_MULTICAST_DELEGATE_OneParam(FOnPrepareWorldsForExplicitSave, TArray<UWorld*> Worlds);
+	UNREALED_API static FOnPrepareWorldsForExplicitSave OnPrepareWorldsForExplicitSave;
+
+	/**
+	 * Broadcasts the OnPrepareWorldsForExplicitSave event to allow validation and changes to worlds before they're saved.
+	 * Used for Editor triggered explicit saves caused by user input such as ctrl+s, or clicking the save button.
+	 * Changes to worlds include dirtying the world itself or affecting its ExternalActors through Adds/Removes/Edits.
+	 * 
+	 * 
+	 * @param Worlds A list of the worlds that will be broadcasted.
+	 */
+	UNREALED_API static void PrepareWorldsForExplicitSave(TArray<UWorld*> Worlds);
 
 	/**
 	 * Saves the active level, prompting the use for checkout if necessary.
