Fix for pending kill asset reaching FAssetThumbnailPool::LoadThumbnail

diff --git a/Engine/Source/Editor/UnrealEd/Private/AssetThumbnail.cpp b/Engine/Source/Editor/UnrealEd/Private/AssetThumbnail.cpp
--- a/Engine/Source/Editor/UnrealEd/Private/AssetThumbnail.cpp
+++ b/Engine/Source/Editor/UnrealEd/Private/AssetThumbnail.cpp
@@ -1178,7 +1178,10 @@ bool FAssetThumbnailPool::LoadThumbnail(TSharedRef<FThumbnailInfo> ThumbnailInfo
 {
 	const FAssetData& AssetData = CustomAssetToRender.IsValid() ? CustomAssetToRender : ThumbnailInfo->AssetData;
 	UObject* Asset = AssetData.FastGetAsset();
-	ensure(!Asset || IsValidChecked(Asset));
+	if (Asset && !IsValidChecked(Asset))
+	{
+		Asset = nullptr;
+	}
 
 	const bool bAreShadersCompiling = (GShaderCompilingManager && GShaderCompilingManager->IsCompiling());
 	if (Asset && !bAreShadersCompiling)
