MPCook - Change replication of package-specific data to use registration as IMPCollector rather than adding data to PackageData->GetOrAddPackageRemoteResult. This reduces the number of interfaces and is easier to set up. It also makes it possible to extend the feature into licensee systems which cannot know about FPackageData.
#rnx

diff --git a/Engine/Source/Editor/UnrealEd/Classes/CookOnTheSide/CookOnTheFlyServer.h b/Engine/Source/Editor/UnrealEd/Classes/CookOnTheSide/CookOnTheFlyServer.h
--- a/Engine/Source/Editor/UnrealEd/Classes/CookOnTheSide/CookOnTheFlyServer.h
+++ b/Engine/Source/Editor/UnrealEd/Classes/CookOnTheSide/CookOnTheFlyServer.h
@@ -143,6 +143,7 @@ ENUM_CLASS_FLAGS(ECookTickFlags);
 
 namespace UE::Cook
 {
+	class FAssetRegistryMPCollector;
 	class FBuildDefinitions;
 	class FCookDirector;
 	class FCookWorkerClient;
@@ -1424,12 +1425,13 @@ private:
 	float DisplayUpdatePeriodSeconds = 0.0f;
 	EIdleStatus IdleStatus = EIdleStatus::Done;
 
+	friend UE::Cook::FAssetRegistryMPCollector;
 	friend UE::Cook::FBeginCookConfigSettings;
-	friend UE::Cook::FInitializeConfigSettings;
 	friend UE::Cook::FCookDirector;
 	friend UE::Cook::FCookWorkerClient;
 	friend UE::Cook::FCookWorkerServer;
 	friend UE::Cook::FGeneratorPackage;
+	friend UE::Cook::FInitializeConfigSettings;
 	friend UE::Cook::FPackageData;
 	friend UE::Cook::FPendingCookedPlatformData;
 	friend UE::Cook::FPlatformManager;
diff --git a/Engine/Source/Editor/UnrealEd/Private/Commandlets/AssetRegistryGenerator.cpp b/Engine/Source/Editor/UnrealEd/Private/Commandlets/AssetRegistryGenerator.cpp
--- a/Engine/Source/Editor/UnrealEd/Private/Commandlets/AssetRegistryGenerator.cpp
+++ b/Engine/Source/Editor/UnrealEd/Private/Commandlets/AssetRegistryGenerator.cpp
@@ -8,7 +8,9 @@
 #include "AssetRegistry/CookTagList.h"
 #include "CollectionManagerModule.h"
 #include "CollectionManagerTypes.h"
+#include "Cooker/CookMPCollector.h"
 #include "Cooker/CookPackageData.h"
+#include "Cooker/CookPlatformManager.h"
 #include "Cooker/CookWorkerClient.h"
 #include "Commandlets/ChunkDependencyInfo.h"
 #include "Commandlets/IChunkDataGenerator.h"
@@ -2583,10 +2585,12 @@ void FAssetRegistryGenerator::UpdateAssetRegistryData(const UPackage& Package,
 		TEXT("Trying to update asset package flags in package '%s' that does not exist"), *PackageName.ToString());
 }
 
-void FAssetRegistryGenerator::UpdateAssetRegistryData(UE::Cook::FPackageData& PackageData, UE::Cook::FAssetRegistryPackageMessage&& Message)
+void FAssetRegistryGenerator::UpdateAssetRegistryData(UE::Cook::FMPCollectorServerMessageContext& Context,
+	UE::Cook::FAssetRegistryPackageMessage&& Message)
 {
 	LLM_SCOPE_BYTAG(Cooker_GeneratedAssetRegistry);
-	const FName PackageName = PackageData.GetPackageName();
+	const FName PackageName = Context.GetPackageName();
+	check(!PackageName.IsNone());
 	PreviousPackagesToUpdate.Remove(PackageName);
 
 	for (FAssetData& AssetData : Message.AssetDatas)
@@ -2628,12 +2632,15 @@ void FAssetRegistryReporterRemote::UpdateAssetRegistryData(FPackageData& Package
 		DiskSize = -1;
 	}
 
+	FName PackageName = PackageData.GetPackageName();
 	FAssetRegistryPackageMessage Message;
+	Message.PackageName = PackageName;
+	Message.TargetPlatform = TargetPlatform;
 	Message.PackageFlags = NewPackageFlags;
 	Message.DiskSize = DiskSize;
 
 	// Add to the message all the AssetDatas in the package from the global AssetRegistry
-	IAssetRegistry::Get()->GetAssetsByPackageName(PackageData.GetPackageName(), Message.AssetDatas,
+	IAssetRegistry::Get()->GetAssetsByPackageName(PackageName, Message.AssetDatas,
 		bIncludeOnlyDiskAssets, false /* SkipARFilteredAssets */);
 
 	// Also add AssetDatas for any assets that were created during PostLoad or cooking and are not in the global assetregistry
@@ -2658,15 +2665,22 @@ void FAssetRegistryReporterRemote::UpdateAssetRegistryData(FPackageData& Package
 	Message.CookTags = MoveTemp(InArchiveCookTagList.ObjectToTags);
 
 	// Send the message to the director
-	PackageData.GetOrAddPackageRemoteResult().AddMessage(PackageData, TargetPlatform, Message);
+	FCbWriter Writer;
+	Writer.BeginObject();
+	Message.Write(Writer);
+	Writer.EndObject();
+
+	PackageUpdateMessages.Add(PackageName, Writer.Save().AsObject());
 }
 
-void FAssetRegistryPackageMessage::Write(FCbWriter& Writer, const FPackageData& PackageData, const ITargetPlatform* TargetPlatform) const
+void FAssetRegistryPackageMessage::Write(FCbWriter& Writer) const
 {
+	check(!PackageName.IsNone());
+
 	Writer.BeginArray("A");
 	for (const FAssetData& AssetData : AssetDatas)
 	{
-		check(AssetData.PackageName == PackageData.GetPackageName() && !AssetData.AssetName.IsNone()); // We replicate only regular Assets of the form PackageName.AssetName
+		check(AssetData.PackageName == PackageName && !AssetData.AssetName.IsNone()); // We replicate only regular Assets of the form PackageName.AssetName
 		AssetData.NetworkWrite(Writer, false /* bWritePackageName */);
 
 	}
@@ -2680,8 +2694,10 @@ void FAssetRegistryPackageMessage::Write(FCbWriter& Writer, const FPackageData&
 	Writer << "S" << DiskSize;
 }
 
-bool FAssetRegistryPackageMessage::TryRead(FCbObjectView Object, FPackageData& PackageData, const ITargetPlatform* TargetPlatform)
+bool FAssetRegistryPackageMessage::TryRead(FCbObjectView Object)
 {
+	check(!PackageName.IsNone());
+
 	LLM_SCOPE_BYTAG(Cooker_GeneratedAssetRegistry);
 	FCbFieldView AssetDatasField = Object["A"];
 	FCbArrayView AssetDatasArray = AssetDatasField.AsArrayView();
@@ -2693,7 +2709,7 @@ bool FAssetRegistryPackageMessage::TryRead(FCbObjectView Object, FPackageData& P
 	for (FCbFieldView ElementField : AssetDatasArray)
 	{
 		FAssetData& AssetData = AssetDatas.Emplace_GetRef();
-		if (!AssetData.TryNetworkRead(ElementField, false /* bReadPackageName */, PackageData.GetPackageName()))
+		if (!AssetData.TryNetworkRead(ElementField, false /* bReadPackageName */, PackageName))
 		{
 			return false;
 		}
@@ -2717,6 +2733,55 @@ bool FAssetRegistryPackageMessage::TryRead(FCbObjectView Object, FPackageData& P
 
 FGuid FAssetRegistryPackageMessage::MessageType(TEXT("0588DCCEBF1742399EC1E011FC97E4DC"));
 
+FAssetRegistryMPCollector::FAssetRegistryMPCollector(UCookOnTheFlyServer& InCOTFS)
+	: COTFS(InCOTFS)
+{
+}
+
+void FAssetRegistryMPCollector::ClientTickPackage(FMPCollectorClientTickPackageContext& Context)
+{
+	bool bLoggedWarning = false;
+	for (const ITargetPlatform* TargetPlatform : Context.GetPlatforms())
+	{
+		FPlatformData* PlatformData = COTFS.PlatformManager->GetPlatformData(TargetPlatform);
+		FAssetRegistryReporterRemote& RegistryReporter = static_cast<FAssetRegistryReporterRemote&>(*PlatformData->RegistryReporter);
+		FCbObject Message;
+		if (!RegistryReporter.PackageUpdateMessages.RemoveAndCopyValue(Context.GetPackageName(), Message))
+		{
+			UE_CLOG(!bLoggedWarning, LogAssetRegistryGenerator, Warning,
+				TEXT("ClientTickPackage was called for package %s, but UpdateAssetRegistryData was not called for that package. We will not have up to date AssetDataTags in the generated AssetRegistry."),
+				*Context.GetPackageName().ToString());
+			bLoggedWarning = true;
+		}
+		else
+		{
+			Context.AddPlatformMessage(TargetPlatform, MoveTemp(Message));
+		}
+	}
+}
+
+void FAssetRegistryMPCollector::ServerReceiveMessage(FMPCollectorServerMessageContext& Context, FCbObjectView Message)
+{
+	FName PackageName = Context.GetPackageName();
+	const ITargetPlatform* TargetPlatform = Context.GetTargetPlatform();
+	check(PackageName.IsValid() && TargetPlatform);
+
+	FAssetRegistryPackageMessage ARMessage;
+	ARMessage.PackageName = PackageName;
+	ARMessage.TargetPlatform = TargetPlatform;
+	if (!ARMessage.TryRead(Message))
+	{
+		UE_LOG(LogCook, Error, TEXT("Corrupt AssetRegistryPackageMessage received from CookWorker %d. It will be ignored."),
+			Context.GetProfileId());
+	}
+	else
+	{
+		FAssetRegistryGenerator* RegistryGenerator = COTFS.PlatformManager->GetPlatformData(TargetPlatform)->RegistryGenerator.Get();
+		check(RegistryGenerator); // The TargetPlatform came from OrderedSessionPlatforms, and the RegistryGenerator should exist for any of those platforms
+		RegistryGenerator->UpdateAssetRegistryData(Context, MoveTemp(ARMessage));
+	}
+}
+
 }
 
 bool FAssetRegistryGenerator::UpdateAssetPackageFlags(const FName& PackageName, const uint32 PackageFlags)
diff --git a/Engine/Source/Editor/UnrealEd/Private/Cooker/CookDirector.cpp b/Engine/Source/Editor/UnrealEd/Private/Cooker/CookDirector.cpp
--- a/Engine/Source/Editor/UnrealEd/Private/Cooker/CookDirector.cpp
+++ b/Engine/Source/Editor/UnrealEd/Private/Cooker/CookDirector.cpp
@@ -129,6 +129,7 @@ FCookDirector::FCookDirector(UCookOnTheFlyServer& InCOTFS, int32 CookProcessCoun
 		{
 			HandleHeartbeatMessage(Context, bReadSuccessful, MoveTemp(Message));
 		}, TEXT("HandleHeartbeatMessage")));
+	Register(new FAssetRegistryMPCollector(COTFS));
 
 	LastTickTimeSeconds = FPlatformTime::Seconds();
 
diff --git a/Engine/Source/Editor/UnrealEd/Private/Cooker/CookMPCollector.cpp b/Engine/Source/Editor/UnrealEd/Private/Cooker/CookMPCollector.cpp
--- a/Engine/Source/Editor/UnrealEd/Private/Cooker/CookMPCollector.cpp
+++ b/Engine/Source/Editor/UnrealEd/Private/Cooker/CookMPCollector.cpp
@@ -41,6 +41,27 @@ const ITargetPlatform* FMPCollectorClientTickContext::IntToPlatform(uint8 Platfo
 	return IntToPlatformImpl(PlatformAsInt, Platforms);
 }
 
+void FMPCollectorClientTickPackageContext::AddMessage(FCbObject Object)
+{
+	Messages.Emplace(nullptr, MoveTemp(Object));
+}
+
+void FMPCollectorClientTickPackageContext::AddPlatformMessage(const ITargetPlatform* TargetPlatform, FCbObject Object)
+{
+	check(TargetPlatform && Platforms.Contains(TargetPlatform));
+	Messages.Emplace(TargetPlatform, MoveTemp(Object));
+}
+
+uint8 FMPCollectorClientTickPackageContext::PlatformToInt(const ITargetPlatform* Platform) const
+{
+	return PlatformToIntImpl(Platform, Platforms);
+}
+
+const ITargetPlatform* FMPCollectorClientTickPackageContext::IntToPlatform(uint8 PlatformAsInt) const
+{
+	return IntToPlatformImpl(PlatformAsInt, Platforms);
+}
+
 uint8 FMPCollectorClientMessageContext::PlatformToInt(const ITargetPlatform* Platform) const
 {
 	return PlatformToIntImpl(Platform, Platforms);
diff --git a/Engine/Source/Editor/UnrealEd/Private/Cooker/CookMPCollector.h b/Engine/Source/Editor/UnrealEd/Private/Cooker/CookMPCollector.h
--- a/Engine/Source/Editor/UnrealEd/Private/Cooker/CookMPCollector.h
+++ b/Engine/Source/Editor/UnrealEd/Private/Cooker/CookMPCollector.h
@@ -41,6 +41,26 @@ private:
 	friend class FCookWorkerClient;
 };
 
+class FMPCollectorClientTickPackageContext
+{
+public:
+	TConstArrayView<const ITargetPlatform*> GetPlatforms() const { return Platforms; }
+	FName GetPackageName() const { return PackageName; }
+
+	void AddMessage(FCbObject Object);
+	void AddPlatformMessage(const ITargetPlatform* Platform, FCbObject Object);
+
+	uint8 PlatformToInt(const ITargetPlatform* Platform) const;
+	const ITargetPlatform* IntToPlatform(uint8 PlatformAsInt) const;
+
+private:
+	TArray<TPair<const ITargetPlatform*, FCbObject>> Messages;
+	TConstArrayView<const ITargetPlatform*> Platforms;
+	FName PackageName;
+
+	friend class FCookWorkerClient;
+};
+
 class FMPCollectorClientMessageContext
 {
 public:
@@ -66,11 +86,18 @@ public:
 	int32 GetProfileId() const { return ProfileId; }
 	FCookWorkerServer* GetCookWorkerServer() { return Server; }
 
+	bool HasPackageName() const { return !PackageName.IsNone(); }
+	FName GetPackageName() const { return PackageName; }
+	bool HasTargetPlatform() const { return TargetPlatform != nullptr; }
+	const ITargetPlatform* GetTargetPlatform() const { return TargetPlatform; }
+
 private:
 	TConstArrayView<const ITargetPlatform*> Platforms;
-	FWorkerId WorkerId;
-	int32 ProfileId;
+	FName PackageName;
 	FCookWorkerServer* Server = nullptr;
+	const ITargetPlatform* TargetPlatform = nullptr;
+	int32 ProfileId;
+	FWorkerId WorkerId;
 
 	friend class FCookWorkerServer;
 };
@@ -88,6 +115,7 @@ public:
 	virtual const TCHAR* GetDebugName() const = 0;
 
 	virtual void ClientTick(FMPCollectorClientTickContext& Context) {}
+	virtual void ClientTickPackage(FMPCollectorClientTickPackageContext& Context) {}
 	virtual void ClientReceiveMessage(FMPCollectorClientMessageContext& Context, FCbObjectView Message) {}
 	virtual void ServerReceiveMessage(FMPCollectorServerMessageContext& Context, FCbObjectView Message) {}
 };
diff --git a/Engine/Source/Editor/UnrealEd/Private/Cooker/CookPackageData.cpp b/Engine/Source/Editor/UnrealEd/Private/Cooker/CookPackageData.cpp
--- a/Engine/Source/Editor/UnrealEd/Private/Cooker/CookPackageData.cpp
+++ b/Engine/Source/Editor/UnrealEd/Private/Cooker/CookPackageData.cpp
@@ -778,7 +778,6 @@ void FPackageData::OnEnterSave()
 	check(!HasPrepareSaveFailed());
 	CheckObjectCacheEmpty();
 	CheckCookedPlatformDataEmpty();
-	check(!PackageRemoteResult);
 }
 
 void FPackageData::OnExitSave()
@@ -787,22 +786,6 @@ void FPackageData::OnExitSave()
 	ClearObjectCache();
 	SetHasPrepareSaveFailed(false);
 	SetIsPrepareSaveRequiresGC(false);
-	PackageRemoteResult.Reset();
-}
-
-FPackageRemoteResult& FPackageData::GetOrAddPackageRemoteResult()
-{
-	check(GetState() == EPackageState::Save);
-	if (!PackageRemoteResult)
-	{
-		PackageRemoteResult = MakeUnique<FPackageRemoteResult>();
-	}
-	return *PackageRemoteResult;
-}
-
-TUniquePtr<FPackageRemoteResult>& FPackageData::GetPackageRemoteResult()
-{
-	return PackageRemoteResult;
 }
 
 void FPackageData::OnEnterInProgress()
diff --git a/Engine/Source/Editor/UnrealEd/Private/Cooker/CookPackageData.h b/Engine/Source/Editor/UnrealEd/Private/Cooker/CookPackageData.h
--- a/Engine/Source/Editor/UnrealEd/Private/Cooker/CookPackageData.h
+++ b/Engine/Source/Editor/UnrealEd/Private/Cooker/CookPackageData.h
@@ -16,7 +16,6 @@
 #include "HAL/Platform.h"
 #include "Misc/EnumClassFlags.h"
 #include "Misc/Optional.h"
-#include "PackageResultsMessage.h"
 #include "Templates/SharedPointer.h"
 #include "UObject/GCObject.h"
 #include "UObject/NameTypes.h"
@@ -496,9 +495,6 @@ public:
 
 	/** Marshall this PackageData to a ConstructData that can be used later or on a remote machine to reconstruct it. */
 	FConstructPackageData CreateConstructData();
-
-	FPackageRemoteResult& GetOrAddPackageRemoteResult();
-	TUniquePtr<FPackageRemoteResult>& GetPackageRemoteResult();
 private:
 	friend struct UE::Cook::FPackageDatas;
 
@@ -569,7 +565,6 @@ private:
 	void OnPackageDataFirstRequested(FInstigator&& InInstigator);
 
 	TUniquePtr<FGeneratorPackage> GeneratorPackage;
-	TUniquePtr<FPackageRemoteResult> PackageRemoteResult;
 	FGeneratorPackage* GeneratedOwner;
 	/** Data for each platform that has been interacted with by *this. */
 	TSortedMap<const ITargetPlatform*, FPlatformData> PlatformDatas;
diff --git a/Engine/Source/Editor/UnrealEd/Private/Cooker/CookWorkerClient.cpp b/Engine/Source/Editor/UnrealEd/Private/Cooker/CookWorkerClient.cpp
--- a/Engine/Source/Editor/UnrealEd/Private/Cooker/CookWorkerClient.cpp
+++ b/Engine/Source/Editor/UnrealEd/Private/Cooker/CookWorkerClient.cpp
@@ -46,6 +46,7 @@ FCookWorkerClient::FCookWorkerClient(UCookOnTheFlyServer& InCOTFS)
 		{
 			HandleHeartbeatMessage(Context, bReadSuccessful, MoveTemp(Message));
 		}, TEXT("HandleHeartbeatMessage")));
+	Register(new FAssetRegistryMPCollector(COTFS));
 }
 
 FCookWorkerClient::~FCookWorkerClient()
@@ -149,65 +150,54 @@ void FCookWorkerClient::ReportDemoteToIdle(const FPackageData& PackageData, ESup
 		return;
 	}
 	FPackageRemoteResult& Result = PendingResults.Emplace_GetRef();
-	Result.PackageName = PackageData.GetPackageName();
-	Result.SuppressCookReason = Reason;
+	Result.SetPackageName(PackageData.GetPackageName());
+	Result.SetSuppressCookReason(Reason);
 }
 
 void FCookWorkerClient::ReportPromoteToSaveComplete(FPackageData& PackageData)
 {
-	TUniquePtr<FPackageRemoteResult> MovedResult = MoveTemp(PackageData.GetPackageRemoteResult());
-	FPackageRemoteResult* Result;
-	if (MovedResult)
-	{
-		Result = &PendingResults.Emplace_GetRef(MoveTemp(*MovedResult));
-	} 
-	else
-	{
-		Result = &PendingResults.Emplace_GetRef();
-	}
+	FPackageRemoteResult& Result = PendingResults.Emplace_GetRef();
 
-	Result->PackageName = PackageData.GetPackageName();
-	Result->SuppressCookReason = ESuppressCookReason::InvalidSuppressCookReason;
+	FName PackageName = PackageData.GetPackageName();
+	Result.SetPackageName(PackageName);
+	Result.SetSuppressCookReason(ESuppressCookReason::InvalidSuppressCookReason);
+	Result.SetPlatforms(OrderedSessionPlatforms);
 
-	// Sort the platforms to match the OrderedSessionPlatforms order, and add any missing platforms
 	int32 NumPlatforms = OrderedSessionPlatforms.Num();
-	bool bAlreadySorted = true;
 	for (int32 PlatformIndex = 0; PlatformIndex < NumPlatforms; ++PlatformIndex)
 	{
-		if (OrderedSessionPlatforms[PlatformIndex] != Result->Platforms[PlatformIndex].Platform)
-		{
-			bAlreadySorted = false;
-			break;
-		}
+		ITargetPlatform* TargetPlatform = OrderedSessionPlatforms[PlatformIndex];
+		FPackageRemoteResult::FPlatformResult& PlatformResults = Result.GetPlatforms()[PlatformIndex];
+		FPackageData::FPlatformData& PackagePlatformData = PackageData.FindOrAddPlatformData(TargetPlatform);
+		PlatformResults.SetSuccessful(PackagePlatformData.bCookSucceeded);
 	}
-	if (!bAlreadySorted)
+
+	FMPCollectorClientTickPackageContext Context;
+	Context.PackageName = PackageName;
+	Context.Platforms = OrderedSessionPlatforms;
+
+	for (const TPair<FGuid, TRefCountPtr<IMPCollector>>& CollectorPair : Collectors)
 	{
-		TArray<FPackageRemoteResult::FPlatformResult, TInlineAllocator<1>> SortedPlatforms;
-		SortedPlatforms.SetNum(NumPlatforms);
-		for (FPackageRemoteResult::FPlatformResult& ExistingResult : Result->Platforms)
+		IMPCollector* Collector = CollectorPair.Value.GetReference();
+		Collector->ClientTickPackage(Context);
+		if (!Context.Messages.IsEmpty())
 		{
-			int32 PlatformIndex = OrderedSessionPlatforms.IndexOfByKey(ExistingResult.Platform);
-			check(PlatformIndex != INDEX_NONE && SortedPlatforms[PlatformIndex].Platform == nullptr); // Only platforms in the session platforms should have been added, and there should not be duplicates
-			SortedPlatforms[PlatformIndex] = MoveTemp(ExistingResult);
-		}
-		for (int32 PlatformIndex = 0; PlatformIndex < NumPlatforms; ++PlatformIndex)
-		{
-			FPackageRemoteResult::FPlatformResult& Sorted = SortedPlatforms[PlatformIndex];
-			if (!Sorted.Platform)
+			const FGuid& MessageType = CollectorPair.Key;
+			for (TPair<const ITargetPlatform*, FCbObject>& MessagePair : Context.Messages)
 			{
-				Sorted.Platform = OrderedSessionPlatforms[PlatformIndex];
-				Sorted.bSuccessful = false;
+				const ITargetPlatform* TargetPlatform = MessagePair.Key;
+				FCbObject Object = MoveTemp(MessagePair.Value);
+				if (!TargetPlatform)
+				{
+					Result.AddPackageMessage(MessageType, MoveTemp(Object));
+				}
+				else
+				{
+					Result.AddPlatformMessage(TargetPlatform, MessageType, MoveTemp(Object));
+				}
 			}
+			Context.Messages.Reset();
 		}
-		Swap(Result->Platforms, SortedPlatforms);
-	}
-
-	for (int32 PlatformIndex = 0; PlatformIndex < NumPlatforms; ++PlatformIndex)
-	{
-		ITargetPlatform* TargetPlatform = OrderedSessionPlatforms[PlatformIndex];
-		FPackageRemoteResult::FPlatformResult& PlatformResults = Result->Platforms[PlatformIndex];
-		FPackageData::FPlatformData& PackagePlatformData = PackageData.FindOrAddPlatformData(TargetPlatform);
-		PlatformResults.bSuccessful = PackagePlatformData.bCookSucceeded;
 	}
 }
 
diff --git a/Engine/Source/Editor/UnrealEd/Private/Cooker/CookWorkerServer.cpp b/Engine/Source/Editor/UnrealEd/Private/Cooker/CookWorkerServer.cpp
--- a/Engine/Source/Editor/UnrealEd/Private/Cooker/CookWorkerServer.cpp
+++ b/Engine/Source/Editor/UnrealEd/Private/Cooker/CookWorkerServer.cpp
@@ -563,25 +563,35 @@ void FCookWorkerServer::HandleReceiveMessagesInternal()
 	}
 }
 
-void FCookWorkerServer::HandleReceivedPackagePlatformMessages(FPackageData& PackageData, const ITargetPlatform* TargetPlatform, TArray<UE::CompactBinaryTCP::FMarshalledMessage>&& Messages)
+void FCookWorkerServer::HandleReceivedPackagePlatformMessages(FPackageData& PackageData, const ITargetPlatform* TargetPlatform,
+	TArray<UE::CompactBinaryTCP::FMarshalledMessage>&& Messages)
 {
 	check(TickState.TickThread == ECookDirectorThread::SchedulerThread);
+	if (Messages.IsEmpty())
+	{
+		return;
+	}
+
+	FMPCollectorServerMessageContext Context;
+	Context.Platforms = OrderedSessionPlatforms;
+	Context.PackageName = PackageData.GetPackageName();
+	Context.TargetPlatform = TargetPlatform;
+	Context.Server = this;
+	Context.ProfileId = ProfileId;
+	Context.WorkerId = WorkerId;
 
 	for (UE::CompactBinaryTCP::FMarshalledMessage& Message : Messages)
 	{
-		if (Message.MessageType == FAssetRegistryPackageMessage::MessageType)
+		TRefCountPtr<IMPCollector>* Collector = Director.Collectors.Find(Message.MessageType);
+		if (Collector)
 		{
-			FAssetRegistryPackageMessage ARMessage;
-			if (!ARMessage.TryRead(Message.Object, PackageData, TargetPlatform))
-			{
-				LogInvalidMessage(TEXT("FAssetRegistryPackageMessage"));
-			}
-			else
-			{
-				FAssetRegistryGenerator* RegistryGenerator = COTFS.PlatformManager->GetPlatformData(TargetPlatform)->RegistryGenerator.Get();
-				check(RegistryGenerator); // The TargetPlatform came from OrderedSessionPlatforms, and the RegistryGenerator should exist for any of those platforms
-				RegistryGenerator->UpdateAssetRegistryData(PackageData, MoveTemp(ARMessage));
-			}
+			check(*Collector);
+			(*Collector)->ServerReceiveMessage(Context, Message.Object);
+		}
+		else
+		{
+			UE_LOG(LogCook, Error, TEXT("CookWorkerServer received PackageMessage of unknown type %s from CookWorker. Ignoring it."),
+				*Message.MessageType.ToString());
 		}
 	}
 }
@@ -611,49 +621,52 @@ void FCookWorkerServer::RecordResults(FPackageResultsMessage& Message)
 
 	for (FPackageRemoteResult& Result : Message.Results)
 	{
-		FPackageData* PackageData = COTFS.PackageDatas->FindPackageDataByPackageName(Result.PackageName);
+		FPackageData* PackageData = COTFS.PackageDatas->FindPackageDataByPackageName(Result.GetPackageName());
 		if (!PackageData)
 		{
 			UE_LOG(LogCook, Warning, TEXT("CookWorkerServer %d received FPackageResultsMessage for invalid package %s. Ignoring it."),
-				ProfileId, *Result.PackageName.ToString());
+				ProfileId, *Result.GetPackageName().ToString());
 			continue;
 		}
 		if (PendingPackages.Remove(PackageData) != 1)
 		{
 			UE_LOG(LogCook, Warning, TEXT("CookWorkerServer %d received FPackageResultsMessage for package %s which is not a pending package. Ignoring it."),
-				ProfileId, *Result.PackageName.ToString());
+				ProfileId, *Result.GetPackageName().ToString());
 			continue;
 		}
 		PackageData->SetWorkerAssignment(FWorkerId::Invalid(), ESendFlags::QueueNone);
 
 		// MPCOOKTODO: Refactor FSaveCookedPackageContext::FinishPlatform and ::FinishPackage so we can call them from here
 		// to reduce duplication
-		if (Result.SuppressCookReason == ESuppressCookReason::InvalidSuppressCookReason)
+		if (Result.GetSuppressCookReason() == ESuppressCookReason::InvalidSuppressCookReason)
 		{
 			int32 NumPlatforms = OrderedSessionPlatforms.Num();
-			if (Result.Platforms.Num() != NumPlatforms)
+			if (Result.GetPlatforms().Num() != NumPlatforms)
 			{
-				UE_LOG(LogCook, Warning, TEXT("CookWorkerServer %d received FPackageResultsMessage for package %s with an invalid number of platform results: expected %d, actual %d. Ignoring it."),
-					ProfileId, *Result.PackageName.ToString(), NumPlatforms, Result.Platforms.Num());
+				UE_LOG(LogCook, Warning,
+					TEXT("CookWorkerServer %d received FPackageResultsMessage for package %s with an invalid number of platform results: expected %d, actual %d. Ignoring it."),
+					ProfileId, *Result.GetPackageName().ToString(), NumPlatforms, Result.GetPlatforms().Num());
 				continue;
 			}
+
+			HandleReceivedPackagePlatformMessages(*PackageData, nullptr, Result.ReleaseMessages());
 			for (int32 PlatformIndex = 0; PlatformIndex < NumPlatforms; ++PlatformIndex)
 			{
 				ITargetPlatform* TargetPlatform = OrderedSessionPlatforms[PlatformIndex];
-				FPackageRemoteResult::FPlatformResult& PlatformResult = Result.Platforms[PlatformIndex];
-				PackageData->SetPlatformCooked(TargetPlatform, PlatformResult.bSuccessful);
+				FPackageRemoteResult::FPlatformResult& PlatformResult = Result.GetPlatforms()[PlatformIndex];
+				PackageData->SetPlatformCooked(TargetPlatform, PlatformResult.IsSuccessful());
 				// MPCOOKTODO: Call CommitRemotePackage on the PackageWriter
-				HandleReceivedPackagePlatformMessages(*PackageData, TargetPlatform, MoveTemp(PlatformResult.Messages));
+				HandleReceivedPackagePlatformMessages(*PackageData, TargetPlatform, PlatformResult.ReleaseMessages());
 			}
-			if (Result.bReferencedOnlyByEditorOnlyData)
+			if (Result.IsReferencedOnlyByEditorOnlyData())
 			{
-				COTFS.PackageTracker->UncookedEditorOnlyPackages.AddUnique(Result.PackageName);
+				COTFS.PackageTracker->UncookedEditorOnlyPackages.AddUnique(Result.GetPackageName());
 			}
 			COTFS.PromoteToSaveComplete(*PackageData, ESendFlags::QueueAddAndRemove);
 		}
 		else
 		{
-			COTFS.DemoteToIdle(*PackageData, ESendFlags::QueueAddAndRemove, Result.SuppressCookReason);
+			COTFS.DemoteToIdle(*PackageData, ESendFlags::QueueAddAndRemove, Result.GetSuppressCookReason());
 		}
 	}
 	Director.ResetFinalIdleHeartbeatFence();
diff --git a/Engine/Source/Editor/UnrealEd/Private/Cooker/PackageResultsMessage.cpp b/Engine/Source/Editor/UnrealEd/Private/Cooker/PackageResultsMessage.cpp
--- a/Engine/Source/Editor/UnrealEd/Private/Cooker/PackageResultsMessage.cpp
+++ b/Engine/Source/Editor/UnrealEd/Private/Cooker/PackageResultsMessage.cpp
@@ -7,8 +7,28 @@
 namespace UE::Cook
 {
 
-void FPackageRemoteResult::AddMessage(const FPackageData& PackageData, const ITargetPlatform* TargetPlatform, const IPackageMessage& Message)
+TArray<UE::CompactBinaryTCP::FMarshalledMessage> FPackageRemoteResult::FPlatformResult::ReleaseMessages()
 {
+	TArray<UE::CompactBinaryTCP::FMarshalledMessage> Result = MoveTemp(Messages);
+	Messages.Empty();
+	return Result;
+}
+
+TArray<UE::CompactBinaryTCP::FMarshalledMessage> FPackageRemoteResult::ReleaseMessages()
+{
+	TArray<UE::CompactBinaryTCP::FMarshalledMessage> Result = MoveTemp(Messages);
+	Messages.Empty();
+	return Result;
+}
+
+void FPackageRemoteResult::AddPackageMessage(const FGuid& MessageType, FCbObject&& Object)
+{
+	Messages.Add(UE::CompactBinaryTCP::FMarshalledMessage{ MessageType, MoveTemp(Object) });
+}
+
+void FPackageRemoteResult::AddPlatformMessage(const ITargetPlatform* TargetPlatform, const FGuid& MessageType, FCbObject&& Object)
+{
+	check(TargetPlatform != nullptr);
 	int32 PlatformIndex = Platforms.IndexOfByPredicate([TargetPlatform](const FPlatformResult& Result)
 		{ return Result.Platform == TargetPlatform; }
 	);
@@ -23,14 +43,18 @@ void FPackageRemoteResult::AddMessage(const FPackageData& PackageData, const ITa
 		Result->Platform = TargetPlatform;
 	}
 
-	FCbWriter Writer;
-	Writer.BeginObject();
-	Message.Write(Writer, PackageData, TargetPlatform);
-	Writer.EndObject();
-
-	Result->Messages.Add(UE::CompactBinaryTCP::FMarshalledMessage{ Message.GetMessageType(), Writer.Save().AsObject() });
+	Result->Messages.Add(UE::CompactBinaryTCP::FMarshalledMessage{ MessageType, MoveTemp(Object)});
 }
 
+void FPackageRemoteResult::SetPlatforms(TConstArrayView<ITargetPlatform*> OrderedSessionPlatforms)
+{
+	Platforms.Reserve(OrderedSessionPlatforms.Num());
+	for (ITargetPlatform* TargetPlatform : OrderedSessionPlatforms)
+	{
+		FPackageRemoteResult::FPlatformResult& PlatformResult = Platforms.Emplace_GetRef();
+		PlatformResult.Platform = TargetPlatform;
+	}
+}
 
 void FPackageResultsMessage::Write(FCbWriter& Writer) const
 {
@@ -42,6 +66,7 @@ void FPackageResultsMessage::Write(FCbWriter& Writer) const
 			Writer << "N" << Result.PackageName;
 			Writer << "R" << (uint8) Result.SuppressCookReason;
 			Writer << "E" << Result.bReferencedOnlyByEditorOnlyData;
+			WriteMessagesArray(Writer, Result.Messages);
 			Writer.BeginArray("P");
 			for (const FPackageRemoteResult::FPlatformResult& PlatformResult : Result.Platforms)
 			{
@@ -49,17 +74,7 @@ void FPackageResultsMessage::Write(FCbWriter& Writer) const
 				Writer << "S" << PlatformResult.bSuccessful;
 				Writer << "G" << PlatformResult.PackageGuid;
 				Writer << "D" << PlatformResult.TargetDomainDependencies;
-				if (!PlatformResult.Messages.IsEmpty())
-				{
-					// We write a nonhomogenous array of length 2N. 2N+0 is the Message type, 2N+1 is the message object.
-					Writer.BeginArray("M");
-					for (const UE::CompactBinaryTCP::FMarshalledMessage& Message : PlatformResult.Messages)
-					{
-						Writer << Message.MessageType;
-						Writer << Message.Object;
-					}
-					Writer.EndArray();
-				}
+				WriteMessagesArray(Writer, PlatformResult.Messages);
 				Writer.EndObject();
 			}
 			Writer.EndArray();
@@ -69,6 +84,22 @@ void FPackageResultsMessage::Write(FCbWriter& Writer) const
 	Writer.EndArray();
 }
 
+void FPackageResultsMessage::WriteMessagesArray(FCbWriter& Writer,
+	TConstArrayView<UE::CompactBinaryTCP::FMarshalledMessage> InMessages)
+{
+	if (!InMessages.IsEmpty())
+	{
+		// We write a nonhomogenous array of length 2N. 2N+0 is the Message type, 2N+1 is the message object.
+		Writer.BeginArray("M");
+		for (const UE::CompactBinaryTCP::FMarshalledMessage& Message : InMessages)
+		{
+			Writer << Message.MessageType;
+			Writer << Message.Object;
+		}
+		Writer.EndArray();
+	}
+}
+
 bool FPackageResultsMessage::TryRead(FCbObjectView Object)
 {
 	Results.Reset();
@@ -88,6 +119,10 @@ bool FPackageResultsMessage::TryRead(FCbObjectView Object)
 		}
 		Result.SuppressCookReason = static_cast<ESuppressCookReason>(LocalSuppressCookReason);
 		Result.bReferencedOnlyByEditorOnlyData = ResultObject["E"].AsBool();
+		if (!TryReadMessagesArray(ResultObject, Result.Messages))
+		{
+			return false;
+		}
 		Result.Platforms.Reset();
 		for (FCbFieldView PlatformField : ResultObject["P"])
 		{
@@ -98,31 +133,41 @@ bool FPackageResultsMessage::TryRead(FCbObjectView Object)
 			PlatformResult.PackageGuid = PlatformObject["G"].AsUuid();
 			PlatformResult.TargetDomainDependencies = FCbObject::Clone(PlatformObject["D"].AsObjectView());
 
-			// We read a nonhomogenous array of length 2N. 2N+0 is the Message type, 2N+1 is the message object.
-			FCbArrayView MessagesArray = PlatformObject["M"].AsArrayView();
-			PlatformResult.Messages.Reserve(MessagesArray.Num() / 2);
-			FCbFieldViewIterator MessageField = MessagesArray.CreateViewIterator();
-			while (MessageField)
+			if (!TryReadMessagesArray(PlatformObject, PlatformResult.Messages))
 			{
-				UE::CompactBinaryTCP::FMarshalledMessage& Message = PlatformResult.Messages.Emplace_GetRef();
-				Message.MessageType = MessageField->AsUuid();
-				if (MessageField.HasError())
-				{
-					return false;
-				}
-				++MessageField;
-				Message.Object = FCbObject::Clone(MessageField->AsObjectView());
-				if (MessageField.HasError())
-				{
-					return false;
-				}
-				++MessageField;
+				return false;
 			}
 		}
 	}
 	return true;
 }
 
+bool FPackageResultsMessage::TryReadMessagesArray(FCbObjectView ObjectWithMessageField,
+	TArray<UE::CompactBinaryTCP::FMarshalledMessage>& InMessages)
+{
+	// We read a nonhomogenous array of length 2N. 2N+0 is the Message type, 2N+1 is the message object.
+	FCbArrayView MessagesArray = ObjectWithMessageField["M"].AsArrayView();
+	InMessages.Reset(MessagesArray.Num() / 2);
+	FCbFieldViewIterator MessageField = MessagesArray.CreateViewIterator();
+	while (MessageField)
+	{
+		UE::CompactBinaryTCP::FMarshalledMessage& Message = InMessages.Emplace_GetRef();
+		Message.MessageType = MessageField->AsUuid();
+		if (MessageField.HasError())
+		{
+			return false;
+		}
+		++MessageField;
+		Message.Object = FCbObject::Clone(MessageField->AsObjectView());
+		if (MessageField.HasError())
+		{
+			return false;
+		}
+		++MessageField;
+	}
+	return true;
+}
+
 FGuid FPackageResultsMessage::MessageType(TEXT("4631C6C0F6DC4CEFB2B09D3FB0B524DB"));
 
 }
\ No newline at end of file
diff --git a/Engine/Source/Editor/UnrealEd/Private/Cooker/PackageResultsMessage.h b/Engine/Source/Editor/UnrealEd/Private/Cooker/PackageResultsMessage.h
--- a/Engine/Source/Editor/UnrealEd/Private/Cooker/PackageResultsMessage.h
+++ b/Engine/Source/Editor/UnrealEd/Private/Cooker/PackageResultsMessage.h
@@ -10,30 +10,14 @@
 
 class ITargetPlatform;
 namespace UE::Cook { struct FPackageData; }
+namespace UE::Cook { struct FPackageResultsMessage; }
 
 namespace UE::Cook
 {
 
 /**
- * The base class of messages that can be sent as platform-specific package-specific submessages
- * in a FPackageResultMessage. Messages are identified to the remote connection by the Guid identifier
- * from GetMessageType.
- */
-class IPackageMessage
-{
-public:
-	virtual ~IPackageMessage() {}
-	/** Marshall the message to a CompactBinaryObject. */
-	virtual void Write(FCbWriter& Writer, const FPackageData& PackageData, const ITargetPlatform* TargetPlatform) const = 0;
-	/** Unmarshall the message from a CompactBinaryObject. */
-	virtual bool TryRead(FCbObjectView Object, FPackageData& PackageData, const ITargetPlatform* TargetPlatform) = 0;
-	/** Return the Guid that identifies the message to the remote connection. */
-	virtual FGuid GetMessageType() const = 0;
-};
-
-/**
- * Helper struct for FPackageResultsMessage that can also be owned by other types like PackageData.
- * Holds replication information about the result of the Package's save on each platform, including
+ * Helper struct for FPackageResultsMessage.
+ * Holds replication information about the result of a Package's save, including per-platform results and
  * system-specific messages from other systems
  */
 struct FPackageRemoteResult
@@ -42,22 +26,59 @@ public:
 	/** Information about the results for a single platform */
 	struct FPlatformResult
 	{
-		const ITargetPlatform* Platform = nullptr;
-		TArray<UE::CompactBinaryTCP::FMarshalledMessage> Messages;
-		FGuid PackageGuid;
+	public:
+		const ITargetPlatform* GetPlatform() const { return Platform; }
+		void SetPlatform(const ITargetPlatform* InPlatform) { Platform = InPlatform; }
+
+		TConstArrayView<UE::CompactBinaryTCP::FMarshalledMessage> GetMessages() const { return Messages; }
+		TArray<UE::CompactBinaryTCP::FMarshalledMessage> ReleaseMessages();
+
+		const FGuid& GetPackageGuid() const { return PackageGuid; }
+		void SetPackageGuid(const FGuid& InPackageGuid) { PackageGuid = InPackageGuid; }
+
+		FCbObjectView GetTargetDomainDependencies() const { return TargetDomainDependencies; }
+		void SetTargetDomainDependencies(FCbObject&& InObject) { TargetDomainDependencies = MoveTemp(InObject); }
+
+		bool IsSuccessful() const { return bSuccessful; }
+		void SetSuccessful(bool bInSuccessful) { bSuccessful = bInSuccessful; }
+
+	private:
 		FCbObject TargetDomainDependencies;
-		bool bSuccessful = true;
+		FGuid PackageGuid;
+		TArray<UE::CompactBinaryTCP::FMarshalledMessage> Messages;
+		const ITargetPlatform* Platform = nullptr;
+		bool bSuccessful = false;
+
+		friend FPackageRemoteResult;
+		friend FPackageResultsMessage;
 	};
 
-	void AddMessage(const FPackageData& PackageData, const ITargetPlatform* TargetPlatform, const IPackageMessage& Message);
+	FName GetPackageName() const { return PackageName; }
+	void SetPackageName(FName InPackageName) { PackageName = InPackageName; }
 
-public:
+	ESuppressCookReason GetSuppressCookReason() const { return SuppressCookReason; }
+	void SetSuppressCookReason(ESuppressCookReason InSuppressCookReason) { SuppressCookReason = InSuppressCookReason; }
+
+	bool IsReferencedOnlyByEditorOnlyData() const { return bReferencedOnlyByEditorOnlyData; }
+	void SetReferencedOnlyByEditorOnlyData(bool bInReferencedOnlyByEditorOnlyData) { bReferencedOnlyByEditorOnlyData = bInReferencedOnlyByEditorOnlyData; }
+
+	void AddPackageMessage(const FGuid& MessageType, FCbObject&& Object);
+	void AddPlatformMessage(const ITargetPlatform* TargetPlatform, const FGuid& MessageType, FCbObject&& Object);
+	TConstArrayView<UE::CompactBinaryTCP::FMarshalledMessage> GetMessages() const { return Messages; }
+	TArray<UE::CompactBinaryTCP::FMarshalledMessage> ReleaseMessages();
+
+	TArray<FPlatformResult, TInlineAllocator<1>>& GetPlatforms() { return Platforms; }
+	void SetPlatforms(TConstArrayView<ITargetPlatform*> OrderedSessionPlatforms);
+
+private:
+	TArray<FPlatformResult, TInlineAllocator<1>> Platforms;
+	TArray<UE::CompactBinaryTCP::FMarshalledMessage> Messages;
 	FName PackageName;
 	/** If failure reason is InvalidSuppressCookReason, it was saved. Otherwise, holds the suppression reason */
 	ESuppressCookReason SuppressCookReason;
 	bool bReferencedOnlyByEditorOnlyData = false;
-	TArray<FPlatformResult, TInlineAllocator<1>> Platforms;
 
+	friend FPackageResultsMessage;
 };
 
 /** Message from Client to Server giving the results for saved or refused-to-cook packages. */
@@ -72,6 +93,12 @@ public:
 	TArray<FPackageRemoteResult> Results;
 
 	static FGuid MessageType;
+
+private:
+	static void WriteMessagesArray(FCbWriter& Writer,
+		TConstArrayView<UE::CompactBinaryTCP::FMarshalledMessage> InMessages);
+	static bool TryReadMessagesArray(FCbObjectView ObjectWithMessageField,
+		TArray<UE::CompactBinaryTCP::FMarshalledMessage>& InMessages);
 };
 
 }
diff --git a/Engine/Source/Editor/UnrealEd/Public/Commandlets/AssetRegistryGenerator.h b/Engine/Source/Editor/UnrealEd/Public/Commandlets/AssetRegistryGenerator.h
--- a/Engine/Source/Editor/UnrealEd/Public/Commandlets/AssetRegistryGenerator.h
+++ b/Engine/Source/Editor/UnrealEd/Public/Commandlets/AssetRegistryGenerator.h
@@ -7,7 +7,7 @@
 #include "AssetRegistry/AssetData.h"
 #include "AssetRegistry/AssetRegistryState.h"
 #include "Containers/Map.h"
-#include "Cooker/PackageResultsMessage.h"
+#include "Cooker/CookMPCollector.h"
 #include "Misc/AssetRegistryInterface.h"
 #include "Misc/Paths.h"
 #include "Templates/UniquePtr.h"
@@ -204,7 +204,7 @@ public:
 	 * @param bIncludeOnlyDiskAssets Include only disk assets or else also enumerate memory assets
 	 */
 	void UpdateAssetRegistryData(const UPackage& Package, FSavePackageResultStruct& SavePackageResult, FCookTagList&& InArchiveCookTagList, bool bIncludeOnlyDiskAssets = true);
-	void UpdateAssetRegistryData(UE::Cook::FPackageData& PackageData, UE::Cook::FAssetRegistryPackageMessage&& Message);
+	void UpdateAssetRegistryData(UE::Cook::FMPCollectorServerMessageContext& Context, UE::Cook::FAssetRegistryPackageMessage&& Message);
 
 	/**
 	 * Check config to see whether chunk assignments use the AssetManager. If so, run the once-per-process construction
@@ -468,15 +468,20 @@ public:
 private:
 	FCookWorkerClient& Client;
 	const ITargetPlatform* TargetPlatform = nullptr;
+	TMap<FName, FCbObject> PackageUpdateMessages;
+
+	friend FAssetRegistryMPCollector;
 };
 
-class FAssetRegistryPackageMessage : public IPackageMessage
+class FAssetRegistryPackageMessage : public UE::CompactBinaryTCP::IMessage
 {
 public:
-	virtual void Write(FCbWriter& Writer, const FPackageData& PackageData, const ITargetPlatform* TargetPlatform) const override;
-	virtual bool TryRead(FCbObjectView Object, FPackageData& PackageData, const ITargetPlatform* TargetPlatform) override;
+	virtual void Write(FCbWriter& Writer) const override;
+	virtual bool TryRead(FCbObjectView Object) override;
 	virtual FGuid GetMessageType() const override { return MessageType; }
 
+	FName PackageName;
+	const ITargetPlatform* TargetPlatform;
 	TArray<FAssetData> AssetDatas;
 	TMap<FSoftObjectPath, TArray<TPair<FName, FString>>> CookTags;
 	uint32 PackageFlags = 0;
@@ -486,6 +491,20 @@ public:
 	static FGuid MessageType;
 };
 
+class FAssetRegistryMPCollector : public UE::Cook::IMPCollector
+{
+public:
+	FAssetRegistryMPCollector(UCookOnTheFlyServer& InCOTFS);
+
+	virtual FGuid GetMessageType() const { return FAssetRegistryPackageMessage::MessageType; }
+	virtual const TCHAR* GetDebugName() const { return TEXT("AssetRegistry"); }
+
+	virtual void ClientTickPackage(FMPCollectorClientTickPackageContext& Context) override;
+	virtual void ServerReceiveMessage(FMPCollectorServerMessageContext& Context, FCbObjectView Message) override;
+
+private:
+	UCookOnTheFlyServer& COTFS;
+};
 
 
 }
\ No newline at end of file
