Fix an issue where DLC cook unsolicited packages can be requested for cook when AllowUncookedAssetReferences is true.
DLC plugins should be cooked so they only cook DLC packages when AllowUncookedAssetReferences is specified.

diff --git a/Engine/Source/Editor/UnrealEd/Private/CookOnTheFlyServer.cpp b/Engine/Source/Editor/UnrealEd/Private/CookOnTheFlyServer.cpp
--- a/Engine/Source/Editor/UnrealEd/Private/CookOnTheFlyServer.cpp
+++ b/Engine/Source/Editor/UnrealEd/Private/CookOnTheFlyServer.cpp
@@ -2754,6 +2754,11 @@ void UCookOnTheFlyServer::QueueDiscoveredPackageData(UE::Cook::FPackageData& Pac
 		return;
 	}
 
+	if (!PackageData.CanCookForPlatforms())
+	{
+		return;
+	}
+
 	if (PackageData.IsInProgress())
 	{
 		return;
diff --git a/Engine/Source/Editor/UnrealEd/Private/Cooker/CookPackageData.cpp b/Engine/Source/Editor/UnrealEd/Private/Cooker/CookPackageData.cpp
--- a/Engine/Source/Editor/UnrealEd/Private/Cooker/CookPackageData.cpp
+++ b/Engine/Source/Editor/UnrealEd/Private/Cooker/CookPackageData.cpp
@@ -50,7 +50,7 @@ static FAutoConsoleVariableRef CVarPollAsyncPeriod(
 //////////////////////////////////////////////////////////////////////////
 // FPackageData
 FPackageData::FPlatformData::FPlatformData()
-	: bRequested(false), bCookAttempted(false), bCookSucceeded(false), bExplored(false), bSaveTimedOut(false)
+	: bRequested(false), bCookAttempted(false), bCookSucceeded(false), bExplored(false), bSaveTimedOut(false), bCookable(true)
 {
 }
 
@@ -190,6 +190,31 @@ bool FPackageData::HasAllExploredPlatforms(const TArrayView<const ITargetPlatfor
 	return true;
 }
 
+bool FPackageData::CanCookForPlatforms() const
+{
+	if (PlatformDatas.Num() == 0)
+	{
+		return true;
+	}
+
+	if (AreAllRequestedPlatformsExplored())
+	{
+		for (const TPair<const ITargetPlatform*, FPlatformData>& Pair : PlatformDatas)
+		{
+			if (Pair.Value.bCookable)
+			{
+				return true;
+			}
+		}
+		return false;
+	}
+	else
+	{
+		return true;
+	}
+}
+
+
 void FPackageData::SetIsUrgent(bool Value)
 {
 	bool OldValue = static_cast<bool>(bIsUrgent);
diff --git a/Engine/Source/Editor/UnrealEd/Private/Cooker/CookPackageData.h b/Engine/Source/Editor/UnrealEd/Private/Cooker/CookPackageData.h
--- a/Engine/Source/Editor/UnrealEd/Private/Cooker/CookPackageData.h
+++ b/Engine/Source/Editor/UnrealEd/Private/Cooker/CookPackageData.h
@@ -146,6 +146,7 @@ public:
 		bool bCookSucceeded : 1;
 		bool bExplored : 1;
 		bool bSaveTimedOut : 1;
+		bool bCookable : 1;
 	};
 
 	FPackageData(FPackageDatas& PackageDatas, const FName& InPackageName, const FName& InFileName);
@@ -207,6 +208,12 @@ public:
 	/** Return true if and only if every element of Platforms has been explored. Returns true if Platforms is empty. */
 	bool HasAllExploredPlatforms(const TArrayView<const ITargetPlatform* const>& Platforms) const;
 
+	/**
+	 * Return true if this package is cookable for at least one of the platforms.
+	 * If this package has not been explored will return true
+	 */
+	bool CanCookForPlatforms() const;
+
 	/**
 	 * Get the flag for whether this InProgress PackageData has been marked as an urgent request
 	 * (e.g. because it has been requested from the game client during cookonthefly.).
diff --git a/Engine/Source/Editor/UnrealEd/Private/Cooker/CookRequestCluster.cpp b/Engine/Source/Editor/UnrealEd/Private/Cooker/CookRequestCluster.cpp
--- a/Engine/Source/Editor/UnrealEd/Private/Cooker/CookRequestCluster.cpp
+++ b/Engine/Source/Editor/UnrealEd/Private/Cooker/CookRequestCluster.cpp
@@ -826,7 +826,9 @@ void FRequestCluster::FGraphSearch::VisitVertex(const FVertexData& VertexData)
 
 	for (const ITargetPlatform* TargetPlatform : Cluster.Platforms)
 	{
-		PackageData->FindOrAddPlatformData(TargetPlatform).bExplored = true;
+		FPackageData::FPlatformData& PlatformData = PackageData->FindOrAddPlatformData(TargetPlatform);
+		PlatformData.bExplored = true;
+		PlatformData.bCookable = VertexData.bCookable;
 	}
 	bool bAlreadyCooked = PackageData->AreAllRequestedPlatformsCooked(true /* bAllowFailedCooks */);
 
