Add async package loading support for cooking
Fix preloading batch size so it can pipeline with the load/save request

#rnx

diff --git a/Engine/Source/Editor/UnrealEd/Private/CookOnTheFlyServer.cpp b/Engine/Source/Editor/UnrealEd/Private/CookOnTheFlyServer.cpp
--- a/Engine/Source/Editor/UnrealEd/Private/CookOnTheFlyServer.cpp
+++ b/Engine/Source/Editor/UnrealEd/Private/CookOnTheFlyServer.cpp
@@ -1391,6 +1391,7 @@ void UCookOnTheFlyServer::TickCookStatus(UE::Cook::FTickStackData& StackData)
 	}
 
 	UpdateDisplay(StackData, false /* bForceDisplay */);
+	ProcessAsyncLoading(false, false, 0.0f);
 	ProcessUnsolicitedPackages();
 	UpdatePackageFilter();
 	PumpExternalRequests(StackData.Timer);
@@ -5994,11 +5995,14 @@ void UCookOnTheFlyServer::SetInitializeConfigSettings(UE::Cook::FInitializeConfi
 {
 	Settings.MoveToLocal(*this);
 
-	MaxPreloadAllocated = 16;
+	// For preload to actually be able to pipeline with load/batch, we need both RequestBatchSize
+	// and MaxPreloadAllocated to be bigger than LoadBatchSize so that we won't consume all preloads
+	// for every iteration.
+	MaxPreloadAllocated = 32;
 	DesiredSaveQueueLength = 8;
 	DesiredLoadQueueLength = 8;
 	LoadBatchSize = 16;
-	RequestBatchSize = 16;
+	RequestBatchSize = 32;
 	WaitForAsyncSleepSeconds = 1.0f;
 
 
diff --git a/Engine/Source/Editor/UnrealEd/Private/Cooker/CookPackageData.cpp b/Engine/Source/Editor/UnrealEd/Private/Cooker/CookPackageData.cpp
--- a/Engine/Source/Editor/UnrealEd/Private/Cooker/CookPackageData.cpp
+++ b/Engine/Source/Editor/UnrealEd/Private/Cooker/CookPackageData.cpp
@@ -886,6 +886,13 @@ bool FPackageData::TryPreload()
 	}
 	if (FindObjectFast<UPackage>(nullptr, GetPackageName()))
 	{
+		if (AsyncRequest && !AsyncRequest->bHasFinished)
+		{
+			// In case of async loading, the object can be found while still being asynchronously serialized, we need to wait until 
+			// the callback is called and the async request is completely done.
+			return false;
+		}
+
 		// If the package has already loaded, then there is no point in further preloading
 		ClearPreload();
 		SetIsPreloadAttempted(true);
@@ -898,6 +905,24 @@ bool FPackageData::TryPreload()
 		SetIsPreloadAttempted(true);
 		return true;
 	}
+	if (IsAsyncLoadingMultithreaded())
+	{
+		if (!AsyncRequest.IsValid())
+		{
+			PackageDatas.GetMonitor().OnPreloadAllocatedChanged(*this, true);
+			AsyncRequest = MakeShared<FAsyncRequest>();
+			AsyncRequest->RequestID = LoadPackageAsync(
+				GetFileName().ToString(), 
+				FLoadPackageAsyncDelegate::CreateLambda(
+					[AsyncRequest = AsyncRequest](const FName&, UPackage*, EAsyncLoadingResult::Type) { AsyncRequest->bHasFinished = true; }
+				),
+				32 /* Use arbitrary higher priority for preload as we're going to need them very soon */
+			);
+		}
+
+		// always return false so we continue to check the status of the load until FindObjectFast above finds the loaded object
+		return false;
+	}
 	if (!PreloadableFile.Get())
 	{
 		if (FEditorDomain* EditorDomain = FEditorDomain::Get())
@@ -989,6 +1014,17 @@ void FPackageData::FTrackedPreloadableFilePtr::Reset(FPackageData& Owner)
 
 void FPackageData::ClearPreload()
 {
+	if (AsyncRequest)
+	{
+		if (!AsyncRequest->bHasFinished)
+		{
+			FlushAsyncLoading(AsyncRequest->RequestID);
+			check(AsyncRequest->bHasFinished);
+		}
+		PackageDatas.GetMonitor().OnPreloadAllocatedChanged(*this, false);
+		AsyncRequest.Reset();
+	}
+
 	const TSharedPtr<FPreloadableArchive>& FilePtr = PreloadableFile.Get();
 	if (GetIsPreloaded())
 	{
@@ -1015,6 +1051,7 @@ void FPackageData::ClearPreload()
 
 void FPackageData::CheckPreloadEmpty()
 {
+	check(!AsyncRequest);
 	check(!GetIsPreloadAttempted());
 	check(!PreloadableFile.Get());
 	check(!GetIsPreloaded());
diff --git a/Engine/Source/Editor/UnrealEd/Private/Cooker/CookPackageData.h b/Engine/Source/Editor/UnrealEd/Private/Cooker/CookPackageData.h
--- a/Engine/Source/Editor/UnrealEd/Private/Cooker/CookPackageData.h
+++ b/Engine/Source/Editor/UnrealEd/Private/Cooker/CookPackageData.h
@@ -579,6 +579,14 @@ private:
 	FCompletionCallback CompletionCallback;
 	FName PackageName;
 	FName FileName;
+
+	struct FAsyncRequest
+	{
+		int32 RequestID { 0 };
+		std::atomic<bool> bHasFinished { false };
+	};
+	TSharedPtr<FAsyncRequest> AsyncRequest;
+
 	TWeakObjectPtr<UPackage> Package;
 	/** The one-per-CookOnTheFlyServer owner of this PackageData. */
 	FPackageDatas& PackageDatas;
