CompactBinarySerialization: Add LoadCompactBinary and operator<< for TTuple.
#rnx

diff --git a/Engine/Source/Editor/UnrealEd/Private/Cooker/CompactBinaryTCP.h b/Engine/Source/Editor/UnrealEd/Private/Cooker/CompactBinaryTCP.h
--- a/Engine/Source/Editor/UnrealEd/Private/Cooker/CompactBinaryTCP.h
+++ b/Engine/Source/Editor/UnrealEd/Private/Cooker/CompactBinaryTCP.h
@@ -7,6 +7,7 @@
 #include "Containers/Map.h"
 #include "Memory/SharedBuffer.h"
 #include "Misc/Guid.h"
+#include "Serialization/CompactBinaryContainerSerialization.h"
 #include "Serialization/CompactBinarySerialization.h"
 #include "Serialization/CompactBinaryWriter.h"
 #include "UObject/SoftObjectPath.h"
@@ -206,24 +207,6 @@ inline FCbWriter& operator<<(FCbWriter& Writer, const TMap<KeyType, ValueType, S
 	return Writer;
 }
 
-template <typename Key, typename Value>
-FCbWriter& operator<<(FCbWriter& Writer, const TPair<Key, Value>& Pair)
-{
-	Writer.BeginObject();
-	Writer << "K" << Pair.Key;
-	Writer << "V" << Pair.Value;
-	Writer.EndObject();
-	return Writer;
-}
-
-template <typename Key, typename Value>
-bool LoadFromCompactBinary(FCbFieldView Field, TPair<Key, Value>& Pair)
-{
-	bool bOk = LoadFromCompactBinary(Field["K"], Pair.Key);
-	bOk = LoadFromCompactBinary(Field["V"], Pair.Value) & bOk;
-	return bOk;
-}
-
 // FSoftObjectPath has an implicit constructor from FString for backwards compatibility; if we
 // try to create an operator<< for it, it will cause operator<< to be ambiguous.
 // This prevents us from using it directly in containers. To hack around this,
diff --git a/Engine/Source/Runtime/Core/Public/Serialization/CompactBinaryContainerSerialization.h b/Engine/Source/Runtime/Core/Public/Serialization/CompactBinaryContainerSerialization.h
new file mode 100644
--- /dev/null
+++ b/Engine/Source/Runtime/Core/Public/Serialization/CompactBinaryContainerSerialization.h
@@ -0,0 +1,32 @@
+// Copyright Epic Games, Inc. All Rights Reserved.
+
+#pragma once
+
+#include "Serialization/CompactBinary.h"
+#include "Serialization/CompactBinaryWriter.h"
+#include "Templates/Tuple.h"
+
+template <typename... Types>
+inline FCbWriter& operator<<(FCbWriter& Writer, const TTuple<Types...>& Tuple)
+{
+	Writer.BeginArray();
+	VisitTupleElements([&Writer](auto&& ElementValue) { Writer << ElementValue; }, Tuple);
+	Writer.EndArray();
+	return Writer;
+}
+
+template <typename... Types>
+inline bool LoadFromCompactBinary(FCbFieldView Field, TTuple<Types...>& Tuple)
+{
+	bool bOk = Field.AsArrayView().Num() == sizeof...(Types);
+	FCbFieldViewIterator ElementIterator = Field.CreateViewIterator();
+
+	VisitTupleElements([&ElementIterator, &bOk](auto&& ElementValue)
+		{
+			bOk = LoadFromCompactBinary(ElementIterator++, ElementValue) & bOk;
+		},
+		Tuple
+	);
+	return bOk;
+}
+
