Updating FSuppressableWarningDialog to new button and icon styles

diff --git a/Engine/Source/Editor/UnrealEd/Private/Dialogs/Dialogs.cpp b/Engine/Source/Editor/UnrealEd/Private/Dialogs/Dialogs.cpp
--- a/Engine/Source/Editor/UnrealEd/Private/Dialogs/Dialogs.cpp
+++ b/Engine/Source/Editor/UnrealEd/Private/Dialogs/Dialogs.cpp
@@ -23,6 +23,7 @@
 #include "DesktopPlatformModule.h"
 #include "HAL/PlatformApplicationMisc.h"
 #include "HAL/PlatformMisc.h"
+#include "SPrimaryButton.h"
 
 DEFINE_LOG_CATEGORY_STATIC(LogDialogs, Log, All);
 
@@ -618,7 +619,7 @@ public:
 				.HAlign(HAlign_Fill)
 				.VAlign(VAlign_Fill)
 				.FillHeight(1.0f)
-				.Padding(0,5,0,5)
+				.Padding(0)
 				.MaxHeight(550)
 				[
 					SNew( SScrollBox )
@@ -630,6 +631,7 @@ public:
 						+SHorizontalBox::Slot()
 						.AutoWidth()
 						.VAlign(VAlign_Center)
+						.Padding(16) // currently hardcoded until we adjust StandardDialog.SlotPadding
 						[
 							SNew( SImage )
 							.Image(InArgs._Image)
@@ -639,7 +641,7 @@ public:
 						+SHorizontalBox::Slot()
 						.AutoWidth()
 						.VAlign(VAlign_Center)
-						.Padding(5,0,5,0)
+						.Padding(0)
 						[
 							SNew( STextBlock )
 							.WrapTextAt(512.0f)
@@ -669,13 +671,12 @@ public:
 	TSharedRef<SHorizontalBox> ConstructConditionalInternals( const FArguments& InArgs ) 
 	{
 		TSharedRef<SHorizontalBox> HorizontalBox = SNew(SHorizontalBox);
-		TSharedPtr<SUniformGridPanel> UniformGridPanel;
 		
 		// checkbox with user specified text
 		HorizontalBox->AddSlot()
 		.HAlign(HAlign_Left)
-		.Padding(5,0,15,0)
-		.AutoWidth()
+		.Padding(16) // currently hardcoded until we adjust StandardDialog.SlotPadding
+		.FillWidth(1.0)
 		[
 			SNew(SCheckBox)
 			.IsChecked(InArgs._bDefaultCheckValue ? ECheckBoxState::Checked : ECheckBoxState::Unchecked)
@@ -689,35 +690,28 @@ public:
 		];
 		HorizontalBox->AddSlot()
 		.HAlign(HAlign_Right)
-		.Padding(2.f)
+		.Padding(16) // currently hardcoded until we adjust StandardDialog.SlotPadding
+		.AutoWidth()
 		[
-			SAssignNew(UniformGridPanel, SUniformGridPanel)
-			.SlotPadding(FAppStyle::GetMargin("StandardDialog.SlotPadding"))
+			SNew(SPrimaryButton)
+			.Text(InArgs._ConfirmText)
+			.OnClicked(this, &SModalDialogWithCheckbox::OnConfirmClicked)
+			
 		];
 
-		// yes/ok/confirm button
-		UniformGridPanel->AddSlot(0,0)
-		.HAlign(HAlign_Fill)
-		[
-			SNew( SButton )
-			.Text( InArgs._ConfirmText )
-			.OnClicked( this, &SModalDialogWithCheckbox::OnConfirmClicked )
-			.ContentPadding(FAppStyle::GetMargin("StandardDialog.ContentPadding"))
-			.HAlign(HAlign_Center)
-		];
 
 		// Only add a cancel button if required
 		if (InArgs._bHasCancelButton)
 		{
 			// cancel/stop/abort button
-			UniformGridPanel->AddSlot(1,0)
-			.HAlign(HAlign_Fill)
+			HorizontalBox->AddSlot()
+			.HAlign(HAlign_Right)
+			.Padding(0, 16, 16, 16) // currently hardcoded until we adjust StandardDialog.SlotPadding
+			.AutoWidth()
 			[
 				SNew( SButton )
 				.Text( InArgs._CancelText )
 				.OnClicked( this, &SModalDialogWithCheckbox::OnCancelClicked )
-				.ContentPadding(FAppStyle::GetMargin("StandardDialog.ContentPadding"))
-				.HAlign(HAlign_Center)
 			];
 		}
 	
@@ -827,7 +821,7 @@ FSuppressableWarningDialog::FSuppressableWarningDialog(const FSetupInfo& Info)
 		.SupportsMaximize(false) .SupportsMinimize(false);
 
 		// Cache a default image to be used as most cases will not provide their own.
-		static const FSlateBrush* DefaultImage = FAppStyle::GetBrush("NotificationList.DefaultMessage");
+		static const FSlateBrush* DefaultImage = FAppStyle::GetBrush("Icons.WarningWithColor.Large");
 
 		MessageBox = SNew(SModalDialogWithCheckbox)
 			.Message(Prompt)
diff --git a/Engine/Source/Editor/UnrealEd/Public/Dialogs/Dialogs.h b/Engine/Source/Editor/UnrealEd/Public/Dialogs/Dialogs.h
--- a/Engine/Source/Editor/UnrealEd/Public/Dialogs/Dialogs.h
+++ b/Engine/Source/Editor/UnrealEd/Public/Dialogs/Dialogs.h
@@ -146,7 +146,7 @@ public:
 	 * 
 	 * @param FSetupInfo Info - struct used to initialize the dialog. 
 	 */
-	FSuppressableWarningDialog ( const FSuppressableWarningDialog::FSetupInfo& info );
+	FSuppressableWarningDialog ( const FSuppressableWarningDialog::FSetupInfo& Info );
 	
 	/** Launches warning window, returns user response or suppressed */
 	EResult ShowModal() const;
