Fix modal window behavior on Linux.

Tested on Mac, Windows and Linux.  Mac and Windows still work as they did before, Linux now has proper modal behavior

diff --git a/Engine/Source/Editor/UnrealEd/Private/EditorEngine.cpp b/Engine/Source/Editor/UnrealEd/Private/EditorEngine.cpp
--- a/Engine/Source/Editor/UnrealEd/Private/EditorEngine.cpp
+++ b/Engine/Source/Editor/UnrealEd/Private/EditorEngine.cpp
@@ -4168,8 +4168,8 @@ void UEditorEngine::BuildReflectionCaptures(UWorld* World)
 
 void UEditorEngine::EditorAddModalWindow( TSharedRef<SWindow> InModalWindow ) const
 {
-	// If there is already a modal window active, parent this new modal window to the existing window so that it doesnt fall behind
-	TSharedPtr<SWindow> ParentWindow = FSlateApplication::Get().GetActiveModalWindow();
+	// If there is already a window active, parent this new modal window to the existing window so that it doesnt fall behind
+	TSharedPtr<SWindow> ParentWindow = FSlateApplication::Get().GetActiveTopLevelRegularWindow();
 
 	if( !ParentWindow.IsValid() )
 	{
