Fix some non-unity build breaks for the new FBookmarkScoped class which was introduced in CL 23867130.

diff --git a/Engine/Source/Editor/UnrealEd/Public/Bookmarks/BookmarkScoped.h b/Engine/Source/Editor/UnrealEd/Public/Bookmarks/BookmarkScoped.h
--- a/Engine/Source/Editor/UnrealEd/Public/Bookmarks/BookmarkScoped.h
+++ b/Engine/Source/Editor/UnrealEd/Public/Bookmarks/BookmarkScoped.h
@@ -2,7 +2,10 @@
 
 #pragma once
 
+#include "CoreMinimal.h"
+
 class IBookmarkTypeActions;
+class FEditorViewportClient;
 class UBookMark;
 
 /**
