Fix for false bExposeClothingSceneElementMenu resulting in missing Clothing submenu in anim editors

diff --git a/Engine/Source/Editor/UnrealEd/Classes/Preferences/PersonaOptions.h b/Engine/Source/Editor/UnrealEd/Classes/Preferences/PersonaOptions.h
--- a/Engine/Source/Editor/UnrealEd/Classes/Preferences/PersonaOptions.h
+++ b/Engine/Source/Editor/UnrealEd/Classes/Preferences/PersonaOptions.h
@@ -109,10 +109,6 @@ class UNREALED_API UPersonaOptions : public UObject
 {
 	GENERATED_UCLASS_BODY()
 
-	UPersonaOptions()
-		: bExposeClothingSceneElementMenu(true)
-	{}
-
 	/** Whether or not the floor should be aligned to the Skeletal Mesh's bounds by default for the Animation Editor(s)*/
 	UPROPERTY(EditAnywhere, config, Category = "Preview Scene")
 	uint32 bAutoAlignFloorToMesh : 1;
@@ -289,7 +285,7 @@ public:
 
 	FNamePermissionList& GetAllowedAnimationEditorTracks() { return AllowedAnimationEditorTracks; }
 
-	bool bExposeClothingSceneElementMenu;
+	bool bExposeClothingSceneElementMenu = true;
 
 protected:
 	// UObject interface
