MPCook - Replicate PackageStoreManifest from CookWorkers to CookDirector.
#rnx

diff --git a/Engine/Source/Editor/UnrealEd/Classes/CookOnTheSide/CookOnTheFlyServer.h b/Engine/Source/Editor/UnrealEd/Classes/CookOnTheSide/CookOnTheFlyServer.h
--- a/Engine/Source/Editor/UnrealEd/Classes/CookOnTheSide/CookOnTheFlyServer.h
+++ b/Engine/Source/Editor/UnrealEd/Classes/CookOnTheSide/CookOnTheFlyServer.h
@@ -148,6 +148,7 @@ namespace UE::Cook
 	class FCookDirector;
 	class FCookWorkerClient;
 	class FCookWorkerServer;
+	class FPackageWriterMPCollector;
 	class FRequestCluster;
 	class FRequestQueue;
 	class FSaveCookedPackageContext;
@@ -1433,6 +1434,7 @@ private:
 	friend UE::Cook::FGeneratorPackage;
 	friend UE::Cook::FInitializeConfigSettings;
 	friend UE::Cook::FPackageData;
+	friend UE::Cook::FPackageWriterMPCollector;
 	friend UE::Cook::FPendingCookedPlatformData;
 	friend UE::Cook::FPlatformManager;
 	friend UE::Cook::FRequestCluster;
diff --git a/Engine/Source/Editor/UnrealEd/Private/Cooker/CookDirector.cpp b/Engine/Source/Editor/UnrealEd/Private/Cooker/CookDirector.cpp
--- a/Engine/Source/Editor/UnrealEd/Private/Cooker/CookDirector.cpp
+++ b/Engine/Source/Editor/UnrealEd/Private/Cooker/CookDirector.cpp
@@ -130,6 +130,7 @@ FCookDirector::FCookDirector(UCookOnTheFlyServer& InCOTFS, int32 CookProcessCoun
 			HandleHeartbeatMessage(Context, bReadSuccessful, MoveTemp(Message));
 		}, TEXT("HandleHeartbeatMessage")));
 	Register(new FAssetRegistryMPCollector(COTFS));
+	Register(new FPackageWriterMPCollector(COTFS));
 
 	LastTickTimeSeconds = FPlatformTime::Seconds();
 
diff --git a/Engine/Source/Editor/UnrealEd/Private/Cooker/CookMPCollector.h b/Engine/Source/Editor/UnrealEd/Private/Cooker/CookMPCollector.h
--- a/Engine/Source/Editor/UnrealEd/Private/Cooker/CookMPCollector.h
+++ b/Engine/Source/Editor/UnrealEd/Private/Cooker/CookMPCollector.h
@@ -44,7 +44,13 @@ private:
 class FMPCollectorClientTickPackageContext
 {
 public:
+	struct FPlatformData
+	{
+		const ITargetPlatform* TargetPlatform = nullptr;
+		bool bSuccessful = false;
+	};
 	TConstArrayView<const ITargetPlatform*> GetPlatforms() const { return Platforms; }
+	TConstArrayView<FPlatformData> GetPlatformDatas() const { return PlatformDatas; }
 	FName GetPackageName() const { return PackageName; }
 
 	void AddMessage(FCbObject Object);
@@ -56,6 +62,7 @@ public:
 private:
 	TArray<TPair<const ITargetPlatform*, FCbObject>> Messages;
 	TConstArrayView<const ITargetPlatform*> Platforms;
+	TConstArrayView<FPlatformData> PlatformDatas;
 	FName PackageName;
 
 	friend class FCookWorkerClient;
diff --git a/Engine/Source/Editor/UnrealEd/Private/Cooker/CookWorkerClient.cpp b/Engine/Source/Editor/UnrealEd/Private/Cooker/CookWorkerClient.cpp
--- a/Engine/Source/Editor/UnrealEd/Private/Cooker/CookWorkerClient.cpp
+++ b/Engine/Source/Editor/UnrealEd/Private/Cooker/CookWorkerClient.cpp
@@ -47,6 +47,7 @@ FCookWorkerClient::FCookWorkerClient(UCookOnTheFlyServer& InCOTFS)
 			HandleHeartbeatMessage(Context, bReadSuccessful, MoveTemp(Message));
 		}, TEXT("HandleHeartbeatMessage")));
 	Register(new FAssetRegistryMPCollector(COTFS));
+	Register(new FPackageWriterMPCollector(COTFS));
 }
 
 FCookWorkerClient::~FCookWorkerClient()
@@ -172,9 +173,17 @@ void FCookWorkerClient::ReportPromoteToSaveComplete(FPackageData& PackageData)
 		PlatformResults.SetSuccessful(PackagePlatformData.bCookSucceeded);
 	}
 
+	TArray<FMPCollectorClientTickPackageContext::FPlatformData, TInlineAllocator<1>> ContextPlatformDatas;
+	ContextPlatformDatas.Reserve(NumPlatforms);
+	for (FPackageRemoteResult::FPlatformResult& PlatformResult : Result.GetPlatforms())
+	{
+		ContextPlatformDatas.Add(FMPCollectorClientTickPackageContext::FPlatformData
+			{ PlatformResult.GetPlatform(), PlatformResult.IsSuccessful() });
+	}
 	FMPCollectorClientTickPackageContext Context;
 	Context.PackageName = PackageName;
 	Context.Platforms = OrderedSessionPlatforms;
+	Context.PlatformDatas = ContextPlatformDatas;
 
 	for (const TPair<FGuid, TRefCountPtr<IMPCollector>>& CollectorPair : Collectors)
 	{
diff --git a/Engine/Source/Editor/UnrealEd/Private/Cooker/CookWorkerServer.cpp b/Engine/Source/Editor/UnrealEd/Private/Cooker/CookWorkerServer.cpp
--- a/Engine/Source/Editor/UnrealEd/Private/Cooker/CookWorkerServer.cpp
+++ b/Engine/Source/Editor/UnrealEd/Private/Cooker/CookWorkerServer.cpp
@@ -1060,4 +1060,39 @@ bool FHeartbeatMessage::TryRead(FCbObjectView Object)
 	return LoadFromCompactBinary(Object["H"], HeartbeatNumber);
 }
 
+FPackageWriterMPCollector::FPackageWriterMPCollector(UCookOnTheFlyServer& InCOTFS)
+	: COTFS(InCOTFS)
+{
+
+}
+
+void FPackageWriterMPCollector::ClientTickPackage(FMPCollectorClientTickPackageContext& Context)
+{
+	for (const FMPCollectorClientTickPackageContext::FPlatformData& PlatformData : Context.GetPlatformDatas())
+	{
+		if (PlatformData.bSuccessful)
+		{
+			ICookedPackageWriter& PackageWriter = COTFS.FindOrCreatePackageWriter(PlatformData.TargetPlatform);
+			FCbObject Object = PackageWriter.WriteMPCookMessageForPackage(Context.GetPackageName());
+			Context.AddPlatformMessage(PlatformData.TargetPlatform, Object);
+		}
+	}
+}
+
+void FPackageWriterMPCollector::ServerReceiveMessage(FMPCollectorServerMessageContext& Context, FCbObjectView Message)
+{
+	FName PackageName = Context.GetPackageName();
+	const ITargetPlatform* TargetPlatform = Context.GetTargetPlatform();
+	check(PackageName.IsValid() && TargetPlatform);
+
+	ICookedPackageWriter& PackageWriter = COTFS.FindOrCreatePackageWriter(TargetPlatform);
+	if (!PackageWriter.TryReadMPCookMessageForPackage(PackageName, Message))
+	{
+		UE_LOG(LogCook, Error, TEXT("CookWorkerServer received invalidly formatted PackageWriter message from CookWorker %d. Ignoring it."),
+			Context.GetProfileId());
+	}
+}
+
+FGuid FPackageWriterMPCollector::MessageType(TEXT("D2B1CE3FD26644AF9EC28FBADB1BD331"));
+
 }
\ No newline at end of file
diff --git a/Engine/Source/Editor/UnrealEd/Private/Cooker/CookWorkerServer.h b/Engine/Source/Editor/UnrealEd/Private/Cooker/CookWorkerServer.h
--- a/Engine/Source/Editor/UnrealEd/Private/Cooker/CookWorkerServer.h
+++ b/Engine/Source/Editor/UnrealEd/Private/Cooker/CookWorkerServer.h
@@ -381,4 +381,19 @@ public:
 	static FGuid MessageType;
 };
 
+class FPackageWriterMPCollector : public UE::Cook::IMPCollector
+{
+public:
+	FPackageWriterMPCollector(UCookOnTheFlyServer& InCOTFS);
+	virtual FGuid GetMessageType() const { return MessageType; }
+	virtual const TCHAR* GetDebugName() const { return TEXT("PackageWriter"); }
+
+	virtual void ClientTickPackage(FMPCollectorClientTickPackageContext& Context) override;
+	virtual void ServerReceiveMessage(FMPCollectorServerMessageContext& Context, FCbObjectView Message) override;
+
+private:
+	UCookOnTheFlyServer& COTFS;
+	static FGuid MessageType;
+};
+
 }
\ No newline at end of file
