changed FCustomVersionCollectorArchive override serialization of FObjectPtr to avoid loading referenced objects

https://p4-swarm.epicgames.net/reviews/23877027

diff --git a/Engine/Source/Editor/UnrealEd/Private/EditorDomain/EditorDomainUtils.cpp b/Engine/Source/Editor/UnrealEd/Private/EditorDomain/EditorDomainUtils.cpp
--- a/Engine/Source/Editor/UnrealEd/Private/EditorDomain/EditorDomainUtils.cpp
+++ b/Engine/Source/Editor/UnrealEd/Private/EditorDomain/EditorDomainUtils.cpp
@@ -846,6 +846,11 @@ public:
 		return TEXT("FCustomVersionCollectorArchive");
 	}
 
+	virtual FArchive& operator<<(FObjectPtr& Value) override
+	{
+		return *this;
+	}
+
 private:
 	int64 Pos = 0;
 	int64 Max = 0;
