Add Import allowed asset types in the class path permission system.

#rnx

diff --git a/Engine/Plugins/Interchange/Runtime/Source/Pipelines/Private/InterchangeGenericTexturePipeline.cpp b/Engine/Plugins/Interchange/Runtime/Source/Pipelines/Private/InterchangeGenericTexturePipeline.cpp
--- a/Engine/Plugins/Interchange/Runtime/Source/Pipelines/Private/InterchangeGenericTexturePipeline.cpp
+++ b/Engine/Plugins/Interchange/Runtime/Source/Pipelines/Private/InterchangeGenericTexturePipeline.cpp
@@ -21,7 +21,6 @@
 #include "InterchangeVolumeTextureNode.h"
 #include "InterchangeVolumeTextureFactoryNode.h"
 #include "Misc/Paths.h"
-#include "Engine/Texture.h"
 
 #include UE_INLINE_GENERATED_CPP_BY_NAME(InterchangeGenericTexturePipeline)
 
diff --git a/Engine/Source/Developer/AssetTools/Public/IAssetTools.h b/Engine/Source/Developer/AssetTools/Public/IAssetTools.h
--- a/Engine/Source/Developer/AssetTools/Public/IAssetTools.h
+++ b/Engine/Source/Developer/AssetTools/Public/IAssetTools.h
@@ -44,6 +44,8 @@ enum class EAssetClassAction : uint8
 	CreateAsset,
 	/** Whether an asset can be viewed in the content browser */
 	ViewAsset,
+	/** Whether an asset can be imported */
+	ImportAsset,
 	AllAssetActions
 };
 
diff --git a/Engine/Source/Editor/UnrealEd/Private/Factories/EditorFactories.cpp b/Engine/Source/Editor/UnrealEd/Private/Factories/EditorFactories.cpp
--- a/Engine/Source/Editor/UnrealEd/Private/Factories/EditorFactories.cpp
+++ b/Engine/Source/Editor/UnrealEd/Private/Factories/EditorFactories.cpp
@@ -380,6 +380,20 @@ public:
 	}
 };
 
+namespace UE::EditorFactories::Private
+{
+	bool CanImportClass(UClass* Class)
+	{
+		IAssetTools& AssetTools = FAssetToolsModule::GetModule().Get();
+		TSharedPtr<FPathPermissionList> AssetClassPermissionList = AssetTools.GetAssetClassPathPermissionList(EAssetClassAction::ImportAsset);
+		if (Class && AssetClassPermissionList && AssetClassPermissionList->HasFiltering())
+		{
+			return AssetClassPermissionList->PassesFilter(Class->GetPathName());
+		}
+		return true;
+	}
+
+}
 /*------------------------------------------------------------------------------
 	UTexture2DFactoryNew implementation.
 ------------------------------------------------------------------------------*/
@@ -5986,7 +6000,12 @@ bool UReimportFbxStaticMeshFactory::FactoryCanImport(const FString& Filename)
 }
 
 bool UReimportFbxStaticMeshFactory::CanReimport( UObject* Obj, TArray<FString>& OutFilenames )
-{	
+{
+	if (!UE::EditorFactories::Private::CanImportClass(UStaticMesh::StaticClass()))
+	{
+		return false;
+	}
+
 	UStaticMesh* Mesh = Cast<UStaticMesh>(Obj);
 	if(Mesh)
 	{
@@ -6031,6 +6050,11 @@ void UReimportFbxStaticMeshFactory::SetReimportPaths( UObject* Obj, const TArray
 
 EReimportResult::Type UReimportFbxStaticMeshFactory::Reimport( UObject* Obj )
 {
+	if (!UE::EditorFactories::Private::CanImportClass(UStaticMesh::StaticClass()))
+	{
+		return EReimportResult::Failed;
+	}
+
 	UStaticMesh* Mesh = Cast<UStaticMesh>(Obj);
 	if( !Mesh )
 	{
@@ -6312,6 +6336,11 @@ bool UReimportFbxSkeletalMeshFactory::FactoryCanImport(const FString& Filename)
 
 bool UReimportFbxSkeletalMeshFactory::CanReimport( UObject* Obj, TArray<FString>& OutFilenames )
 {
+	if (!UE::EditorFactories::Private::CanImportClass(USkeletalMesh::StaticClass()))
+	{
+		return false;
+	}
+
 	USkeletalMesh* SkeletalMesh = Cast<USkeletalMesh>(Obj);
 	TArray<FString> FactoryExtensions;
 	GetSupportedFileExtensions(FactoryExtensions);
@@ -6371,6 +6400,11 @@ void UReimportFbxSkeletalMeshFactory::SetReimportPaths( UObject* Obj, const FStr
 
 EReimportResult::Type UReimportFbxSkeletalMeshFactory::Reimport( UObject* Obj, int32 SourceFileIndex)
 {
+	if (!UE::EditorFactories::Private::CanImportClass(USkeletalMesh::StaticClass()))
+	{
+		return EReimportResult::Failed;
+	}
+
 	// Only handle valid skeletal meshes
 	if( !Obj || !Obj->IsA( USkeletalMesh::StaticClass() ))
 	{
@@ -6795,7 +6829,11 @@ bool UReimportFbxAnimSequenceFactory::FactoryCanImport(const FString& Filename)
 }
 
 bool UReimportFbxAnimSequenceFactory::CanReimport( UObject* Obj, TArray<FString>& OutFilenames )
-{	
+{
+	if (!UE::EditorFactories::Private::CanImportClass(UAnimSequence::StaticClass()))
+	{
+		return false;
+	}
 	UAnimSequence* AnimSequence = Cast<UAnimSequence>(Obj);
 	if(AnimSequence)
 	{
@@ -6837,6 +6875,11 @@ void UReimportFbxAnimSequenceFactory::SetReimportPaths( UObject* Obj, const TArr
 
 EReimportResult::Type UReimportFbxAnimSequenceFactory::Reimport( UObject* Obj )
 {
+	if (!UE::EditorFactories::Private::CanImportClass(UAnimSequence::StaticClass()))
+	{
+		return EReimportResult::Failed;
+	}
+
 	// Only handle valid skeletal meshes
 	if( !Obj || !Obj->IsA( UAnimSequence::StaticClass() ) )
 	{
diff --git a/Engine/Source/Editor/UnrealEd/Private/Fbx/FbxFactory.cpp b/Engine/Source/Editor/UnrealEd/Private/Fbx/FbxFactory.cpp
--- a/Engine/Source/Editor/UnrealEd/Private/Fbx/FbxFactory.cpp
+++ b/Engine/Source/Editor/UnrealEd/Private/Fbx/FbxFactory.cpp
@@ -429,12 +429,17 @@ UObject* UFbxFactory::FactoryCreateFile
 			bool bCombineMeshes = ImportUI->StaticMeshImportData->bCombineMeshes;
 			bool bCombineMeshesLOD = false;
 
-			if ( ImportUI->MeshTypeToImport == FBXIT_SkeletalMesh )
+			const bool bCanImportStaticMesh = FbxImporter->CanImportClass(UStaticMesh::StaticClass());
+			//We can import skeletal mesh only if we can import skeleton or a skeleton was specified in the import options
+			const bool bCanImportSkeletalMesh = FbxImporter->CanImportClass(USkeletalMesh::StaticClass()) && (ImportOptions->SkeletonForAnimation || FbxImporter->CanImportClass(USkeleton::StaticClass()));
+			const bool bCanImportAnimSequence = FbxImporter->CanImportClass(UAnimSequence::StaticClass());
+
+			if (bCanImportStaticMesh && ImportUI->MeshTypeToImport == FBXIT_SkeletalMesh )
 			{
 				FbxImporter->FillFbxSkelMeshArrayInScene(RootNodeToImport, SkelMeshArray, false, (ImportOptions->bImportAsSkeletalGeometry || ImportOptions->bImportAsSkeletalSkinning));
 				InterestingNodeCount = SkelMeshArray.Num();
 			}
-			else if( ImportUI->MeshTypeToImport == FBXIT_StaticMesh )
+			else if(bCanImportSkeletalMesh && ImportUI->MeshTypeToImport == FBXIT_StaticMesh )
 			{
 				FbxImporter->ApplyTransformSettingsToFbxNode(RootNodeToImport, ImportUI->StaticMeshImportData);
 
@@ -475,7 +480,7 @@ UObject* UFbxFactory::FactoryCreateFile
 				int32 NodeIndex = 0;
 
 				int32 ImportedMeshCount = 0;
-				if ( ImportUI->MeshTypeToImport == FBXIT_StaticMesh )  // static mesh
+				if (bCanImportStaticMesh && ImportUI->MeshTypeToImport == FBXIT_StaticMesh )  // static mesh
 				{
 					UStaticMesh* NewStaticMesh = NULL;
 					if (bCombineMeshes)
@@ -628,7 +633,7 @@ UObject* UFbxFactory::FactoryCreateFile
 					CreatedObject = NewStaticMesh;
 
 				}
-				else if ( ImportUI->MeshTypeToImport == FBXIT_SkeletalMesh )// skeletal mesh
+				else if ( bCanImportSkeletalMesh && ImportUI->MeshTypeToImport == FBXIT_SkeletalMesh )// skeletal mesh
 				{
 					int32 TotalNumNodes = 0;
 
@@ -830,7 +835,7 @@ UObject* UFbxFactory::FactoryCreateFile
 							FFbxErrors::SkeletalMesh_NoMeshFoundOnRoot);
 					}
 				}
-				else if ( ImportUI->MeshTypeToImport == FBXIT_Animation )// animation
+				else if (bCanImportAnimSequence && ImportUI->MeshTypeToImport == FBXIT_Animation )// animation
 				{
 					if (ImportOptions->SkeletonForAnimation)
 					{
diff --git a/Engine/Source/Editor/UnrealEd/Private/Fbx/FbxMainImport.cpp b/Engine/Source/Editor/UnrealEd/Private/Fbx/FbxMainImport.cpp
--- a/Engine/Source/Editor/UnrealEd/Private/Fbx/FbxMainImport.cpp
+++ b/Engine/Source/Editor/UnrealEd/Private/Fbx/FbxMainImport.cpp
@@ -39,6 +39,8 @@
 #include "IMeshReductionInterfaces.h"
 #include "ObjectTools.h"
 #include "Misc/AutomationTest.h"
+#include "AssetToolsModule.h"
+#include "IAssetTools.h"
 
 DEFINE_LOG_CATEGORY(LogFbx);
 
@@ -649,6 +651,17 @@ void FFbxImporter::ReleaseScene()
 	LastMergeBonesChoice = EAppReturnType::Ok;
 }
 
+bool FFbxImporter::CanImportClass(UClass* Class) const
+{
+	IAssetTools& AssetTools = FAssetToolsModule::GetModule().Get();
+	TSharedPtr<FPathPermissionList> AssetClassPermissionList = AssetTools.GetAssetClassPathPermissionList(EAssetClassAction::ImportAsset);
+	if (Class && AssetClassPermissionList && AssetClassPermissionList->HasFiltering())
+	{
+		return AssetClassPermissionList->PassesFilter(Class->GetPathName());
+	}
+	return true;
+}
+
 FBXImportOptions* UnFbx::FFbxImporter::GetImportOptions() const
 {
 	return ImportOptions;
diff --git a/Engine/Source/Editor/UnrealEd/Private/Fbx/FbxMaterialImport.cpp b/Engine/Source/Editor/UnrealEd/Private/Fbx/FbxMaterialImport.cpp
--- a/Engine/Source/Editor/UnrealEd/Private/Fbx/FbxMaterialImport.cpp
+++ b/Engine/Source/Editor/UnrealEd/Private/Fbx/FbxMaterialImport.cpp
@@ -38,7 +38,7 @@ using namespace UnFbx;
 
 UTexture* UnFbx::FFbxImporter::ImportTexture(FbxFileTexture* FbxTexture, bool bSetupAsNormalMap)
 {
-	if (!FbxTexture)
+	if (!FbxTexture || !CanImportClass(UTexture2D::StaticClass()))
 	{
 		return nullptr;
 	}
@@ -629,15 +629,6 @@ UMaterialInterface* UnFbx::FFbxImporter::CreateUnrealMaterial(const FbxSurfaceMa
 		return nullptr;
 	}
 
-	const FString BasePackageName = GetMaterialBasePackageName(GetMaterialFullName(FbxMaterial));
-	const FString Suffix(TEXT(""));
-	FString FinalPackageName;
-	FString FinalMaterialName;
-	FAssetToolsModule& AssetToolsModule = FModuleManager::LoadModuleChecked<FAssetToolsModule>("AssetTools");
-	AssetToolsModule.Get().CreateUniqueAssetName(BasePackageName, Suffix, FinalPackageName, FinalMaterialName);
-
-	UPackage* Package = CreatePackage( *FinalPackageName);
-	
 	// Check if we can use the specified base material to instance from it
 	FBXImportOptions* FbxImportOptions = GetImportOptions();
 	bool bCanInstance = false;
@@ -669,8 +660,30 @@ UMaterialInterface* UnFbx::FFbxImporter::CreateUnrealMaterial(const FbxSurfaceMa
 		bCanInstance &= CanUseMaterialWithInstance(FbxMaterial, FbxSurfaceMaterial::sTransparentColor, FbxImportOptions->BaseOpacityTextureName, FbxImportOptions->BaseMaterial, OutUVSets);
 	}
 
+	//Make sure we can import the material class
+	if ( bCanInstance )
+	{
+		if (!CanImportClass(UMaterialInstanceConstant::StaticClass()))
+		{
+			return nullptr;
+		}
+	}
+	else if ( !CanImportClass(UMaterial::StaticClass()) )
+	{
+		return nullptr;
+	}
+
+	const FString BasePackageName = GetMaterialBasePackageName(GetMaterialFullName(FbxMaterial));
+	const FString Suffix(TEXT(""));
+	FString FinalPackageName;
+	FString FinalMaterialName;
+	FAssetToolsModule& AssetToolsModule = FModuleManager::LoadModuleChecked<FAssetToolsModule>("AssetTools");
+	AssetToolsModule.Get().CreateUniqueAssetName(BasePackageName, Suffix, FinalPackageName, FinalMaterialName);
+
+	UPackage* Package = CreatePackage(*FinalPackageName);
 	UMaterialInterface* UnrealMaterialFinal = nullptr;
-	if (bCanInstance) {
+	if (bCanInstance)
+	{
 		auto MaterialInstanceFactory = NewObject<UMaterialInstanceConstantFactoryNew>();
 		MaterialInstanceFactory->InitialParent = FbxImportOptions->BaseMaterial;
 		UMaterialInstanceConstant* UnrealMaterialConstant = (UMaterialInstanceConstant*)MaterialInstanceFactory->FactoryCreateNew(UMaterialInstanceConstant::StaticClass(), Package, *FinalMaterialName, RF_Standalone | RF_Public, NULL, GWarn);
diff --git a/Engine/Source/Editor/UnrealEd/Private/Fbx/FbxSceneImportFactory.cpp b/Engine/Source/Editor/UnrealEd/Private/Fbx/FbxSceneImportFactory.cpp
--- a/Engine/Source/Editor/UnrealEd/Private/Fbx/FbxSceneImportFactory.cpp
+++ b/Engine/Source/Editor/UnrealEd/Private/Fbx/FbxSceneImportFactory.cpp
@@ -2061,6 +2061,12 @@ UObject* UFbxSceneImportFactory::ImportOneSkeletalMesh(void* VoidRootNodeToImpor
 void UFbxSceneImportFactory::ImportAllSkeletalMesh(void* VoidRootNodeToImport, void* VoidFbxImporter, EObjectFlags Flags, int32& NodeIndex, int32& InterestingNodeCount, TSharedPtr<FFbxSceneInfo> SceneInfo)
 {
 	UnFbx::FFbxImporter* FbxImporter = (UnFbx::FFbxImporter*)VoidFbxImporter;
+	
+	if (!FbxImporter->CanImportClass(USkeletalMesh::StaticClass()) || !FbxImporter->CanImportClass(USkeleton::StaticClass()))
+	{
+		return;
+	}
+
 	FbxNode *RootNodeToImport = (FbxNode *)VoidRootNodeToImport;
 	InterestingNodeCount = 1;
 	TArray< TArray<FbxNode*>* > SkelMeshArray;
@@ -2103,6 +2109,12 @@ void UFbxSceneImportFactory::ImportAllStaticMesh(void* VoidRootNodeToImport, voi
 	TRACE_CPUPROFILER_EVENT_SCOPE(UFbxSceneImportFactory::ImportAllStaticMesh);
 
 	UnFbx::FFbxImporter* FbxImporter = (UnFbx::FFbxImporter*)VoidFbxImporter;
+
+	if (!FbxImporter->CanImportClass(UStaticMesh::StaticClass()))
+	{
+		return;
+	}
+
 	FbxNode *RootNodeToImport = (FbxNode *)VoidRootNodeToImport;
 	
 	//Copy default options to StaticMeshImportData
diff --git a/Engine/Source/Editor/UnrealEd/Private/Fbx/FbxSkeletalMeshImport.cpp b/Engine/Source/Editor/UnrealEd/Private/Fbx/FbxSkeletalMeshImport.cpp
--- a/Engine/Source/Editor/UnrealEd/Private/Fbx/FbxSkeletalMeshImport.cpp
+++ b/Engine/Source/Editor/UnrealEd/Private/Fbx/FbxSkeletalMeshImport.cpp
@@ -1624,7 +1624,7 @@ void UnFbx::FFbxImporter::FillLastImportMaterialNames(TArray<FName> &LastImporte
 
 USkeletalMesh* UnFbx::FFbxImporter::ImportSkeletalMesh(FImportSkeletalMeshArgs &ImportSkeletalMeshArgs)
 {
-	if (ImportSkeletalMeshArgs.NodeArray.Num() == 0)
+	if (ImportSkeletalMeshArgs.NodeArray.Num() == 0 || !CanImportClass(USkeletalMesh::StaticClass()))
 	{
 		return nullptr;
 	}
@@ -2045,7 +2045,7 @@ USkeletalMesh* UnFbx::FFbxImporter::ImportSkeletalMesh(FImportSkeletalMeshArgs &
 
 		EARLY_RETURN_ON_CANCEL(false, CancelCleanup);
 
-		if ((!SkeletalMesh->GetResourceForRendering() || !SkeletalMesh->GetResourceForRendering()->LODRenderData.IsValidIndex(0)) && ImportOptions->bCreatePhysicsAsset)
+		if ((!SkeletalMesh->GetResourceForRendering() || !SkeletalMesh->GetResourceForRendering()->LODRenderData.IsValidIndex(0)) && ImportOptions->bCreatePhysicsAsset && CanImportClass(UPhysicsAsset::StaticClass()))
 		{
 			//We need to have a valid render data to create physic asset
 			SkeletalMesh->Build();
@@ -2056,7 +2056,7 @@ USkeletalMesh* UnFbx::FFbxImporter::ImportSkeletalMesh(FImportSkeletalMeshArgs &
 	if(ImportSkeletalMeshArgs.LodIndex == 0)
 	{
 		// see if we have skeleton set up
-		// if creating skeleton, create skeleeton
+		// if creating skeleton, create skeleton
 		USkeleton* Skeleton = ImportOptions->SkeletonForAnimation;
 		if (Skeleton == NULL)
 		{
@@ -2163,7 +2163,7 @@ USkeletalMesh* UnFbx::FFbxImporter::ImportSkeletalMesh(FImportSkeletalMeshArgs &
 		// We create the physic asset after we create the skeleton since we need the skeleton to correctly build it
 		if (ImportOptions->bCreatePhysicsAsset)
 		{
-			if (SkeletalMesh->GetPhysicsAsset() == NULL)
+			if (SkeletalMesh->GetPhysicsAsset() == NULL && CanImportClass(UPhysicsAsset::StaticClass()))
 			{
 				FString ObjectName = FString::Printf(TEXT("%s_PhysicsAsset"), *SkeletalMesh->GetName());
 				UPhysicsAsset * NewPhysicsAsset = CreateAsset<UPhysicsAsset>(ImportSkeletalMeshArgs.InParent->GetName(), ObjectName, true);
@@ -2368,6 +2368,11 @@ UObject* UnFbx::FFbxImporter::CreateAssetOfClass(UClass* AssetClass, FString Par
 
 void UnFbx::FFbxImporter::SetupAnimationDataFromMesh(USkeletalMesh* SkeletalMesh, UObject* InParent, TArray<FbxNode*>& NodeArray, UFbxAnimSequenceImportData* TemplateImportData, const FString& Name)
 {
+	if (!CanImportClass(UAnimSequence::StaticClass()))
+	{
+		return;
+	}
+
 	USkeleton* Skeleton = SkeletalMesh->GetSkeleton();
 
 	if (Scene->GetSrcObjectCount<FbxAnimStack>() > 0)
@@ -4464,6 +4469,11 @@ void UnFbx::FFbxImporter::ImportMorphTargetsInternal( TArray<FbxNode*>& SkelMesh
 // Import Morph target
 void UnFbx::FFbxImporter::ImportFbxMorphTarget(TArray<FbxNode*> &SkelMeshNodeArray, USkeletalMesh* BaseSkelMesh, int32 LODIndex, FSkeletalMeshImportData &BaseSkeletalMeshImportData)
 {
+	if (!CanImportClass(UMorphTarget::StaticClass()))
+	{
+		return;
+	}
+
 	//Stack the PostEditChange call, it will call post edit change when it will go out of scope
 	FScopedSkeletalMeshPostEditChange ScopedPostEditChange(BaseSkelMesh);
 	FFbxScopedOperation ScopedImportOperation(this);
diff --git a/Engine/Source/Editor/UnrealEd/Private/SkeletalMeshEdit.cpp b/Engine/Source/Editor/UnrealEd/Private/SkeletalMeshEdit.cpp
--- a/Engine/Source/Editor/UnrealEd/Private/SkeletalMeshEdit.cpp
+++ b/Engine/Source/Editor/UnrealEd/Private/SkeletalMeshEdit.cpp
@@ -568,7 +568,7 @@ void UnFbx::FFbxImporter::GetAnimationIntervalMultiLayer(FbxNode* RootNode, FbxA
 UAnimSequence * UnFbx::FFbxImporter::ImportAnimations(USkeleton* Skeleton, UObject* Outer, TArray<FbxNode*>& SortedLinks, const FString& Name, UFbxAnimSequenceImportData* TemplateImportData, TArray<FbxNode*>& NodeArray)
 {
 	// we need skeleton to create animsequence
-	if (Skeleton == nullptr)
+	if (Skeleton == nullptr || !CanImportClass(UAnimSequence::StaticClass()))
 	{
 		return nullptr;
 	}
diff --git a/Engine/Source/Editor/UnrealEd/Public/FbxImporter.h b/Engine/Source/Editor/UnrealEd/Public/FbxImporter.h
--- a/Engine/Source/Editor/UnrealEd/Public/FbxImporter.h
+++ b/Engine/Source/Editor/UnrealEd/Public/FbxImporter.h
@@ -629,6 +629,11 @@ public:
 		TransformSettingsToFbxApply.Reset();
 	}
 
+	/**
+	 * The asset tool have a filter mecanism for ImportAsset, return true if the asset can be imported, false otherwise
+	 */
+	bool CanImportClass(UClass* Class) const;
+
 	/**
 	 * Detect if the FBX file has skeletal mesh model. If there is deformer definition, then there is skeletal mesh.
 	 * In this function, we don't need to import the scene. But the open process is time-consume if the file is large.
diff --git a/Engine/Source/Runtime/Interchange/Engine/Private/Tasks/InterchangeTaskImportObject.cpp b/Engine/Source/Runtime/Interchange/Engine/Private/Tasks/InterchangeTaskImportObject.cpp
--- a/Engine/Source/Runtime/Interchange/Engine/Private/Tasks/InterchangeTaskImportObject.cpp
+++ b/Engine/Source/Runtime/Interchange/Engine/Private/Tasks/InterchangeTaskImportObject.cpp
@@ -1,6 +1,7 @@
 // Copyright Epic Games, Inc. All Rights Reserved.
 #include "InterchangeTaskImportObject.h"
 
+#include "AssetToolsModule.h"
 #include "Async/TaskGraphInterfaces.h"
 #include "CoreMinimal.h"
 #include "InterchangeAssetImportData.h"
@@ -10,6 +11,7 @@
 #include "InterchangeResult.h"
 #include "InterchangeSourceData.h"
 #include "InterchangeTranslatorBase.h"
+#include "Misc/NamePermissionList.h"
 #include "Misc/Paths.h"
 #include "PackageUtils/PackageUtils.h"
 #include "Stats/Stats.h"
@@ -117,6 +119,23 @@ namespace UE
 				return false;
 			}
 
+			bool CanImportClass(UClass* Class)
+			{
+				if (!Class)
+				{
+					return true;
+				}
+				IAssetTools& AssetTools = FAssetToolsModule::GetModule().Get();
+				TSharedPtr<FPathPermissionList> AssetClassPermissionList = AssetTools.GetAssetClassPathPermissionList(EAssetClassAction::ImportAsset);
+				if (AssetClassPermissionList && AssetClassPermissionList->HasFiltering())
+				{
+					if (!AssetClassPermissionList->PassesFilter(Class->GetPathName()))
+					{
+						return false;
+					}
+				}
+				return true;
+			}
 		}//ns Private
 	}//ns Interchange
 }//ns UE
@@ -131,11 +150,12 @@ void UE::Interchange::FTaskImportObject_GameThread::DoTask(ENamedThreads::Type C
 	check(AsyncHelper.IsValid());
 
 	//Verify if the task was cancel
-	if (AsyncHelper->bCancel)
+	if (AsyncHelper->bCancel
+		|| !FactoryNode
+		|| !UE::Interchange::Private::CanImportClass(FactoryNode->GetObjectClass()))
 	{
 		return;
 	}
-
 	//The create package thread must always execute on the game thread
 	check(IsInGameThread());
 
@@ -269,7 +289,9 @@ void UE::Interchange::FTaskImportObject_Async::DoTask(ENamedThreads::Type Curren
 	check(AsyncHelper.IsValid());
 
 	//Verify if the task was cancel
-	if (AsyncHelper->bCancel)
+	if (AsyncHelper->bCancel
+		|| !FactoryNode
+		|| !UE::Interchange::Private::CanImportClass(FactoryNode->GetObjectClass()))
 	{
 		return;
 	}
