Fix CIS build
#rnx

diff --git a/Engine/Source/Editor/UnrealEd/Private/Fbx/FbxFactory.cpp b/Engine/Source/Editor/UnrealEd/Private/Fbx/FbxFactory.cpp
--- a/Engine/Source/Editor/UnrealEd/Private/Fbx/FbxFactory.cpp
+++ b/Engine/Source/Editor/UnrealEd/Private/Fbx/FbxFactory.cpp
@@ -434,12 +434,12 @@ UObject* UFbxFactory::FactoryCreateFile
 			const bool bCanImportSkeletalMesh = FbxImporter->CanImportClass(USkeletalMesh::StaticClass()) && (ImportOptions->SkeletonForAnimation || FbxImporter->CanImportClass(USkeleton::StaticClass()));
 			const bool bCanImportAnimSequence = FbxImporter->CanImportClass(UAnimSequence::StaticClass());
 
-			if (bCanImportStaticMesh && ImportUI->MeshTypeToImport == FBXIT_SkeletalMesh )
+			if (bCanImportSkeletalMesh && ImportUI->MeshTypeToImport == FBXIT_SkeletalMesh )
 			{
 				FbxImporter->FillFbxSkelMeshArrayInScene(RootNodeToImport, SkelMeshArray, false, (ImportOptions->bImportAsSkeletalGeometry || ImportOptions->bImportAsSkeletalSkinning));
 				InterestingNodeCount = SkelMeshArray.Num();
 			}
-			else if(bCanImportSkeletalMesh && ImportUI->MeshTypeToImport == FBXIT_StaticMesh )
+			else if(bCanImportStaticMesh && ImportUI->MeshTypeToImport == FBXIT_StaticMesh )
 			{
 				FbxImporter->ApplyTransformSettingsToFbxNode(RootNodeToImport, ImportUI->StaticMeshImportData);
 
diff --git a/Engine/Source/Editor/UnrealEd/Private/Fbx/FbxMainImport.cpp b/Engine/Source/Editor/UnrealEd/Private/Fbx/FbxMainImport.cpp
--- a/Engine/Source/Editor/UnrealEd/Private/Fbx/FbxMainImport.cpp
+++ b/Engine/Source/Editor/UnrealEd/Private/Fbx/FbxMainImport.cpp
@@ -14,7 +14,6 @@
 #include "Misc/SecureHash.h"
 #include "Factories/FbxSkeletalMeshImportData.h"
 #include "Factories/FbxTextureImportData.h"
-
 #include "Materials/MaterialInterface.h"
 #include "Rendering/SkeletalMeshLODImporterData.h"
 #include "Logging/TokenizedMessage.h"
@@ -41,6 +40,7 @@
 #include "Misc/AutomationTest.h"
 #include "AssetToolsModule.h"
 #include "IAssetTools.h"
+#include "Misc/NamePermissionList.h"
 
 DEFINE_LOG_CATEGORY(LogFbx);
 
