Fixed CIS static analysis

diff --git a/Engine/Source/Editor/UnrealEd/Private/EditorLevelUtils.cpp b/Engine/Source/Editor/UnrealEd/Private/EditorLevelUtils.cpp
--- a/Engine/Source/Editor/UnrealEd/Private/EditorLevelUtils.cpp
+++ b/Engine/Source/Editor/UnrealEd/Private/EditorLevelUtils.cpp
@@ -780,7 +780,7 @@ ULevelStreaming* UEditorLevelUtils::CreateNewStreamingLevelForWorld(UWorld& InWo
 		// Make sure to uninitialize the world since the level will be used as a streaming level
 		// This will make sure that the initialization order will be respected.
 		// One example is world partition initialization done inside ULevel::OnLevelLoaded.
-		if (NewLevelWorld->bIsWorldInitialized)
+		if (NewLevelWorld && NewLevelWorld->bIsWorldInitialized)
 		{
 			NewLevelWorld->CleanupWorld();
 		}
