Attempt to catch any reinstanced objects referenced by the transaction buffer

This now subscribes to FCoreUObjectDelegates::OnObjectsReinstanced directly instead of only relying on the ARO fix-up

#rnx

diff --git a/Engine/Source/Editor/UnrealEd/Classes/Editor/TransBuffer.h b/Engine/Source/Editor/UnrealEd/Classes/Editor/TransBuffer.h
--- a/Engine/Source/Editor/UnrealEd/Classes/Editor/TransBuffer.h
+++ b/Engine/Source/Editor/UnrealEd/Classes/Editor/TransBuffer.h
@@ -63,6 +63,8 @@ public:
 
 protected:
 	
+	void OnObjectsReinstanced(const TMap<UObject*, UObject*>& OldToNewInstances);
+
 	/** Implementation of the begin function. Used to create a specific transaction type */
 	template<typename TTransaction>
 	int32 BeginInternal( const TCHAR* SessionContext, const FText& Description )
diff --git a/Engine/Source/Editor/UnrealEd/Private/EditorTransaction.cpp b/Engine/Source/Editor/UnrealEd/Private/EditorTransaction.cpp
--- a/Engine/Source/Editor/UnrealEd/Private/EditorTransaction.cpp
+++ b/Engine/Source/Editor/UnrealEd/Private/EditorTransaction.cpp
@@ -1119,6 +1119,8 @@ void UTransBuffer::Initialize(SIZE_T InMaxMemory)
 	Reset( NSLOCTEXT("UnrealEd", "Startup", "Startup") );
 	CheckState();
 
+	FCoreUObjectDelegates::OnObjectsReinstanced.AddUObject(this, &UTransBuffer::OnObjectsReinstanced);
+
 	UE_LOG(LogInit, Log, TEXT("Transaction tracking system initialized") );
 }
 
@@ -1150,6 +1152,42 @@ void UTransBuffer::FinishDestroy()
 	Super::FinishDestroy();
 }
 
+void UTransBuffer::OnObjectsReinstanced(const TMap<UObject*, UObject*>& OldToNewInstances)
+{
+	if (OldToNewInstances.Num() == 0)
+	{
+		return;
+	}
+
+	class FReinstancingReferenceCollector : public FReferenceCollector
+	{
+	public:
+		virtual bool IsIgnoringArchetypeRef() const override
+		{
+			return false;
+		}
+
+		virtual bool IsIgnoringTransient() const override
+		{
+			return false;
+		}
+
+		virtual void HandleObjectReference(UObject*& Object, const UObject* ReferencingObject, const FProperty* ReferencingProperty) override
+		{
+			if (UObject* NewObject = OldToNewInstancesPtr->FindRef(Object))
+			{
+				Object = NewObject;
+			}
+		}
+
+		const TMap<UObject*, UObject*>* OldToNewInstancesPtr = nullptr;
+	};
+
+	FReinstancingReferenceCollector Collector;
+	Collector.OldToNewInstancesPtr = &OldToNewInstances;
+	CallAddReferencedObjects(Collector);
+}
+
 void UTransBuffer::AddReferencedObjects(UObject* InThis, FReferenceCollector& Collector)
 {	
 	UTransBuffer* This = CastChecked<UTransBuffer>(InThis);
