Cooker: MPCook: Replicate package hashes in LooseCookedPackageWriter.
#rnx

diff --git a/Engine/Source/Developer/IoStoreUtilities/Private/ZenStoreWriter.cpp b/Engine/Source/Developer/IoStoreUtilities/Private/ZenStoreWriter.cpp
--- a/Engine/Source/Developer/IoStoreUtilities/Private/ZenStoreWriter.cpp
+++ b/Engine/Source/Developer/IoStoreUtilities/Private/ZenStoreWriter.cpp
@@ -1050,14 +1050,20 @@ void FZenStoreWriter::MarkPackagesUpToDate(TArrayView<const FName> UpToDatePacka
 	}
 }
 
-FCbObject FZenStoreWriter::WriteMPCookMessageForPackage(FName PackageName)
+TFuture<FCbObject> FZenStoreWriter::WriteMPCookMessageForPackage(FName PackageName)
 {
+	// MPCOOKTODO: Replicate the other accumulated data in FZenStoreWriter
 	FCbWriter Writer;
 	Writer.BeginObject();
 	Writer.SetName("Manifest");
 	PackageStoreManifest.WritePackage(Writer, PackageName);
 	Writer.EndObject();
-	return Writer.Save().AsObject();
+	FCbObject Object = Writer.Save().AsObject();
+
+	TPromise<FCbObject> Promise;
+	Promise.SetValue(MoveTemp(Object));
+	return Promise.GetFuture();
+
 }
 
 bool FZenStoreWriter::TryReadMPCookMessageForPackage(FName PackageName, FCbObjectView Message)
diff --git a/Engine/Source/Developer/IoStoreUtilities/Public/ZenStoreWriter.h b/Engine/Source/Developer/IoStoreUtilities/Public/ZenStoreWriter.h
--- a/Engine/Source/Developer/IoStoreUtilities/Public/ZenStoreWriter.h
+++ b/Engine/Source/Developer/IoStoreUtilities/Public/ZenStoreWriter.h
@@ -100,7 +100,7 @@ public:
 	IOSTOREUTILITIES_API virtual void RemoveCookedPackages(TArrayView<const FName> PackageNamesToRemove) override;
 	IOSTOREUTILITIES_API virtual void RemoveCookedPackages() override;
 	IOSTOREUTILITIES_API virtual void MarkPackagesUpToDate(TArrayView<const FName> UpToDatePackages) override;
-	IOSTOREUTILITIES_API virtual FCbObject WriteMPCookMessageForPackage(FName PackageName) override;
+	IOSTOREUTILITIES_API virtual TFuture<FCbObject> WriteMPCookMessageForPackage(FName PackageName) override;
 	IOSTOREUTILITIES_API virtual bool TryReadMPCookMessageForPackage(FName PackageName, FCbObjectView Message) override;
 
 	IOSTOREUTILITIES_API virtual TMap<FName, TRefCountPtr<FPackageHashes>>& GetPackageHashes() override
diff --git a/Engine/Source/Editor/UnrealEd/Private/Cooker/CookMPCollector.cpp b/Engine/Source/Editor/UnrealEd/Private/Cooker/CookMPCollector.cpp
--- a/Engine/Source/Editor/UnrealEd/Private/Cooker/CookMPCollector.cpp
+++ b/Engine/Source/Editor/UnrealEd/Private/Cooker/CookMPCollector.cpp
@@ -46,12 +46,23 @@ void FMPCollectorClientTickPackageContext::AddMessage(FCbObject Object)
 	Messages.Emplace(nullptr, MoveTemp(Object));
 }
 
+void FMPCollectorClientTickPackageContext::AddAsyncMessage(TFuture<FCbObject>&& ObjectFuture)
+{
+	AsyncMessages.Emplace(nullptr, MoveTemp(ObjectFuture));
+}
+
 void FMPCollectorClientTickPackageContext::AddPlatformMessage(const ITargetPlatform* TargetPlatform, FCbObject Object)
 {
 	check(TargetPlatform && Platforms.Contains(TargetPlatform));
 	Messages.Emplace(TargetPlatform, MoveTemp(Object));
 }
 
+void FMPCollectorClientTickPackageContext::AddAsyncPlatformMessage(const ITargetPlatform* TargetPlatform, TFuture<FCbObject>&& ObjectFuture)
+{
+	check(TargetPlatform && Platforms.Contains(TargetPlatform));
+	AsyncMessages.Emplace(TargetPlatform, MoveTemp(ObjectFuture));
+}
+
 uint8 FMPCollectorClientTickPackageContext::PlatformToInt(const ITargetPlatform* Platform) const
 {
 	return PlatformToIntImpl(Platform, Platforms);
diff --git a/Engine/Source/Editor/UnrealEd/Private/Cooker/CookMPCollector.h b/Engine/Source/Editor/UnrealEd/Private/Cooker/CookMPCollector.h
--- a/Engine/Source/Editor/UnrealEd/Private/Cooker/CookMPCollector.h
+++ b/Engine/Source/Editor/UnrealEd/Private/Cooker/CookMPCollector.h
@@ -2,6 +2,7 @@
 
 #pragma once
 
+#include "Async/Future.h"
 #include "Containers/Array.h"
 #include "Containers/ArrayView.h"
 #include "Cooker/CompactBinaryTCP.h"
@@ -54,13 +55,16 @@ public:
 	FName GetPackageName() const { return PackageName; }
 
 	void AddMessage(FCbObject Object);
+	void AddAsyncMessage(TFuture<FCbObject>&& ObjectFuture);
 	void AddPlatformMessage(const ITargetPlatform* Platform, FCbObject Object);
+	void AddAsyncPlatformMessage(const ITargetPlatform* Platform, TFuture<FCbObject>&& ObjectFuture);
 
 	uint8 PlatformToInt(const ITargetPlatform* Platform) const;
 	const ITargetPlatform* IntToPlatform(uint8 PlatformAsInt) const;
 
 private:
 	TArray<TPair<const ITargetPlatform*, FCbObject>> Messages;
+	TArray<TPair<const ITargetPlatform*, TFuture<FCbObject>>> AsyncMessages;
 	TConstArrayView<const ITargetPlatform*> Platforms;
 	TConstArrayView<FPlatformData> PlatformDatas;
 	FName PackageName;
diff --git a/Engine/Source/Editor/UnrealEd/Private/Cooker/CookWorkerClient.cpp b/Engine/Source/Editor/UnrealEd/Private/Cooker/CookWorkerClient.cpp
--- a/Engine/Source/Editor/UnrealEd/Private/Cooker/CookWorkerClient.cpp
+++ b/Engine/Source/Editor/UnrealEd/Private/Cooker/CookWorkerClient.cpp
@@ -58,6 +58,21 @@ FCookWorkerClient::~FCookWorkerClient()
 		UE_LOG(LogCook, Warning, TEXT("CookWorker was destroyed before it finished Disconnect. The CookDirector may be missing some information."));
 	}
 	Sockets::CloseSocket(ServerSocket);
+
+	// Before destructing, wait on all of the Futures that could have async access to *this from a TaskThread
+	TArray<FPendingResultNeedingAsyncWork> LocalPendingResultsNeedingAsyncWork;
+	{
+		FScopeLock PendingResultsScopeLock(&PendingResultsLock);
+		for (TPair<FPackageRemoteResult*, FPendingResultNeedingAsyncWork>& Pair : PendingResultsNeedingAsyncWork)
+		{
+			LocalPendingResultsNeedingAsyncWork.Add(MoveTemp(Pair.Value));
+		}
+		PendingResultsNeedingAsyncWork.Empty();
+	}
+	for (FPendingResultNeedingAsyncWork& PendingResult : LocalPendingResultsNeedingAsyncWork)
+	{
+		PendingResult.CompletionFuture.Get();
+	}
 }
 
 bool FCookWorkerClient::TryConnect(FDirectorConnectionInfo&& ConnectInfo)
@@ -150,32 +165,36 @@ void FCookWorkerClient::ReportDemoteToIdle(const FPackageData& PackageData, ESup
 	{
 		return;
 	}
-	FPackageRemoteResult& Result = PendingResults.Emplace_GetRef();
-	Result.SetPackageName(PackageData.GetPackageName());
-	Result.SetSuppressCookReason(Reason);
+	TUniquePtr<FPackageRemoteResult> Result(new FPackageRemoteResult());
+	Result->SetPackageName(PackageData.GetPackageName());
+	Result->SetSuppressCookReason(Reason);
+
+	FScopeLock PendingResultsScopeLock(&PendingResultsLock);
+	PendingResults.Add(MoveTemp(Result));
 }
 
 void FCookWorkerClient::ReportPromoteToSaveComplete(FPackageData& PackageData)
 {
-	FPackageRemoteResult& Result = PendingResults.Emplace_GetRef();
+	TUniquePtr<FPackageRemoteResult> ResultOwner(new FPackageRemoteResult());
+	FPackageRemoteResult* Result = ResultOwner.Get();
 
 	FName PackageName = PackageData.GetPackageName();
-	Result.SetPackageName(PackageName);
-	Result.SetSuppressCookReason(ESuppressCookReason::InvalidSuppressCookReason);
-	Result.SetPlatforms(OrderedSessionPlatforms);
+	Result->SetPackageName(PackageName);
+	Result->SetSuppressCookReason(ESuppressCookReason::InvalidSuppressCookReason);
+	Result->SetPlatforms(OrderedSessionPlatforms);
 
 	int32 NumPlatforms = OrderedSessionPlatforms.Num();
 	for (int32 PlatformIndex = 0; PlatformIndex < NumPlatforms; ++PlatformIndex)
 	{
 		ITargetPlatform* TargetPlatform = OrderedSessionPlatforms[PlatformIndex];
-		FPackageRemoteResult::FPlatformResult& PlatformResults = Result.GetPlatforms()[PlatformIndex];
+		FPackageRemoteResult::FPlatformResult& PlatformResults = Result->GetPlatforms()[PlatformIndex];
 		FPackageData::FPlatformData& PackagePlatformData = PackageData.FindOrAddPlatformData(TargetPlatform);
 		PlatformResults.SetSuccessful(PackagePlatformData.bCookSucceeded);
 	}
 
 	TArray<FMPCollectorClientTickPackageContext::FPlatformData, TInlineAllocator<1>> ContextPlatformDatas;
 	ContextPlatformDatas.Reserve(NumPlatforms);
-	for (FPackageRemoteResult::FPlatformResult& PlatformResult : Result.GetPlatforms())
+	for (FPackageRemoteResult::FPlatformResult& PlatformResult : Result->GetPlatforms())
 	{
 		ContextPlatformDatas.Add(FMPCollectorClientTickPackageContext::FPlatformData
 			{ PlatformResult.GetPlatform(), PlatformResult.IsSuccessful() });
@@ -189,23 +208,66 @@ void FCookWorkerClient::ReportPromoteToSaveComplete(FPackageData& PackageData)
 	{
 		IMPCollector* Collector = CollectorPair.Value.GetReference();
 		Collector->ClientTickPackage(Context);
-		if (!Context.Messages.IsEmpty())
+		const FGuid& MessageType = CollectorPair.Key;
+		for (TPair<const ITargetPlatform*, FCbObject>& MessagePair : Context.Messages)
 		{
-			const FGuid& MessageType = CollectorPair.Key;
-			for (TPair<const ITargetPlatform*, FCbObject>& MessagePair : Context.Messages)
+			const ITargetPlatform* TargetPlatform = MessagePair.Key;
+			FCbObject Object = MoveTemp(MessagePair.Value);
+			if (!TargetPlatform)
 			{
-				const ITargetPlatform* TargetPlatform = MessagePair.Key;
-				FCbObject Object = MoveTemp(MessagePair.Value);
-				if (!TargetPlatform)
-				{
-					Result.AddPackageMessage(MessageType, MoveTemp(Object));
-				}
-				else
-				{
-					Result.AddPlatformMessage(TargetPlatform, MessageType, MoveTemp(Object));
-				}
+				Result->AddPackageMessage(MessageType, MoveTemp(Object));
+			}
+			else
+			{
+				Result->AddPlatformMessage(TargetPlatform, MessageType, MoveTemp(Object));
+			}
+		}
+		Context.Messages.Reset();
+		for (TPair<const ITargetPlatform*, TFuture<FCbObject>>& MessagePair : Context.AsyncMessages)
+		{
+			const ITargetPlatform* TargetPlatform = MessagePair.Key;
+			TFuture<FCbObject> ObjectFuture = MoveTemp(MessagePair.Value);
+			if (!TargetPlatform)
+			{
+				Result->AddAsyncPackageMessage(MessageType, MoveTemp(ObjectFuture));
+			}
+			else
+			{
+				Result->AddAsyncPlatformMessage(TargetPlatform, MessageType, MoveTemp(ObjectFuture));
 			}
-			Context.Messages.Reset();
+		}
+		Context.AsyncMessages.Reset();
+	}
+
+	++(Result->GetUserRefCount()); // Used to test whether the async Future still needs to access *this
+	TFuture<void> CompletionFuture = Result->GetCompletionFuture().Then(
+	[this, Result](TFuture<int>&& OldFuture)
+	{
+		FScopeLock PendingResultsScopeLock(&PendingResultsLock);
+		FPendingResultNeedingAsyncWork PendingResult;
+		PendingResultsNeedingAsyncWork.RemoveAndCopyValue(Result, PendingResult);
+
+		// Result might have not been added into PendingResultsNeedingAsyncWork yet, and also could have
+		// been removed by cancellation from e.g. CookWorkerClient destructor.
+		if (PendingResult.PendingResult)
+		{
+			PendingResults.Add(MoveTemp(PendingResult.PendingResult));
+		}
+		--(Result->GetUserRefCount());
+	});
+
+	{
+		FScopeLock PendingResultsScopeLock(&PendingResultsLock);
+		if (Result->GetUserRefCount() == 0)
+		{
+			// Result->GetCompletionFuture() has already been called
+			check(Result->IsComplete());
+			PendingResults.Add(MoveTemp(ResultOwner));
+		}
+		else
+		{
+			PendingResultsNeedingAsyncWork.Add(Result,
+				FPendingResultNeedingAsyncWork{ MoveTemp(ResultOwner), MoveTemp(CompletionFuture) });
 		}
 	}
 }
@@ -405,18 +467,29 @@ void FCookWorkerClient::PumpSendMessages()
 
 void FCookWorkerClient::SendPendingResults()
 {
-	if (!PendingResults.IsEmpty())
+	FPackageResultsMessage Message;
+	{
+		FScopeLock PendingResultsScopeLock(&PendingResultsLock);
+		if (!PendingResults.IsEmpty())
+		{
+			Message.Results.Reserve(PendingResults.Num());
+			for (TUniquePtr<FPackageRemoteResult>& Result : PendingResults)
+			{
+				Message.Results.Add(MoveTemp(*Result));
+			}
+			PendingResults.Reset();
+		}
+	}
+	if (!Message.Results.IsEmpty())
 	{
-		FPackageResultsMessage Message;
-		Message.Results = MoveTemp(PendingResults);
 		SendMessage(Message);
-		PendingResults.Reset();
 	}
+
 	if (!PendingDiscoveredPackages.IsEmpty())
 	{
-		FDiscoveredPackagesMessage Message;
-		Message.Packages = MoveTemp(PendingDiscoveredPackages);
-		SendMessage(Message);
+		FDiscoveredPackagesMessage DiscoveredMessage;
+		DiscoveredMessage.Packages = MoveTemp(PendingDiscoveredPackages);
+		SendMessage(DiscoveredMessage);
 		PendingDiscoveredPackages.Reset();
 	}
 }
diff --git a/Engine/Source/Editor/UnrealEd/Private/Cooker/CookWorkerClient.h b/Engine/Source/Editor/UnrealEd/Private/Cooker/CookWorkerClient.h
--- a/Engine/Source/Editor/UnrealEd/Private/Cooker/CookWorkerClient.h
+++ b/Engine/Source/Editor/UnrealEd/Private/Cooker/CookWorkerClient.h
@@ -6,8 +6,10 @@
 #include "Containers/Array.h"
 #include "CookOnTheSide/CookOnTheFlyServer.h"
 #include "CookTypes.h"
+#include "HAL/CriticalSection.h"
 #include "IPAddress.h"
 #include "Misc/Guid.h"
+#include "Templates/UniquePtr.h"
 
 namespace UE::Cook { class FMPCollectorClientMessageContext; }
 namespace UE::Cook { class IMPCollector; }
@@ -117,10 +119,24 @@ private:
 		FHeartbeatMessage&& Message);
 
 private:
+	/**
+	 * A PendingResult constructed during ReportPromoteToSaveComplete that is not yet ready to 
+	 * send because it has some asynchronous messages still pending.
+	 */
+	struct FPendingResultNeedingAsyncWork
+	{
+		FPendingResultNeedingAsyncWork() = default;
+		FPendingResultNeedingAsyncWork(FPendingResultNeedingAsyncWork&&) = default;
+		FPendingResultNeedingAsyncWork& operator=(FPendingResultNeedingAsyncWork&&) = default;
+
+		TUniquePtr<FPackageRemoteResult> PendingResult;
+		TFuture<void> CompletionFuture;
+	};
+private:
+	// Variables Read/Write only from the Scheduler thread
 	TSharedPtr<FInternetAddr> DirectorAddr;
 	TUniquePtr<FInitialConfigMessage> InitialConfigMessage;
 	TArray<ITargetPlatform*> OrderedSessionPlatforms;
-	TArray<FPackageRemoteResult> PendingResults;
 	TArray<FDiscoveredPackage> PendingDiscoveredPackages;
 	TMap<FGuid, TRefCountPtr<IMPCollector>> Collectors;
 	UE::CompactBinaryTCP::FSendBuffer SendBuffer;
@@ -133,6 +149,11 @@ private:
 	EConnectStatus ConnectStatus = EConnectStatus::Uninitialized;
 	ECookMode::Type DirectorCookMode = ECookMode::CookByTheBook;
 	bool bHasRunFinished = false;
+
+	// Variables Read/Write only within PendingResultsLock
+	FCriticalSection PendingResultsLock;
+	TArray<TUniquePtr<FPackageRemoteResult>> PendingResults;
+	TMap<FPackageRemoteResult*, FPendingResultNeedingAsyncWork> PendingResultsNeedingAsyncWork;
 };
 
 }
\ No newline at end of file
diff --git a/Engine/Source/Editor/UnrealEd/Private/Cooker/CookWorkerServer.cpp b/Engine/Source/Editor/UnrealEd/Private/Cooker/CookWorkerServer.cpp
--- a/Engine/Source/Editor/UnrealEd/Private/Cooker/CookWorkerServer.cpp
+++ b/Engine/Source/Editor/UnrealEd/Private/Cooker/CookWorkerServer.cpp
@@ -1073,8 +1073,8 @@ void FPackageWriterMPCollector::ClientTickPackage(FMPCollectorClientTickPackageC
 		if (PlatformData.bSuccessful)
 		{
 			ICookedPackageWriter& PackageWriter = COTFS.FindOrCreatePackageWriter(PlatformData.TargetPlatform);
-			FCbObject Object = PackageWriter.WriteMPCookMessageForPackage(Context.GetPackageName());
-			Context.AddPlatformMessage(PlatformData.TargetPlatform, Object);
+			TFuture<FCbObject> ObjectFuture = PackageWriter.WriteMPCookMessageForPackage(Context.GetPackageName());
+			Context.AddAsyncPlatformMessage(PlatformData.TargetPlatform, MoveTemp(ObjectFuture));
 		}
 	}
 }
diff --git a/Engine/Source/Editor/UnrealEd/Private/Cooker/DiffPackageWriter.h b/Engine/Source/Editor/UnrealEd/Private/Cooker/DiffPackageWriter.h
--- a/Engine/Source/Editor/UnrealEd/Private/Cooker/DiffPackageWriter.h
+++ b/Engine/Source/Editor/UnrealEd/Private/Cooker/DiffPackageWriter.h
@@ -100,7 +100,7 @@ public:
 		Inner->UpdateSaveArguments(SaveArgs);
 	}
 	virtual bool IsAnotherSaveNeeded(FSavePackageResultStruct& PreviousResult, FSavePackageArgs& SaveArgs) override;
-	virtual FCbObject WriteMPCookMessageForPackage(FName PackageName) override
+	virtual TFuture<FCbObject> WriteMPCookMessageForPackage(FName PackageName) override
 	{
 		return Inner->WriteMPCookMessageForPackage(PackageName);
 	}
@@ -236,7 +236,7 @@ public:
 	}
 	virtual void UpdateSaveArguments(FSavePackageArgs& SaveArgs) override;
 	virtual bool IsAnotherSaveNeeded(FSavePackageResultStruct& PreviousResult, FSavePackageArgs& SaveArgs) override;
-	virtual FCbObject WriteMPCookMessageForPackage(FName PackageName) override
+	virtual TFuture<FCbObject> WriteMPCookMessageForPackage(FName PackageName) override
 	{
 		return Inner->WriteMPCookMessageForPackage(PackageName);
 	}
diff --git a/Engine/Source/Editor/UnrealEd/Private/Cooker/LooseCookedPackageWriter.cpp b/Engine/Source/Editor/UnrealEd/Private/Cooker/LooseCookedPackageWriter.cpp
--- a/Engine/Source/Editor/UnrealEd/Private/Cooker/LooseCookedPackageWriter.cpp
+++ b/Engine/Source/Editor/UnrealEd/Private/Cooker/LooseCookedPackageWriter.cpp
@@ -10,6 +10,7 @@
 #include "Containers/Set.h"
 #include "Containers/StringView.h"
 #include "Cooker/AsyncIODelete.h"
+#include "Cooker/CompactBinaryTCP.h"
 #include "Cooker/CookPackageData.h"
 #include "Cooker/CookTypes.h"
 #include "HAL/FileManager.h"
@@ -62,7 +63,6 @@ FLooseCookedPackageWriter::FLooseCookedPackageWriter(const FString& InOutputPath
 	, PackageStoreManifest(InOutputPath)
 	, PluginsToRemap(InPluginsToRemap)
 	, AsyncIODelete(InAsyncIODelete)
-	, bIterateSharedBuild(false)
 {
 }
 
@@ -313,14 +313,20 @@ void FLooseCookedPackageWriter::AsyncSaveOutputFiles(FRecord& Record, FCommitCon
 	UE::SavePackageUtilities::IncrementOutstandingAsyncWrites();
 
 	TRefCountPtr<FPackageHashes> ThisPackageHashes;
+	TUniquePtr<TPromise<int>> PackageHashesCompletionPromise;
 	
 	if (EnumHasAnyFlags(Context.Info.WriteOptions, EWriteOptions::ComputeHash))
 	{
 		ThisPackageHashes = new FPackageHashes();
+		if (bProvidePerPackageResults)
+		{
+			PackageHashesCompletionPromise.Reset(new TPromise<int>());
+			ThisPackageHashes->CompletionFuture = PackageHashesCompletionPromise->GetFuture();
+		}
 
 		bool bAlreadyExisted = false;
 		{
-			FScopeLock ConcurrentSaveScopeLock(&ConcurrentSaveLock);
+			FScopeLock PackageHashesScopeLock(&PackageHashesLock);
 			TRefCountPtr<FPackageHashes>& ExistingPackageHashes = AllPackageHashes.FindOrAdd(Context.Info.PackageName);
 			// This calculation of bAlreadyExisted looks weird but we're finding the _refcount_, not the hashes. So if it gets
 			// constructed, it's not actually assigned a pointer.
@@ -333,7 +339,10 @@ void FLooseCookedPackageWriter::AsyncSaveOutputFiles(FRecord& Record, FCommitCon
 		}
 	}
 
-	UE::Tasks::Launch(TEXT("HashAndWriteLooseCookedFile"), [OutputFiles = MoveTemp(Context.OutputFiles), WriteOptions = Context.Info.WriteOptions, ThisPackageHashes = MoveTemp(ThisPackageHashes)]()
+	UE::Tasks::Launch(TEXT("HashAndWriteLooseCookedFile"),
+		[OutputFiles = MoveTemp(Context.OutputFiles), WriteOptions = Context.Info.WriteOptions,
+		ThisPackageHashes = MoveTemp(ThisPackageHashes), PackageHashesCompletionPromise=MoveTemp(PackageHashesCompletionPromise)]
+		() mutable
 		{
 			FMD5 AccumulatedHash;
 			for (const FWriteFileData& OutputFile : OutputFiles)
@@ -346,6 +355,13 @@ void FLooseCookedPackageWriter::AsyncSaveOutputFiles(FRecord& Record, FCommitCon
 				ThisPackageHashes->PackageHash.Set(AccumulatedHash);
 			}
 
+			if (PackageHashesCompletionPromise)
+			{
+				// Note that setting this Promise might call arbitrary code from anything that subscribed
+				// to ThisPackageHashes->CompletionFuture.Then(). So don't call it inside a lock.
+				PackageHashesCompletionPromise->SetValue(0);
+			}
+
 			// This is used to release the game thread to access the hashes
 			UE::SavePackageUtilities::DecrementOutstandingAsyncWrites();
 		},
@@ -525,6 +541,7 @@ void FLooseCookedPackageWriter::BeginCook(const FCookInfo& Info)
 	else
 	{
 		PackageStoreManifest.SetTrackPackageData(true);
+		bProvidePerPackageResults = true;
 	}
 	AllPackageHashes.Empty();
 }
@@ -705,19 +722,82 @@ void FLooseCookedPackageWriter::MarkPackagesUpToDate(TArrayView<const FName> UpT
 {
 }
 
-FCbObject FLooseCookedPackageWriter::WriteMPCookMessageForPackage(FName PackageName)
+TFuture<FCbObject> FLooseCookedPackageWriter::WriteMPCookMessageForPackage(FName PackageName)
 {
-	FCbWriter Writer;
-	Writer.BeginObject();
-	Writer.SetName("Manifest");
-	PackageStoreManifest.WritePackage(Writer, PackageName);
-	Writer.EndObject();
-	return Writer.Save().AsObject();
+	FCbWriter ManifestWriter;
+	PackageStoreManifest.WritePackage(ManifestWriter, PackageName);
+	FCbFieldIterator ManifestField = ManifestWriter.Save();
+
+	TRefCountPtr<FPackageHashes> PackageHashes;
+	{
+		FScopeLock PackageHashesScopeLock(&PackageHashesLock);
+		AllPackageHashes.RemoveAndCopyValue(PackageName, PackageHashes);
+	}
+
+	auto ComposeMessage = [ManifestField = MoveTemp(ManifestField)](FPackageHashes* PackageHashes)
+	{
+		FCbWriter Writer;
+		Writer.BeginObject();
+		Writer << "Manifest" << ManifestField;
+		if (PackageHashes)
+		{
+			Writer << "PackageHash" << PackageHashes->PackageHash;
+			Writer << "ChunkHashes" << PackageHashes->ChunkHashes;
+		}
+		Writer.EndObject();
+		return Writer.Save().AsObject();
+	};
+
+	if (PackageHashes && PackageHashes->CompletionFuture.IsValid())
+	{
+		TUniquePtr<TPromise<FCbObject>> Promise(new TPromise<FCbObject>());
+		TFuture<FCbObject> ResultFuture = Promise->GetFuture();
+		PackageHashes->CompletionFuture.Next(
+			[PackageHashes, Promise = MoveTemp(Promise), ComposeMessage = MoveTemp(ComposeMessage)]
+			(int)
+			{
+				Promise->SetValue(ComposeMessage(PackageHashes.GetReference()));
+			});
+		return ResultFuture;
+	}
+	else
+	{
+		TPromise<FCbObject> Promise;
+		Promise.SetValue(ComposeMessage(PackageHashes.GetReference()));
+		return Promise.GetFuture();
+	}
 }
 
 bool FLooseCookedPackageWriter::TryReadMPCookMessageForPackage(FName PackageName, FCbObjectView Message)
 {
-	return PackageStoreManifest.TryReadPackage(Message["Manifest"], PackageName);
+	if (!PackageStoreManifest.TryReadPackage(Message["Manifest"], PackageName))
+	{
+		return false;
+	}
+
+	bool bOk = true;
+	TRefCountPtr<FPackageHashes> ThisPackageHashes(new FPackageHashes());
+	if (LoadFromCompactBinary(Message["PackageHash"], ThisPackageHashes->PackageHash))
+	{
+		bOk = LoadFromCompactBinary(Message["ChunkHashes"], ThisPackageHashes->ChunkHashes) & bOk;
+		if (bOk)
+		{
+			bool bAlreadyExisted = false;
+			{
+				FScopeLock PackageHashesScopeLock(&PackageHashesLock);
+				TRefCountPtr<FPackageHashes>& ExistingPackageHashes = AllPackageHashes.FindOrAdd(PackageName);
+				bAlreadyExisted = ExistingPackageHashes.IsValid();
+				ExistingPackageHashes = ThisPackageHashes;
+			}
+			if (bAlreadyExisted)
+			{
+				UE_LOG(LogSavePackage, Error, TEXT("FLooseCookedPackageWriter encountered the same package twice in a cook! (%s)"),
+					*PackageName.ToString());
+			}
+		}
+	}
+
+	return bOk;
 }
 
 void FLooseCookedPackageWriter::RemoveCookedPackages()
diff --git a/Engine/Source/Editor/UnrealEd/Private/Cooker/LooseCookedPackageWriter.h b/Engine/Source/Editor/UnrealEd/Private/Cooker/LooseCookedPackageWriter.h
--- a/Engine/Source/Editor/UnrealEd/Private/Cooker/LooseCookedPackageWriter.h
+++ b/Engine/Source/Editor/UnrealEd/Private/Cooker/LooseCookedPackageWriter.h
@@ -62,7 +62,7 @@ public:
 	virtual void RemoveCookedPackages(TArrayView<const FName> PackageNamesToRemove) override;
 	virtual void RemoveCookedPackages() override;
 	virtual void MarkPackagesUpToDate(TArrayView<const FName> UpToDatePackages) override;
-	virtual FCbObject WriteMPCookMessageForPackage(FName PackageName) override;
+	virtual TFuture<FCbObject> WriteMPCookMessageForPackage(FName PackageName) override;
 	virtual bool TryReadMPCookMessageForPackage(FName PackageName, FCbObjectView Message) override;
 	virtual bool GetPreviousCookedBytes(const FPackageInfo& Info, FPreviousCookedBytesData& OutData) override;
 	virtual void CompleteExportsArchiveForDiff(const FPackageInfo& Info, FLargeMemoryWriter& ExportsArchive) override;
@@ -147,7 +147,7 @@ private:
 	TMap<FName, TRefCountPtr<FPackageHashes>> AllPackageHashes;
 
 	TMap<FName, FName> UncookedPathToCookedPath;
-	FCriticalSection ConcurrentSaveLock;
+	FCriticalSection PackageHashesLock;
 	FString OutputPath;
 	FString MetadataDirectoryPath;
 	const ITargetPlatform& TargetPlatform;
@@ -155,5 +155,6 @@ private:
 	FPackageStoreManifest PackageStoreManifest;
 	const TArray<TSharedRef<IPlugin>>& PluginsToRemap;
 	FAsyncIODelete& AsyncIODelete;
-	bool bIterateSharedBuild;
+	bool bIterateSharedBuild = false;
+	bool bProvidePerPackageResults = false;
 };
diff --git a/Engine/Source/Editor/UnrealEd/Private/Cooker/PackageResultsMessage.cpp b/Engine/Source/Editor/UnrealEd/Private/Cooker/PackageResultsMessage.cpp
--- a/Engine/Source/Editor/UnrealEd/Private/Cooker/PackageResultsMessage.cpp
+++ b/Engine/Source/Editor/UnrealEd/Private/Cooker/PackageResultsMessage.cpp
@@ -3,6 +3,7 @@
 #include "PackageResultsMessage.h"
 
 #include "CookPackageData.h"
+#include "Misc/ScopeLock.h"
 
 namespace UE::Cook
 {
@@ -26,24 +27,94 @@ void FPackageRemoteResult::AddPackageMessage(const FGuid& MessageType, FCbObject
 	Messages.Add(UE::CompactBinaryTCP::FMarshalledMessage{ MessageType, MoveTemp(Object) });
 }
 
-void FPackageRemoteResult::AddPlatformMessage(const ITargetPlatform* TargetPlatform, const FGuid& MessageType, FCbObject&& Object)
+void FPackageRemoteResult::AddAsyncPackageMessage(const FGuid& MessageType, TFuture<FCbObject>&& ObjectFuture)
+{
+	FAsyncMessage& AsyncMessage = AsyncMessages.Emplace_GetRef();
+	AsyncMessage.MessageType = MessageType;
+	AsyncMessage.Future = MoveTemp(ObjectFuture);
+}
+
+void FPackageRemoteResult::AddPlatformMessage(const ITargetPlatform* TargetPlatform, const FGuid& MessageType,
+	FCbObject&& Object)
 {
 	check(TargetPlatform != nullptr);
-	int32 PlatformIndex = Platforms.IndexOfByPredicate([TargetPlatform](const FPlatformResult& Result)
+	FPlatformResult* Result = Platforms.FindByPredicate([TargetPlatform](const FPlatformResult& Result)
 		{ return Result.Platform == TargetPlatform; }
 	);
-	FPlatformResult* Result;
-	if (PlatformIndex != INDEX_NONE)
+	check(Result);
+	Result->Messages.Add(UE::CompactBinaryTCP::FMarshalledMessage{ MessageType, MoveTemp(Object) });
+}
+
+void FPackageRemoteResult::AddAsyncPlatformMessage(const ITargetPlatform* TargetPlatform, const FGuid& MessageType,
+	TFuture<FCbObject>&& ObjectFuture)
+{
+	check(TargetPlatform != nullptr);
+	FAsyncMessage& AsyncMessage = AsyncMessages.Emplace_GetRef();
+	AsyncMessage.MessageType = MessageType;
+	AsyncMessage.Future = MoveTemp(ObjectFuture);
+	AsyncMessage.TargetPlatform = TargetPlatform;
+}
+
+bool FPackageRemoteResult::IsComplete()
+{
+	FinalizeAsyncMessages();
+	FScopeLock AsyncWorkScopeLock(&AsyncSupport->AsyncWorkLock);
+	return bAsyncMessagesComplete;
+}
+
+TFuture<int> FPackageRemoteResult::GetCompletionFuture()
+{
+	FinalizeAsyncMessages();
+	return AsyncSupport->CompletionFuture.GetFuture();
+}
+
+void FPackageRemoteResult::FinalizeAsyncMessages()
+{
+	if (bAsyncMessagesFinalized)
 	{
-		Result = &Platforms[PlatformIndex];
+		return;
 	}
-	else
+	bAsyncMessagesFinalized = true;
+	AsyncSupport = MakeUnique<FAsyncSupport>();
+
+	if (AsyncMessages.IsEmpty())
 	{
-		Result = &Platforms.Emplace_GetRef();
-		Result->Platform = TargetPlatform;
+		bAsyncMessagesComplete = true;
+		AsyncSupport->CompletionFuture.EmplaceValue(0);
+		return;
 	}
 
-	Result->Messages.Add(UE::CompactBinaryTCP::FMarshalledMessage{ MessageType, MoveTemp(Object)});
+	// As of now we have no multithreading that accesses this. The first time it can occur is after calling the first Next.
+	NumIncompleteAsyncWork = AsyncMessages.Num();
+	for (FAsyncMessage& Message : AsyncMessages)
+	{
+		Message.Future.Next([this, Message = &Message](FCbObject Object)
+		{
+			bool bLocalAsyncMessagesComplete = false;
+			{
+				FScopeLock AsyncWorkScopeLock(&AsyncSupport->AsyncWorkLock);
+				check(!Message->bCompleted);
+				Message->bCompleted = true;
+				if (!Message->TargetPlatform)
+				{
+					AddPackageMessage(Message->MessageType, MoveTemp(Object));
+				}
+				else
+				{
+					AddPlatformMessage(Message->TargetPlatform, Message->MessageType, MoveTemp(Object));
+				}
+				if (--NumIncompleteAsyncWork == 0)
+				{
+					bAsyncMessagesComplete = true;
+					bLocalAsyncMessagesComplete = true;
+				}
+			}
+			if (bLocalAsyncMessagesComplete)
+			{
+				AsyncSupport->CompletionFuture.EmplaceValue(0);
+			}
+		});
+	}
 }
 
 void FPackageRemoteResult::SetPlatforms(TConstArrayView<ITargetPlatform*> OrderedSessionPlatforms)
@@ -72,8 +143,6 @@ void FPackageResultsMessage::Write(FCbWriter& Writer) const
 			{
 				Writer.BeginObject();
 				Writer << "S" << PlatformResult.bSuccessful;
-				Writer << "G" << PlatformResult.PackageGuid;
-				Writer << "D" << PlatformResult.TargetDomainDependencies;
 				WriteMessagesArray(Writer, PlatformResult.Messages);
 				Writer.EndObject();
 			}
@@ -130,8 +199,6 @@ bool FPackageResultsMessage::TryRead(FCbObjectView Object)
 			FCbObjectView PlatformObject = PlatformField.AsObjectView();
 
 			PlatformResult.bSuccessful = PlatformObject["S"].AsBool();
-			PlatformResult.PackageGuid = PlatformObject["G"].AsUuid();
-			PlatformResult.TargetDomainDependencies = FCbObject::Clone(PlatformObject["D"].AsObjectView());
 
 			if (!TryReadMessagesArray(PlatformObject, PlatformResult.Messages))
 			{
diff --git a/Engine/Source/Editor/UnrealEd/Private/Cooker/PackageResultsMessage.h b/Engine/Source/Editor/UnrealEd/Private/Cooker/PackageResultsMessage.h
--- a/Engine/Source/Editor/UnrealEd/Private/Cooker/PackageResultsMessage.h
+++ b/Engine/Source/Editor/UnrealEd/Private/Cooker/PackageResultsMessage.h
@@ -2,9 +2,11 @@
 
 #pragma once
 
+#include "Async/Future.h"
 #include "CompactBinaryTCP.h"
 #include "Containers/Array.h"
 #include "CookTypes.h"
+#include "HAL/CriticalSection.h"
 #include "Templates/UniquePtr.h"
 #include "UObject/NameTypes.h"
 
@@ -33,18 +35,10 @@ public:
 		TConstArrayView<UE::CompactBinaryTCP::FMarshalledMessage> GetMessages() const { return Messages; }
 		TArray<UE::CompactBinaryTCP::FMarshalledMessage> ReleaseMessages();
 
-		const FGuid& GetPackageGuid() const { return PackageGuid; }
-		void SetPackageGuid(const FGuid& InPackageGuid) { PackageGuid = InPackageGuid; }
-
-		FCbObjectView GetTargetDomainDependencies() const { return TargetDomainDependencies; }
-		void SetTargetDomainDependencies(FCbObject&& InObject) { TargetDomainDependencies = MoveTemp(InObject); }
-
 		bool IsSuccessful() const { return bSuccessful; }
 		void SetSuccessful(bool bInSuccessful) { bSuccessful = bInSuccessful; }
 
 	private:
-		FCbObject TargetDomainDependencies;
-		FGuid PackageGuid;
 		TArray<UE::CompactBinaryTCP::FMarshalledMessage> Messages;
 		const ITargetPlatform* Platform = nullptr;
 		bool bSuccessful = false;
@@ -53,6 +47,12 @@ public:
 		friend FPackageResultsMessage;
 	};
 
+	FPackageRemoteResult() = default;
+	FPackageRemoteResult(FPackageRemoteResult&&) = default;
+	FPackageRemoteResult(const FPackageRemoteResult&) = delete;
+	FPackageRemoteResult& operator=(FPackageRemoteResult&&) = default;
+	FPackageRemoteResult& operator=(const FPackageRemoteResult&) = delete;
+
 	FName GetPackageName() const { return PackageName; }
 	void SetPackageName(FName InPackageName) { PackageName = InPackageName; }
 
@@ -63,21 +63,76 @@ public:
 	void SetReferencedOnlyByEditorOnlyData(bool bInReferencedOnlyByEditorOnlyData) { bReferencedOnlyByEditorOnlyData = bInReferencedOnlyByEditorOnlyData; }
 
 	void AddPackageMessage(const FGuid& MessageType, FCbObject&& Object);
+	void AddAsyncPackageMessage(const FGuid& MessageType, TFuture<FCbObject>&& ObjectFuture);
 	void AddPlatformMessage(const ITargetPlatform* TargetPlatform, const FGuid& MessageType, FCbObject&& Object);
+	void AddAsyncPlatformMessage(const ITargetPlatform* TargetPlatform, const FGuid& MessageType, TFuture<FCbObject>&& ObjectFuture);
+
+	// GetMessages and ReleaseMessages are not thread-safe until IsComplete returns true or GetCompletionFuture().Get()/.Next().
 	TConstArrayView<UE::CompactBinaryTCP::FMarshalledMessage> GetMessages() const { return Messages; }
 	TArray<UE::CompactBinaryTCP::FMarshalledMessage> ReleaseMessages();
 
+	bool IsComplete();
+	TFuture<int> GetCompletionFuture();
+
 	TArray<FPlatformResult, TInlineAllocator<1>>& GetPlatforms() { return Platforms; }
 	void SetPlatforms(TConstArrayView<ITargetPlatform*> OrderedSessionPlatforms);
 
+	/**
+	 * A non-atomic RefCount that can be used for storage of a refcount by the user (e.g. CookWorkerClient)
+	 * If used from multiple threads, the user must access it only within the user's external critical section.
+	 */
+	int32& GetUserRefCount() { return UserRefCount; }
+
 private:
-	TArray<FPlatformResult, TInlineAllocator<1>> Platforms;
-	TArray<UE::CompactBinaryTCP::FMarshalledMessage> Messages;
+	/** A TFuture and status data that was received from an asynchronous IMPCollector. */
+	struct FAsyncMessage
+	{
+		FAsyncMessage() = default;
+		FAsyncMessage(FAsyncMessage&&) = default;
+		FAsyncMessage(const FAsyncMessage&) = delete;
+
+		FGuid MessageType;
+		TFuture<FCbObject> Future;
+		const ITargetPlatform* TargetPlatform = nullptr;
+		bool bCompleted = false;
+	};
+	/**
+	 * Some of the fields used when writing async messages on clients; these fields are otherwise unused.
+	 * These fields do not support Move construction or assignment, or memmove, so to support TArray
+	 * of FPackageRemoteResult we have to store these fields in a separate allocation.
+	 */
+	struct FAsyncSupport
+	{
+		TPromise<int> CompletionFuture;
+		FCriticalSection AsyncWorkLock;
+	};
+
+	/**
+	 * If any async messages have been stored, subscribe to their Futures to pull their resultant messages
+	 * and trigger this struct's ComplectionFuture when they are all done.
+	 */
+	void FinalizeAsyncMessages();
+
+private:
+	// Fields read/writable only from the owner thread.
+	TArray<FAsyncMessage> AsyncMessages;
 	FName PackageName;
 	/** If failure reason is InvalidSuppressCookReason, it was saved. Otherwise, holds the suppression reason */
 	ESuppressCookReason SuppressCookReason;
 	bool bReferencedOnlyByEditorOnlyData = false;
 
+	// Fields guarded by AsyncSupport->AsyncWorkLock. They can only be read or written if either AsyncSupport is nullptr
+	// or if within AsyncSupport->AsyncWorkLock.
+	TArray<FPlatformResult, TInlineAllocator<1>> Platforms;
+	TArray<UE::CompactBinaryTCP::FMarshalledMessage> Messages;
+	TUniquePtr<FAsyncSupport> AsyncSupport;
+	int32 NumIncompleteAsyncWork = 0;
+	bool bAsyncMessagesFinalized = false;
+	bool bAsyncMessagesComplete = false;
+
+	// Fields Read/Write only within an external critical section
+	int32 UserRefCount = 0;
+
 	friend FPackageResultsMessage;
 };
 
diff --git a/Engine/Source/Runtime/Core/Private/Misc/SecureHash.cpp b/Engine/Source/Runtime/Core/Private/Misc/SecureHash.cpp
--- a/Engine/Source/Runtime/Core/Private/Misc/SecureHash.cpp
+++ b/Engine/Source/Runtime/Core/Private/Misc/SecureHash.cpp
@@ -4,6 +4,7 @@
 
 #include "Async/AsyncWork.h"
 #include "HAL/FileManager.h"
+#include "Memory/MemoryView.h"
 #include "Misc/Paths.h"
 #include "Misc/StringBuilder.h"
 #include "Serialization/CompactBinaryWriter.h"
@@ -600,6 +601,28 @@ FMD5Hash FMD5Hash::HashFileFromArchive( FArchive* Ar, TArray<uint8>* Buffer)
 	return Hash;
 }
 
+FCbWriter& FMD5Hash::WriteCompactBinary(FCbWriter& Writer) const
+{
+	// TODO: Should we write bIsValid as well?
+	Writer.AddBinary(MakeMemoryView(Bytes, GetSize()));
+	return Writer;
+}
+
+bool LoadFromCompactBinary(FCbFieldView Field, FMD5Hash& OutHash)
+{
+	FMemoryView BinaryView = Field.AsBinaryView();
+	if (BinaryView.GetSize() != OutHash.GetSize())
+	{
+		OutHash.bIsValid = false;
+		return false;
+	}
+
+	FMemory::Memcpy(OutHash.Bytes, BinaryView.GetData(), OutHash.GetSize());
+	OutHash.bIsValid = true;
+	return true;
+}
+
+
 
 /*-----------------------------------------------------------------------------
 	SHA-1
diff --git a/Engine/Source/Runtime/Core/Public/Misc/SecureHash.h b/Engine/Source/Runtime/Core/Public/Misc/SecureHash.h
--- a/Engine/Source/Runtime/Core/Public/Misc/SecureHash.h
+++ b/Engine/Source/Runtime/Core/Public/Misc/SecureHash.h
@@ -188,6 +188,12 @@ private:
 	/** The bytes this hash comprises */
 	uint8 Bytes[16];
 
+	friend inline FCbWriter& operator<<(FCbWriter& Writer, const FMD5Hash& Hash) // Hidden friends must be inlined to be performant
+	{
+		return Hash.WriteCompactBinary(Writer);
+	}
+	CORE_API FCbWriter& WriteCompactBinary(FCbWriter& Writer) const;
+	friend CORE_API bool LoadFromCompactBinary(FCbFieldView Field, FMD5Hash& OutHash); // inlining not available because we don't want FCbFieldView defined
 	friend CORE_API FString LexToString(const FMD5Hash&);
 	friend CORE_API void LexFromString(FMD5Hash& Hash, const TCHAR*);
 };
diff --git a/Engine/Source/Runtime/Core/Public/Serialization/PackageWriter.h b/Engine/Source/Runtime/Core/Public/Serialization/PackageWriter.h
--- a/Engine/Source/Runtime/Core/Public/Serialization/PackageWriter.h
+++ b/Engine/Source/Runtime/Core/Public/Serialization/PackageWriter.h
@@ -194,14 +194,20 @@ public:
 ENUM_CLASS_FLAGS(IPackageWriter::EWriteOptions);
 
 /** Struct containing hashes computed during cooked package writing. */
-struct FPackageHashes : FRefCountBase
+struct FPackageHashes : FThreadSafeRefCountedObject
 {
-	// Hashes for each chunk saved by the package.
+	/** Hashes for each chunk saved by the package. */
 	TMap<FIoChunkId, FIoHash> ChunkHashes;
 
-	// This is a hash representing the entire package. Note this is
-	// not consistently computed across PackageWriters!
+	/** This is a hash representing the entire package. Not consistently computed across PackageWriters! */
 	FMD5Hash PackageHash;
+
+	/**
+	 * A Future that is triggered after all packages have been stored on *this.
+	 * Left as an invalid TFuture when hashes are not async; caller should check for IsValid before
+	 * chaining with .Then or .Next.
+	 */
+	TFuture<int> CompletionFuture;
 };
 
 /** Interface for cooking that writes cooked packages to storage usable by the runtime game. */
@@ -329,12 +335,12 @@ public:
 		return false;
 	}
 	/**
-	 * Create a CompactBinary Object message that replicates all of the package data from package save that is
-	 * collected in memory and written at end of cook rather than being written to disk during package save.
+	 * Asynchronously create a CompactBinary Object message that replicates all of the package data from package save
+	 * that is collected in memory and written at end of cook rather than being written to disk during package save.
 	 * Used during MPCook to transfer this information from CookWorker to CookDirector. Called after CommitPackage,
 	 * and only on CookWorkers.
 	 */
-	virtual FCbObject WriteMPCookMessageForPackage(FName PackageName) = 0;
+	virtual TFuture<FCbObject> WriteMPCookMessageForPackage(FName PackageName) = 0;
 
 	/** Read PackageData written by WriteMPCookMessageForPackage on a CookWorker. Called only on CookDirector. */
 	virtual bool TryReadMPCookMessageForPackage(FName PackageName, FCbObjectView Message) = 0;
