Cooker: Some function-level optimizations for CookByTheBook.
Results from  a 2.3hr cook session:
EvaluateGarbageCollectionResults: Skip SoftGC diagnostics if not running a memory diagnostic cook.
112s -> 0.05s

GetNeverCookPackageFileNames: Use FDirectoryTree instead of linear search.
20s -> 1.2s

AssignLayerChunkDelegate: Cache the ChunkLayerAssignment map from ini because the function can be called multiple times.
11.5s -> 0.01s

InitializeFromExistingAndPrune: Use FName WriteToString instead of ToString.
6.8s -> 6.8s (no change)

TJsonStringWriter::Close: Reserve the OutString to make only a single allocation
2.8s -> 2.2s

MakePathRelativeTo: Use TInlineAllocators and FStringView in the ParseIntoArray arrays, use TArrayView.RightChop instead of RemoveAt, presize the output string.
2.7s -> 1.5s

#rnx

diff --git a/Engine/Source/Editor/UnrealEd/Private/Commandlets/AssetRegistryGenerator.cpp b/Engine/Source/Editor/UnrealEd/Private/Commandlets/AssetRegistryGenerator.cpp
--- a/Engine/Source/Editor/UnrealEd/Private/Commandlets/AssetRegistryGenerator.cpp
+++ b/Engine/Source/Editor/UnrealEd/Private/Commandlets/AssetRegistryGenerator.cpp
@@ -368,15 +368,20 @@ static void AssignLayerChunkDelegate(const FAssignLayerChunkMap* ChunkManifest,
 {
 	OutChunkLayer = 0;
 
-	FConfigFile PlatformIniFile;
-	FConfigCacheIni::LoadLocalIniFile(PlatformIniFile, TEXT("Game"), true, *Platform);
-	TArray<FString> ChunkLayerAssignmentArray;
-	PlatformIniFile.GetArray(TEXT("/Script/UnrealEd.ProjectPackagingSettings"), TEXT("ChunkLayerAssignment"), ChunkLayerAssignmentArray);
+	static TMap<FString, TUniquePtr<TMap<int32, int32>>> PlatformChunkLayerAssignments;
+	TUniquePtr<TMap<int32, int32>>& ChunkLayerAssignment = PlatformChunkLayerAssignments.FindOrAdd(Platform);
+	if (!ChunkLayerAssignment)
+	{
+		FConfigFile PlatformIniFile;
+		FConfigCacheIni::LoadLocalIniFile(PlatformIniFile, TEXT("Game"), true, *Platform);
+		TArray<FString> ChunkLayerAssignmentArray;
+		PlatformIniFile.GetArray(TEXT("/Script/UnrealEd.ProjectPackagingSettings"), TEXT("ChunkLayerAssignment"), ChunkLayerAssignmentArray);
 
-	TMap<int32, int32> ChunkLayerAssignment;
-	ParseChunkLayerAssignment(ChunkLayerAssignmentArray, ChunkLayerAssignment);
+		ChunkLayerAssignment.Reset(new TMap<int32, int32>());
+		ParseChunkLayerAssignment(ChunkLayerAssignmentArray, *ChunkLayerAssignment);
+	}
 
-	int32* LayerId = ChunkLayerAssignment.Find(ChunkIndex);
+	int32* LayerId = ChunkLayerAssignment->Find(ChunkIndex);
 	if (LayerId)
 	{
 		OutChunkLayer = *LayerId;
diff --git a/Engine/Source/Editor/UnrealEd/Private/CookOnTheFlyServer.cpp b/Engine/Source/Editor/UnrealEd/Private/CookOnTheFlyServer.cpp
--- a/Engine/Source/Editor/UnrealEd/Private/CookOnTheFlyServer.cpp
+++ b/Engine/Source/Editor/UnrealEd/Private/CookOnTheFlyServer.cpp
@@ -17,6 +17,7 @@
 #include "Async/ParallelFor.h"
 #include "Commandlets/AssetRegistryGenerator.h"
 #include "Commandlets/ShaderPipelineCacheToolsCommandlet.h"
+#include "Containers/DirectoryTree.h"
 #include "Containers/RingBuffer.h"
 #include "Cooker/AsyncIODelete.h"
 #include "Cooker/CookDirector.h"
@@ -5028,42 +5029,39 @@ void UCookOnTheFlyServer::EvaluateGarbageCollectionResults(bool bWasDueToOOM, bo
 		// get loaded again.
 		PackageTracker->AddExpectedNeverLoadPackages(ExpectedFreedPackageNames);
 
-		// If some packages we expected to be freed were not freed, show the reference chains for why
-		// they were not freed.
-		TArray<UPackage*> PackagesReferencedOutsideOfCooker;
-		for (FWeakObjectPtr& WeakPtr : CookByTheBookOptions->SessionStartupObjects)
-		{
-			UObject* Object = WeakPtr.Get();
-			if (!Object)
-			{
-				continue;
-			}
-			UPackage* Package = Object->GetPackage();
-			ExpectedFreedPackageNames.Remove(Package->GetFName());
-		}
-		PackageTracker->ForEachLoadedPackage(
-			[this, &PackagesReferencedOutsideOfCooker](UPackage* Package)
-			{
-				if (ExpectedFreedPackageNames.Contains(Package->GetFName()))
-				{
-					PackagesReferencedOutsideOfCooker.Add(Package);
-				}
-			});
-		if (PackagesReferencedOutsideOfCooker.Num() > 0)
-		{
-			// Only show diagnostics if LLM is on. We could add a separate setting for this, but it's more convenient to combine it with the LLM enabled setting
+		// Only show diagnostics if LLM is on, because they are somewhat expensive. We could add a separate setting
+		// for this, but it's more convenient to combine it with the LLM enabled setting
 #if ENABLE_LOW_LEVEL_MEM_TRACKER
-			bool bShowDiagnostics = FLowLevelMemTracker::Get().IsEnabled();
+		bool bShowDiagnostics = FLowLevelMemTracker::Get().IsEnabled();
 #else
-			constexpr bool bShowDiagnostics = false;
+		constexpr bool bShowDiagnostics = false;
 #endif
-			if (bShowDiagnostics)
+		if (bShowDiagnostics)
+		{
+			// If some packages we expected to be freed were not freed, show the reference chains for why
+			// they were not freed.
+			TArray<UPackage*> PackagesReferencedOutsideOfCooker;
+			for (FWeakObjectPtr& WeakPtr : CookByTheBookOptions->SessionStartupObjects)
 			{
-				UE::Cook::DumpPackageReferencers(PackagesReferencedOutsideOfCooker);
+				UObject* Object = WeakPtr.Get();
+				if (!Object)
+				{
+					continue;
+				}
+				UPackage* Package = Object->GetPackage();
+				ExpectedFreedPackageNames.Remove(Package->GetFName());
 			}
-			else
+			PackageTracker->ForEachLoadedPackage(
+				[this, &PackagesReferencedOutsideOfCooker](UPackage* Package)
+				{
+					if (ExpectedFreedPackageNames.Contains(Package->GetFName()))
+					{
+						PackagesReferencedOutsideOfCooker.Add(Package);
+					}
+				});
+			if (PackagesReferencedOutsideOfCooker.Num() > 0)
 			{
-				GLog->Logf(TEXT("Soft Garbage Collect diagnostics are not displayed because llm is disabled. Run with -llm or -trace=memtag to see diagnostics."));
+				UE::Cook::DumpPackageReferencers(PackagesReferencedOutsideOfCooker);
 			}
 		}
 	}
@@ -10208,33 +10206,22 @@ TArray<FName> UCookOnTheFlyServer::GetNeverCookPackageFileNames(TArrayView<const
 	TArray<FString> NeverCookPackagesPaths;
 	if (AssetRegistry->IsSearchAllAssets() && !AssetRegistry->IsLoadingAssets())
 	{
-		TArray<FString> NeverCookDirectoryLongPackageNames;
-		NeverCookDirectoryLongPackageNames.Reserve(NeverCookDirectories.Num());
+		TDirectoryTree<int32> NeverCookDirectoryTree;
 		for (const FString& LocalDirectory : NeverCookDirectories)
 		{
 			FString PackagePath;
 			if (FPackageName::TryConvertFilenameToLongPackageName(LocalDirectory, PackagePath))
 			{
-				NeverCookDirectoryLongPackageNames.Add(PackagePath);
+				NeverCookDirectoryTree.FindOrAdd(PackagePath);
 			}
 		}
 
 		FString PackageNameStr;
 		AssetRegistry->EnumerateAllPackages(
-			[&NeverCookPackagesPaths, &NeverCookDirectoryLongPackageNames, &PackageNameStr](FName PackageName, const FAssetPackageData& PackageData)
+			[&NeverCookPackagesPaths, &NeverCookDirectoryTree, &PackageNameStr](FName PackageName, const FAssetPackageData& PackageData)
 			{
 				PackageName.ToString(PackageNameStr);
-				bool bIsInNeverCook = false;
-				for (const FString& NeverCookLongPackageName : NeverCookDirectoryLongPackageNames)
-				{
-					if (FPathViews::IsParentPathOf(NeverCookLongPackageName, PackageNameStr))
-					{
-						bIsInNeverCook = true;
-						break;
-					}
-				}
-
-				if (bIsInNeverCook)
+				if (NeverCookDirectoryTree.ContainsPathOrParent(PackageNameStr))
 				{
 					FString LocalFileName;
 					if (PackageData.Extension != EPackageExtension::Unspecified && PackageData.Extension != EPackageExtension::Custom)
diff --git a/Engine/Source/Runtime/AssetRegistry/Private/AssetRegistryState.cpp b/Engine/Source/Runtime/AssetRegistry/Private/AssetRegistryState.cpp
--- a/Engine/Source/Runtime/AssetRegistry/Private/AssetRegistryState.cpp
+++ b/Engine/Source/Runtime/AssetRegistry/Private/AssetRegistryState.cpp
@@ -16,6 +16,7 @@
 #include "Misc/FileHelper.h"
 #include "Misc/Paths.h"
 #include "Misc/PathViews.h"
+#include "Misc/StringBuilder.h"
 #include "NameTableArchive.h"
 #include "PackageReader.h"
 #include "Serialization/ArrayReader.h"
@@ -211,7 +212,7 @@ void FAssetRegistryState::InitializeFromExistingAndPrune(const FAssetRegistrySta
 			bRemoveAssetData = true;
 		}
 		else if (Options.bFilterAssetDataWithNoTags && AssetData->TagsAndValues.Num() == 0 &&
-			!FPackageName::IsLocalizedPackage(AssetData->PackageName.ToString()))
+			!FPackageName::IsLocalizedPackage(WriteToString<256>(AssetData->PackageName)))
 		{
 			bRemoveAssetData = true;
 			bRemoveDependencyData = Options.bFilterDependenciesWithNoTags;
@@ -245,7 +246,7 @@ void FAssetRegistryState::InitializeFromExistingAndPrune(const FAssetRegistrySta
 		{
 			// Only add if also in asset data map, or script package
 			if (CachedAssetsByPackageName.Find(Pair.Key) ||
-				FPackageName::IsScriptPackage(Pair.Key.ToString()))
+				FPackageName::IsScriptPackage(WriteToString<256>(Pair.Key)))
 			{
 				FAssetPackageData* NewData = CreateOrGetAssetPackageData(Pair.Key);
 				*NewData = *Pair.Value;
@@ -269,7 +270,7 @@ void FAssetRegistryState::InitializeFromExistingAndPrune(const FAssetRegistrySta
 		else if (Id.IsPackage() &&
 			!CachedAssetsByPackageName.Contains(Id.PackageName) &&
 			!RequiredDependNodePackages.Contains(Id.PackageName) &&
-			!FPackageName::IsScriptPackage(Id.PackageName.ToString()))
+			!FPackageName::IsScriptPackage(WriteToString<256>(Id.PackageName)))
 		{
 			bRemoveDependsNode = true;
 		}
@@ -383,7 +384,7 @@ void FAssetRegistryState::InitializeFromExisting(const FAssetDataMap& AssetDataM
 	{
 		for (const TPair<FName, FAssetPackageData*>& Pair : AssetPackageDataMap)
 		{
-			bool bIsScriptPackage = FPackageName::IsScriptPackage(Pair.Key.ToString());
+			bool bIsScriptPackage = FPackageName::IsScriptPackage(WriteToString<256>(Pair.Key));
 			if (InInitializationMode == EInitializationMode::OnlyUpdateNew && CachedPackageData.Find(Pair.Key))
 			{
 				continue;
@@ -504,7 +505,7 @@ void FAssetRegistryState::PruneAssetData(const TSet<FName>& RequiredPackages, co
 			bRemoveAssetData = true;
 		}
 		else if (Options.bFilterAssetDataWithNoTags && AssetData->TagsAndValues.Num() == 0 &&
-			!FPackageName::IsLocalizedPackage(AssetData->PackageName.ToString()))
+			!FPackageName::IsLocalizedPackage(WriteToString<256>(AssetData->PackageName)))
 		{
 			bRemoveAssetData = true;
 			bRemoveDependencyData = Options.bFilterDependenciesWithNoTags;
@@ -551,7 +552,7 @@ void FAssetRegistryState::PruneAssetData(const TSet<FName>& RequiredPackages, co
 		else if (Id.IsPackage() &&
 			!CachedAssetsByPackageName.Contains(Id.PackageName) &&
 			!RequiredDependNodePackages.Contains(Id.PackageName) &&
-			!FPackageName::IsScriptPackage(Id.PackageName.ToString()))
+			!FPackageName::IsScriptPackage(WriteToString<256>(Id.PackageName)))
 		{
 			bRemoveDependsNode = true;
 		}
diff --git a/Engine/Source/Runtime/Core/Private/Containers/String.cpp b/Engine/Source/Runtime/Core/Private/Containers/String.cpp
--- a/Engine/Source/Runtime/Core/Private/Containers/String.cpp
+++ b/Engine/Source/Runtime/Core/Private/Containers/String.cpp
@@ -18,6 +18,7 @@
 #include "Misc/VarArgs.h"
 #include "Serialization/Archive.h"
 #include "String/HexToBytes.h"
+#include "String/Parse.h"
 #include "Templates/MemoryOps.h"
 #include "Templates/RemoveReference.h"
 #include "Templates/UnrealTemplate.h"
@@ -1277,39 +1278,10 @@ bool FString::IsNumeric() const
 	return FCString::IsNumeric(Data.GetData());
 }
 
-/**
- * Breaks up a delimited string into elements of a string array.
- *
- * @param	InArray		The array to fill with the string pieces
- * @param	pchDelim	The string to delimit on
- * @param	InCullEmpty	If 1, empty strings are not added to the array
- *
- * @return	The number of elements in InArray
- */
 int32 FString::ParseIntoArray( TArray<FString>& OutArray, const TCHAR* pchDelim, const bool InCullEmpty ) const
 {
-	// Make sure the delimit string is not null or empty
 	check(pchDelim);
-	OutArray.Reset();
-	const TCHAR *Start = **this;
-	const int32 DelimLength = FCString::Strlen(pchDelim);
-	if (Start && *Start != TEXT('\0') && DelimLength)
-	{
-		while( const TCHAR *At = FCString::Strstr(Start,pchDelim) )
-		{
-			if (!InCullEmpty || At-Start)
-			{
-				OutArray.Emplace(UE_PTRDIFF_TO_INT32(At-Start),Start);
-			}
-			Start = At + DelimLength;
-		}
-		if (!InCullEmpty || *Start)
-		{
-			OutArray.Emplace(Start);
-		}
-
-	}
-	return OutArray.Num();
+	return UE::String::ParseIntoArray(OutArray, *this, pchDelim, InCullEmpty);
 }
 
 bool FString::MatchesWildcard(const TCHAR* InWildcard, int32 InWildcardLen, ESearchCase::Type SearchCase) const
diff --git a/Engine/Source/Runtime/Core/Private/Misc/Paths.cpp b/Engine/Source/Runtime/Core/Private/Misc/Paths.cpp
--- a/Engine/Source/Runtime/Core/Private/Misc/Paths.cpp
+++ b/Engine/Source/Runtime/Core/Private/Misc/Paths.cpp
@@ -2,22 +2,24 @@
 
 // Core includes.
 #include "Misc/Paths.h"
-#include "UObject/NameTypes.h"
-#include "Logging/LogMacros.h"
+
+#include "Containers/StringView.h"
 #include "HAL/FileManager.h"
 #include "HAL/PlatformFile.h"
-#include "Misc/Parse.h"
-#include "Misc/ScopeLock.h"
-#include "Misc/CommandLine.h"
-#include "Internationalization/Text.h"
 #include "Internationalization/Internationalization.h"
-#include "Misc/Guid.h"
-#include "Misc/ConfigCacheIni.h"
+#include "Internationalization/Text.h"
+#include "Logging/LogMacros.h"
 #include "Misc/App.h"
+#include "Misc/CommandLine.h"
+#include "Misc/ConfigCacheIni.h"
 #include "Misc/DataDrivenPlatformInfoRegistry.h"
 #include "Misc/EngineVersion.h"
+#include "Misc/Guid.h"
 #include "Misc/LazySingleton.h"
-#include "Containers/StringView.h"
+#include "Misc/Parse.h"
+#include "Misc/ScopeLock.h"
+#include "String/Parse.h"
+#include "UObject/NameTypes.h"
 
 DEFINE_LOG_CATEGORY_STATIC(LogPaths, Log, All);
 
@@ -1295,10 +1297,12 @@ bool FPaths::MakePathRelativeTo( FString& InPath, const TCHAR* InRelativeTo )
 	Source.ReplaceInline(TEXT("\\"), TEXT("/"), ESearchCase::CaseSensitive);
 	Target.ReplaceInline(TEXT("\\"), TEXT("/"), ESearchCase::CaseSensitive);
 
-	TArray<FString> TargetArray;
-	Target.ParseIntoArray(TargetArray, TEXT("/"), true);
-	TArray<FString> SourceArray;
-	Source.ParseIntoArray(SourceArray, TEXT("/"), true);
+	TArray<FStringView, TInlineAllocator<16>> TargetArrayBuffer;
+	UE::String::ParseIntoArray(TargetArrayBuffer, Target, TEXT("/"), true);
+	TArrayView<FStringView> TargetArray(TargetArrayBuffer);
+	TArray<FStringView, TInlineAllocator<16>> SourceArrayBuffer;
+	UE::String::ParseIntoArray(SourceArrayBuffer, Source, TEXT("/"), true);
+	TArrayView<FStringView> SourceArray(SourceArrayBuffer);
 
 	if (TargetArray.Num() && SourceArray.Num())
 	{
@@ -1315,20 +1319,30 @@ bool FPaths::MakePathRelativeTo( FString& InPath, const TCHAR* InRelativeTo )
 
 	while (TargetArray.Num() && SourceArray.Num() && TargetArray[0] == SourceArray[0])
 	{
-		TargetArray.RemoveAt(0);
-		SourceArray.RemoveAt(0);
+		TargetArray.RightChopInline(1);
+		SourceArray.RightChopInline(1);
 	}
+
+	FStringView ParentDirSeparator = TEXTVIEW("../");
+	FStringView DirSeparator = TEXTVIEW("/");
+	int32 ResultsLen = SourceArray.Num() * ParentDirSeparator.Len() + (TargetArray.Num() - 1)* DirSeparator.Len();
+	for (const FStringView& TargetDir : TargetArray)
+	{
+		ResultsLen += TargetDir.Len();
+	}
+
 	FString Result;
+	Result.Reserve(ResultsLen);
 	for (int32 Index = 0; Index < SourceArray.Num(); Index++)
 	{
-		Result += TEXT("../");
+		Result += ParentDirSeparator;
 	}
 	for (int32 Index = 0; Index < TargetArray.Num(); Index++)
 	{
 		Result += TargetArray[Index];
 		if (Index + 1 < TargetArray.Num())
 		{
-			Result += TEXT("/");
+			Result += DirSeparator;
 		}
 	}
 	
diff --git a/Engine/Source/Runtime/Core/Public/String/Parse.h b/Engine/Source/Runtime/Core/Public/String/Parse.h
new file mode 100644
--- /dev/null
+++ b/Engine/Source/Runtime/Core/Public/String/Parse.h
@@ -0,0 +1,55 @@
+// Copyright Epic Games, Inc. All Rights Reserved.
+
+#pragma once
+
+#include "Containers/Array.h"
+#include "Containers/StringView.h"
+#include "Containers/UnrealString.h"
+#include "Misc/AssertionMacros.h"
+#include "Misc/CoreMiscDefines.h"
+#include "Misc/CString.h"
+
+namespace UE::String
+{
+
+/**
+ * Breaks up a delimited string into elements of a string array.
+ *
+ * @param	InArray		The array to fill with the string pieces. Can be an Array of Strings or StringViews.
+ * @param	Text		The input string to parse.
+ * @param	pchDelim	The string to delimit on. (Legacy behavior, may change): If empty, output array is empty.
+ * @param	InCullEmpty	If 1, empty strings are not added to the array
+ *
+ * @return	The number of elements in InArray
+ */
+template <typename OutElementType, typename AllocatorType>
+int32 ParseIntoArray(TArray<OutElementType, AllocatorType>& OutArray, FStringView Text, FStringView pchDelim, const bool InCullEmpty)
+{
+	OutArray.Reset();
+	const TCHAR* Start = Text.GetData();
+	const TCHAR* TextEnd = Text.GetData() + Text.Len();
+	const int32 DelimLength = pchDelim.Len();
+	int32 RemainingLen = UE_PTRDIFF_TO_INT32(TextEnd - Start);
+	// TODO Legacy behavior: if DelimLength is 0 we return an empty array. Can we change this to return { Text }?
+	if (Start && Start != TextEnd && DelimLength)
+	{
+		while (const TCHAR *At = FCString::Strnstr(Start, RemainingLen,	pchDelim.GetData(), DelimLength))
+		{
+			if (!InCullEmpty || At-Start)
+			{
+				OutArray.Emplace(FStringView(Start, UE_PTRDIFF_TO_INT32(At-Start)));
+			}
+			Start = At + DelimLength;
+			RemainingLen = UE_PTRDIFF_TO_INT32(TextEnd - Start);
+		}
+		if (!InCullEmpty || Start != TextEnd)
+		{
+			OutArray.Emplace(FStringView(Start, RemainingLen));
+		}
+
+	}
+	return OutArray.Num();
+}
+
+
+} // UE::String
diff --git a/Engine/Source/Runtime/Json/Public/Serialization/JsonWriter.h b/Engine/Source/Runtime/Json/Public/Serialization/JsonWriter.h
--- a/Engine/Source/Runtime/Json/Public/Serialization/JsonWriter.h
+++ b/Engine/Source/Runtime/Json/Public/Serialization/JsonWriter.h
@@ -579,16 +579,13 @@ public:
 
 	virtual bool Close() override
 	{
-		FString Out;
-
+		OutString->Reset(Bytes.Num()/sizeof(TCHAR));
 		for (int32 i = 0; i < Bytes.Num(); i+=sizeof(TCHAR))
 		{
 			TCHAR* Char = static_cast<TCHAR*>(static_cast<void*>(&Bytes[i]));
-			Out += *Char;
+			*OutString += *Char;
 		}
 
-		*OutString = Out;
-
 		return TJsonWriter<TCHAR, PrintPolicy>::Close();
 	}
 
