[UnrealEd]
* CookWorkerClient.cpp - Worked around c++20 compile error

diff --git a/Engine/Source/Editor/UnrealEd/Private/Cooker/CookWorkerClient.cpp b/Engine/Source/Editor/UnrealEd/Private/Cooker/CookWorkerClient.cpp
--- a/Engine/Source/Editor/UnrealEd/Private/Cooker/CookWorkerClient.cpp
+++ b/Engine/Source/Editor/UnrealEd/Private/Cooker/CookWorkerClient.cpp
@@ -266,8 +266,10 @@ void FCookWorkerClient::ReportPromoteToSaveComplete(FPackageData& PackageData)
 		}
 		else
 		{
-			PendingResultsNeedingAsyncWork.Add(Result,
-				FPendingResultNeedingAsyncWork{ MoveTemp(ResultOwner), MoveTemp(CompletionFuture) });
+			FPendingResultNeedingAsyncWork Work;
+			Work.PendingResult = MoveTemp(ResultOwner);
+			Work.CompletionFuture = MoveTemp(CompletionFuture);
+			PendingResultsNeedingAsyncWork.Add(Result, MoveTemp(Work));
 		}
 	}
 }
