Create a system to control which extension importer should support.

#rnx

diff --git a/Engine/Plugins/Enterprise/DatasmithImporter/Source/DatasmithExternalSource/Private/DatasmithFileUriResolver.cpp b/Engine/Plugins/Enterprise/DatasmithImporter/Source/DatasmithExternalSource/Private/DatasmithFileUriResolver.cpp
--- a/Engine/Plugins/Enterprise/DatasmithImporter/Source/DatasmithExternalSource/Private/DatasmithFileUriResolver.cpp
+++ b/Engine/Plugins/Enterprise/DatasmithImporter/Source/DatasmithExternalSource/Private/DatasmithFileUriResolver.cpp
@@ -7,6 +7,7 @@
 #include "DatasmithTranslatorManager.h"
 
 #if WITH_EDITOR
+#include "AssetToolsModule.h"
 #include "DesktopPlatformModule.h"
 #include "Framework/Application/SlateApplication.h"
 #include "IDesktopPlatform.h"
@@ -49,7 +50,8 @@ namespace UE::DatasmithImporter
 		FString FileTypes;
 		FString Extensions;
 
-		ObjectTools::AppendFormatsFileExtensions(Formats, FileTypes, Extensions);
+		IAssetTools& AssetTools = FModuleManager::LoadModuleChecked<FAssetToolsModule>(TEXT("AssetTools")).Get();
+		ObjectTools::AppendFormatsFileExtensions(Formats, FileTypes, Extensions, AssetTools.GetSupportedImportExtension());
 
 		const FString FormatString = FString::Printf(TEXT("All Files (%s)|%s|%s"), *Extensions, *Extensions, *FileTypes);
 		const FString Title = NSLOCTEXT("DatasmithFileUriResolver", "BrowseSourceDialogTitle", "Select Datasmith Source File").ToString();
diff --git a/Engine/Plugins/Enterprise/DatasmithImporter/Source/DatasmithImporter/Private/DatasmithImporterHelper.cpp b/Engine/Plugins/Enterprise/DatasmithImporter/Source/DatasmithImporter/Private/DatasmithImporterHelper.cpp
--- a/Engine/Plugins/Enterprise/DatasmithImporter/Source/DatasmithImporter/Private/DatasmithImporterHelper.cpp
+++ b/Engine/Plugins/Enterprise/DatasmithImporter/Source/DatasmithImporter/Private/DatasmithImporterHelper.cpp
@@ -5,6 +5,7 @@
 #include "Utility/DatasmithImporterUtils.h"
 
 #include "AssetRegistry/AssetRegistryModule.h"
+#include "AssetToolsModule.h"
 #include "BusyCursor.h"
 #include "ContentBrowserModule.h"
 #include "DesktopPlatformModule.h"
@@ -253,7 +254,8 @@ FString FDatasmithImporterHelper::GetFilterStringInternal(UFactory* Factory)
 	FString AllExtensions;
 	TMultiMap<uint32, UFactory*> FilterIndexToFactory;
 
-	ObjectTools::GenerateFactoryFileExtensions(Factories, FileTypes, AllExtensions, FilterIndexToFactory);
+	IAssetTools& AssetTools = FModuleManager::LoadModuleChecked<FAssetToolsModule>(TEXT("AssetTools")).Get();
+	ObjectTools::GenerateFactoryFileExtensions(Factories, FileTypes, AllExtensions, FilterIndexToFactory, AssetTools.GetSupportedImportExtension());
 
 	return FString::Printf(TEXT("All Files (%s)|%s|%s"), *AllExtensions, *AllExtensions, *FileTypes);
 }
diff --git a/Engine/Plugins/Experimental/Enterprise/DatasmithInterchange/Source/DatasmithInterchange/Private/DatasmithInterchangeModule.cpp b/Engine/Plugins/Experimental/Enterprise/DatasmithInterchange/Source/DatasmithInterchange/Private/DatasmithInterchangeModule.cpp
--- a/Engine/Plugins/Experimental/Enterprise/DatasmithInterchange/Source/DatasmithInterchange/Private/DatasmithInterchangeModule.cpp
+++ b/Engine/Plugins/Experimental/Enterprise/DatasmithInterchange/Source/DatasmithInterchange/Private/DatasmithInterchangeModule.cpp
@@ -24,6 +24,7 @@
 #include "Logging/LogMacros.h"
 
 #if WITH_EDITOR
+#include "AssetToolsModule.h"
 #include "DatasmithImporterModule.h"
 #include "DatasmithUtils.h"
 #include "DesktopPlatformModule.h"
@@ -132,8 +133,8 @@ void FDatasmithInterchangeModule::OnImportInterchange()
 	FString Extensions;
 
 	FString DefaultLocation(FEditorDirectories::Get().GetLastDirectory(ELastDirectory::GENERIC_IMPORT));
-
-	ObjectTools::AppendFormatsFileExtensions(Formats, FileTypes, Extensions);
+	IAssetTools& AssetTools = FModuleManager::LoadModuleChecked<FAssetToolsModule>(TEXT("AssetTools")).Get();
+	ObjectTools::AppendFormatsFileExtensions(Formats, FileTypes, Extensions, AssetTools.GetSupportedImportExtension());
 
 	const FString FormatString = FString::Printf(TEXT("All Files (%s)|%s|%s"), *Extensions, *Extensions, *FileTypes);
 	const FString Title = LOCTEXT("BrowseSourceDialogTitle", "Select Datasmith Source File").ToString();
diff --git a/Engine/Plugins/Interchange/Editor/Source/Utilities/Private/InterchangeOpenFileDialog.cpp b/Engine/Plugins/Interchange/Editor/Source/Utilities/Private/InterchangeOpenFileDialog.cpp
--- a/Engine/Plugins/Interchange/Editor/Source/Utilities/Private/InterchangeOpenFileDialog.cpp
+++ b/Engine/Plugins/Interchange/Editor/Source/Utilities/Private/InterchangeOpenFileDialog.cpp
@@ -1,6 +1,7 @@
 // Copyright Epic Games, Inc. All Rights Reserved. 
 #include "InterchangeOpenFileDialog.h"
 
+#include "AssetToolsModule.h"
 #include "DesktopPlatformModule.h"
 #include "Framework/Application/SlateApplication.h"
 #include "IDesktopPlatform.h"
@@ -16,7 +17,8 @@ namespace UE::Interchange::Utilities::Private
 		FString FileTypes;
 		FString Extensions;
 
-		ObjectTools::AppendFormatsFileExtensions(TranslatorFormats, FileTypes, Extensions);
+		IAssetTools& AssetTools = FModuleManager::LoadModuleChecked<FAssetToolsModule>(TEXT("AssetTools")).Get();
+		ObjectTools::AppendFormatsFileExtensions(TranslatorFormats, FileTypes, Extensions, AssetTools.GetSupportedImportExtension());
 
 		const FString FormatString = FString::Printf(TEXT("All Files (%s)|%s|%s"), *Extensions, *Extensions, *FileTypes);
 
diff --git a/Engine/Source/Developer/AssetTools/Private/AssetTools.cpp b/Engine/Source/Developer/AssetTools/Private/AssetTools.cpp
--- a/Engine/Source/Developer/AssetTools/Private/AssetTools.cpp
+++ b/Engine/Source/Developer/AssetTools/Private/AssetTools.cpp
@@ -1287,6 +1287,9 @@ UAssetToolsImpl::UAssetToolsImpl(const FObjectInitializer& ObjectInitializer)
 		}
 	}
 
+	
+	GConfig->GetArray(TEXT("AssetTools"), TEXT("SupportedImportExtensions"), SupportedImportExtensionArray, GEditorIni);
+
 	TArray<FString> DenyListedViewPath;
 	GConfig->GetArray(TEXT("AssetTools"), TEXT("DenyListAssetPaths"), DenyListedViewPath, GEditorIni);
 	for (const FString& Path : DenyListedViewPath)
@@ -5591,6 +5594,11 @@ const TSharedRef<FPathPermissionList>& UAssetToolsImpl::GetAssetClassPathPermiss
 	return Empty;
 }
 
+const TArray<FString>& UAssetToolsImpl::GetSupportedImportExtension() const
+{
+	return SupportedImportExtensionArray;
+}
+
 TSet<EBlueprintType>& UAssetToolsImpl::GetAllowedBlueprintTypes()
 {
 	return AllowedBlueprintTypes;
@@ -5637,14 +5645,13 @@ TArray<UObject*> UAssetToolsImpl::ImportAssetsWithDialogImplementation(const FSt
 	}
 
 	TMultiMap<uint32, UFactory*> FilterIndexToFactory;
-
 	// Generate the file types and extensions represented by the selected factories
-	ObjectTools::GenerateFactoryFileExtensions(Factories, FileTypes, AllExtensions, FilterIndexToFactory);
+	ObjectTools::GenerateFactoryFileExtensions(Factories, FileTypes, AllExtensions, FilterIndexToFactory, GetSupportedImportExtension());
 
 	if (UInterchangeManager::IsInterchangeImportEnabled())
 	{
 		TArray<FString> InterchangeFileExtensions = UInterchangeManager::GetInterchangeManager().GetSupportedFormats(EInterchangeTranslatorType::Assets);
-		ObjectTools::AppendFormatsFileExtensions(InterchangeFileExtensions, FileTypes, AllExtensions, FilterIndexToFactory);
+		ObjectTools::AppendFormatsFileExtensions(InterchangeFileExtensions, FileTypes, AllExtensions, FilterIndexToFactory, GetSupportedImportExtension());
 	}
 
 	FileTypes = FString::Printf(TEXT("All Files (%s)|%s|%s"), *AllExtensions, *AllExtensions, *FileTypes);
diff --git a/Engine/Source/Developer/AssetTools/Private/AssetTools.h b/Engine/Source/Developer/AssetTools/Private/AssetTools.h
--- a/Engine/Source/Developer/AssetTools/Private/AssetTools.h
+++ b/Engine/Source/Developer/AssetTools/Private/AssetTools.h
@@ -135,6 +135,7 @@ public:
 	UE_DEPRECATED(5.1, "Class names are now represented by path names. Please use GetAssetClassPathPermissionList.")
 	TSharedRef<FNamePermissionList>& GetAssetClassPermissionList(EAssetClassAction AssetClassAction);
 	virtual const TSharedRef<FPathPermissionList>& GetAssetClassPathPermissionList(EAssetClassAction AssetClassAction) const override;
+	virtual const TArray<FString>& GetSupportedImportExtension() const override;
 	virtual TSet<EBlueprintType>& GetAllowedBlueprintTypes() override;
 	virtual TSharedRef<FPathPermissionList>& GetFolderPermissionList() override;
 	virtual TSharedRef<FPathPermissionList>& GetWritableFolderPermissionList() override;
@@ -238,6 +239,8 @@ private:
 	/** Permission lists of assets by class path name, one for each EAssetClassAction */
 	TArray<TSharedRef<FPathPermissionList>> AssetClassPermissionList;
 
+	TArray<FString> SupportedImportExtensionArray;
+
 	/** Which types of BlueprintFactories are allowed to be returned through GetNewAssetFactories. Empty set allows all types. */
 	TSet<EBlueprintType> AllowedBlueprintTypes;
 
diff --git a/Engine/Source/Developer/AssetTools/Public/IAssetTools.h b/Engine/Source/Developer/AssetTools/Public/IAssetTools.h
--- a/Engine/Source/Developer/AssetTools/Public/IAssetTools.h
+++ b/Engine/Source/Developer/AssetTools/Public/IAssetTools.h
@@ -570,6 +570,9 @@ public:
 	/** Get asset class permission list for content browser and other systems */
 	virtual const TSharedRef<FPathPermissionList>& GetAssetClassPathPermissionList(EAssetClassAction AssetClassAction) const = 0;
 
+	/** Get extension list allowed for importer */
+	virtual const TArray<FString>& GetSupportedImportExtension() const = 0;
+
 	/** Which BlueprintTypes are allowed to be created. An empty list should allow everything. */
 	virtual TSet<EBlueprintType>& GetAllowedBlueprintTypes() = 0;
 
diff --git a/Engine/Source/Editor/UnrealEd/Private/Editor.cpp b/Engine/Source/Editor/UnrealEd/Private/Editor.cpp
--- a/Engine/Source/Editor/UnrealEd/Private/Editor.cpp
+++ b/Engine/Source/Editor/UnrealEd/Private/Editor.cpp
@@ -779,19 +779,20 @@ void FReimportManager::GetNewReimportPath(UObject* Obj, TArray<FString>& InOutFi
 
 	TMultiMap<uint32, UFactory*> DummyFilterIndexToFactory;
 
+	IAssetTools& AssetTools = FModuleManager::LoadModuleChecked<FAssetToolsModule>(TEXT("AssetTools")).Get();
 	//Append the interchange supported translator formats for this object
 	if (UInterchangeManager::IsInterchangeImportEnabled())
 	{
 		//Get the extension interchange can translate for this object
 		TArray<FString> TranslatorFormats = UInterchangeManager::GetInterchangeManager().GetSupportedFormatsForObject(Obj);
-		ObjectTools::AppendFormatsFileExtensions(TranslatorFormats, FileTypes, AllExtensions, DummyFilterIndexToFactory);
+		ObjectTools::AppendFormatsFileExtensions(TranslatorFormats, FileTypes, AllExtensions, DummyFilterIndexToFactory, AssetTools.GetSupportedImportExtension());
 	}
 
 	//If this object was not import with interchange add the legacy formats
 	if(AllExtensions.IsEmpty())
 	{
 		// Generate the file types and extensions represented by the selected factories
-		ObjectTools::GenerateFactoryFileExtensions(Factories, FileTypes, AllExtensions, DummyFilterIndexToFactory);
+		ObjectTools::GenerateFactoryFileExtensions(Factories, FileTypes, AllExtensions, DummyFilterIndexToFactory, AssetTools.GetSupportedImportExtension());
 	}
 
 	
diff --git a/Engine/Source/Editor/UnrealEd/Private/FileHelpers.cpp b/Engine/Source/Editor/UnrealEd/Private/FileHelpers.cpp
--- a/Engine/Source/Editor/UnrealEd/Private/FileHelpers.cpp
+++ b/Engine/Source/Editor/UnrealEd/Private/FileHelpers.cpp
@@ -492,13 +492,13 @@ FString FEditorFileUtils::GetFilterString(EFileInteraction Interaction)
 				FString FileTypes;
 				FString AllExtensions;
 				TMultiMap<uint32, UFactory*> FilterIndexToFactory;
-
-				ObjectTools::GenerateFactoryFileExtensions(Factories, FileTypes, AllExtensions, FilterIndexToFactory);
+				IAssetTools& AssetTools = FModuleManager::LoadModuleChecked<FAssetToolsModule>(TEXT("AssetTools")).Get();
+				ObjectTools::GenerateFactoryFileExtensions(Factories, FileTypes, AllExtensions, FilterIndexToFactory, AssetTools.GetSupportedImportExtension());
 
 				if (UInterchangeManager::IsInterchangeImportEnabled())
 				{
 					TArray<FString> InterchangeFileExtensions = UInterchangeManager::GetInterchangeManager().GetSupportedFormats(EInterchangeTranslatorType::Scenes);
-					ObjectTools::AppendFormatsFileExtensions(InterchangeFileExtensions, FileTypes, AllExtensions);
+					ObjectTools::AppendFormatsFileExtensions(InterchangeFileExtensions, FileTypes, AllExtensions, AssetTools.GetSupportedImportExtension());
 				}
 
 				FileTypes = FString::Printf(TEXT("All Files (%s)|%s|%s"), *AllExtensions, *AllExtensions, *FileTypes);
diff --git a/Engine/Source/Editor/UnrealEd/Private/ObjectTools.cpp b/Engine/Source/Editor/UnrealEd/Private/ObjectTools.cpp
--- a/Engine/Source/Editor/UnrealEd/Private/ObjectTools.cpp
+++ b/Engine/Source/Editor/UnrealEd/Private/ObjectTools.cpp
@@ -4329,7 +4329,10 @@ namespace ObjectTools
 	 * @param	out_Descriptions	Array of format descriptions associated with the current factory; should equal the number of extensions
 	 * @param	out_Extensions		Array of format extensions associated with the current factory; should equal the number of descriptions
 	 */
-	void InternalGetFormatInfo(const TArray<FString>& Formats, TArray<FString>& out_Descriptions, TArray<FString>& out_Extensions )
+	void InternalGetFormatInfo(const TArray<FString>& Formats
+		, TArray<FString>& out_Descriptions
+		, TArray<FString>& out_Extensions
+		, const TArray<FString>& SupportedExtensions)
 	{
 		// Iterate over each formats.
 		for ( TArray<FString>::TConstIterator FormatIter( Formats ); FormatIter; ++FormatIter )
@@ -4343,6 +4346,18 @@ namespace ObjectTools
 			for ( int32 ComponentIndex = 0; ComponentIndex < FormatComponents.Num(); ComponentIndex += 2 )
 			{
 				check( FormatComponents.IsValidIndex( ComponentIndex + 1 ) );
+
+				//If we have a valid SupportedExtensions array, skip any extension not in the SupportedExtension array.
+				if (SupportedExtensions.Num() > 0)
+				{
+					FString& RefExtension = FormatComponents[ComponentIndex];
+					if (!SupportedExtensions.FindByPredicate([&RefExtension](const FString& Item){return Item.Equals(RefExtension, ESearchCase::IgnoreCase);}))
+					{
+						//Skip this extension
+						continue;
+					}
+				}
+
 				out_Extensions.Add( FormatComponents[ComponentIndex] );
 				out_Descriptions.Add( FormatComponents[ComponentIndex + 1] );
 			}
@@ -4356,12 +4371,16 @@ namespace ObjectTools
 	 * @param	out_Filetypes	File types supported by the provided factory, concatenated into a string
 	 * @param	out_Extensions	Extensions supported by the provided factory, concatenated into a string
 	 */
-	void GenerateFactoryFileExtensions( UFactory* InFactory, FString& out_Filetypes, FString& out_Extensions, TMultiMap<uint32, UFactory*>& out_FilterIndexToFactory )
+	void GenerateFactoryFileExtensions( UFactory* InFactory
+		, FString& out_Filetypes
+		, FString& out_Extensions
+		, TMultiMap<uint32, UFactory*>& out_FilterIndexToFactory
+		, const TArray<FString>& SupportedExtensions)
 	{
 		// Place the factory in an array and call the overloaded version of this function
 		TArray<UFactory*> FactoryArray;
 		FactoryArray.Add( InFactory );
-		GenerateFactoryFileExtensions( FactoryArray, out_Filetypes, out_Extensions, out_FilterIndexToFactory );
+		GenerateFactoryFileExtensions( FactoryArray, out_Filetypes, out_Extensions, out_FilterIndexToFactory, SupportedExtensions );
 	}
 
 	/**
@@ -4371,7 +4390,11 @@ namespace ObjectTools
 	 * @param	out_Filetypes	File types supported by the provided factory, concatenated into a string
 	 * @param	out_Extensions	Extensions supported by the provided factory, concatenated into a string
 	 */
-	void GenerateFactoryFileExtensions( const TArray<UFactory*>& InFactories, FString& out_Filetypes, FString& out_Extensions, TMultiMap<uint32, UFactory*>& out_FilterIndexToFactory )
+	void GenerateFactoryFileExtensions( const TArray<UFactory*>& InFactories
+		, FString& out_Filetypes
+		, FString& out_Extensions
+		, TMultiMap<uint32, UFactory*>& out_FilterIndexToFactory
+		, const TArray<FString>& SupportedExtensions)
 	{
 		// Store all the descriptions and their corresponding extensions in a map
 		TMultiMap<FString, FString> DescToExtensionMap;
@@ -4385,7 +4408,7 @@ namespace ObjectTools
 
 			TArray<FString> Descriptions;
 			TArray<FString> Extensions;
-			InternalGetFormatInfo( CurFactory->GetFormats(), Descriptions, Extensions);
+			InternalGetFormatInfo( CurFactory->GetFormats(), Descriptions, Extensions, SupportedExtensions);
 			check( Descriptions.Num() == Extensions.Num() );
 
 			// Make sure to only store each key, value pair once
@@ -4497,19 +4520,26 @@ namespace ObjectTools
 		}
 	}
 
-	void AppendFormatsFileExtensions(const TArray<FString>& InFormats, FString& out_FileTypes, FString& out_Extensions)
+	void AppendFormatsFileExtensions(const TArray<FString>& InFormats
+		, FString& out_FileTypes
+		, FString& out_Extensions
+		, const TArray<FString>& SupportedExtensions)
 	{
 		TArray<FString> Descriptions;
 		TArray<FString> Extensions;
-		InternalGetFormatInfo(InFormats, Descriptions, Extensions);
+		InternalGetFormatInfo(InFormats, Descriptions, Extensions, SupportedExtensions);
 		InternalAppendFileExtensions(Descriptions, Extensions, out_FileTypes, out_Extensions);
 	}
 
-	void AppendFormatsFileExtensions(const TArray<FString>& InFormats, FString& out_FileTypes, FString& out_Extensions, TMultiMap<uint32, UFactory*>& out_FilterIndexToFactory)
+	void AppendFormatsFileExtensions(const TArray<FString>& InFormats
+		, FString& out_FileTypes
+		, FString& out_Extensions
+		, TMultiMap<uint32, UFactory*>& out_FilterIndexToFactory
+		, const TArray<FString>& SupportedExtensions)
 	{
 		TArray<FString> Descriptions;
 		TArray<FString> Extensions;
-		InternalGetFormatInfo(InFormats, Descriptions, Extensions);
+		InternalGetFormatInfo(InFormats, Descriptions, Extensions, SupportedExtensions);
 		InternalAppendFileExtensions(Descriptions, Extensions, out_FileTypes, out_Extensions);
 		uint32 MaxKeyNumber = 0;
 		TSet<uint32> Keys;
@@ -4531,12 +4561,15 @@ namespace ObjectTools
 	/**
 	 * Generates a list of file types for a given class.
 	 */
-	void AppendFactoryFileExtensions ( UFactory* InFactory, FString& out_Filetypes, FString& out_Extensions )
+	void AppendFactoryFileExtensions ( UFactory* InFactory
+		, FString& out_Filetypes
+		, FString& out_Extensions
+		, const TArray<FString>& SupportedExtensions)
 	{
 		check(InFactory);
 		TArray<FString> Descriptions;
 		TArray<FString> Extensions;
-		InternalGetFormatInfo( InFactory->GetFormats(), Descriptions, Extensions);
+		InternalGetFormatInfo( InFactory->GetFormats(), Descriptions, Extensions, SupportedExtensions);
 		InternalAppendFileExtensions( Descriptions, Extensions, out_Filetypes, out_Extensions);
 	}
 
diff --git a/Engine/Source/Editor/UnrealEd/Public/ObjectTools.h b/Engine/Source/Editor/UnrealEd/Public/ObjectTools.h
--- a/Engine/Source/Editor/UnrealEd/Public/ObjectTools.h
+++ b/Engine/Source/Editor/UnrealEd/Public/ObjectTools.h
@@ -488,8 +488,13 @@ namespace ObjectTools
 	 * @param	InFactory		Factory whose supported file types and extensions should be retrieved
 	 * @param	out_Filetypes	File types supported by the provided factory, concatenated into a string
 	 * @param	out_Extensions	Extensions supported by the provided factory, concatenated into a string
+	 * @param   SupportedExtensions			If not null only extension in the list can be added
 	 */
-	UNREALED_API void GenerateFactoryFileExtensions( const UFactory* InFactory, FString& out_Filetypes, FString& out_Extensions, TMultiMap<uint32, UFactory*>& out_FilterIndexToFactory );
+	UNREALED_API void GenerateFactoryFileExtensions( const UFactory* InFactory
+		, FString& out_Filetypes
+		, FString& out_Extensions
+		, TMultiMap<uint32, UFactory*>& out_FilterIndexToFactory
+		, const TArray<FString>& SupportedExtensions = TArray<FString>());
 
 	/**
 	 * Populates two strings with all of the file types and extensions the provided factories support.
@@ -497,13 +502,22 @@ namespace ObjectTools
 	 * @param	InFactories		Factories whose supported file types and extensions should be retrieved
 	 * @param	out_Filetypes	File types supported by the provided factory, concatenated into a string
 	 * @param	out_Extensions	Extensions supported by the provided factory, concatenated into a string
+	 * @param   SupportedExtensions			If not null only extension in the list can be added
 	 */
-	UNREALED_API void GenerateFactoryFileExtensions( const TArray<UFactory*>& InFactories, FString& out_Filetypes, FString& out_Extensions, TMultiMap<uint32, UFactory*>& out_FilterIndexToFactory );
+	UNREALED_API void GenerateFactoryFileExtensions( const TArray<UFactory*>& InFactories
+		, FString& out_Filetypes
+		, FString& out_Extensions
+		, TMultiMap<uint32, UFactory*>& out_FilterIndexToFactory
+		, const TArray<FString>& SupportedExtensions = TArray<FString>());
 
 	/**
 	 * Generates a list of file types for a given class.
+	 * @param   SupportedExtensions			If not null only extension in the list can be added
 	 */
-	UNREALED_API void AppendFactoryFileExtensions( UFactory* InFactory, FString& out_Filetypes, FString& out_Extensions );
+	UNREALED_API void AppendFactoryFileExtensions( UFactory* InFactory
+		, FString& out_Filetypes
+		, FString& out_Extensions
+		, const TArray<FString>& SupportedExtensions = TArray<FString>());
 
 	/**
 	 * Populates two strings with all of the file types and extensions the format list provides.
@@ -511,8 +525,12 @@ namespace ObjectTools
 	 * @param	InFormats		Array of supported file types. Each entry needs to be of the form "ext;Description" where ext is the file extension. 
 	 * @param	out_FileTypes	File types supported by the provided array of formats, concatenated into a string
 	 * @param	out_Extensions	Extensions supported by the provided array of formats, concatenated into a string
+	 * @param   SupportedExtensions			If not null only extension in the list can be added
 	 */
-	UNREALED_API void AppendFormatsFileExtensions(const TArray<FString>& InFormats, FString& out_FileTypes, FString& out_Extensions);
+	UNREALED_API void AppendFormatsFileExtensions(const TArray<FString>& InFormats
+		, FString& out_FileTypes
+		, FString& out_Extensions
+		, const TArray<FString>& SupportedExtensions = TArray<FString>());
 
 	/**
 	 * Populates two strings with all of the file types and extensions the format list provides.
@@ -521,8 +539,13 @@ namespace ObjectTools
 	 * @param	out_FileTypes	File types supported by the provided array of formats, concatenated into a string
 	 * @param	out_Extensions	Extensions supported by the provided array of formats, concatenated into a string
 	 * @param	out_FilterIndexToFactory	Add INDEX_NONE entry for all provided Formats
+	 * @param   SupportedExtensions			If not null only extension in the list can be added
 	 */
-	UNREALED_API void AppendFormatsFileExtensions(const TArray<FString>& InFormats, FString& out_FileTypes, FString& out_Extensions, TMultiMap<uint32, UFactory*>& out_FilterIndexToFactory);
+	UNREALED_API void AppendFormatsFileExtensions(const TArray<FString>& InFormats
+		, FString& out_FileTypes
+		, FString& out_Extensions
+		, TMultiMap<uint32, UFactory*>& out_FilterIndexToFactory
+		, const TArray<FString>& SupportedExtensions = TArray<FString>());
 
 	/**
 	 * Iterates over all classes and assembles a list of non-abstract UExport-derived type instances.
