UEditorAssetSubsystem::ListAssets: Fixed Reserve() being called on the wrong array

diff --git a/Engine/Source/Editor/UnrealEd/Private/Subsystems/EditorAssetSubsystem.cpp b/Engine/Source/Editor/UnrealEd/Private/Subsystems/EditorAssetSubsystem.cpp
--- a/Engine/Source/Editor/UnrealEd/Private/Subsystems/EditorAssetSubsystem.cpp
+++ b/Engine/Source/Editor/UnrealEd/Private/Subsystems/EditorAssetSubsystem.cpp
@@ -1435,7 +1435,7 @@ TArray<FString> UEditorAssetSubsystem::ListAssets(const FString& DirectoryPath,
 
 	if (AssetDatas.Num() > 0)
 	{
-		AssetDatas.Reserve(AssetDatas.Num());
+		AssetPaths.Reserve(AssetDatas.Num());
 		for (const FAssetData& AssetData : AssetDatas)
 		{
 			AssetPaths.Add(AssetData.GetObjectPathString());
