Add option to not show the class view in ClassPickerDialog.

diff --git a/Engine/Source/Editor/ClassViewer/Public/ClassViewerModule.h b/Engine/Source/Editor/ClassViewer/Public/ClassViewerModule.h
--- a/Engine/Source/Editor/ClassViewer/Public/ClassViewerModule.h
+++ b/Engine/Source/Editor/ClassViewer/Public/ClassViewerModule.h
@@ -142,6 +142,9 @@ public:
 	/** (true) Will show the default classes if they exist. */
 	bool bShowDefaultClasses;
 
+	/** (true) Will show the classes viewer. */
+	bool bShowClassesViewer;
+
 public:
 
 	/** Constructor */
@@ -162,6 +165,7 @@ public:
 		, bEditorClassesOnly(false)
 		, InitiallySelectedClass(nullptr)
 		, bShowDefaultClasses(true)
+		, bShowClassesViewer(true)
 	{
 	}
 };
diff --git a/Engine/Source/Editor/UnrealEd/Private/SClassPickerDialog.cpp b/Engine/Source/Editor/UnrealEd/Private/SClassPickerDialog.cpp
--- a/Engine/Source/Editor/UnrealEd/Private/SClassPickerDialog.cpp
+++ b/Engine/Source/Editor/UnrealEd/Private/SClassPickerDialog.cpp
@@ -39,16 +39,18 @@ void SClassPickerDialog::Construct(const FArguments& InArgs)
 	bPressedOk = false;
 	ChosenClass = NULL;
 
-	TSharedPtr<SListView<TSharedPtr<FClassPickerDefaults>>> DefaultClassViewer;
-
-	ClassViewer = StaticCastSharedRef<SClassViewer>(FModuleManager::LoadModuleChecked<FClassViewerModule>("ClassViewer").CreateClassViewer(InArgs._Options, FOnClassPicked::CreateSP(this,&SClassPickerDialog::OnClassPicked)));
-
 	FClassViewerModule& ClassViewerModule = FModuleManager::LoadModuleChecked<FClassViewerModule>("ClassViewer");
-	const TSharedPtr<IClassViewerFilter>& GlobalClassFilter = ClassViewerModule.GetGlobalClassViewerFilter();
-	TSharedRef<FClassViewerFilterFuncs> ClassFilterFuncs = ClassViewerModule.CreateFilterFuncs();
+
+	if (InArgs._Options.bShowClassesViewer)
+	{
+		ClassViewer = StaticCastSharedRef<SClassViewer>(ClassViewerModule.CreateClassViewer(InArgs._Options, FOnClassPicked::CreateSP(this, &SClassPickerDialog::OnClassPicked)));
+	}
 
 	if (InArgs._Options.bShowDefaultClasses)
 	{
+		const TSharedPtr<IClassViewerFilter>& GlobalClassFilter = ClassViewerModule.GetGlobalClassViewerFilter();
+		TSharedRef<FClassViewerFilterFuncs> ClassFilterFuncs = ClassViewerModule.CreateFilterFuncs();
+
 		// Load in default settings
 		for (const FClassPickerDefaults& DefaultObj : GUnrealEd->GetUnrealEdOptions()->GetNewAssetDefaultClasses())
 		{
@@ -96,15 +98,108 @@ void SClassPickerDialog::Construct(const FArguments& InArgs)
 
 	const bool bHasDefaultClasses = AssetDefaultClasses.Num() > 0;
 
+	TSharedPtr<SListView<TSharedPtr<FClassPickerDefaults>>> DefaultClassViewer;
+	if (bHasDefaultClasses)
+	{
+		SAssignNew(DefaultClassViewer, SListView < TSharedPtr<FClassPickerDefaults> >)
+			.ItemHeight(24)
+			.SelectionMode(ESelectionMode::None)
+			.ListItemsSource(&AssetDefaultClasses)
+			.OnGenerateRow(this, &SClassPickerDialog::GenerateListRow);
+	}
+
 	bool bExpandDefaultClassPicker = true;
 	bool bExpandCustomClassPicker = !bHasDefaultClasses;
-
-	if (bHasDefaultClasses)
+	if (bHasDefaultClasses && ClassViewer)
 	{
 		GConfig->GetBool(TEXT("/Script/UnrealEd.UnrealEdOptions"), TEXT("bExpandClassPickerDefaultClassList"), bExpandDefaultClassPicker, GEditorIni);
 		GConfig->GetBool(TEXT("/Script/UnrealEd.UnrealEdOptions"), TEXT("bExpandCustomClassPickerClassList"), bExpandCustomClassPicker, GEditorIni);
 	}
 
+	ensureMsgf(ClassViewer || DefaultClassViewer, TEXT("Either ClassViewer or DefaultClassViewer should be used."));
+
+	TSharedRef<SVerticalBox> Container = SNew(SVerticalBox);
+
+	if (DefaultClassViewer)
+	{
+		Container->AddSlot()
+		.AutoHeight()
+		[
+			SNew(SExpandableArea)
+			.BorderImage(FStyleDefaults::GetNoBrush())
+			.BodyBorderImage(FAppStyle::Get().GetBrush("Brushes.Recessed"))
+			.HeaderPadding(FMargin(5.0f, 3.0f))
+			.AllowAnimatedTransition(false)
+			.InitiallyCollapsed(!bExpandDefaultClassPicker)
+			.OnAreaExpansionChanged(this, &SClassPickerDialog::OnDefaultAreaExpansionChanged)
+			.HeaderContent()
+			[
+				SNew(STextBlock)
+				.Text(NSLOCTEXT("SClassPickerDialog", "CommonClassesAreaTitle", "Common"))
+				.TextStyle(FAppStyle::Get(), "ButtonText")
+				.Font(FAppStyle::Get().GetFontStyle("NormalFontBold"))
+				.TransformPolicy(ETextTransformPolicy::ToUpper)
+			]
+			.BodyContent()
+			[
+				DefaultClassViewer.ToSharedRef()
+			]
+		];
+	}
+
+	if (ClassViewer)
+	{
+		FMargin Padding = DefaultClassViewer ? FMargin(0.0f, 10.0f, 0.0f, 0.0f) : FMargin();
+		Container->AddSlot()
+		.AutoHeight()
+		.Padding(Padding)
+		[
+			SNew(SExpandableArea)
+			.BorderImage(FStyleDefaults::GetNoBrush())
+			.BodyBorderImage(FAppStyle::Get().GetBrush("Brushes.Recessed"))
+			.HeaderPadding(FMargin(5.0f, 3.0f))
+			.AllowAnimatedTransition(false)
+			.MaxHeight(320.f)
+			.InitiallyCollapsed(!bExpandCustomClassPicker)
+			.OnAreaExpansionChanged(this, &SClassPickerDialog::OnCustomAreaExpansionChanged)
+			.HeaderContent()
+			[
+				SNew(STextBlock)
+				.Text(NSLOCTEXT("SClassPickerDialog", "AllClassesAreaTitle", "All Classes"))
+				.TextStyle(FAppStyle::Get(), "ButtonText")
+				.Font(FAppStyle::Get().GetFontStyle("NormalFontBold"))
+				.TransformPolicy(ETextTransformPolicy::ToUpper)
+			]
+			.BodyContent()
+			[
+				ClassViewer.ToSharedRef()
+			]
+		];
+	}
+
+	Container->AddSlot()
+	.HAlign(HAlign_Right)
+	.VAlign(VAlign_Bottom)
+	.Padding(8.f)
+	[
+		SNew(SUniformGridPanel)
+		.SlotPadding(FAppStyle::GetMargin("StandardDialog.SlotPadding"))
+		+SUniformGridPanel::Slot(0,0)
+		[
+			SNew(SPrimaryButton)
+			.Text(NSLOCTEXT("SClassPickerDialog", "ClassPickerSelectButton", "Select"))
+			.Visibility( this, &SClassPickerDialog::GetSelectButtonVisibility )
+			.OnClicked(this, &SClassPickerDialog::OnClassPickerConfirmed)
+		]
+		+SUniformGridPanel::Slot(1,0)
+		[
+			SNew(SButton)
+			.Text(NSLOCTEXT("SClassPickerDialog", "ClassPickerCancelButton", "Cancel"))
+			.HAlign(HAlign_Center)
+			.OnClicked(this, &SClassPickerDialog::OnClassPickerCanceled)
+		]
+	];
+
 	ChildSlot
 	[
 		SNew(SBorder)
@@ -113,93 +208,18 @@ void SClassPickerDialog::Construct(const FArguments& InArgs)
 			SNew(SBox)
 			.WidthOverride(610.0f)
 			[
-				SNew(SVerticalBox)
-				+SVerticalBox::Slot()
-				.AutoHeight()
-				[
-					SNew(SExpandableArea)
-					.BorderImage(FStyleDefaults::GetNoBrush())
-					.BodyBorderImage(FAppStyle::Get().GetBrush("Brushes.Recessed"))
-					.HeaderPadding(FMargin(5.0f, 3.0f))
-					.AllowAnimatedTransition(false)
-					.InitiallyCollapsed(!bExpandDefaultClassPicker)
-					.OnAreaExpansionChanged(this, &SClassPickerDialog::OnDefaultAreaExpansionChanged)
-					.HeaderContent()
-					[
-						SNew(STextBlock)
-						.Text(NSLOCTEXT("SClassPickerDialog", "CommonClassesAreaTitle", "Common"))
-						.TextStyle(FAppStyle::Get(), "ButtonText")
-						.Font(FAppStyle::Get().GetFontStyle("NormalFontBold"))
-						.TransformPolicy(ETextTransformPolicy::ToUpper)
-					]
-					.BodyContent()
-					[
-						SAssignNew(DefaultClassViewer, SListView < TSharedPtr<FClassPickerDefaults> >)
-						.ItemHeight(24)
-						.SelectionMode(ESelectionMode::None)
-						.ListItemsSource(&AssetDefaultClasses)
-						.OnGenerateRow(this, &SClassPickerDialog::GenerateListRow)
-						.Visibility(bHasDefaultClasses? EVisibility::Visible: EVisibility::Collapsed)
-					]
-				]
-				+SVerticalBox::Slot()
-				.AutoHeight()
-				.Padding(0.0f, 10.0f, 0.0f, 0.0f)
-				[
-					SNew(SExpandableArea)
-					.BorderImage(FStyleDefaults::GetNoBrush())
-					.BodyBorderImage(FAppStyle::Get().GetBrush("Brushes.Recessed"))
-					.HeaderPadding(FMargin(5.0f, 3.0f))
-					.AllowAnimatedTransition(false)
-					.MaxHeight(320.f)
-					.InitiallyCollapsed(!bExpandCustomClassPicker)
-					.OnAreaExpansionChanged(this, &SClassPickerDialog::OnCustomAreaExpansionChanged)
-					.HeaderContent()
-					[
-						SNew(STextBlock)
-						.Text(NSLOCTEXT("SClassPickerDialog", "AllClassesAreaTitle", "All Classes"))
-						.TextStyle(FAppStyle::Get(), "ButtonText")
-						.Font(FAppStyle::Get().GetFontStyle("NormalFontBold"))
-						.TransformPolicy(ETextTransformPolicy::ToUpper)
-					]
-					.BodyContent()
-					[
-						ClassViewer.ToSharedRef()
-					]
-				]
-				+SVerticalBox::Slot()
-				.HAlign(HAlign_Right)
-				.VAlign(VAlign_Bottom)
-				.Padding(8.0f)
-				[
-					SNew(SUniformGridPanel)
-					.SlotPadding(FAppStyle::GetMargin("StandardDialog.SlotPadding"))
-					+SUniformGridPanel::Slot(0,0)
-					[
-						SNew(SPrimaryButton)
-						.Text(NSLOCTEXT("SClassPickerDialog", "ClassPickerSelectButton", "Select"))
-						.Visibility( this, &SClassPickerDialog::GetSelectButtonVisibility )
-						.OnClicked(this, &SClassPickerDialog::OnClassPickerConfirmed)
-					]
-					+SUniformGridPanel::Slot(1,0)
-					[
-						SNew(SButton)
-						.Text(NSLOCTEXT("SClassPickerDialog", "ClassPickerCancelButton", "Cancel"))
-						.HAlign(HAlign_Center)
-						.OnClicked(this, &SClassPickerDialog::OnClassPickerCanceled)
-					]
-				]
+				Container
 			]
 		]
 	];
 
 	if (WeakParentWindow.IsValid())
 	{
-		if (bExpandCustomClassPicker)
+		if (bExpandCustomClassPicker && ClassViewer)
 		{
 			WeakParentWindow.Pin().Get()->SetWidgetToFocusOnActivate(ClassViewer);
 		}
-		else
+		else if (DefaultClassViewer)
 		{
 			WeakParentWindow.Pin().Get()->SetWidgetToFocusOnActivate(DefaultClassViewer);
 		}
@@ -344,7 +364,7 @@ FReply SClassPickerDialog::OnClassPickerCanceled()
 
 void SClassPickerDialog::OnDefaultAreaExpansionChanged(bool bExpanded)
 {
-	if (bExpanded && WeakParentWindow.IsValid())
+	if (bExpanded && WeakParentWindow.IsValid() && ClassViewer)
 	{
 		WeakParentWindow.Pin().Get()->SetWidgetToFocusOnActivate(ClassViewer);
 	}
@@ -357,6 +377,7 @@ void SClassPickerDialog::OnDefaultAreaExpansionChanged(bool bExpanded)
 
 void SClassPickerDialog::OnCustomAreaExpansionChanged(bool bExpanded)
 {
+	check(ClassViewer);
 	if (bExpanded && WeakParentWindow.IsValid())
 	{
 		WeakParentWindow.Pin().Get()->SetWidgetToFocusOnActivate(ClassViewer);
@@ -381,7 +402,10 @@ EVisibility SClassPickerDialog::GetSelectButtonVisibility() const
 /** Overridden from SWidget: Called when a key is pressed down - capturing copy */
 FReply SClassPickerDialog::OnKeyDown( const FGeometry& MyGeometry, const FKeyEvent& InKeyEvent )
 {
-	WeakParentWindow.Pin().Get()->SetWidgetToFocusOnActivate(ClassViewer);
+	if (ClassViewer)
+	{
+		WeakParentWindow.Pin().Get()->SetWidgetToFocusOnActivate(ClassViewer);
+	}
 
 	if (InKeyEvent.GetKey() == EKeys::Escape)
 	{
@@ -391,7 +415,7 @@ FReply SClassPickerDialog::OnKeyDown( const FGeometry& MyGeometry, const FKeyEve
 	{
 		OnClassPickerConfirmed();
 	}
-	else
+	else if (ClassViewer)
 	{
 		return ClassViewer->OnKeyDown(MyGeometry, InKeyEvent);
 	}
