PropertyAccessEditor: Fix non-unity compile

diff --git a/Engine/Source/Editor/UnrealEd/Public/IPropertyAccessEditor.h b/Engine/Source/Editor/UnrealEd/Public/IPropertyAccessEditor.h
--- a/Engine/Source/Editor/UnrealEd/Public/IPropertyAccessEditor.h
+++ b/Engine/Source/Editor/UnrealEd/Public/IPropertyAccessEditor.h
@@ -15,6 +15,7 @@ class UEdGraph;
 class FExtender;
 class SWidget;
 struct FSlateBrush;
+struct FButtonStyle;
 struct FEdGraphPinType;
 class IPropertyAccessLibraryCompiler;
 struct FPropertyAccessLibrary;
