MPCook: Allow Cook::IMPCollectors to also append messages onto demoted (failed cook) packages. The PackageManifest has an entry for non-cooked packages.
Downgrade failed to respond to retraction request message to Display. It can occur when the CookWorker is busy with long package loads due to huge maps.
Remove TODO comments that are now TODONE.
#rnx

diff --git a/Engine/Source/Developer/IoStoreUtilities/Private/PackageStoreManifest.cpp b/Engine/Source/Developer/IoStoreUtilities/Private/PackageStoreManifest.cpp
--- a/Engine/Source/Developer/IoStoreUtilities/Private/PackageStoreManifest.cpp
+++ b/Engine/Source/Developer/IoStoreUtilities/Private/PackageStoreManifest.cpp
@@ -390,9 +390,7 @@ void FPackageStoreManifest::WritePackage(FCbWriter& Writer, FName PackageName)
 		PackageFileChunkIds.RemoveAndCopyValue(PackageName, FileChunkIds);
 	}
 
-	UE_CLOG(!bHasPackageInfo, LogIoStore, Error,
-		TEXT("ClientTickPackage was called for package %s, but CommitPackage was not called for that package. We will not have up to date PackageStoreManifest data."),
-		*PackageName.ToString());
+	// For a failed package, CommitPackage may have never been called. Send empty values in that case.
 
 	Writer.BeginObject();
 	Writer << "ExportBundleChunkIds" << PackageInfo.ExportBundleChunkIds;
diff --git a/Engine/Source/Editor/UnrealEd/Private/Commandlets/AssetRegistryGenerator.cpp b/Engine/Source/Editor/UnrealEd/Private/Commandlets/AssetRegistryGenerator.cpp
--- a/Engine/Source/Editor/UnrealEd/Private/Commandlets/AssetRegistryGenerator.cpp
+++ b/Engine/Source/Editor/UnrealEd/Private/Commandlets/AssetRegistryGenerator.cpp
@@ -2746,17 +2746,24 @@ FAssetRegistryMPCollector::FAssetRegistryMPCollector(UCookOnTheFlyServer& InCOTF
 void FAssetRegistryMPCollector::ClientTickPackage(FMPCollectorClientTickPackageContext& Context)
 {
 	bool bLoggedWarning = false;
-	for (const ITargetPlatform* TargetPlatform : Context.GetPlatforms())
+	for (const FMPCollectorClientTickPackageContext::FPlatformData& ContextPlatformData : Context.GetPlatformDatas())
 	{
+		const ITargetPlatform* TargetPlatform = ContextPlatformData.TargetPlatform;
 		FPlatformData* PlatformData = COTFS.PlatformManager->GetPlatformData(TargetPlatform);
 		FAssetRegistryReporterRemote& RegistryReporter = static_cast<FAssetRegistryReporterRemote&>(*PlatformData->RegistryReporter);
 		FCbObject Message;
+
 		if (!RegistryReporter.PackageUpdateMessages.RemoveAndCopyValue(Context.GetPackageName(), Message))
 		{
-			UE_CLOG(!bLoggedWarning, LogAssetRegistryGenerator, Warning,
-				TEXT("ClientTickPackage was called for package %s, but UpdateAssetRegistryData was not called for that package. We will not have up to date AssetDataTags in the generated AssetRegistry."),
-				*Context.GetPackageName().ToString());
-			bLoggedWarning = true;
+			// For a failed package, UpdateAssetRegistryData will never have been called. Silently skip it and do not send a message.
+			if (ContextPlatformData.bSuccessful)
+			{
+				// For a successful package, UpdateAssetRegistryData should have been called
+				UE_CLOG(!bLoggedWarning, LogAssetRegistryGenerator, Warning,
+					TEXT("ClientTickPackage was called for package %s, but UpdateAssetRegistryData was not called for that package. We will not have up to date AssetDataTags in the generated AssetRegistry."),
+					*Context.GetPackageName().ToString());
+				bLoggedWarning = true;
+			}
 		}
 		else
 		{
diff --git a/Engine/Source/Editor/UnrealEd/Private/CookOnTheFlyServer.cpp b/Engine/Source/Editor/UnrealEd/Private/CookOnTheFlyServer.cpp
--- a/Engine/Source/Editor/UnrealEd/Private/CookOnTheFlyServer.cpp
+++ b/Engine/Source/Editor/UnrealEd/Private/CookOnTheFlyServer.cpp
@@ -8495,7 +8495,6 @@ void UCookOnTheFlyServer::CookByTheBookFinished()
 				// Add the package hashes to the relevant AssetPackageDatas.
 				// PackageHashes are gated by requiring UPackage::WaitForAsyncFileWrites(), which is called above.
 				FCookSavePackageContext& SaveContext = FindOrCreateSaveContext(TargetPlatform);
-				// MPCOOKTODO: Need to replicate AllPackageHashes from CookWorkers
 				TMap<FName, TRefCountPtr<FPackageHashes>>& AllPackageHashes = SaveContext.PackageWriter->GetPackageHashes();
 				for (TPair<FName, TRefCountPtr<FPackageHashes>>& HashSet : AllPackageHashes)
 				{
diff --git a/Engine/Source/Editor/UnrealEd/Private/Cooker/CookDirector.cpp b/Engine/Source/Editor/UnrealEd/Private/Cooker/CookDirector.cpp
--- a/Engine/Source/Editor/UnrealEd/Private/Cooker/CookDirector.cpp
+++ b/Engine/Source/Editor/UnrealEd/Private/Cooker/CookDirector.cpp
@@ -404,8 +404,6 @@ void FCookDirector::AssignRequests(TArray<FWorkerId>&& InWorkers, TArray<TRefCou
 		}
 	}
 
-	// MPCOOKTODO: Sort each batch from leaf to root
-
 	// Assign each batch to the FCookWorkerServer in RemoteWorkers;
 	// the CookWorkerServer's tick will handle sending the message to the remote process
 	for (FWorkerId WorkerId : InWorkers)
@@ -1539,7 +1537,7 @@ void FCookDirector::FRetractionHandler::TickFromSchedulerThread(bool bAnyIdle, b
 			}
 		}
 		bOutComplete = false;
-		UE_CLOG(!IsCookIgnoreTimeouts(), LogCook, Warning, TEXT("%s has not responded to a RetractionRequest message for %.1f seconds. Continuing to wait..."),
+		UE_CLOG(!IsCookIgnoreTimeouts(), LogCook, Display, TEXT("%s has not responded to a RetractionRequest message for %.1f seconds. Continuing to wait..."),
 			*Director.GetDisplayName(ExpectedWorker), static_cast<float>(CurrentTime - MessageSentTimeSeconds));
 		LastWarnTimeSeconds = CurrentTime;
 		return;
diff --git a/Engine/Source/Editor/UnrealEd/Private/Cooker/CookWorkerClient.cpp b/Engine/Source/Editor/UnrealEd/Private/Cooker/CookWorkerClient.cpp
--- a/Engine/Source/Editor/UnrealEd/Private/Cooker/CookWorkerClient.cpp
+++ b/Engine/Source/Editor/UnrealEd/Private/Cooker/CookWorkerClient.cpp
@@ -165,12 +165,14 @@ void FCookWorkerClient::ReportDemoteToIdle(const FPackageData& PackageData, ESup
 	{
 		return;
 	}
-	TUniquePtr<FPackageRemoteResult> Result(new FPackageRemoteResult());
-	Result->SetPackageName(PackageData.GetPackageName());
-	Result->SetSuppressCookReason(Reason);
+	TUniquePtr<FPackageRemoteResult> ResultOwner(new FPackageRemoteResult());
+	FName PackageName = PackageData.GetPackageName();
+	ResultOwner->SetPackageName(PackageName);
+	ResultOwner->SetSuppressCookReason(Reason);
+	// Set the platforms, use the default values for each platform (e.g. bSuccessful=false)
+	ResultOwner->SetPlatforms(OrderedSessionPlatforms);
 
-	FScopeLock PendingResultsScopeLock(&PendingResultsLock);
-	PendingResults.Add(MoveTemp(Result));
+	ReportPackageMessage(PackageName, MoveTemp(ResultOwner));
 }
 
 void FCookWorkerClient::ReportPromoteToSaveComplete(FPackageData& PackageData)
@@ -192,8 +194,15 @@ void FCookWorkerClient::ReportPromoteToSaveComplete(FPackageData& PackageData)
 		PlatformResults.SetSuccessful(PackagePlatformData.bCookSucceeded);
 	}
 
+	ReportPackageMessage(PackageName, MoveTemp(ResultOwner));
+}
+
+void FCookWorkerClient::ReportPackageMessage(FName PackageName, TUniquePtr<FPackageRemoteResult>&& ResultOwner)
+{
+	FPackageRemoteResult* Result = ResultOwner.Get();
+
 	TArray<FMPCollectorClientTickPackageContext::FPlatformData, TInlineAllocator<1>> ContextPlatformDatas;
-	ContextPlatformDatas.Reserve(NumPlatforms);
+	ContextPlatformDatas.Reserve(Result->GetPlatforms().Num());
 	for (FPackageRemoteResult::FPlatformResult& PlatformResult : Result->GetPlatforms())
 	{
 		ContextPlatformDatas.Add(FMPCollectorClientTickPackageContext::FPlatformData
diff --git a/Engine/Source/Editor/UnrealEd/Private/Cooker/CookWorkerClient.h b/Engine/Source/Editor/UnrealEd/Private/Cooker/CookWorkerClient.h
--- a/Engine/Source/Editor/UnrealEd/Private/Cooker/CookWorkerClient.h
+++ b/Engine/Source/Editor/UnrealEd/Private/Cooker/CookWorkerClient.h
@@ -111,6 +111,9 @@ private:
 	void AssignPackages(FAssignPackagesMessage& Message);
 	/** Tick the registered collectors. */
 	void TickCollectors(FTickStackData& StackData, bool bFlush);
+	/** Helper for ReportDemote/ReportPromote: Collect IMPCollectors and asynchronously add the message to pending. */
+	void ReportPackageMessage(FName PackageName, TUniquePtr<FPackageRemoteResult>&& ResultOwner);
+	
 	void HandleAbortPackagesMessage(FMPCollectorClientMessageContext& Context, bool bReadSuccessful,
 		FAbortPackagesMessage&& Message);
 	void HandleRetractionMessage(FMPCollectorClientMessageContext& Context, bool bReadSuccessful,
diff --git a/Engine/Source/Editor/UnrealEd/Private/Cooker/CookWorkerServer.cpp b/Engine/Source/Editor/UnrealEd/Private/Cooker/CookWorkerServer.cpp
--- a/Engine/Source/Editor/UnrealEd/Private/Cooker/CookWorkerServer.cpp
+++ b/Engine/Source/Editor/UnrealEd/Private/Cooker/CookWorkerServer.cpp
@@ -655,7 +655,6 @@ void FCookWorkerServer::RecordResults(FPackageResultsMessage& Message)
 				ITargetPlatform* TargetPlatform = OrderedSessionPlatforms[PlatformIndex];
 				FPackageRemoteResult::FPlatformResult& PlatformResult = Result.GetPlatforms()[PlatformIndex];
 				PackageData->SetPlatformCooked(TargetPlatform, PlatformResult.IsSuccessful());
-				// MPCOOKTODO: Call CommitRemotePackage on the PackageWriter
 				HandleReceivedPackagePlatformMessages(*PackageData, TargetPlatform, PlatformResult.ReleaseMessages());
 			}
 			if (Result.IsReferencedOnlyByEditorOnlyData())
@@ -1070,12 +1069,9 @@ void FPackageWriterMPCollector::ClientTickPackage(FMPCollectorClientTickPackageC
 {
 	for (const FMPCollectorClientTickPackageContext::FPlatformData& PlatformData : Context.GetPlatformDatas())
 	{
-		if (PlatformData.bSuccessful)
-		{
-			ICookedPackageWriter& PackageWriter = COTFS.FindOrCreatePackageWriter(PlatformData.TargetPlatform);
-			TFuture<FCbObject> ObjectFuture = PackageWriter.WriteMPCookMessageForPackage(Context.GetPackageName());
-			Context.AddAsyncPlatformMessage(PlatformData.TargetPlatform, MoveTemp(ObjectFuture));
-		}
+		ICookedPackageWriter& PackageWriter = COTFS.FindOrCreatePackageWriter(PlatformData.TargetPlatform);
+		TFuture<FCbObject> ObjectFuture = PackageWriter.WriteMPCookMessageForPackage(Context.GetPackageName());
+		Context.AddAsyncPlatformMessage(PlatformData.TargetPlatform, MoveTemp(ObjectFuture));
 	}
 }
 
