Fix double auto-checkout in editor

Removes auto-checkout from the editor itself to rely on auto-checkout from the unsaved assets module.

diff --git a/Engine/Source/Editor/UnrealEd/Private/UnrealEdEngine.cpp b/Engine/Source/Editor/UnrealEd/Private/UnrealEdEngine.cpp
--- a/Engine/Source/Editor/UnrealEd/Private/UnrealEdEngine.cpp
+++ b/Engine/Source/Editor/UnrealEd/Private/UnrealEdEngine.cpp
@@ -669,9 +669,6 @@ void UUnrealEdEngine::OnSourceControlStateUpdated(const FSourceControlOperationR
 		// Get the source control state of the package
 		ISourceControlProvider& SourceControlProvider = ISourceControlModule::Get().GetProvider();
 
-		TArray<TWeakObjectPtr<UPackage>> PackagesToAutomaticallyCheckOut;
-		TArray<FString> FilesToAutomaticallyCheckOut;
-
 		const UEditorLoadingSavingSettings* Settings = GetDefault<UEditorLoadingSavingSettings>();
 		for (const TWeakObjectPtr<UPackage>& PackagePtr : Packages)
 		{
@@ -692,12 +689,7 @@ void UUnrealEdEngine::OnSourceControlStateUpdated(const FSourceControlOperationR
 						}
 						else
 						{
-							if (Settings->GetAutomaticallyCheckoutOnAssetModification())
-							{
-								PackagesToAutomaticallyCheckOut.Add(PackagePtr);
-								FilesToAutomaticallyCheckOut.Add(SourceControlHelpers::PackageFilename(Package));
-							}
-							else
+							if (!Settings->GetAutomaticallyCheckoutOnAssetModification())
 							{
 								PackageToNotifyState.Add(PackagePtr, NS_PendingPrompt);
 								bShowPackageNotification = true;
@@ -712,11 +704,6 @@ void UUnrealEdEngine::OnSourceControlStateUpdated(const FSourceControlOperationR
 				}
 			}
 		}
-
-		if (FilesToAutomaticallyCheckOut.Num() > 0)
-		{
-			SourceControlProvider.Execute(ISourceControlOperation::Create<FCheckOut>(), SourceControlHelpers::AbsoluteFilenames(FilesToAutomaticallyCheckOut), EConcurrency::Asynchronous, FSourceControlOperationComplete::CreateUObject(this, &UUnrealEdEngine::OnPackagesCheckedOut, PackagesToAutomaticallyCheckOut));
-		}
 	}
 }
 
