Improve the notification when a save is requested but no files need to be saved.

diff --git a/Engine/Source/Editor/UnrealEd/Private/FileHelpers.cpp b/Engine/Source/Editor/UnrealEd/Private/FileHelpers.cpp
--- a/Engine/Source/Editor/UnrealEd/Private/FileHelpers.cpp
+++ b/Engine/Source/Editor/UnrealEd/Private/FileHelpers.cpp
@@ -3570,8 +3570,7 @@ static void InternalNotifyNoPackagesSaved(const bool bUseDialog)
 {
 	if (bUseDialog)
 	{
-		FNotificationInfo NotificationInfo(LOCTEXT("NoAssetsToSave", "No new changes to save!"));
-		NotificationInfo.Image = FAppStyle::GetBrush(FTokenizedMessage::GetSeverityIconName(EMessageSeverity::Info));
+		FNotificationInfo NotificationInfo(LOCTEXT("NoAssetsToSave", "All files are already saved."));
 		NotificationInfo.bFireAndForget = true;
 		NotificationInfo.ExpireDuration = 4.0f; // Need this message to last a little longer than normal since the user may have expected there to be modified files.
 		NotificationInfo.bUseThrobber = true;
@@ -3579,7 +3578,7 @@ static void InternalNotifyNoPackagesSaved(const bool bUseDialog)
 	}
 	else
 	{
-		UE_LOG(LogFileHelpers, Log, TEXT("%s"), *LOCTEXT("NoAssetsToSave", "No new changes to save!").ToString());
+		UE_LOG(LogFileHelpers, Log, TEXT("%s"), *LOCTEXT("NoAssetsToSave", "All files are already saved.").ToString());
 	}
 }
 
