Fix for FORCEINLINE method not found

diff --git a/Engine/Source/Editor/UnrealEd/Public/Toolkits/BaseToolkit.h b/Engine/Source/Editor/UnrealEd/Public/Toolkits/BaseToolkit.h
--- a/Engine/Source/Editor/UnrealEd/Public/Toolkits/BaseToolkit.h
+++ b/Engine/Source/Editor/UnrealEd/Public/Toolkits/BaseToolkit.h
@@ -269,5 +269,5 @@ protected:
 	TSharedPtr<FToolkitBuilder> ToolkitBuilder;
 
 	/** returns whether or not this FModeToolkit has a ToolkitBuilder defined */
-	FORCEINLINE bool HasToolkitBuilder() const;
+	bool HasToolkitBuilder() const;
 };
