Assertion failed BracketCounter >= 0 after attempting to undo the importation of CesiumMan.gltf
#fix ensure that controller changes are only transacted when reimporting an animation sequence

diff --git a/Engine/Plugins/Animation/AnimationData/Source/AnimationData/Private/AnimSequencerController.cpp b/Engine/Plugins/Animation/AnimationData/Source/AnimationData/Private/AnimSequencerController.cpp
--- a/Engine/Plugins/Animation/AnimationData/Source/AnimationData/Private/AnimSequencerController.cpp
+++ b/Engine/Plugins/Animation/AnimationData/Source/AnimationData/Private/AnimSequencerController.cpp
@@ -42,7 +42,7 @@ void UAnimSequencerController::OpenBracket(const FText& InTitle, bool bShouldTra
 {
 	ValidateModel();
 
-	if (UE::FChangeTransactor::CanTransactChanges() && !ChangeTransactor.IsTransactionPending())
+	if (UE::FChangeTransactor::CanTransactChanges() && !ChangeTransactor.IsTransactionPending() && bShouldTransact)
 	{
 		ChangeTransactor.OpenTransaction(InTitle);
 
@@ -74,7 +74,7 @@ void UAnimSequencerController::CloseBracket(bool bShouldTransact /*= true*/)
 
 	if (BracketDepth == 0)
 	{
-		if (UE::FChangeTransactor::CanTransactChanges())
+		if (UE::FChangeTransactor::CanTransactChanges() && bShouldTransact)
 		{
 			ensure(ChangeTransactor.IsTransactionPending());
 
diff --git a/Engine/Plugins/Interchange/Runtime/Source/Import/Private/Animation/InterchangeAnimSequenceFactory.cpp b/Engine/Plugins/Interchange/Runtime/Source/Import/Private/Animation/InterchangeAnimSequenceFactory.cpp
--- a/Engine/Plugins/Interchange/Runtime/Source/Import/Private/Animation/InterchangeAnimSequenceFactory.cpp
+++ b/Engine/Plugins/Interchange/Runtime/Source/Import/Private/Animation/InterchangeAnimSequenceFactory.cpp
@@ -153,7 +153,8 @@ namespace UE::Interchange::Private
 		, const int32 CurveFlags
 		, const bool bDoNotImportCurveWithZero
 		, const bool bMorphTargetCurve
-		, const bool bMaterialCurve)
+		, const bool bMaterialCurve
+		, const bool bShouldTransact)
 	{
 		if (!TargetSequence || CurvePayload.Curves.Num() != 1 || CurveName.IsEmpty())
 		{
@@ -198,7 +199,6 @@ namespace UE::Interchange::Private
 
 		FAnimationCurveIdentifier FloatCurveId(NewName, ERawCurveTrackTypes::RCT_Float);
 
-
 		IAnimationDataModel* DataModel = TargetSequence->GetDataModel();
 		IAnimationDataController& Controller = TargetSequence->GetController();
 
@@ -206,19 +206,19 @@ namespace UE::Interchange::Private
 		if (TargetCurve == nullptr)
 		{
 			// Need to add the curve first
-			Controller.AddCurve(FloatCurveId, AACF_DefaultCurve | CurveFlags);
+			Controller.AddCurve(FloatCurveId, AACF_DefaultCurve | CurveFlags, bShouldTransact);
 			TargetCurve = DataModel->FindFloatCurve(FloatCurveId);
 		}
 		else
 		{
 			// Need to update any of the flags
-			Controller.SetCurveFlags(FloatCurveId, CurveFlags | TargetCurve->GetCurveTypeFlags());
+			Controller.SetCurveFlags(FloatCurveId, CurveFlags | TargetCurve->GetCurveTypeFlags(), bShouldTransact);
 		}
 
 		// Should be valid at this point
 		ensure(TargetCurve);
 
-		Controller.UpdateCurveNamesFromSkeleton(Skeleton, ERawCurveTrackTypes::RCT_Float);
+		Controller.UpdateCurveNamesFromSkeleton(Skeleton, ERawCurveTrackTypes::RCT_Float, bShouldTransact);
 
 		//MorphTarget curves are shift to 0 second
 		const float CurveStartTime = CurvePayload.Curves[0].GetFirstKey().Time;
@@ -227,7 +227,7 @@ namespace UE::Interchange::Private
 			CurvePayload.Curves[0].ShiftCurve(-CurveStartTime);
 		}
 		// Set actual keys on curve within the model
-		Controller.SetCurveKeys(FloatCurveId, CurvePayload.Curves[0].GetConstRefOfKeys());
+		Controller.SetCurveKeys(FloatCurveId, CurvePayload.Curves[0].GetConstRefOfKeys(), bShouldTransact);
 
 		if (bMaterialCurve || bMorphTargetCurve)
 		{
@@ -236,26 +236,26 @@ namespace UE::Interchange::Private
 		return true;
 	}
 
-	bool CreateMorphTargetCurve(UAnimSequence* TargetSequence, FAnimationCurvePayloadData& CurvePayload, const FString& CurveName, int32 CurveFlags, bool bDoNotImportCurveWithZero)
+	bool CreateMorphTargetCurve(UAnimSequence* TargetSequence, FAnimationCurvePayloadData& CurvePayload, const FString& CurveName, int32 CurveFlags, bool bDoNotImportCurveWithZero, bool bShouldTransact)
 	{
 		constexpr bool bIsMorphTargetCurve = true;
 		constexpr bool bIsMaterialCurve = false;
-		return InternalCreateCurve(TargetSequence, CurvePayload, CurveName, CurveFlags, bDoNotImportCurveWithZero, bIsMorphTargetCurve, bIsMaterialCurve);
+		return InternalCreateCurve(TargetSequence, CurvePayload, CurveName, CurveFlags, bDoNotImportCurveWithZero, bIsMorphTargetCurve, bIsMaterialCurve, bShouldTransact);
 	}
 
-	bool CreateMaterialCurve(UAnimSequence* TargetSequence, FAnimationCurvePayloadData& CurvePayload, const FString& CurveName, int32 CurveFlags, bool bDoNotImportCurveWithZero)
+	bool CreateMaterialCurve(UAnimSequence* TargetSequence, FAnimationCurvePayloadData& CurvePayload, const FString& CurveName, int32 CurveFlags, bool bDoNotImportCurveWithZero, bool bShouldTransact)
 	{
 		constexpr bool bIsMorphTargetCurve = false;
 		constexpr bool bIsMaterialCurve = true;
-		return InternalCreateCurve(TargetSequence, CurvePayload, CurveName, CurveFlags, bDoNotImportCurveWithZero, bIsMorphTargetCurve, bIsMaterialCurve);
+		return InternalCreateCurve(TargetSequence, CurvePayload, CurveName, CurveFlags, bDoNotImportCurveWithZero, bIsMorphTargetCurve, bIsMaterialCurve, bShouldTransact);
 	}
 
-	bool CreateAttributeCurve(UAnimSequence* TargetSequence, FAnimationCurvePayloadData& CurvePayload, const FString& CurveName, int32 CurveFlags, bool bDoNotImportCurveWithZero)
+	bool CreateAttributeCurve(UAnimSequence* TargetSequence, FAnimationCurvePayloadData& CurvePayload, const FString& CurveName, int32 CurveFlags, bool bDoNotImportCurveWithZero, bool bShouldTransact)
 	{
 		//This curve don't animate morph target or material parameter.
 		constexpr bool bIsMorphTargetCurve = false;
 		constexpr bool bIsMaterialCurve = false;
-		return InternalCreateCurve(TargetSequence, CurvePayload, CurveName, CurveFlags, bDoNotImportCurveWithZero, bIsMorphTargetCurve, bIsMaterialCurve);
+		return InternalCreateCurve(TargetSequence, CurvePayload, CurveName, CurveFlags, bDoNotImportCurveWithZero, bIsMorphTargetCurve, bIsMaterialCurve, bShouldTransact);
 	}
 
 	void RetrieveAnimationPayloads(UAnimSequence* AnimSequence
@@ -264,6 +264,7 @@ namespace UE::Interchange::Private
 		, const UInterchangeSkeletonFactoryNode* SkeletonFactoryNode
 		, const IInterchangeAnimationPayloadInterface* AnimSequenceTranslatorPayloadInterface
 		, const FString& AssetName
+		, const bool bIsReimporting
 		, TArray<FString> OutCurvesNotFound)
 	{
 		TMap<const UInterchangeSceneNode*, TFuture<TOptional<UE::Interchange::FAnimationBakeTransformPayloadData>>> AnimationPayloads;
@@ -282,6 +283,9 @@ namespace UE::Interchange::Private
 		TArray<FString> SkeletonNodes;
 		GetSkeletonSceneNodeFlatListRecursive(NodeContainer, SkeletonRootUid, SkeletonNodes);
 
+		
+		const bool bShouldTransact = bIsReimporting;
+
 		bool bImportBoneTracks = false;
 		AnimSequenceFactoryNode->GetCustomImportBoneTracks(bImportBoneTracks);
 		if (bImportBoneTracks)
@@ -329,12 +333,12 @@ namespace UE::Interchange::Private
 			}
 
 			//This destroy all previously imported animation raw data
-			Controller.RemoveAllBoneTracks();
+			Controller.RemoveAllBoneTracks(bShouldTransact);
 
 			const FFrameRate ResampleFrameRate(SampleRate, 1);
 			const FFrameNumber NumberOfFrames = ResampleFrameRate.AsFrameNumber(SequenceLength);
-			Controller.SetFrameRate(ResampleFrameRate);
-			Controller.SetNumberOfFrames(FGenericPlatformMath::Max<int32>(NumberOfFrames.Value, 1));
+			Controller.SetFrameRate(ResampleFrameRate, bShouldTransact);
+			Controller.SetNumberOfFrames(FGenericPlatformMath::Max<int32>(NumberOfFrames.Value, 1), bShouldTransact);
 
 			FTransform GlobalOffsetTransform;
 			bool bBakeMeshes = false;
@@ -432,8 +436,8 @@ namespace UE::Interchange::Private
 					&& TimeKeys.Num() == BakeKeyCount);
 
 				//add new track
-				Controller.AddBoneCurve(BoneName);
-				Controller.SetBoneTrackKeys(BoneName, RawTrack.PosKeys, RawTrack.RotKeys, RawTrack.ScaleKeys);
+				Controller.AddBoneCurve(BoneName, bShouldTransact);
+				Controller.SetBoneTrackKeys(BoneName, RawTrack.PosKeys, RawTrack.RotKeys, RawTrack.ScaleKeys, bShouldTransact);
 			}
 		}
 
@@ -462,13 +466,13 @@ namespace UE::Interchange::Private
 			for (const FSmartName& CurveName : CurveSmartNamesToRemove)
 			{
 				const FAnimationCurveIdentifier CurveId(CurveName, ERawCurveTrackTypes::RCT_Float);
-				Controller.RemoveCurve(CurveId);
+				Controller.RemoveCurve(CurveId, bShouldTransact);
 			}
 		}
 
 		if (bDeleteExistingNonCurveCustomAttributes)
 		{
-			Controller.RemoveAllAttributes();
+			Controller.RemoveAllAttributes(bShouldTransact);
 		}
 
 		
@@ -573,7 +577,7 @@ namespace UE::Interchange::Private
 						}
 					}
 					constexpr int32 CurveFlags = 0;
-					CreateMorphTargetCurve(AnimSequence, CurvePayload, AnimationCurveMorphTargetNodeNames.FindChecked(CurveNameAndPayload.Key), CurveFlags, bDoNotImportCurveWithZero);
+					CreateMorphTargetCurve(AnimSequence, CurvePayload, AnimationCurveMorphTargetNodeNames.FindChecked(CurveNameAndPayload.Key), CurveFlags, bDoNotImportCurveWithZero, bShouldTransact);
 				}
 			}
 
@@ -646,11 +650,11 @@ namespace UE::Interchange::Private
 						constexpr int32 CurveFlags = 0;
 						if (bMaterialDriveParameterOnCustomAttribute || IsCurveHookToMaterial(CurveName))
 						{
-							CreateMaterialCurve(AnimSequence, CurvePayload, CurveName, CurveFlags, bDoNotImportCurveWithZero);
+							CreateMaterialCurve(AnimSequence, CurvePayload, CurveName, CurveFlags, bDoNotImportCurveWithZero, bShouldTransact);
 						}
 						else
 						{
-							CreateAttributeCurve(AnimSequence, CurvePayload, CurveName, CurveFlags, bDoNotImportCurveWithZero);
+							CreateAttributeCurve(AnimSequence, CurvePayload, CurveName, CurveFlags, bDoNotImportCurveWithZero, bShouldTransact);
 						}
 					}
 				}
@@ -948,13 +952,14 @@ UObject* UInterchangeAnimSequenceFactory::ImportObjectSourceData(const FImportAs
 				FrameRate = UE::Interchange::Animation::ConvertSampleRatetoFrameRate(SampleRate);
 			}
 		}
+		const bool bShouldTransact = bIsReImport;
 		IAnimationDataController& Controller = AnimSequence->GetController();
-		Controller.OpenBracket(NSLOCTEXT("InterchangeAnimSequenceFactory", "ImportAnimationInterchange_Bracket", "Importing Animation (Interchange)"));
+		Controller.OpenBracket(NSLOCTEXT("InterchangeAnimSequenceFactory", "ImportAnimationInterchange_Bracket", "Importing Animation (Interchange)"), bShouldTransact);
 		AnimSequence->SetSkeleton(Skeleton);
 		Controller.InitializeModel();
 		AnimSequence->ImportFileFramerate = SampleRate;
 		AnimSequence->ImportResampleFramerate = SampleRate;
-		Controller.SetFrameRate(FrameRate);
+		Controller.SetFrameRate(FrameRate, bShouldTransact);
 
 		TArray<FString> CurvesNotFound;
 		UE::Interchange::Private::RetrieveAnimationPayloads(AnimSequence
@@ -963,6 +968,7 @@ UObject* UInterchangeAnimSequenceFactory::ImportObjectSourceData(const FImportAs
 			, SkeletonFactoryNode
 			, AnimSequenceTranslatorPayloadInterface
 			, Arguments.AssetName
+			, bIsReImport
 			, CurvesNotFound);
 
 		if (CurvesNotFound.Num())
@@ -977,7 +983,7 @@ UObject* UInterchangeAnimSequenceFactory::ImportObjectSourceData(const FImportAs
 			}
 		}
 		Controller.NotifyPopulated();
-		Controller.CloseBracket(false);
+		Controller.CloseBracket(bShouldTransact);
 	}
 
 	if (!bIsReImport)
diff --git a/Engine/Source/Developer/AnimationDataController/Private/AnimDataController.cpp b/Engine/Source/Developer/AnimationDataController/Private/AnimDataController.cpp
--- a/Engine/Source/Developer/AnimationDataController/Private/AnimDataController.cpp
+++ b/Engine/Source/Developer/AnimationDataController/Private/AnimDataController.cpp
@@ -39,7 +39,7 @@ void UAnimDataController::OpenBracket(const FText& InTitle, bool bShouldTransact
 {
 	ValidateModel();
 
-	if (UE::FChangeTransactor::CanTransactChanges() && !ChangeTransactor.IsTransactionPending())
+	if (UE::FChangeTransactor::CanTransactChanges() && !ChangeTransactor.IsTransactionPending() && bShouldTransact)
 	{
 		ChangeTransactor.OpenTransaction(InTitle);
 
@@ -71,7 +71,7 @@ void UAnimDataController::CloseBracket(bool bShouldTransact /*= true*/)
 
 	if (BracketDepth == 0)
 	{
-		if (UE::FChangeTransactor::CanTransactChanges())
+		if (UE::FChangeTransactor::CanTransactChanges() && bShouldTransact)
 		{
 			ensure(ChangeTransactor.IsTransactionPending());
 
diff --git a/Engine/Source/Editor/UnrealEd/Private/Fbx/ReimportFbxSceneFactory.cpp b/Engine/Source/Editor/UnrealEd/Private/Fbx/ReimportFbxSceneFactory.cpp
--- a/Engine/Source/Editor/UnrealEd/Private/Fbx/ReimportFbxSceneFactory.cpp
+++ b/Engine/Source/Editor/UnrealEd/Private/Fbx/ReimportFbxSceneFactory.cpp
@@ -1458,14 +1458,14 @@ EReimportResult::Type UReimportFbxSceneFactory::ReimportSkeletalMesh(void* VoidF
 							// since to know full path, reimport will need to do same
 							UFbxAnimSequenceImportData* ImportData = UFbxAnimSequenceImportData::GetImportDataForAnimSequence(DestSeq, AnimSequenceImportData);
 							ImportData->Update(UFactory::CurrentFilename);
-							FbxImporter->ImportAnimation(Mesh->GetSkeleton(), DestSeq, CurrentFilename, SortedLinks, FBXMeshNodeArray, CurAnimStack, ResampleRate, AnimTimeSpan);
+							FbxImporter->ImportAnimation(Mesh->GetSkeleton(), DestSeq, CurrentFilename, SortedLinks, FBXMeshNodeArray, CurAnimStack, ResampleRate, AnimTimeSpan, false);
 						}
 						else
 						{
 							//Reimport in a existing sequence
 							if (FbxImporter->ValidateAnimStack(SortedLinks, FBXMeshNodeArray, CurAnimStack, ResampleRate, true, FbxImporter->ImportOptions->bSnapToClosestFrameBoundary, AnimTimeSpan))
 							{
-								FbxImporter->ImportAnimation(Mesh->GetSkeleton(), DestSeq, CurrentFilename, SortedLinks, FBXMeshNodeArray, CurAnimStack, ResampleRate, AnimTimeSpan);
+								FbxImporter->ImportAnimation(Mesh->GetSkeleton(), DestSeq, CurrentFilename, SortedLinks, FBXMeshNodeArray, CurAnimStack, ResampleRate, AnimTimeSpan, true);
 							}
 						}
 					}
diff --git a/Engine/Source/Editor/UnrealEd/Private/SkeletalMeshEdit.cpp b/Engine/Source/Editor/UnrealEd/Private/SkeletalMeshEdit.cpp
--- a/Engine/Source/Editor/UnrealEd/Private/SkeletalMeshEdit.cpp
+++ b/Engine/Source/Editor/UnrealEd/Private/SkeletalMeshEdit.cpp
@@ -283,7 +283,7 @@ bool UEditorEngine::ReimportFbxAnimation( USkeleton* Skeleton, UAnimSequence* An
 					if (FbxImporter->ValidateAnimStack(SortedLinks, FBXMeshNodeArray, CurAnimStack, ResampleRate, bImportMorphTracks, FbxImporter->ImportOptions->bSnapToClosestFrameBoundary, AnimTimeSpan))
 					{
 						AnimSequence->ImportResampleFramerate = ResampleRate;
-						FbxImporter->ImportAnimation(Skeleton, AnimSequence, Filename, SortedLinks, FBXMeshNodeArray, CurAnimStack, ResampleRate, AnimTimeSpan);
+						FbxImporter->ImportAnimation(Skeleton, AnimSequence, Filename, SortedLinks, FBXMeshNodeArray, CurAnimStack, ResampleRate, AnimTimeSpan, true);
 					}
 				}
 				else
@@ -664,7 +664,7 @@ UAnimSequence * UnFbx::FFbxImporter::ImportAnimations(USkeleton* Skeleton, UObje
 
 		DestSeq->GetController().InitializeModel();
 
-		ImportAnimation(Skeleton, DestSeq, Name, SortedLinks, NodeArray, CurAnimStack, ResampleRate, AnimTimeSpan);
+		ImportAnimation(Skeleton, DestSeq, Name, SortedLinks, NodeArray, CurAnimStack, ResampleRate, AnimTimeSpan, false);
 
 		if (bCreateAsset)
 		{			
@@ -1363,7 +1363,7 @@ namespace AnimationTransformDebug
 	}
 }
 
-bool UnFbx::FFbxImporter::ImportCurveToAnimSequence(class UAnimSequence * TargetSequence, const FString& CurveName, const FbxAnimCurve* FbxCurve, int32 CurveFlags,const FbxTimeSpan& AnimTimeSpan, float ValueScale/*=1.f*/) const
+bool UnFbx::FFbxImporter::ImportCurveToAnimSequence(class UAnimSequence * TargetSequence, const FString& CurveName, const FbxAnimCurve* FbxCurve, int32 CurveFlags,const FbxTimeSpan& AnimTimeSpan, const bool bReimport, float ValueScale/*=1.f*/) const
 {
 	if (TargetSequence && FbxCurve)
 	{
@@ -1383,7 +1383,7 @@ bool UnFbx::FFbxImporter::ImportCurveToAnimSequence(class UAnimSequence * Target
 
 		FAnimationCurveIdentifier FloatCurveId(NewName, ERawCurveTrackTypes::RCT_Float);
 
-
+		const bool bShouldTransact = bReimport;
 		IAnimationDataModel* DataModel = TargetSequence->GetDataModel();
 		IAnimationDataController& Controller = TargetSequence->GetController();
 
@@ -1391,19 +1391,19 @@ bool UnFbx::FFbxImporter::ImportCurveToAnimSequence(class UAnimSequence * Target
 		if (TargetCurve == nullptr)
 		{
 			// Need to add the curve first
-			Controller.AddCurve(FloatCurveId, AACF_DefaultCurve | CurveFlags);
+			Controller.AddCurve(FloatCurveId, AACF_DefaultCurve | CurveFlags, bShouldTransact);
 			TargetCurve = DataModel->FindFloatCurve(FloatCurveId);
 		}
 		else
 		{
 			// Need to update any of the flags
-			Controller.SetCurveFlags(FloatCurveId, CurveFlags | TargetCurve->GetCurveTypeFlags());
+			Controller.SetCurveFlags(FloatCurveId, CurveFlags | TargetCurve->GetCurveTypeFlags(), bShouldTransact);
 		}
 		
 		// Should be valid at this point
 		ensure(TargetCurve);
 
-		Controller.UpdateCurveNamesFromSkeleton(Skeleton, ERawCurveTrackTypes::RCT_Float);
+		Controller.UpdateCurveNamesFromSkeleton(Skeleton, ERawCurveTrackTypes::RCT_Float, bShouldTransact);
 
 		FRichCurve RichCurve;
 		constexpr bool bNegative = false;
@@ -1415,7 +1415,7 @@ bool UnFbx::FFbxImporter::ImportCurveToAnimSequence(class UAnimSequence * Target
 			}
 
 			// Set actual keys on curve within the model
-			Controller.SetCurveKeys(FloatCurveId, RichCurve.GetConstRefOfKeys());
+			Controller.SetCurveKeys(FloatCurveId, RichCurve.GetConstRefOfKeys(), bShouldTransact);
 
 			return true;
 		}
@@ -1425,10 +1425,11 @@ bool UnFbx::FFbxImporter::ImportCurveToAnimSequence(class UAnimSequence * Target
 }
 
 
-bool UnFbx::FFbxImporter::ImportRichCurvesToAnimSequence(UAnimSequence* TargetSequence, const TArray<FString>& CurveNames, const TArray<FRichCurve> RichCurves, int32 CurveFlags) const
+bool UnFbx::FFbxImporter::ImportRichCurvesToAnimSequence(UAnimSequence* TargetSequence, const TArray<FString>& CurveNames, const TArray<FRichCurve> RichCurves, int32 CurveFlags, const bool bReimport) const
 {
 	if (TargetSequence && CurveNames.Num() > 0 && CurveNames.Num() == RichCurves.Num())
 	{
+		const bool bShouldTransact = bReimport;
 		for (int32 CurveIndex = 0; CurveIndex < CurveNames.Num(); ++CurveIndex)
 		{
 			FName Name = *CurveNames[CurveIndex];
@@ -1455,22 +1456,22 @@ bool UnFbx::FFbxImporter::ImportRichCurvesToAnimSequence(UAnimSequence* TargetSe
 			if (TargetCurve == nullptr)
 			{
 				// Need to add the curve first
-				Controller.AddCurve(FloatCurveId, AACF_DefaultCurve | CurveFlags);
+				Controller.AddCurve(FloatCurveId, AACF_DefaultCurve | CurveFlags, bShouldTransact);
 				TargetCurve = DataModel->FindFloatCurve(FloatCurveId);
 			}
 			else
 			{
 				// Need to update any of the flags
-				Controller.SetCurveFlags(FloatCurveId, CurveFlags | TargetCurve->GetCurveTypeFlags());
+				Controller.SetCurveFlags(FloatCurveId, CurveFlags | TargetCurve->GetCurveTypeFlags(), bShouldTransact);
 			}
 		
 			// Should be valid at this point
 			ensure(TargetCurve);
 
-			Controller.UpdateCurveNamesFromSkeleton(Skeleton, ERawCurveTrackTypes::RCT_Float);
+			Controller.UpdateCurveNamesFromSkeleton(Skeleton, ERawCurveTrackTypes::RCT_Float, bShouldTransact);
 
 			// Set actual keys on curve within the model
-			Controller.SetCurveKeys(FloatCurveId, RichCurves[CurveIndex].GetConstRefOfKeys());
+			Controller.SetCurveKeys(FloatCurveId, RichCurves[CurveIndex].GetConstRefOfKeys(), bShouldTransact);
 
 		}
 		
@@ -1614,7 +1615,7 @@ void FillCurveAttributeToBone(TArray<float>& OutFrameTimes, TArray<AttributeType
 	}
 }
 
-bool UnFbx::FFbxImporter::ImportCustomAttributeToBone(UAnimSequence* TargetSequence, FbxProperty& InProperty, FName BoneName, const FString& CurveName, const FbxAnimCurve* FbxCurve, const FbxTimeSpan& AnimTimeSpan, float ValueScale/*=1.f*/)
+bool UnFbx::FFbxImporter::ImportCustomAttributeToBone(UAnimSequence* TargetSequence, FbxProperty& InProperty, FName BoneName, const FString& CurveName, const FbxAnimCurve* FbxCurve, const FbxTimeSpan& AnimTimeSpan, const bool bReimport, float ValueScale/*=1.f*/)
 {
 	if (TargetSequence)
 	{
@@ -1755,25 +1756,26 @@ bool ShouldImportCurve(FbxAnimCurve* Curve, bool bDoNotImportWithZeroValues)
 	return false;
 }
 
-bool UnFbx::FFbxImporter::ImportAnimation(USkeleton* Skeleton, UAnimSequence * DestSeq, const FString& FileName, TArray<FbxNode*>& SortedLinks, TArray<FbxNode*>& NodeArray, FbxAnimStack* CurAnimStack, const int32 ResampleRate, const FbxTimeSpan AnimTimeSpan)
+bool UnFbx::FFbxImporter::ImportAnimation(USkeleton* Skeleton, UAnimSequence * DestSeq, const FString& FileName, TArray<FbxNode*>& SortedLinks, TArray<FbxNode*>& NodeArray, FbxAnimStack* CurAnimStack, const int32 ResampleRate, const FbxTimeSpan AnimTimeSpan, const bool bReimport)
 {
 	// @todo : the length might need to change w.r.t. sampling keys
 	FbxTime SequenceLength = AnimTimeSpan.GetDuration();
 	float PreviousSequenceLength = DestSeq->GetPlayLength();
 
+	const bool bShouldTransact = bReimport;
 	IAnimationDataController& Controller = DestSeq->GetController();
-	Controller.OpenBracket(LOCTEXT("ImportAnimation_Bracket", "Importing Animation"));
+	Controller.OpenBracket(LOCTEXT("ImportAnimation_Bracket", "Importing Animation"), bShouldTransact);
 
 	//This destroy all previously imported animation raw data
-	Controller.RemoveAllBoneTracks();
+	Controller.RemoveAllBoneTracks(bShouldTransact);
 
 	// First set frame rate
 	const FFrameRate ResampleFrameRate(ResampleRate, 1);
-	Controller.SetFrameRate(ResampleFrameRate);
+	Controller.SetFrameRate(ResampleFrameRate, bShouldTransact);
 
 	// if you have one pose(thus 0.f duration), it still contains animation, so we'll need to consider that as MINIMUM_ANIMATION_LENGTH time length
 	const FFrameNumber NumberOfFrames = ResampleFrameRate.AsFrameNumber(SequenceLength.GetSecondDouble());
-	Controller.SetNumberOfFrames(FGenericPlatformMath::Max<int32>(NumberOfFrames.Value, 1));
+	Controller.SetNumberOfFrames(FGenericPlatformMath::Max<int32>(NumberOfFrames.Value, 1), bShouldTransact);
 
 	if(PreviousSequenceLength > MINIMUM_ANIMATION_LENGTH && DestSeq->GetDataModel()->GetNumberOfFloatCurves() > 0)
 	{
@@ -1785,7 +1787,7 @@ bool UnFbx::FFbxImporter::ImportAnimation(USkeleton* Skeleton, UAnimSequence * D
 			for (const FFloatCurve& Curve : DestSeq->GetDataModel()->GetFloatCurves())
 			{
 				const FAnimationCurveIdentifier CurveId(Curve.Name, ERawCurveTrackTypes::RCT_Float);
-				Controller.ScaleCurve(CurveId, 0.f, ScaleFactor);
+				Controller.ScaleCurve(CurveId, 0.f, ScaleFactor, bShouldTransact);
 			}
 		}
 	}
@@ -1812,13 +1814,13 @@ bool UnFbx::FFbxImporter::ImportAnimation(USkeleton* Skeleton, UAnimSequence * D
 		for (auto CurveName : CurveSmartNamesToRemove)
 		{
 			const FAnimationCurveIdentifier CurveId(CurveName, ERawCurveTrackTypes::RCT_Float);
-			Controller.RemoveCurve(CurveId);
+			Controller.RemoveCurve(CurveId, bShouldTransact);
 		}
 	}
 
 	if (ImportOptions->bDeleteExistingNonCurveCustomAttributes)
 	{
-		Controller.RemoveAllAttributes();
+		Controller.RemoveAllAttributes(bShouldTransact);
 	}
 	
 	const bool bReimportWarnings = GetDefault<UEditorPerProjectUserSettings>()->bAnimationReimportWarnings;
@@ -1836,7 +1838,7 @@ bool UnFbx::FFbxImporter::ImportAnimation(USkeleton* Skeleton, UAnimSequence * D
 	// import blend shape curves
 	//
 	int32 CurveAttributeKeyCount = 0;
-	ImportBlendShapeCurves(AnimImportSettings, CurAnimStack, CurveAttributeKeyCount);
+	ImportBlendShapeCurves(AnimImportSettings, CurAnimStack, CurveAttributeKeyCount, bReimport);
 
 	// importing custom attribute START
 	TArray<FString> CurvesNotFound;
@@ -1844,7 +1846,7 @@ bool UnFbx::FFbxImporter::ImportAnimation(USkeleton* Skeleton, UAnimSequence * D
 	{
 		int CustomAttributeKeyCount = 0;
 		
-		ImportAnimationCustomAttribute(AnimImportSettings, CustomAttributeKeyCount, CurvesNotFound);
+		ImportAnimationCustomAttribute(AnimImportSettings, CustomAttributeKeyCount, CurvesNotFound, bReimport);
 
 		CurveAttributeKeyCount = FMath::Max(CurveAttributeKeyCount, CustomAttributeKeyCount);
 	}
@@ -1878,13 +1880,13 @@ bool UnFbx::FFbxImporter::ImportAnimation(USkeleton* Skeleton, UAnimSequence * D
 	if (ImportOptions->bImportBoneTracks)
 	{
 		FbxNode* SkeletalMeshRootNode = NodeArray.Num() > 0 ? NodeArray[0] : nullptr;
-		ImportBoneTracks(Skeleton, AnimImportSettings, SkeletalMeshRootNode, ResampleRate, TransformDebugData, TotalNumKeys);
+		ImportBoneTracks(Skeleton, AnimImportSettings, SkeletalMeshRootNode, ResampleRate, TransformDebugData, TotalNumKeys, bReimport);
 
 		AnimationTransformDebug::OutputAnimationTransformDebugData(TransformDebugData, TotalNumKeys, RefSkeleton);
 	}
 	else if (CurveAttributeKeyCount > 0)
 	{
-		Controller.SetFrameRate(FFrameRate(ResampleRate, 1));
+		Controller.SetFrameRate(FFrameRate(ResampleRate, 1), bShouldTransact);
 	}
 
 	// Import bone metadata to AnimSequence
@@ -1894,7 +1896,7 @@ bool UnFbx::FFbxImporter::ImportAnimation(USkeleton* Skeleton, UAnimSequence * D
 	}
 
 	Controller.NotifyPopulated();
-	Controller.CloseBracket();
+	Controller.CloseBracket(bShouldTransact);
 
 	// Reregister skeletal mesh components so they reflect the updated animation
 	for (TObjectIterator<USkeletalMeshComponent> Iter; Iter; ++Iter)
@@ -1904,7 +1906,7 @@ bool UnFbx::FFbxImporter::ImportAnimation(USkeleton* Skeleton, UAnimSequence * D
 	return true;
 }
 
-void UnFbx::FFbxImporter::ImportBlendShapeCurves(FAnimCurveImportSettings& AnimImportSettings, FbxAnimStack* CurAnimStack, int32& OutKeyCount)
+void UnFbx::FFbxImporter::ImportBlendShapeCurves(FAnimCurveImportSettings& AnimImportSettings, FbxAnimStack* CurAnimStack, int32& OutKeyCount, const bool bReimport)
 {
 	FText CurrentExportMessage = LOCTEXT("BeginImportMorphTargetCurves", "Importing Morph Target Curves");
 	FScopedSlowTask SlowTaskNode(AnimImportSettings.NodeArray.Num(), CurrentExportMessage);
@@ -1978,7 +1980,7 @@ void UnFbx::FFbxImporter::ImportBlendShapeCurves(FAnimCurveImportSettings& AnimI
 								{
 									// now see if we have one already exists. If so, just overwrite that. if not, add new one. 
 
-									if (ImportCurveToAnimSequence(AnimImportSettings.DestSeq, *ChannelName, Curve, 0, AnimImportSettings.AnimTimeSpan, 0.01f /** for some reason blend shape values are coming as 100 scaled **/))
+									if (ImportCurveToAnimSequence(AnimImportSettings.DestSeq, *ChannelName, Curve, 0, AnimImportSettings.AnimTimeSpan, bReimport, 0.01f /** for some reason blend shape values are coming as 100 scaled **/))
 									{
 										OutKeyCount = FMath::Max(OutKeyCount, Curve->KeyGetCount());
 										// this one doesn't reset Material curve to false, it just accumulate if true. 
@@ -2039,7 +2041,7 @@ void UnFbx::FFbxImporter::ImportBlendShapeCurves(FAnimCurveImportSettings& AnimI
 										}
 									}
 
-									if (ImportRichCurvesToAnimSequence(AnimImportSettings.DestSeq, CurveNames, Results, 0))
+									if (ImportRichCurvesToAnimSequence(AnimImportSettings.DestSeq, CurveNames, Results, 0, bReimport))
 									{
 										OutKeyCount = FMath::Max(OutKeyCount, Curve->KeyGetCount());
 										for (const FString& CurveName : CurveNames)
@@ -2067,7 +2069,7 @@ void UnFbx::FFbxImporter::ImportBlendShapeCurves(FAnimCurveImportSettings& AnimI
 	}
 }
 
-void UnFbx::FFbxImporter::ImportAnimationCustomAttribute(FAnimCurveImportSettings& AnimImportSettings, int32& OutKeyCount, TArray<FString>& OutCurvesNotFound)
+void UnFbx::FFbxImporter::ImportAnimationCustomAttribute(FAnimCurveImportSettings& AnimImportSettings, int32& OutKeyCount, TArray<FString>& OutCurvesNotFound, const bool bReimport)
 {
 	FScopedSlowTask SlowTask(AnimImportSettings.SortedLinks.Num(), LOCTEXT("BeginImportCustomAttributeCurves", "Importing Custom Attribute Curves"), true);
 	SlowTask.MakeDialog();
@@ -2103,7 +2105,7 @@ void UnFbx::FFbxImporter::ImportAnimationCustomAttribute(FAnimCurveImportSetting
 		if (!bImportAllAttributesOnBone)
 		{
 			FbxAnimUtils::ExtractAttributeCurves(Node, ImportOptions->bDoNotImportCurveWithZero,
-				[this, &LinkIndex, &DestSeq, &AnimImportSettings, &OutCurvesNotFound, &OutKeyCount, &SlowTask](FbxAnimCurve* InCurve, const FString& InCurveName)
+				[this, &LinkIndex, &DestSeq, &AnimImportSettings, &OutCurvesNotFound, &OutKeyCount, &SlowTask, bReimport](FbxAnimCurve* InCurve, const FString& InCurveName)
 			{
 				FFormatNamedArguments Args;
 				Args.Add(TEXT("CurveName"), FText::FromString(InCurveName));
@@ -2111,7 +2113,7 @@ void UnFbx::FFbxImporter::ImportAnimationCustomAttribute(FAnimCurveImportSetting
 				SlowTask.EnterProgressFrame(0, StatusUpate);
 
 				int32 CurveFlags = AACF_DefaultCurve;
-				if (ImportCurveToAnimSequence(DestSeq, InCurveName, InCurve, CurveFlags, AnimImportSettings.AnimTimeSpan))
+				if (ImportCurveToAnimSequence(DestSeq, InCurveName, InCurve, CurveFlags, AnimImportSettings.AnimTimeSpan, bReimport))
 				{
 					USkeleton* SeqSkeleton = DestSeq->GetSkeleton();
 					OutKeyCount = FMath::Max(OutKeyCount, InCurve->KeyGetCount());
@@ -2141,9 +2143,9 @@ void UnFbx::FFbxImporter::ImportAnimationCustomAttribute(FAnimCurveImportSetting
 		}
 
 		FbxAnimUtils::ExtractNodeAttributes(Node, ImportOptions->bDoNotImportCurveWithZero, bImportAllAttributesOnBone,
-			[this, &OutKeyCount, &DestSeq, &AnimImportSettings, &BoneName](FbxProperty& InProperty, FbxAnimCurve* InCurve, const FString& InCurveName)
+			[this, &OutKeyCount, &DestSeq, &AnimImportSettings, &BoneName, bReimport](FbxProperty& InProperty, FbxAnimCurve* InCurve, const FString& InCurveName)
 		{
-			if (ImportCustomAttributeToBone(DestSeq, InProperty, BoneName, InCurveName, InCurve, AnimImportSettings.AnimTimeSpan))
+			if (ImportCustomAttributeToBone(DestSeq, InProperty, BoneName, InCurveName, InCurve, AnimImportSettings.AnimTimeSpan, bReimport))
 			{
 				OutKeyCount = FMath::Max(OutKeyCount, InCurve ? InCurve->KeyGetCount() : 1);
 			}
@@ -2151,7 +2153,7 @@ void UnFbx::FFbxImporter::ImportAnimationCustomAttribute(FAnimCurveImportSetting
 	}
 }
 
-void UnFbx::FFbxImporter::ImportBoneTracks(USkeleton* Skeleton, FAnimCurveImportSettings& AnimImportSettings, FbxNode* SkeletalMeshRootNode, const int32 ResampleRate, TArray<AnimationTransformDebug::FAnimationTransformDebugData>& TransformDebugData, int32& OutTotalNumKeys)
+void UnFbx::FFbxImporter::ImportBoneTracks(USkeleton* Skeleton, FAnimCurveImportSettings& AnimImportSettings, FbxNode* SkeletalMeshRootNode, const int32 ResampleRate, TArray<AnimationTransformDebug::FAnimationTransformDebugData>& TransformDebugData, int32& OutTotalNumKeys, const bool bReimport)
 {
 	FScopedSlowTask SlowTask(AnimImportSettings.FbxRawBoneNames.Num(), LOCTEXT("BeginImportAnimation", "Importing Animation"), true);
 	SlowTask.MakeDialog();
@@ -2162,6 +2164,7 @@ void UnFbx::FFbxImporter::ImportBoneTracks(USkeleton* Skeleton, FAnimCurveImport
 
 	UAnimSequence* DestSeq = AnimImportSettings.DestSeq;
 	const FbxTimeSpan& AnimTimeSpan = AnimImportSettings.AnimTimeSpan;
+	const bool bShouldTransact = bReimport;
 
 	// Build additional transform matrix
 	UFbxAnimSequenceImportData* TemplateData = Cast<UFbxAnimSequenceImportData>(DestSeq->AssetImportData);
@@ -2194,7 +2197,7 @@ void UnFbx::FFbxImporter::ImportBoneTracks(USkeleton* Skeleton, FAnimCurveImport
 
 	IAnimationDataModel* DataModel = DestSeq->GetDataModel();
 	IAnimationDataController& Controller = DestSeq->GetController();	
-	Controller.SetFrameRate(FFrameRate(ResampleRate, 1));
+	Controller.SetFrameRate(FFrameRate(ResampleRate, 1), bShouldTransact);
 
 	for (int32 SourceTrackIdx = 0; SourceTrackIdx < AnimImportSettings.FbxRawBoneNames.Num(); ++SourceTrackIdx)
 	{
@@ -2324,8 +2327,8 @@ void UnFbx::FFbxImporter::ImportBoneTracks(USkeleton* Skeleton, FAnimCurveImport
 				if (BoneTreeIndex != INDEX_NONE)
 				{
 					//add new track
-					Controller.AddBoneCurve(BoneName);
-					Controller.SetBoneTrackKeys(BoneName, RawTrack.PosKeys, RawTrack.RotKeys, RawTrack.ScaleKeys);
+					Controller.AddBoneCurve(BoneName, bShouldTransact);
+					Controller.SetBoneTrackKeys(BoneName, RawTrack.PosKeys, RawTrack.RotKeys, RawTrack.ScaleKeys, bShouldTransact);
 					NewDebugData.SetTrackData(BoneTreeIndex, BoneName);
 
 					// add mapping to skeleton bone track
@@ -2355,7 +2358,7 @@ void UnFbx::FFbxImporter::ImportBoneTracks(USkeleton* Skeleton, FAnimCurveImport
 									{
 										if (Model->FindAttribute(AttributeIdentifier))
 										{
-											Controller.RemoveAttribute(AttributeIdentifier);
+											Controller.RemoveAttribute(AttributeIdentifier, bShouldTransact);
 										}
 									}
 
@@ -2383,7 +2386,7 @@ void UnFbx::FFbxImporter::ImportBoneTracks(USkeleton* Skeleton, FAnimCurveImport
 
 									// create the attribute and add the transform keys
 									const int32 NumAttributeKeys = (bReduceKeys) ? 1 : NumKeysForTrack;
-									UE::Anim::AddTypedCustomAttribute<FTransformAnimationAttribute, FTransform>(FName(BoneName), TargetBoneName, DestSeq, MakeArrayView(TimeKeys.GetData(), NumAttributeKeys), MakeArrayView(TransformValues.GetData(), NumAttributeKeys));
+									UE::Anim::AddTypedCustomAttribute<FTransformAnimationAttribute, FTransform>(FName(BoneName), TargetBoneName, DestSeq, MakeArrayView(TimeKeys.GetData(), NumAttributeKeys), MakeArrayView(TransformValues.GetData(), NumAttributeKeys), bShouldTransact);
 								}
 							}
 						}
diff --git a/Engine/Source/Editor/UnrealEd/Public/FbxImporter.h b/Engine/Source/Editor/UnrealEd/Public/FbxImporter.h
--- a/Engine/Source/Editor/UnrealEd/Public/FbxImporter.h
+++ b/Engine/Source/Editor/UnrealEd/Public/FbxImporter.h
@@ -903,7 +903,7 @@ public:
 	 * @param ResampleRate	Resample Rate for data
 	 * @param AnimTimeSpan	AnimTimeSpan
 	 */
-	bool ImportAnimation(USkeleton* Skeleton, UAnimSequence* DestSeq, const FString& FileName, TArray<FbxNode*>& SortedLinks, TArray<FbxNode*>& NodeArray, FbxAnimStack* CurAnimStack, const int32 ResampleRate, const FbxTimeSpan AnimTimeSpan);
+	bool ImportAnimation(USkeleton* Skeleton, UAnimSequence* DestSeq, const FString& FileName, TArray<FbxNode*>& SortedLinks, TArray<FbxNode*>& NodeArray, FbxAnimStack* CurAnimStack, const int32 ResampleRate, const FbxTimeSpan AnimTimeSpan, const bool bReimport);
 	/**
 	* Calculate the global Sample Rate for all the nodes in the FbxAnimStack pass in parameter
 	*
@@ -1283,8 +1283,9 @@ private:
 	 * @param AnimImportSettings	Common settings to import animation.
 	 * @param CurAnimStack			The current anim stack we are importing.
 	 * @param OutKeyCount			Out parameter returning the number of keys imported, used to set the number of keys in the sequencer.
+ 	 * @param bReimport				Flag indicating whether or not this operation is part of reimporting.
 	 */
-	void ImportBlendShapeCurves(FAnimCurveImportSettings& AnimImportSettings, FbxAnimStack* CurAnimStack, int32& OutKeyCount);
+	void ImportBlendShapeCurves(FAnimCurveImportSettings& AnimImportSettings, FbxAnimStack* CurAnimStack, int32& OutKeyCount, const bool bReimport);
 
 	/**
 	 * Import the custom attributes curves into the UAnimSequence.
@@ -1292,8 +1293,9 @@ private:
 	 * @param AnimImportSettings	Common settings to import animation.
 	 * @param OutKeyCount			Out parameter returning the number of keys imported, used to set the number of keys in the sequencer.
 	 * @param OutCurvesNotFound		Out parameter returning a list of curves name already present in the AnimSequence that were not imported, indicating that some curve data are missing during a reimport.
+ 	 * @param bReimport				Flag indicating whether or not this operation is part of reimporting.
 	 */
-	void ImportAnimationCustomAttribute(FAnimCurveImportSettings& AnimImportSettings, int32& OutKeyCount, TArray<FString>& OutCurvesNotFound);
+	void ImportAnimationCustomAttribute(FAnimCurveImportSettings& AnimImportSettings, int32& OutKeyCount, TArray<FString>& OutCurvesNotFound, const bool bReimport);
 
 	/**
 	 * Import the bone transforms curves into the UAnimSequence.
@@ -1304,8 +1306,9 @@ private:
 	 * @param ResampleRate			The rate at which the animations are resampled.
 	 * @param TransformDebugData	Out parameter data for internal debugging.
 	 * @param OutTotalNumKeys		Out parameter returning the number of keys imported, used to set the number of keys in the sequencer.
+ 	 * @param bReimport				Flag indicating whether or not this operation is part of reimporting.
 	 */
-	void ImportBoneTracks(USkeleton* Skeleton, FAnimCurveImportSettings& AnimImportSettings, FbxNode* SkeletalMeshRootNode, const int32 ResampleRate, TArray<AnimationTransformDebug::FAnimationTransformDebugData>& TransformDebugData, int32& OutTotalNumKeys);
+	void ImportBoneTracks(USkeleton* Skeleton, FAnimCurveImportSettings& AnimImportSettings, FbxNode* SkeletalMeshRootNode, const int32 ResampleRate, TArray<AnimationTransformDebug::FAnimationTransformDebugData>& TransformDebugData, int32& OutTotalNumKeys, const bool bReimport);
 
 public:
 	// current Fbx scene we are importing. Make sure to release it after import
@@ -1836,12 +1839,12 @@ private:
 	/**
 	 * Import FbxCurve to anim sequence
 	 */
-	bool ImportCurveToAnimSequence(class UAnimSequence * TargetSequence, const FString& CurveName, const FbxAnimCurve* FbxCurve, int32 CurveFlags,const FbxTimeSpan& AnimTimeSpan, float ValueScale = 1.f) const;
+	bool ImportCurveToAnimSequence(class UAnimSequence * TargetSequence, const FString& CurveName, const FbxAnimCurve* FbxCurve, int32 CurveFlags,const FbxTimeSpan& AnimTimeSpan, const bool bReimport, float ValueScale = 1.f) const;
 
 	/**
 	 * Import rich Curves to anim sequence
 	 */
-	bool ImportRichCurvesToAnimSequence(class UAnimSequence * TargetSequence, const TArray<FString>& CurveNames, const TArray<FRichCurve> RichCurves, int32 CurveFlags) const;
+	bool ImportRichCurvesToAnimSequence(class UAnimSequence * TargetSequence, const TArray<FString>& CurveNames, const TArray<FRichCurve> RichCurves, int32 CurveFlags, const bool bReimport) const;
 	
 	/**
 	 * Given a primary blend shape channel curve and inbetween target full weights,
@@ -1862,7 +1865,7 @@ private:
 	 *
 	 * @return Returns true if the given custom attribute was properly added to the bone, false otherwise.
 	 */
-	bool ImportCustomAttributeToBone(class UAnimSequence* TargetSequence, FbxProperty& InProperty, FName BoneName, const FString& CurveName, const FbxAnimCurve* FbxCurve, const FbxTimeSpan& AnimTimeSpan, float ValueScale=1.f);
+	bool ImportCustomAttributeToBone(class UAnimSequence* TargetSequence, FbxProperty& InProperty, FName BoneName, const FString& CurveName, const FbxAnimCurve* FbxCurve, const FbxTimeSpan& AnimTimeSpan, const bool bReimport, float ValueScale=1.f);
 };
 
 
diff --git a/Engine/Source/Runtime/Engine/Classes/Animation/BuiltInAttributeTypes.h b/Engine/Source/Runtime/Engine/Classes/Animation/BuiltInAttributeTypes.h
--- a/Engine/Source/Runtime/Engine/Classes/Animation/BuiltInAttributeTypes.h
+++ b/Engine/Source/Runtime/Engine/Classes/Animation/BuiltInAttributeTypes.h
@@ -375,7 +375,7 @@ namespace UE
 #if WITH_EDITOR
 		/** Helper functionality allowing the user to add an attribute with a typed value array */
 		template<typename AttributeType, typename ValueType>
-		bool AddTypedCustomAttribute(const FName& AttributeName, const FName& BoneName, UAnimSequenceBase* AnimSequenceBase, TArrayView<const float> Keys, TArrayView<const ValueType> Values)
+		bool AddTypedCustomAttribute(const FName& AttributeName, const FName& BoneName, UAnimSequenceBase* AnimSequenceBase, TArrayView<const float> Keys, TArrayView<const ValueType> Values, bool bShouldTransact = true)
 		{
 			const FAnimationAttributeIdentifier Identifier = UAnimationAttributeIdentifierExtensions::CreateAttributeIdentifier(AnimSequenceBase, AttributeName, BoneName, AttributeType::StaticStruct());
 
@@ -390,7 +390,7 @@ namespace UE
 					return Attribute;
 				});
 
-				return Controller.SetTypedAttributeKeys<AttributeType>(Identifier, Keys, MakeArrayView(AttributeValues));
+				return Controller.SetTypedAttributeKeys<AttributeType>(Identifier, Keys, MakeArrayView(AttributeValues), bShouldTransact);
 			}
 
 			return false;
