Removed populating MaterialEditor context menu and palette with MaterialFunction assets in restrictive mode.

Note: this change is done to prevent users from creating materials with functions that are not allowed and would fail to verify. In the future we should improve on this fix by instead filtering material function assets by the FolderPermissionList taking under consideratino that folders may be aliased.

#p4v-preflight-copy 24028662

diff --git a/Engine/Source/Editor/UnrealEd/Private/MaterialGraphSchema.cpp b/Engine/Source/Editor/UnrealEd/Private/MaterialGraphSchema.cpp
--- a/Engine/Source/Editor/UnrealEd/Private/MaterialGraphSchema.cpp
+++ b/Engine/Source/Editor/UnrealEd/Private/MaterialGraphSchema.cpp
@@ -47,6 +47,10 @@
 
 #define LOCTEXT_NAMESPACE "MaterialGraphSchema"
 
+static bool GDisableMaterialFunctionContextMenuActions = false;
+static FAutoConsoleVariableRef DisableMaterialFunctionContextMenuActionsCVar(TEXT("r.Material.DisableMaterialFunctionContextMenuActions"),
+	GDisableMaterialFunctionContextMenuActions, TEXT("Prevents MaterialFunction creation actions from being added to MaterialEditor context menus (False: disabled, True: enabled"));
+
 int32 UMaterialGraphSchema::CurrentCacheRefreshID = 0;
 
 ////////////////////////////////////////
@@ -917,6 +921,13 @@ TSharedPtr<FEdGraphSchemaAction> UMaterialGraphSchema::GetCreateCommentAction()
 
 void UMaterialGraphSchema::GetMaterialFunctionActions(FGraphActionMenuBuilder& ActionMenuBuilder) const
 {
+	// Note: this is a temporary security solution. We need to improve on this by allowing functions that in included folders, 
+	// also considering that some folders may be aliased to different paths.
+	if (GDisableMaterialFunctionContextMenuActions)
+	{
+		return;
+	}
+
 	// Get type of dragged pin
 	uint32 FromPinType = 0;
 	if (ActionMenuBuilder.FromPin)
