Cooker: Use a block allocator for FPackageData to avoid expense of individual malloc calls during the creation-in-bulk of FPackageDatas from base game during DLC cook.
For an empty DLC on a large test project:
UCookCommandlet::Main 8.8s -> 7.3s
#rnx

diff --git a/Engine/Source/Editor/UnrealEd/Private/CookOnTheFlyServer.cpp b/Engine/Source/Editor/UnrealEd/Private/CookOnTheFlyServer.cpp
--- a/Engine/Source/Editor/UnrealEd/Private/CookOnTheFlyServer.cpp
+++ b/Engine/Source/Editor/UnrealEd/Private/CookOnTheFlyServer.cpp
@@ -4847,17 +4847,17 @@ void UCookOnTheFlyServer::PostGarbageCollect()
 
 	GCKeepObjects.Empty();
 
-	for (FPackageData* PackageData : *PackageDatas.Get())
+	PackageDatas->LockAndEnumeratePackageDatas([](FPackageData* PackageData)
 	{
 		if (FGeneratorPackage* GeneratorPackage = PackageData->GetGeneratorPackage())
 		{
 			GeneratorPackage->PostGarbageCollect();
 		}
-	}
-	for (FPackageData* PackageData : *PackageDatas.Get())
+	});
+	PackageDatas->LockAndEnumeratePackageDatas([](FPackageData* PackageData)
 	{
 		PackageData->SetKeepReferencedDuringGC(false);
-	}
+	});
 
 	CookedPackageCountSinceLastGC = 0;
 }
@@ -7040,14 +7040,14 @@ void UCookOnTheFlyServer::OnRequestClusterCompleted(const UE::Cook::FRequestClus
 	{
 		if (bHybridIterativeDebug)
 		{
-			for (FPackageData* PackageData : *PackageDatas)
+			PackageDatas->LockAndEnumeratePackageDatas([](FPackageData* PackageData)
 			{
 				if (!PackageData->AreAllRequestedPlatformsExplored())
 				{
 					UE_LOG(LogCook, Warning, TEXT("Missing dependency: existing requested Package %s was not explored in the first cluster."),
 						*WriteToString<256>(PackageData->GetPackageName()));
 				}
-			}
+			});
 			PackageDatas->SetLogDiscoveredPackages(true); // Turn this on for the rest of the cook after the initial cluster
 		}
 	}
@@ -8608,10 +8608,10 @@ void UCookOnTheFlyServer::ShutdownCookSession()
 		CookAsCookWorkerFinished();
 	}
 
-	for (UE::Cook::FPackageData* PackageData : *PackageDatas)
+	PackageDatas->LockAndEnumeratePackageDatas([](UE::Cook::FPackageData* PackageData)
 	{
 		PackageData->DestroyGeneratorPackage();
-	}
+	});
 
 	if (IsCookByTheBookMode())
 	{
@@ -8863,7 +8863,7 @@ void UCookOnTheFlyServer::ClearPlatformCookedData(const ITargetPlatform* TargetP
 
 void UCookOnTheFlyServer::ResetCook(TConstArrayView<TPair<const ITargetPlatform*, bool>> TargetPlatforms)
 {
-	for (UE::Cook::FPackageData* PackageData : *PackageDatas.Get())
+	PackageDatas->LockAndEnumeratePackageDatas([TargetPlatforms](UE::Cook::FPackageData* PackageData)
 	{
 		for (const TPair<const ITargetPlatform*, bool>& Pair : TargetPlatforms)
 		{
@@ -8881,7 +8881,7 @@ void UCookOnTheFlyServer::ResetCook(TConstArrayView<TPair<const ITargetPlatform*
 				}
 			}
 		}
-	}
+	});
 
 	TArray<FName> PackageNames;
 	for (const TPair<const ITargetPlatform*, bool>& Pair : TargetPlatforms)
@@ -9216,13 +9216,13 @@ void UCookOnTheFlyServer::BeginCookSandbox(FBeginCookContext& BeginContext)
 			// Set the NumPackagesIterativelySkipped field to include all of the already CookedPackages
 			COOK_STAT(DetailedCookStats::NumPackagesIterativelySkipped = 0);
 			const ITargetPlatform* TargetPlatform = AlreadyCookedPlatforms[0];
-			for (UE::Cook::FPackageData* PackageData : *PackageDatas)
+			PackageDatas->LockAndEnumeratePackageDatas([TargetPlatform](UE::Cook::FPackageData* PackageData)
 			{
 				if (PackageData->HasCookedPlatform(TargetPlatform, true /* bIncludeFailed */))
 				{
 					COOK_STAT(++DetailedCookStats::NumPackagesIterativelySkipped);
 				}
-			}
+			});
 		}
 	}
 
@@ -9974,6 +9974,7 @@ void UCookOnTheFlyServer::RecordDLCPackagesFromBaseGame(FBeginCookContext& Begin
 		}
 	}
 
+	bool bFirstAddExistingPackageDatas = true;
 	for (const ITargetPlatform* TargetPlatform : BeginContext.TargetPlatforms)
 	{
 		SCOPED_BOOT_TIMING("AddCookedPlatforms");
@@ -9997,7 +9998,7 @@ void UCookOnTheFlyServer::RecordDLCPackagesFromBaseGame(FBeginCookContext& Begin
 		if (ActivePackageList.Num() > 0)
 		{
 			SCOPED_BOOT_TIMING("AddPackageDataByFileNamesForPlatform");
-			PackageDatas->AddExistingPackageDatasForPlatform(ActivePackageList, TargetPlatform);
+			PackageDatas->AddExistingPackageDatasForPlatform(ActivePackageList, TargetPlatform, bFirstAddExistingPackageDatas);
 		}
 
 		TArray<FName>& PlatformBasedPackages = CookByTheBookOptions->BasedOnReleaseCookedPackages.FindOrAdd(PlatformName);
@@ -10006,6 +10007,7 @@ void UCookOnTheFlyServer::RecordDLCPackagesFromBaseGame(FBeginCookContext& Begin
 		{
 			PlatformBasedPackages.Add(PackageData.NormalizedFileName);
 		}
+		bFirstAddExistingPackageDatas = false;
 	}
 
 	FString ExtraReleaseVersionAssetsFile;
@@ -10433,8 +10435,10 @@ bool UCookOnTheFlyServer::GetAllPackageFilenamesFromAssetRegistry(const FString&
 		int32 NumPackages = SerializedState->GetNumAssets();
 		TArray<const FAssetData*> AssetDatas;
 		TSet<FName> PackageNames;
+		FString PackageNameStrBuffer;
 		AssetDatas.Reserve(NumPackages);
 		OutPackageDatas.Reserve(NumPackages);
+		PackageNames.Reserve(NumPackages);
 
 		// Convert the Map of RegistryData into an Array of FAssetData and populate PackageNames in the output array
 		SerializedState->EnumerateAllAssets([&](const FAssetData& RegistryData)
@@ -10452,7 +10456,8 @@ bool UCookOnTheFlyServer::GetAllPackageFilenamesFromAssetRegistry(const FString&
 			{
 				return;
 			}
-			if (FPackageName::GetPackageMountPoint(PackageName.ToString()).IsNone())
+			PackageName.ToString(PackageNameStrBuffer);
+			if (FPackageName::GetPackageMountPoint(PackageNameStrBuffer).IsNone())
 			{
 				// Skip any packages that are not currently mounted; if we tried to find their FileNames below
 				// we would get log spam
diff --git a/Engine/Source/Editor/UnrealEd/Private/Cooker/CookPackageData.cpp b/Engine/Source/Editor/UnrealEd/Private/Cooker/CookPackageData.cpp
--- a/Engine/Source/Editor/UnrealEd/Private/Cooker/CookPackageData.cpp
+++ b/Engine/Source/Editor/UnrealEd/Private/Cooker/CookPackageData.cpp
@@ -6,6 +6,7 @@
 #include "Algo/Count.h"
 #include "Algo/Find.h"
 #include "Algo/Sort.h"
+#include "Algo/Unique.h"
 #include "AssetCompilingManager.h"
 #include "AssetRegistry/IAssetRegistry.h"
 #include "Async/ParallelFor.h"
@@ -381,7 +382,7 @@ void FPackageData::ClearCookProgress(const ITargetPlatform* TargetPlatform)
 	}
 }
 
-const TSortedMap<const ITargetPlatform*, FPackageData::FPlatformData>& FPackageData::GetPlatformDatas() const
+const TSortedMap<const ITargetPlatform*, FPackageData::FPlatformData, TInlineAllocator<1>>& FPackageData::GetPlatformDatas() const
 {
 	return PlatformDatas;
 }
@@ -1259,7 +1260,7 @@ bool FPackageData::HasReferencedObjects() const
 
 void FPackageData::RemapTargetPlatforms(const TMap<ITargetPlatform*, ITargetPlatform*>& Remap)
 {
-	typedef TSortedMap<const ITargetPlatform*, FPlatformData> MapType;
+	typedef TSortedMap<const ITargetPlatform*, FPlatformData, TInlineAllocator<1>> MapType;
 	MapType NewPlatformDatas;
 	NewPlatformDatas.Reserve(PlatformDatas.Num());
 	for (TPair<const ITargetPlatform*, FPlatformData>& ExistingPair : PlatformDatas)
@@ -2155,6 +2156,8 @@ FPackageDatas::FPackageDatas(UCookOnTheFlyServer& InCookOnTheFlyServer)
 	: CookOnTheFlyServer(InCookOnTheFlyServer)
 	, LastPollAsyncTime(0)
 {
+	Allocator.SetMinBlockSize(1024);
+	Allocator.SetMaxBlockSize(65536);
 }
 
 void FPackageDatas::SetBeginCookConfigSettings(FStringView CookShowInstigator)
@@ -2370,9 +2373,9 @@ FPackageData& FPackageDatas::CreatePackageData(FName PackageName, FName FileName
 {
 	check(!PackageName.IsNone());
 	check(!FileName.IsNone());
-	FPackageData* PackageData = new FPackageData(*this, PackageName, FileName);
 
 	FWriteScopeLock ExistenceWriteLock(ExistenceLock);
+	FPackageData* PackageData = Allocator.NewElement(*this, PackageName, FileName);
 	FPackageData*& ExistingByPackageName = PackageNameToPackageData.FindOrAdd(PackageName);
 	FPackageData*& ExistingByFileName = FileNameToPackageData.FindOrAdd(FileName);
 	if (ExistingByPackageName)
@@ -2387,7 +2390,6 @@ FPackageData& FPackageDatas::CreatePackageData(FName PackageName, FName FileName
 	check(!ExistingByFileName);
 	ExistingByPackageName = PackageData;
 	ExistingByFileName = PackageData;
-	PackageDatas.Add(PackageData);
 	return *PackageData;
 }
 
@@ -2521,71 +2523,55 @@ FName FPackageDatas::GetStandardFileName(FStringView InFileName)
 }
 
 void FPackageDatas::AddExistingPackageDatasForPlatform(TConstArrayView<FConstructPackageData> ExistingPackages,
-	const ITargetPlatform* TargetPlatform)
+	const ITargetPlatform* TargetPlatform, bool bExpectPackageDatasAreNew)
 {
 	int32 NumPackages = ExistingPackages.Num();
-
-	// Make the list unique
-	TMap<FName, FName> UniquePackages;
-	UniquePackages.Reserve(NumPackages);
-	for (const FConstructPackageData& PackageToAdd : ExistingPackages)
+	if (NumPackages == 0)
 	{
-		FName& AddedFileName = UniquePackages.FindOrAdd(PackageToAdd.PackageName, PackageToAdd.NormalizedFileName);
-		check(AddedFileName == PackageToAdd.NormalizedFileName);
+		return;
 	}
-	TArray<FConstructPackageData> UniqueArray;
-	if (UniquePackages.Num() != NumPackages)
-	{
-		NumPackages = UniquePackages.Num();
-		UniqueArray.Reserve(NumPackages);
-		for (TPair<FName, FName>& Pair : UniquePackages)
+
+	// Make the list unique
+	TArray<FConstructPackageData> UniqueArray(ExistingPackages);
+	Algo::Sort(UniqueArray, [](const FConstructPackageData& A, const FConstructPackageData& B)
 		{
-			UniqueArray.Add(FConstructPackageData{ Pair.Key, Pair.Value });
-		}
-		ExistingPackages = UniqueArray;
-	}
+			return A.PackageName.FastLess(B.PackageName);
+		});
+	UniqueArray.SetNum(Algo::Unique(UniqueArray, [](const FConstructPackageData& A, const FConstructPackageData& B)
+		{
+			return A.PackageName == B.PackageName;
+		}));
+	ExistingPackages = UniqueArray;
 
-	// parallelize the read-only operations (and write NewPackageDataObjects by index which has no threading issues)
-	TArray<FPackageData*> NewPackageDataObjects;
-	NewPackageDataObjects.AddZeroed(NumPackages);
 	FWriteScopeLock ExistenceWriteLock(ExistenceLock);
-	ParallelFor(NumPackages,
-		[&ExistingPackages, TargetPlatform, &NewPackageDataObjects, this](int Index)
+	if (bExpectPackageDatasAreNew)
+	{
+		Allocator.ReserveDelta(NumPackages);
+		FileNameToPackageData.Reserve(FileNameToPackageData.Num() + NumPackages);
+		PackageNameToPackageData.Reserve(PackageNameToPackageData.Num() + NumPackages);
+	}
+
+	// Create the PackageDatas and mark them as cooked
+	for (const FConstructPackageData& ConstructData : ExistingPackages)
 	{
-		FName PackageName = ExistingPackages[Index].PackageName;
-		FName NormalizedFileName = ExistingPackages[Index].NormalizedFileName;
+		FName PackageName = ConstructData.PackageName;
+		FName NormalizedFileName = ConstructData.NormalizedFileName;
 		check(!PackageName.IsNone());
 		check(!NormalizedFileName.IsNone());
 
-		FPackageData** PackageDataMapAddr = FileNameToPackageData.Find(NormalizedFileName);
-		FPackageData* PackageData = nullptr;
-		if (PackageDataMapAddr != nullptr)
-		{
-			PackageData = *PackageDataMapAddr;
-		}
-		else
-		{
-			// create the package data and remember it for updating caches after the the ParallelFor
-			PackageData = new FPackageData(*this, PackageName, NormalizedFileName);
-			NewPackageDataObjects[Index] = PackageData;
-		}
-		PackageData->SetPlatformCooked(TargetPlatform, true /* Succeeded */);
-	});
-
-	// update cache for all newly created objects (copied from CreatePackageData)
-	for (FPackageData* PackageData : NewPackageDataObjects)
-	{
-		if (PackageData)
+		FPackageData*& PackageData = FileNameToPackageData.FindOrAdd(NormalizedFileName, nullptr);
+		if (!PackageData)
 		{
-			FPackageData* ExistingByFileName = FileNameToPackageData.Add(PackageData->FileName, PackageData);
-			// We looked up by FileName in the loop; it should still have been unset before the write we just did
-			check(ExistingByFileName == PackageData);
-			FPackageData* ExistingByPackageName = PackageNameToPackageData.FindOrAdd(PackageData->PackageName, PackageData);
+			// create the package data (copied from CreatePackageData)
+			FPackageData* NewPackageData = Allocator.NewElement(*this, PackageName, NormalizedFileName);
+			FPackageData* ExistingByPackageName = PackageNameToPackageData.FindOrAdd(PackageName, NewPackageData);
 			// If no other CreatePackageData added the FileName, then they should not have added
 			// the PackageName either
-			check(ExistingByPackageName == PackageData);
-			PackageDatas.Add(PackageData);
+			check(ExistingByPackageName == NewPackageData);
+
+			PackageData = NewPackageData;
 		}
+		PackageData->SetPlatformCooked(TargetPlatform, true /* Succeeded */);
 	}
 }
 
@@ -2635,7 +2621,8 @@ int32 FPackageDatas::GetNumCooked()
 void FPackageDatas::GetCookedPackagesForPlatform(const ITargetPlatform* Platform, TArray<FPackageData*>& CookedPackages,
 	bool bGetFailedCookedPackages, bool bGetSuccessfulCookedPackages)
 {
-	for (FPackageData* PackageData : PackageDatas)
+	LockAndEnumeratePackageDatas(
+	[Platform, &CookedPackages, bGetFailedCookedPackages, bGetSuccessfulCookedPackages](FPackageData* PackageData)
 	{
 		ECookResult CookResults = PackageData->GetCookResults(Platform);
 		if (((CookResults == ECookResult::Succeeded) & (bGetSuccessfulCookedPackages != 0)) |
@@ -2643,7 +2630,7 @@ void FPackageDatas::GetCookedPackagesForPlatform(const ITargetPlatform* Platform
 		{
 			CookedPackages.Add(PackageData);
 		}
-	}
+	});
 }
 
 void FPackageDatas::Clear()
@@ -2654,32 +2641,36 @@ void FPackageDatas::Clear()
 	SaveQueue.Empty();
 	PackageNameToPackageData.Empty();
 	FileNameToPackageData.Empty();
-	for (FPackageData* PackageData : PackageDatas)
-	{
-		PackageData->ClearReferences();
-	}
-	for (FPackageData* PackageData : PackageDatas)
 	{
-		delete PackageData;
+		// All references must be cleared before any PackageDatas are destroyed
+		EnumeratePackageDatasWithinLock([](FPackageData* PackageData)
+		{
+			PackageData->ClearReferences();
+		});
+		EnumeratePackageDatasWithinLock([](FPackageData* PackageData)
+		{
+			PackageData->~FPackageData();
+		});
+		Allocator.Empty();
 	}
-	PackageDatas.Empty();
+
 	ShowInstigatorPackageData = nullptr;
 }
 
 void FPackageDatas::ClearCookedPlatforms()
 {
-	for (FPackageData* PackageData : PackageDatas)
+	LockAndEnumeratePackageDatas([](FPackageData* PackageData)
 	{
 		PackageData->ClearCookProgress();
-	}
+	});
 }
 
 void FPackageDatas::OnRemoveSessionPlatform(const ITargetPlatform* TargetPlatform)
 {
-	for (FPackageData* PackageData : PackageDatas)
+	LockAndEnumeratePackageDatas([TargetPlatform](FPackageData* PackageData)
 	{
 		PackageData->OnRemoveSessionPlatform(TargetPlatform);
-	}
+	});
 }
 
 constexpr int32 PendingPlatformDataReservationSize = 128;
@@ -2789,22 +2780,12 @@ void FPackageDatas::PollPendingCookedPlatformDatas(bool bForce, double& LastCook
 	}
 }
 
-TArray<FPackageData*>::RangedForIteratorType FPackageDatas::begin()
-{
-	return PackageDatas.begin();
-}
-
-TArray<FPackageData*>::RangedForIteratorType FPackageDatas::end()
-{
-	return PackageDatas.end();
-}
-
 void FPackageDatas::RemapTargetPlatforms(const TMap<ITargetPlatform*, ITargetPlatform*>& Remap)
 {
-	for (FPackageData* PackageData : PackageDatas)
+	LockAndEnumeratePackageDatas([&Remap](FPackageData* PackageData)
 	{
 		PackageData->RemapTargetPlatforms(Remap);
-	}
+	});
 	ForEachPendingCookedPlatformData([&Remap](FPendingCookedPlatformData& CookedPlatformData)
 	{
 		CookedPlatformData.RemapTargetPlatforms(Remap);
diff --git a/Engine/Source/Editor/UnrealEd/Private/Cooker/CookPackageData.h b/Engine/Source/Editor/UnrealEd/Private/Cooker/CookPackageData.h
--- a/Engine/Source/Editor/UnrealEd/Private/Cooker/CookPackageData.h
+++ b/Engine/Source/Editor/UnrealEd/Private/Cooker/CookPackageData.h
@@ -16,7 +16,10 @@
 #include "HAL/Platform.h"
 #include "Misc/EnumClassFlags.h"
 #include "Misc/Optional.h"
+#include "Misc/ScopeRWLock.h"
 #include "Templates/SharedPointer.h"
+#include "Templates/UniquePtr.h"
+#include "TypedBlockAllocator.h"
 #include "UObject/GCObject.h"
 #include "UObject/NameTypes.h"
 #include "UObject/PackageResourceManager.h"
@@ -260,7 +263,7 @@ public:
 	void ClearCookProgress(const ITargetPlatform* TargetPlatform);
 
 	/** Access the information about platforms interacted with by *this. */
-	const TSortedMap<const ITargetPlatform*, FPlatformData>& GetPlatformDatas() const;
+	const TSortedMap<const ITargetPlatform*, FPlatformData, TInlineAllocator<1>>& GetPlatformDatas() const;
 
 	/** Add a platform if not already existing and return a writable pointer to its flags. */
 	FPlatformData& FindOrAddPlatformData(const ITargetPlatform* TargetPlatform);
@@ -574,7 +577,7 @@ private:
 	TUniquePtr<FGeneratorPackage> GeneratorPackage;
 	FGeneratorPackage* GeneratedOwner;
 	/** Data for each platform that has been interacted with by *this. */
-	TSortedMap<const ITargetPlatform*, FPlatformData> PlatformDatas;
+	TSortedMap<const ITargetPlatform*, FPlatformData, TInlineAllocator<1>> PlatformDatas;
 	TArray<FWeakObjectPtr> CachedObjectsInOuter;
 	FCompletionCallback CompletionCallback;
 	FName PackageName;
@@ -1164,7 +1167,7 @@ public:
 
 	/** Create and mark-cooked a batch of PackageDatas, used by DLC for cooked-in-earlier-release packages. */
 	void AddExistingPackageDatasForPlatform(TConstArrayView<FConstructPackageData> ExistingPackages,
-		const ITargetPlatform* TargetPlatform);
+		const ITargetPlatform* TargetPlatform, bool bExpectPackageDatasAreNew);
 
 	/**
 	 * Try to find the PackageData for the given PackageName.
@@ -1215,10 +1218,6 @@ public:
 	 */
 	void PollPendingCookedPlatformDatas(bool bForce, double& LastCookableObjectTickTime);
 
-	/** RangedFor methods for iterating over all FPackageData managed by this FPackageDatas */
-	TArray<FPackageData*>::RangedForIteratorType begin();
-	TArray<FPackageData*>::RangedForIteratorType end();
-
 	/** Swap all ITargetPlatform* stored on this instance according to the mapping in @param Remap. */
 	void RemapTargetPlatforms(const TMap<ITargetPlatform*, ITargetPlatform*>& Remap);
 
@@ -1228,21 +1227,31 @@ public:
 
 	/** Called when a PackageData assigns its instigator, for debugging. */
 	void DebugInstigator(FPackageData& PackageData);
+
+	/** Enter the required locks and enumerate all created PackageDatas. */
+	template <typename CallbackType>
+	void LockAndEnumeratePackageDatas(CallbackType&& Callback);
+
 private:
 	/**
 	 * Construct a new FPackageData with the given PackageName and FileName and store references to it in the maps.
 	 * New FPackageData are always created in the Idle state.
 	 */
 	FPackageData& CreatePackageData(FName PackageName, FName FileName);
+	/** Called from within ExistenceLock, enumerate all created PackageDatas. */
+	template <typename CallbackType>
+	void EnumeratePackageDatasWithinLock(CallbackType&& Callback);
 	/** Return whether a filename exists for the package on disk, and return it, unnormalized. */
 	static bool TryLookupFileNameOnDisk(FName PackageName, FString& OutFileName);
 	/** Return the corresponding PackageName if the normalized filename exists on disk. */
 	static FName LookupPackageNameOnDisk(FName NormalizedFileName, bool bExactMatchRequired, FName& FoundFileName);
 
-	/** Collection of pointers to all FPackageData that have been constructed by *this. */
-	TArray<FPackageData*> PackageDatas;
+	/** Allocator for PackageDatas Guarded by ExistenceLock. */
+	TTypedBlockAllocator<FPackageData> Allocator;
 	FPackageDataMonitor Monitor;
+	/** Guarded by ExistenceLock */
 	TMap<FName, FPackageData*> PackageNameToPackageData;
+	/** Guarded by ExistenceLock */
 	TMap<FName, FPackageData*> FileNameToPackageData;
 	TRingBuffer<FPendingCookedPlatformDataContainer> PendingCookedPlatformDataLists;
 	int32 PendingCookedPlatformDataNum = 0;
@@ -1260,6 +1269,19 @@ private:
 	static IAssetRegistry* AssetRegistry;
 };
 
+template <typename CallbackType>
+inline void FPackageDatas::LockAndEnumeratePackageDatas(CallbackType&& Callback)
+{
+	FReadScopeLock ExistenceReadLock(ExistenceLock);
+	EnumeratePackageDatasWithinLock(Forward<CallbackType>(Callback));
+}
+
+template <typename CallbackType>
+inline void FPackageDatas::EnumeratePackageDatasWithinLock(CallbackType&& Callback)
+{
+	Allocator.EnumerateAllocations(Forward<CallbackType>(Callback));
+}
+
 /**
  * A debug-only scope class to confirm that each FPackageData removed from a container during a Pump function
  * is added to the container for its new state before leaving the Pump function.
diff --git a/Engine/Source/Editor/UnrealEd/Private/Cooker/TypedBlockAllocator.h b/Engine/Source/Editor/UnrealEd/Private/Cooker/TypedBlockAllocator.h
new file mode 100644
--- /dev/null
+++ b/Engine/Source/Editor/UnrealEd/Private/Cooker/TypedBlockAllocator.h
@@ -0,0 +1,268 @@
+// Copyright Epic Games, Inc. All Rights Reserved.
+
+#pragma once
+
+#include "Containers/Array.h"
+#include "Containers/ArrayView.h"
+#include "HAL/Platform.h"
+#include "HAL/PlatformMath.h"
+#include "Templates/TypeCompatibleBytes.h"
+#include "Templates/UniquePtr.h"
+#include "Templates/UnrealTemplate.h"
+
+/**
+ * An allocator that batches allocation calls into blocks to reduce malloc costs.
+ * 
+ * Not threadsafe; must be guarded by critical section if used from multiple threads.
+ * 
+ * Block size is adaptive; it doubles with each new block, up to MaxBlockSize.
+ * 
+ * Allocated memory is not released until Empty is called.
+ */
+template <typename ElementType>
+class TTypedBlockAllocator
+{
+public:
+	/** Set the MinimumBlockSize. Default is 1024. All new blocks will be this size or larger. */
+	void SetMinBlockSize(int32 BlockSize);
+	int32 GetMinBlockSize() const;
+	/**
+	 * Set the MaximumBlockSize. Default is 65536. New blocks added by Alloc will be this size or
+	 * less (ignored if MaximumBlockSize < MinimumBlockSize). Blocks added by Reserve are not
+	 * limited by MaximumBlockSize. 
+	 */
+	void SetMaxBlockSize(int32 BlockSize);
+	int32 GetMaxBlockSize() const;
+
+	/** Return the memory for an Element without calling a constructor. */
+	ElementType* Alloc();
+	/**
+	 * Make the memory for an Element returned from Alloc or Construct available again to Alloc.
+	 * Does not call the ElementType destructor. Does not return the memory to the inner allocator.
+	 */
+	void Free(ElementType* Element);
+	/** Allocate an Element and call constructor with the given arguments. */
+	template <typename... ArgsType>
+	ElementType* NewElement(ArgsType&&... Args);
+	/** Call destructor on the given element and call Free on its memory. */
+	void DeleteElement(ElementType* Element);
+
+	/**
+	 * Call the given Callback on every element that has been returned from Alloc. Not valid to call if any
+	 * elements have been Freed because it is too expensive to prevent calling the Callback on the freed elements.
+	 */
+	template <typename CallbackType>
+	void EnumerateAllocations(CallbackType&& Callback);
+
+	/**
+	 * Allocate enough memory from the inner allocator to ensure that AllocationCount more calls to Alloc can be made
+	 * without further calls to the inner allocator.
+	 * 
+	 * @param AllocationCount The number of future calls to Alloc that should reserved.
+	 * @param InMaxBlockSize If non-zero, the maximum capacity of any Blocks allocated by ReserveDelta.
+	 *        If zero, MaxBlockSize is used.
+	 */
+	void ReserveDelta(int32 AllocationCount, int32 InMaxBlockSize = 0);
+	/**
+	 * Release all allocated memory. For performance, Empty does not require that allocations have been destructed
+	 * or freed; caller is responsible for calling any necessary ElementType destructors and for dropping all
+	 * references to the allocated ElementTypes.
+	 */
+	void Empty();
+
+private:
+	struct FAllocationBlock
+	{
+		FAllocationBlock(int32 InCapacity);
+
+		TUniquePtr<TTypeCompatibleBytes<ElementType>[]> Elements;
+		int32 NextIndex;
+		int32 Capacity;
+	};
+	TArray<FAllocationBlock> Blocks;
+	ElementType* FreeList = nullptr;
+	int32 NextBlock = 0;
+	int32 NumAllocations = 0;
+	int32 NumFreeList = 0;
+	int32 MinBlockSize = 1024;
+	int32 MaxBlockSize = 65536;
+};
+
+
+template <typename ElementType>
+inline void TTypedBlockAllocator<ElementType>::SetMinBlockSize(int32 BlockSize)
+{
+	check(BlockSize >= 0);
+	MinBlockSize = FMath::Max(BlockSize, 1);
+}
+
+template <typename ElementType>
+inline int32 TTypedBlockAllocator<ElementType>::GetMinBlockSize() const
+{
+	return MinBlockSize;
+}
+
+template <typename ElementType>
+inline void TTypedBlockAllocator<ElementType>::SetMaxBlockSize(int32 BlockSize)
+{
+	check(BlockSize >= 0);
+	MaxBlockSize = BlockSize;
+}
+
+template <typename ElementType>
+inline int32 TTypedBlockAllocator<ElementType>::GetMaxBlockSize() const
+{
+	return MaxBlockSize;
+}
+
+template <typename ElementType>
+ElementType* TTypedBlockAllocator<ElementType>::Alloc()
+{
+	ElementType* Result = nullptr;
+	if (FreeList)
+	{
+		check(NumFreeList > 0);
+		Result = FreeList;
+		FreeList = *reinterpret_cast<ElementType**>(FreeList);
+		++NumAllocations;
+		--NumFreeList;
+		return Result;
+	}
+
+	FAllocationBlock* Block = nullptr;
+	for (; NextBlock < Blocks.Num(); ++NextBlock)
+	{
+		if (Blocks[NextBlock].NextIndex < Blocks[NextBlock].Capacity)
+		{
+			Block = &Blocks[NextBlock];
+			break;
+		}
+	}
+
+	if (!Block)
+	{
+		check(NextBlock == Blocks.Num());
+		// Double our NumAllocations with each new block, up to a maximum block size
+		int32 BlockCapacity = FMath::Max(NumAllocations, MinBlockSize);
+		if (MaxBlockSize > MinBlockSize)
+		{
+			BlockCapacity = FMath::Min(BlockCapacity, MaxBlockSize);
+		}
+		Block = &Blocks.Emplace_GetRef(BlockCapacity);
+	}
+
+	Result = Block->Elements[Block->NextIndex].GetTypedPtr();
+	++Block->NextIndex;
+	++NumAllocations;
+	return Result;
+}
+
+template <typename ElementType>
+void TTypedBlockAllocator<ElementType>::Free(ElementType* Element)
+{
+	if constexpr (sizeof(ElementType) <= sizeof(ElementType*))
+	{
+		static_assert(sizeof(ElementType) >= sizeof(ElementType*),
+			"The FreeList is implemented by storing pointers within each freed element. To use the FreeList, elementsize must be >= pointer size.");
+	}
+	check(NumAllocations > 0);
+
+	*reinterpret_cast<ElementType**>(Element) = FreeList;
+	FreeList = Element;
+	--NumAllocations;
+	++NumFreeList;
+}
+
+template <typename ElementType>
+template <typename... ArgsType>
+ElementType* TTypedBlockAllocator<ElementType>::NewElement(ArgsType&&... Args)
+{
+	return new(Alloc()) ElementType(Forward<ArgsType>(Args)...);
+}
+
+template <typename ElementType>
+void TTypedBlockAllocator<ElementType>::DeleteElement(ElementType* Element)
+{
+	// We need a typedef here because VC won't compile the destructor call below if ElementType itself has a member
+	// called ElementType
+	typedef ElementType TTypedBlockAllocatorDestructElementType;
+	((TTypedBlockAllocatorDestructElementType*)Element)->
+		TTypedBlockAllocatorDestructElementType::~TTypedBlockAllocatorDestructElementType();
+	Free(Element);
+}
+
+template <typename ElementType>
+template <typename CallbackType>
+void TTypedBlockAllocator<ElementType>::EnumerateAllocations(CallbackType&& Callback)
+{
+	checkf(FreeList == nullptr, TEXT("It is invalid to call EnumerateAllocations after calling Free or Destruct."));
+	for (const FAllocationBlock& Block : Blocks)
+	{
+		for (TTypeCompatibleBytes<ElementType>& Bytes :
+			TArrayView<TTypeCompatibleBytes<ElementType>>(Block.Elements.Get(), Block.NextIndex))
+		{
+			ElementType* Element = Bytes.GetTypedPtr();
+			Callback(Element);
+		}
+	}
+}
+
+template <typename ElementType>
+void TTypedBlockAllocator<ElementType>::ReserveDelta(int32 AllocationCount, int32 InMaxBlockSize)
+{
+	int32 DeltaAllocationCount = AllocationCount - NumFreeList;
+	for (const FAllocationBlock& Block: TConstArrayView<FAllocationBlock>(Blocks).RightChop(NextBlock))
+	{
+		DeltaAllocationCount -= Block.Capacity - Block.NextIndex;
+	}
+	if (InMaxBlockSize == 0)
+	{
+		InMaxBlockSize = MaxBlockSize;
+	}
+	if (DeltaAllocationCount <= 0)
+	{
+		return;
+	}
+
+	// Allocate blocks until we have enough capacity to cover the Reservation. As with blocks allocated from
+	// Alloc, set the unclamped capacity high enough to double our number of allocations. But also set it high
+	// enough to cover the remaining count, and use the different MaxBlockSize if passed in.
+	int32 BlockCapacity = FMath::Max(DeltaAllocationCount, NumAllocations);
+	BlockCapacity = FMath::Max(BlockCapacity, MinBlockSize);
+	if (InMaxBlockSize > MinBlockSize)
+	{
+		BlockCapacity = FMath::Min(BlockCapacity, InMaxBlockSize);
+	}
+	Blocks.Emplace(BlockCapacity);
+	DeltaAllocationCount -= BlockCapacity;
+
+	if (DeltaAllocationCount > 0)
+	{
+		// For further blocks, use the MaxBlockSize
+		check(InMaxBlockSize > MinBlockSize); // Otherwise we would have allocated all of DeltaAllocationCount in a single block
+		BlockCapacity = InMaxBlockSize;
+		while (DeltaAllocationCount > 0)
+		{
+			Blocks.Emplace(BlockCapacity);
+			DeltaAllocationCount -= BlockCapacity;
+		}
+	}
+}
+
+template <typename ElementType>
+void TTypedBlockAllocator<ElementType>::Empty()
+{
+	Blocks.Empty();
+	FreeList = nullptr;
+	NextBlock = 0;
+	NumAllocations = 0;
+	NumFreeList = 0;
+}
+
+template <typename ElementType>
+TTypedBlockAllocator<ElementType>::FAllocationBlock::FAllocationBlock(int32 InCapacity)
+	: Elements(InCapacity > 0 ? new TTypeCompatibleBytes<ElementType>[InCapacity] : nullptr)
+	, NextIndex(0)
+	, Capacity(InCapacity)
+{
+}
