Fixed gradient preview in the Color Curve Editor

diff --git a/Engine/Source/Editor/UnrealEd/Private/SColorGradientEditor.cpp b/Engine/Source/Editor/UnrealEd/Private/SColorGradientEditor.cpp
--- a/Engine/Source/Editor/UnrealEd/Private/SColorGradientEditor.cpp
+++ b/Engine/Source/Editor/UnrealEd/Private/SColorGradientEditor.cpp
@@ -209,9 +209,6 @@ int32 SColorGradientEditor::OnPaint( const FPaintArgs& Args, const FGeometry& Al
 
 			// Sample the curve
 			FLinearColor Color = CurveOwner->GetLinearColorValue( Time );
-			//Slate MakeGradient call expects the linear colors to be pre-converted to sRGB
-			FColor ColorNosRGB = Color.ToFColor(true);
-			Color = ColorNosRGB.ReinterpretAsLinear();
 			if( !bHasAnyAlphaKeys )
 			{
 				// Only show alpha if there is at least one key.  For some curves, alpha may not be important
