Remove specific version of WinSDK when installing visual studio as old requested version cannot be installed in VS2022

#rnx

diff --git a/Engine/Source/Editor/UnrealEd/Private/SourceCodeNavigation.cpp b/Engine/Source/Editor/UnrealEd/Private/SourceCodeNavigation.cpp
--- a/Engine/Source/Editor/UnrealEd/Private/SourceCodeNavigation.cpp
+++ b/Engine/Source/Editor/UnrealEd/Private/SourceCodeNavigation.cpp
@@ -2193,7 +2193,7 @@ FString FSourceCodeNavigationImpl::GetSuggestedIDEInstallerFileName()
 void FSourceCodeNavigationImpl::LaunchIDEInstaller(const FString& Filepath)
 {
 #if PLATFORM_WINDOWS
-	auto Params = TEXT("--productId \"Microsoft.VisualStudio.Product.Community\" --add \"Microsoft.VisualStudio.Workload.NativeDesktop\" --add \"Microsoft.VisualStudio.Workload.NativeGame\" --add \"Component.Unreal\" --add \"Microsoft.VisualStudio.Component.Windows10SDK.17763\" --campaign \"EpicGames_UE4\"");
+	auto Params = TEXT("--productId \"Microsoft.VisualStudio.Product.Community\" --add \"Microsoft.VisualStudio.Workload.NativeDesktop\" --add \"Microsoft.VisualStudio.Workload.NativeGame\" --add \"Component.Unreal\" --add \"Microsoft.VisualStudio.Component.Windows10SDK\" --campaign \"EpicGames_UE5\"");
 	FPlatformProcess::ExecElevatedProcess(*Filepath, Params, nullptr);
 #endif
 }
