Change asset tools import extension support to FNamePermissionList instead of FString. This allow to control the permissions like all the other permission.

diff --git a/Engine/Plugins/Editor/ContentBrowser/ContentBrowserAssetDataSource/Source/ContentBrowserAssetDataSource/Private/ContentBrowserAssetDataSource.cpp b/Engine/Plugins/Editor/ContentBrowser/ContentBrowserAssetDataSource/Source/ContentBrowserAssetDataSource/Private/ContentBrowserAssetDataSource.cpp
--- a/Engine/Plugins/Editor/ContentBrowser/ContentBrowserAssetDataSource/Source/ContentBrowserAssetDataSource/Private/ContentBrowserAssetDataSource.cpp
+++ b/Engine/Plugins/Editor/ContentBrowser/ContentBrowserAssetDataSource/Source/ContentBrowserAssetDataSource/Private/ContentBrowserAssetDataSource.cpp
@@ -1859,6 +1859,21 @@ bool UContentBrowserAssetDataSource::CanHandleDragDropEvent(const FContentBrowse
 			{
 				NewDragCursor = EMouseCursor::SlashedCircle;
 			}
+			else if (ExternalDragDropOp->HasFiles())
+			{
+				bool bSupportOneFile = false;
+				for (const FString& File : ExternalDragDropOp->GetFiles())
+				{
+					if (AssetTools->IsImportExtensionAllowed(FPaths::GetExtension(File)))
+					{
+						bSupportOneFile = true;
+					}
+				}
+				if (!bSupportOneFile)
+				{
+					NewDragCursor = EMouseCursor::SlashedCircle;
+				}
+			}
 			ExternalDragDropOp->SetCursorOverride(NewDragCursor);
 
 			return true; // We will handle this drop, even if the result is invalid (eg, read-only folder)
@@ -1890,10 +1905,35 @@ bool UContentBrowserAssetDataSource::HandleDragDropOnItem(const FContentBrowserI
 		if (TSharedPtr<FExternalDragOperation> ExternalDragDropOp = InDragDropEvent.GetOperationAs<FExternalDragOperation>())
 		{
 			FText ErrorMsg;
+			FString UnsupportedFiles;
 			if (ExternalDragDropOp->HasFiles() && ContentBrowserAssetData::CanModifyPath(AssetTools, FolderPayload->GetInternalPath(), &ErrorMsg))
 			{
-				// Delay import until next tick to avoid blocking the process that files were dragged from
-				GEditor->GetEditorSubsystem<UImportSubsystem>()->ImportNextTick(ExternalDragDropOp->GetFiles(), FolderPayload->GetInternalPath().ToString());
+				TArray<FString> ImportFiles;
+				for (const FString& File : ExternalDragDropOp->GetFiles())
+				{
+					if (AssetTools->IsImportExtensionAllowed(FPaths::GetExtension(File)))
+					{
+						ImportFiles.AddUnique(File);
+					}
+					else
+					{
+						if (!UnsupportedFiles.IsEmpty())
+						{
+							UnsupportedFiles += TEXT("\n");
+						}
+						UnsupportedFiles += File;
+					}
+				}
+				if (ImportFiles.Num() > 0)
+				{
+					// Delay import until next tick to avoid blocking the process that files were dragged from
+					GEditor->GetEditorSubsystem<UImportSubsystem>()->ImportNextTick(ImportFiles, FolderPayload->GetInternalPath().ToString());
+				}
+			}
+
+			if (!UnsupportedFiles.IsEmpty())
+			{
+				ErrorMsg = FText::Format(LOCTEXT("HandleDragDropOnItemUnsupportedFile", "Unsupported file format to import:\n\t{0}"), FText::FromString(UnsupportedFiles));
 			}
 
 			if (!ErrorMsg.IsEmpty())
diff --git a/Engine/Plugins/Interchange/Editor/Source/Utilities/Private/InterchangeOpenFileDialog.cpp b/Engine/Plugins/Interchange/Editor/Source/Utilities/Private/InterchangeOpenFileDialog.cpp
--- a/Engine/Plugins/Interchange/Editor/Source/Utilities/Private/InterchangeOpenFileDialog.cpp
+++ b/Engine/Plugins/Interchange/Editor/Source/Utilities/Private/InterchangeOpenFileDialog.cpp
@@ -1,7 +1,6 @@
 // Copyright Epic Games, Inc. All Rights Reserved. 
 #include "InterchangeOpenFileDialog.h"
 
-#include "AssetToolsModule.h"
 #include "DesktopPlatformModule.h"
 #include "Framework/Application/SlateApplication.h"
 #include "IDesktopPlatform.h"
@@ -17,8 +16,7 @@ namespace UE::Interchange::Utilities::Private
 		FString FileTypes;
 		FString Extensions;
 
-		IAssetTools& AssetTools = FModuleManager::LoadModuleChecked<FAssetToolsModule>(TEXT("AssetTools")).Get();
-		ObjectTools::AppendFormatsFileExtensions(TranslatorFormats, FileTypes, Extensions, AssetTools.GetSupportedImportExtension());
+		ObjectTools::AppendFormatsFileExtensions(TranslatorFormats, FileTypes, Extensions);
 
 		const FString FormatString = FString::Printf(TEXT("All Files (%s)|%s|%s"), *Extensions, *Extensions, *FileTypes);
 
diff --git a/Engine/Source/Developer/AssetTools/Private/AssetTools.cpp b/Engine/Source/Developer/AssetTools/Private/AssetTools.cpp
--- a/Engine/Source/Developer/AssetTools/Private/AssetTools.cpp
+++ b/Engine/Source/Developer/AssetTools/Private/AssetTools.cpp
@@ -1260,6 +1260,7 @@ UAssetToolsImpl::UAssetToolsImpl(const FObjectInitializer& ObjectInitializer)
 	, AssetFixUpRedirectors(MakeShareable(new FAssetFixUpRedirectors))
 	, NextUserCategoryBit(EAssetTypeCategories::FirstUser)
 	, AssetClassPermissionList_DEPRECATED(MakeShared<FNamePermissionList>())
+	, ImportExtensionPermissionList(MakeShared<FNamePermissionList>())
 	, FolderPermissionList(MakeShared<FPathPermissionList>())
 	, WritableFolderPermissionList(MakeShared<FPathPermissionList>())
 	, CreateAssetsAsExternallyReferenceable(true)
@@ -1287,8 +1288,13 @@ UAssetToolsImpl::UAssetToolsImpl(const FObjectInitializer& ObjectInitializer)
 		}
 	}
 
-	
+	TArray<FString> SupportedImportExtensionArray;
 	GConfig->GetArray(TEXT("AssetTools"), TEXT("SupportedImportExtensions"), SupportedImportExtensionArray, GEditorIni);
+	for (FString& Extension : SupportedImportExtensionArray)
+	{
+		ImportExtensionPermissionList->AddAllowListItem("AssetToolsConfigFile", FName(*Extension));
+	}
+
 
 	TArray<FString> DenyListedViewPath;
 	GConfig->GetArray(TEXT("AssetTools"), TEXT("DenyListAssetPaths"), DenyListedViewPath, GEditorIni);
@@ -5610,9 +5616,16 @@ const TSharedRef<FPathPermissionList>& UAssetToolsImpl::GetAssetClassPathPermiss
 	return Empty;
 }
 
-const TArray<FString>& UAssetToolsImpl::GetSupportedImportExtension() const
+const TSharedRef<FNamePermissionList>& UAssetToolsImpl::GetImportExtensionPermissionList() const
+{
+	return ImportExtensionPermissionList;
+}
+
+bool UAssetToolsImpl::IsImportExtensionAllowed(const FString& Extension) const
 {
-	return SupportedImportExtensionArray;
+	//Always test the extension in lower case
+	FName NameExtension(*Extension);
+	return ImportExtensionPermissionList->PassesFilter(NameExtension);
 }
 
 TSet<EBlueprintType>& UAssetToolsImpl::GetAllowedBlueprintTypes()
@@ -5662,12 +5675,12 @@ TArray<UObject*> UAssetToolsImpl::ImportAssetsWithDialogImplementation(const FSt
 
 	TMultiMap<uint32, UFactory*> FilterIndexToFactory;
 	// Generate the file types and extensions represented by the selected factories
-	ObjectTools::GenerateFactoryFileExtensions(Factories, FileTypes, AllExtensions, FilterIndexToFactory, GetSupportedImportExtension());
+	ObjectTools::GenerateFactoryFileExtensions(Factories, FileTypes, AllExtensions, FilterIndexToFactory);
 
 	if (UInterchangeManager::IsInterchangeImportEnabled())
 	{
 		TArray<FString> InterchangeFileExtensions = UInterchangeManager::GetInterchangeManager().GetSupportedFormats(EInterchangeTranslatorType::Assets);
-		ObjectTools::AppendFormatsFileExtensions(InterchangeFileExtensions, FileTypes, AllExtensions, FilterIndexToFactory, GetSupportedImportExtension());
+		ObjectTools::AppendFormatsFileExtensions(InterchangeFileExtensions, FileTypes, AllExtensions, FilterIndexToFactory);
 	}
 
 	FileTypes = FString::Printf(TEXT("All Files (%s)|%s|%s"), *AllExtensions, *AllExtensions, *FileTypes);
diff --git a/Engine/Source/Developer/AssetTools/Private/AssetTools.h b/Engine/Source/Developer/AssetTools/Private/AssetTools.h
--- a/Engine/Source/Developer/AssetTools/Private/AssetTools.h
+++ b/Engine/Source/Developer/AssetTools/Private/AssetTools.h
@@ -135,7 +135,8 @@ public:
 	UE_DEPRECATED(5.1, "Class names are now represented by path names. Please use GetAssetClassPathPermissionList.")
 	TSharedRef<FNamePermissionList>& GetAssetClassPermissionList(EAssetClassAction AssetClassAction);
 	virtual const TSharedRef<FPathPermissionList>& GetAssetClassPathPermissionList(EAssetClassAction AssetClassAction) const override;
-	virtual const TArray<FString>& GetSupportedImportExtension() const override;
+	virtual const TSharedRef<FNamePermissionList>& GetImportExtensionPermissionList() const override;
+	virtual bool IsImportExtensionAllowed(const FString& Extension) const override;
 	virtual TSet<EBlueprintType>& GetAllowedBlueprintTypes() override;
 	virtual TSharedRef<FPathPermissionList>& GetFolderPermissionList() override;
 	virtual TSharedRef<FPathPermissionList>& GetWritableFolderPermissionList() override;
@@ -239,8 +240,8 @@ private:
 	/** Permission lists of assets by class path name, one for each EAssetClassAction */
 	TArray<TSharedRef<FPathPermissionList>> AssetClassPermissionList;
 
-	TArray<FString> SupportedImportExtensionArray;
-
+	TSharedRef<FNamePermissionList> ImportExtensionPermissionList;
+	
 	/** Which types of BlueprintFactories are allowed to be returned through GetNewAssetFactories. Empty set allows all types. */
 	TSet<EBlueprintType> AllowedBlueprintTypes;
 
diff --git a/Engine/Source/Developer/AssetTools/Public/IAssetTools.h b/Engine/Source/Developer/AssetTools/Public/IAssetTools.h
--- a/Engine/Source/Developer/AssetTools/Public/IAssetTools.h
+++ b/Engine/Source/Developer/AssetTools/Public/IAssetTools.h
@@ -570,8 +570,10 @@ public:
 	/** Get asset class permission list for content browser and other systems */
 	virtual const TSharedRef<FPathPermissionList>& GetAssetClassPathPermissionList(EAssetClassAction AssetClassAction) const = 0;
 
-	/** Get extension list allowed for importer */
-	virtual const TArray<FString>& GetSupportedImportExtension() const = 0;
+	/** Get extension permission list allowed for importer */
+	virtual const TSharedRef<FNamePermissionList>& GetImportExtensionPermissionList() const = 0;
+
+	virtual bool IsImportExtensionAllowed(const FString& Extension) const = 0;
 
 	/** Which BlueprintTypes are allowed to be created. An empty list should allow everything. */
 	virtual TSet<EBlueprintType>& GetAllowedBlueprintTypes() = 0;
diff --git a/Engine/Source/Editor/UnrealEd/Private/Editor.cpp b/Engine/Source/Editor/UnrealEd/Private/Editor.cpp
--- a/Engine/Source/Editor/UnrealEd/Private/Editor.cpp
+++ b/Engine/Source/Editor/UnrealEd/Private/Editor.cpp
@@ -785,14 +785,14 @@ void FReimportManager::GetNewReimportPath(UObject* Obj, TArray<FString>& InOutFi
 	{
 		//Get the extension interchange can translate for this object
 		TArray<FString> TranslatorFormats = UInterchangeManager::GetInterchangeManager().GetSupportedFormatsForObject(Obj);
-		ObjectTools::AppendFormatsFileExtensions(TranslatorFormats, FileTypes, AllExtensions, DummyFilterIndexToFactory, AssetTools.GetSupportedImportExtension());
+		ObjectTools::AppendFormatsFileExtensions(TranslatorFormats, FileTypes, AllExtensions, DummyFilterIndexToFactory);
 	}
 
 	//If this object was not import with interchange add the legacy formats
 	if(AllExtensions.IsEmpty())
 	{
 		// Generate the file types and extensions represented by the selected factories
-		ObjectTools::GenerateFactoryFileExtensions(Factories, FileTypes, AllExtensions, DummyFilterIndexToFactory, AssetTools.GetSupportedImportExtension());
+		ObjectTools::GenerateFactoryFileExtensions(Factories, FileTypes, AllExtensions, DummyFilterIndexToFactory);
 	}
 
 	
diff --git a/Engine/Source/Editor/UnrealEd/Private/FileHelpers.cpp b/Engine/Source/Editor/UnrealEd/Private/FileHelpers.cpp
--- a/Engine/Source/Editor/UnrealEd/Private/FileHelpers.cpp
+++ b/Engine/Source/Editor/UnrealEd/Private/FileHelpers.cpp
@@ -493,12 +493,12 @@ FString FEditorFileUtils::GetFilterString(EFileInteraction Interaction)
 				FString AllExtensions;
 				TMultiMap<uint32, UFactory*> FilterIndexToFactory;
 				IAssetTools& AssetTools = FModuleManager::LoadModuleChecked<FAssetToolsModule>(TEXT("AssetTools")).Get();
-				ObjectTools::GenerateFactoryFileExtensions(Factories, FileTypes, AllExtensions, FilterIndexToFactory, AssetTools.GetSupportedImportExtension());
+				ObjectTools::GenerateFactoryFileExtensions(Factories, FileTypes, AllExtensions, FilterIndexToFactory);
 
 				if (UInterchangeManager::IsInterchangeImportEnabled())
 				{
 					TArray<FString> InterchangeFileExtensions = UInterchangeManager::GetInterchangeManager().GetSupportedFormats(EInterchangeTranslatorType::Scenes);
-					ObjectTools::AppendFormatsFileExtensions(InterchangeFileExtensions, FileTypes, AllExtensions, AssetTools.GetSupportedImportExtension());
+					ObjectTools::AppendFormatsFileExtensions(InterchangeFileExtensions, FileTypes, AllExtensions);
 				}
 
 				FileTypes = FString::Printf(TEXT("All Files (%s)|%s|%s"), *AllExtensions, *AllExtensions, *FileTypes);
diff --git a/Engine/Source/Editor/UnrealEd/Private/ObjectTools.cpp b/Engine/Source/Editor/UnrealEd/Private/ObjectTools.cpp
--- a/Engine/Source/Editor/UnrealEd/Private/ObjectTools.cpp
+++ b/Engine/Source/Editor/UnrealEd/Private/ObjectTools.cpp
@@ -16,6 +16,7 @@
 #include "Misc/ScopedSlowTask.h"
 #include "Misc/App.h"
 #include "Misc/FileHelper.h"
+#include "Misc/NamePermissionList.h"
 #include "Modules/ModuleManager.h"
 #include "UObject/UObjectHash.h"
 #include "UObject/UObjectIterator.h"
@@ -4334,9 +4335,9 @@ namespace ObjectTools
 	 */
 	void InternalGetFormatInfo(const TArray<FString>& Formats
 		, TArray<FString>& out_Descriptions
-		, TArray<FString>& out_Extensions
-		, const TArray<FString>& SupportedExtensions)
+		, TArray<FString>& out_Extensions)
 	{
+		IAssetTools& AssetTools = FModuleManager::LoadModuleChecked<FAssetToolsModule>(TEXT("AssetTools")).Get();
 		// Iterate over each formats.
 		for ( TArray<FString>::TConstIterator FormatIter( Formats ); FormatIter; ++FormatIter )
 		{
@@ -4350,15 +4351,11 @@ namespace ObjectTools
 			{
 				check( FormatComponents.IsValidIndex( ComponentIndex + 1 ) );
 
-				//If we have a valid SupportedExtensions array, skip any extension not in the SupportedExtension array.
-				if (SupportedExtensions.Num() > 0)
+				FString& RefExtension = FormatComponents[ComponentIndex];
+				if (!AssetTools.IsImportExtensionAllowed(RefExtension))
 				{
-					FString& RefExtension = FormatComponents[ComponentIndex];
-					if (!SupportedExtensions.FindByPredicate([&RefExtension](const FString& Item){return Item.Equals(RefExtension, ESearchCase::IgnoreCase);}))
-					{
-						//Skip this extension
-						continue;
-					}
+					//Skip this extension
+					continue;
 				}
 
 				out_Extensions.Add( FormatComponents[ComponentIndex] );
@@ -4377,13 +4374,12 @@ namespace ObjectTools
 	void GenerateFactoryFileExtensions( UFactory* InFactory
 		, FString& out_Filetypes
 		, FString& out_Extensions
-		, TMultiMap<uint32, UFactory*>& out_FilterIndexToFactory
-		, const TArray<FString>& SupportedExtensions)
+		, TMultiMap<uint32, UFactory*>& out_FilterIndexToFactory)
 	{
 		// Place the factory in an array and call the overloaded version of this function
 		TArray<UFactory*> FactoryArray;
 		FactoryArray.Add( InFactory );
-		GenerateFactoryFileExtensions( FactoryArray, out_Filetypes, out_Extensions, out_FilterIndexToFactory, SupportedExtensions );
+		GenerateFactoryFileExtensions( FactoryArray, out_Filetypes, out_Extensions, out_FilterIndexToFactory);
 	}
 
 	/**
@@ -4396,8 +4392,7 @@ namespace ObjectTools
 	void GenerateFactoryFileExtensions( const TArray<UFactory*>& InFactories
 		, FString& out_Filetypes
 		, FString& out_Extensions
-		, TMultiMap<uint32, UFactory*>& out_FilterIndexToFactory
-		, const TArray<FString>& SupportedExtensions)
+		, TMultiMap<uint32, UFactory*>& out_FilterIndexToFactory)
 	{
 		// Store all the descriptions and their corresponding extensions in a map
 		TMultiMap<FString, FString> DescToExtensionMap;
@@ -4411,7 +4406,7 @@ namespace ObjectTools
 
 			TArray<FString> Descriptions;
 			TArray<FString> Extensions;
-			InternalGetFormatInfo( CurFactory->GetFormats(), Descriptions, Extensions, SupportedExtensions);
+			InternalGetFormatInfo( CurFactory->GetFormats(), Descriptions, Extensions);
 			check( Descriptions.Num() == Extensions.Num() );
 
 			// Make sure to only store each key, value pair once
@@ -4525,24 +4520,22 @@ namespace ObjectTools
 
 	void AppendFormatsFileExtensions(const TArray<FString>& InFormats
 		, FString& out_FileTypes
-		, FString& out_Extensions
-		, const TArray<FString>& SupportedExtensions)
+		, FString& out_Extensions)
 	{
 		TArray<FString> Descriptions;
 		TArray<FString> Extensions;
-		InternalGetFormatInfo(InFormats, Descriptions, Extensions, SupportedExtensions);
+		InternalGetFormatInfo(InFormats, Descriptions, Extensions);
 		InternalAppendFileExtensions(Descriptions, Extensions, out_FileTypes, out_Extensions);
 	}
 
 	void AppendFormatsFileExtensions(const TArray<FString>& InFormats
 		, FString& out_FileTypes
 		, FString& out_Extensions
-		, TMultiMap<uint32, UFactory*>& out_FilterIndexToFactory
-		, const TArray<FString>& SupportedExtensions)
+		, TMultiMap<uint32, UFactory*>& out_FilterIndexToFactory)
 	{
 		TArray<FString> Descriptions;
 		TArray<FString> Extensions;
-		InternalGetFormatInfo(InFormats, Descriptions, Extensions, SupportedExtensions);
+		InternalGetFormatInfo(InFormats, Descriptions, Extensions);
 		InternalAppendFileExtensions(Descriptions, Extensions, out_FileTypes, out_Extensions);
 		uint32 MaxKeyNumber = 0;
 		TSet<uint32> Keys;
@@ -4566,13 +4559,12 @@ namespace ObjectTools
 	 */
 	void AppendFactoryFileExtensions ( UFactory* InFactory
 		, FString& out_Filetypes
-		, FString& out_Extensions
-		, const TArray<FString>& SupportedExtensions)
+		, FString& out_Extensions)
 	{
 		check(InFactory);
 		TArray<FString> Descriptions;
 		TArray<FString> Extensions;
-		InternalGetFormatInfo( InFactory->GetFormats(), Descriptions, Extensions, SupportedExtensions);
+		InternalGetFormatInfo( InFactory->GetFormats(), Descriptions, Extensions);
 		InternalAppendFileExtensions( Descriptions, Extensions, out_Filetypes, out_Extensions);
 	}
 
diff --git a/Engine/Source/Editor/UnrealEd/Private/Subsystems/ImportSubsystem.cpp b/Engine/Source/Editor/UnrealEd/Private/Subsystems/ImportSubsystem.cpp
--- a/Engine/Source/Editor/UnrealEd/Private/Subsystems/ImportSubsystem.cpp
+++ b/Engine/Source/Editor/UnrealEd/Private/Subsystems/ImportSubsystem.cpp
@@ -10,6 +10,8 @@
 #include "Misc/NamePermissionList.h"
 #include "InterchangeManager.h"
 #include "InterchangeFactoryBase.h"
+#include "Framework/Notifications/NotificationManager.h"
+#include "Widgets/Notifications/SNotificationList.h"
 
 class FImportFilesByPath : public IImportSubsystemTask
 {
@@ -152,8 +154,34 @@ void UImportSubsystem::ImportNextTick(const TArray<FString>& Files, const FStrin
 		AssetToolsModule.Get().NotifyBlockedByWritableFolderFilter();
 		return;
 	}
+	FString UnsupportedFiles;
+	TArray<FString> ImportFiles;
+	for (const FString& File : Files)
+	{
+		if (AssetToolsModule.Get().IsImportExtensionAllowed(FPaths::GetExtension(File)))
+		{
+			ImportFiles.AddUnique(File);
+		}
+		else
+		{
+			if (!UnsupportedFiles.IsEmpty())
+			{
+				UnsupportedFiles += TEXT("\n");
+			}
+			UnsupportedFiles += File;
+		}
+	}
+	if (ImportFiles.Num() > 0)
+	{
+		PendingTasks.Enqueue(MakeShared<FImportFilesByPath>(ImportFiles, DestinationPath));
+	}
+
+	if (!UnsupportedFiles.IsEmpty())
+	{
+		FText ErrorMsg = FText::Format(NSLOCTEXT("UImportSubsystem", "ImportNextTickUnsupportedExtension", "Unsupported file format to import:\n\t{0}"), FText::FromString(UnsupportedFiles));
+		FSlateNotificationManager::Get().AddNotification(FNotificationInfo(ErrorMsg));
+	}
 
-	PendingTasks.Enqueue(MakeShared<FImportFilesByPath>(Files, DestinationPath));
 	GEditor->GetTimerManager()->SetTimerForNextTick(this, &UImportSubsystem::HandleNextTick);
 }
 
diff --git a/Engine/Source/Editor/UnrealEd/Public/ObjectTools.h b/Engine/Source/Editor/UnrealEd/Public/ObjectTools.h
--- a/Engine/Source/Editor/UnrealEd/Public/ObjectTools.h
+++ b/Engine/Source/Editor/UnrealEd/Public/ObjectTools.h
@@ -14,6 +14,7 @@
 #include "UObject/GCObject.h"
 #include "CollectionManagerTypes.h"
 
+class FNamePermissionList;
 class FTextureRenderTargetResource;
 class SWindow;
 class UExporter;
@@ -493,8 +494,7 @@ namespace ObjectTools
 	UNREALED_API void GenerateFactoryFileExtensions( const UFactory* InFactory
 		, FString& out_Filetypes
 		, FString& out_Extensions
-		, TMultiMap<uint32, UFactory*>& out_FilterIndexToFactory
-		, const TArray<FString>& SupportedExtensions = TArray<FString>());
+		, TMultiMap<uint32, UFactory*>& out_FilterIndexToFactory);
 
 	/**
 	 * Populates two strings with all of the file types and extensions the provided factories support.
@@ -507,8 +507,7 @@ namespace ObjectTools
 	UNREALED_API void GenerateFactoryFileExtensions( const TArray<UFactory*>& InFactories
 		, FString& out_Filetypes
 		, FString& out_Extensions
-		, TMultiMap<uint32, UFactory*>& out_FilterIndexToFactory
-		, const TArray<FString>& SupportedExtensions = TArray<FString>());
+		, TMultiMap<uint32, UFactory*>& out_FilterIndexToFactory);
 
 	/**
 	 * Generates a list of file types for a given class.
@@ -516,8 +515,7 @@ namespace ObjectTools
 	 */
 	UNREALED_API void AppendFactoryFileExtensions( UFactory* InFactory
 		, FString& out_Filetypes
-		, FString& out_Extensions
-		, const TArray<FString>& SupportedExtensions = TArray<FString>());
+		, FString& out_Extensions);
 
 	/**
 	 * Populates two strings with all of the file types and extensions the format list provides.
@@ -529,8 +527,7 @@ namespace ObjectTools
 	 */
 	UNREALED_API void AppendFormatsFileExtensions(const TArray<FString>& InFormats
 		, FString& out_FileTypes
-		, FString& out_Extensions
-		, const TArray<FString>& SupportedExtensions = TArray<FString>());
+		, FString& out_Extensions);
 
 	/**
 	 * Populates two strings with all of the file types and extensions the format list provides.
@@ -544,8 +541,7 @@ namespace ObjectTools
 	UNREALED_API void AppendFormatsFileExtensions(const TArray<FString>& InFormats
 		, FString& out_FileTypes
 		, FString& out_Extensions
-		, TMultiMap<uint32, UFactory*>& out_FilterIndexToFactory
-		, const TArray<FString>& SupportedExtensions = TArray<FString>());
+		, TMultiMap<uint32, UFactory*>& out_FilterIndexToFactory);
 
 	/**
 	 * Iterates over all classes and assembles a list of non-abstract UExport-derived type instances.
diff --git a/Engine/Source/Runtime/Interchange/Engine/Private/InterchangeManager.cpp b/Engine/Source/Runtime/Interchange/Engine/Private/InterchangeManager.cpp
--- a/Engine/Source/Runtime/Interchange/Engine/Private/InterchangeManager.cpp
+++ b/Engine/Source/Runtime/Interchange/Engine/Private/InterchangeManager.cpp
@@ -5,6 +5,7 @@
 #include "AssetRegistry/AssetData.h"
 #include "AssetRegistry/AssetDataTagMap.h"
 #include "AssetRegistry/AssetRegistryModule.h"
+#include "AssetToolsModule.h"
 #include "Async/Async.h"
 #include "CoreGlobals.h"
 #include "CoreMinimal.h"
@@ -807,6 +808,13 @@ bool UInterchangeManager::CanTranslateSourceData(const UInterchangeSourceData* S
 	{
 		return false;
 	}
+	
+	IAssetTools& AssetTools = FModuleManager::LoadModuleChecked<FAssetToolsModule>(TEXT("AssetTools")).Get();
+	const FString Extension = FPaths::GetExtension(SourceData->GetFilename());
+	if (!AssetTools.IsImportExtensionAllowed(Extension))
+	{
+		return false;
+	}
 
 	return GetTranslatorForSourceData(SourceData) != nullptr;
 }
