Cooker and WorldPartition: GeneratedPackages: Fix generated packages to not be marked as LoadedByEditorPropertiesOnly if their generator package is not LoadedByEditorPropertiesOnly.
#rnx

diff --git a/Engine/Source/Editor/UnrealEd/Private/Cooker/CookPackageData.cpp b/Engine/Source/Editor/UnrealEd/Private/Cooker/CookPackageData.cpp
--- a/Engine/Source/Editor/UnrealEd/Private/Cooker/CookPackageData.cpp
+++ b/Engine/Source/Editor/UnrealEd/Private/Cooker/CookPackageData.cpp
@@ -1625,6 +1625,10 @@ UPackage* FGeneratorPackage::CreateGeneratedUPackage(FCookGenerationInfo& Genera
 	GeneratedPackage->SetPersistentGuid(InOwnerPackage->GetPersistentGuid());
 	GeneratedPackage->SetPackageFlags(PKG_CookGenerated);
 	GeneratedInfo.SetHasCreatedPackage(true);
+	if (!InOwnerPackage->IsLoadedByEditorPropertiesOnly())
+	{
+		GeneratedPackage->SetLoadedByEditorPropertiesOnly(false);
+	}
 	return GeneratedPackage;
 }
 
