Fix CIS issue.

diff --git a/Engine/Source/Editor/UnrealEd/Classes/MaterialGraph/MaterialGraph.h b/Engine/Source/Editor/UnrealEd/Classes/MaterialGraph/MaterialGraph.h
--- a/Engine/Source/Editor/UnrealEd/Classes/MaterialGraph/MaterialGraph.h
+++ b/Engine/Source/Editor/UnrealEd/Classes/MaterialGraph/MaterialGraph.h
@@ -7,6 +7,7 @@
 #include "Templates/UniquePtr.h"
 #include "EdGraph/EdGraph.h"
 #include "Materials/Material.h"
+#include "RenderUtils.h"
 #include "MaterialGraph.generated.h"
 
 class UMaterialExpressionComment;
