Optimize GetAllPackageFilenamesFromAssetRegistry
  - 1.7s to 333ms for a small project warm cook

#rnx

diff --git a/Engine/Source/Editor/UnrealEd/Private/CookOnTheFlyServer.cpp b/Engine/Source/Editor/UnrealEd/Private/CookOnTheFlyServer.cpp
--- a/Engine/Source/Editor/UnrealEd/Private/CookOnTheFlyServer.cpp
+++ b/Engine/Source/Editor/UnrealEd/Private/CookOnTheFlyServer.cpp
@@ -10429,64 +10429,67 @@ bool UCookOnTheFlyServer::GetAllPackageFilenamesFromAssetRegistry(const FString&
 			SerializedState->Serialize(*Reader.Get(), FAssetRegistrySerializationOptions());
 		}
 
-		
 		check(OutPackageDatas.Num() == 0);
 
-		int32 NumPackages = SerializedState->GetNumAssets();
-		TArray<const FAssetData*> AssetDatas;
-		TSet<FName> PackageNames;
-		FString PackageNameStrBuffer;
-		AssetDatas.Reserve(NumPackages);
-		OutPackageDatas.Reserve(NumPackages);
-		PackageNames.Reserve(NumPackages);
+		// Apply lock striping to reduce contention.
+		constexpr int32 UNIQUEPACKAGENAMES_BUCKETS = 31; /* prime number for best distribution using modulo */
 
-		// Convert the Map of RegistryData into an Array of FAssetData and populate PackageNames in the output array
-		SerializedState->EnumerateAllAssets([&](const FAssetData& RegistryData)
+		struct FUniquePackageNames
 		{
-			// If we want to reevaluate (try cooking again) the uncooked packages (packages that were found to be empty when we cooked them before),
-			// then remove the uncooked packages from the set of known packages. Uncooked packages are identified by PackageFlags == 0.
-			if (bReevaluateUncookedPackages && (RegistryData.PackageFlags == 0))
-			{
-				return;
-			}
-			FName PackageName = RegistryData.PackageName;
-			bool bPackageAlreadyAdded;
-			PackageNames.Add(PackageName, &bPackageAlreadyAdded);
-			if (bPackageAlreadyAdded)
-			{
-				return;
-			}
-			PackageName.ToString(PackageNameStrBuffer);
-			if (FPackageName::GetPackageMountPoint(PackageNameStrBuffer).IsNone())
-			{
-				// Skip any packages that are not currently mounted; if we tried to find their FileNames below
-				// we would get log spam
-				return;
-			}
+			FRWLock Lock;
+			TSet<FName> Names;
+		} UniquePackageNames[UNIQUEPACKAGENAMES_BUCKETS];
 
-			AssetDatas.Add(&RegistryData);
-			FConstructPackageData& PackageData = OutPackageDatas.Emplace_GetRef();
-			PackageData.PackageName = PackageName;
-
-			// For any PackageNames that already have PackageDatas, mark them ahead of the loop to
-			// skip the effort of checking whether they exist on disk inside the loop
-			FPackageData* ExistingPackageData = PackageDatas->FindPackageDataByPackageName(PackageName);
-			if (ExistingPackageData)
-			{
-				PackageData.NormalizedFileName = ExistingPackageData->GetFileName();
-			}
-		});
-		NumPackages = AssetDatas.Num();
+		const int32 NumPackages = SerializedState->GetAssetDataMap().Num();
+		OutPackageDatas.SetNum(NumPackages);
 
+		// Convert the Map of RegistryData into an Array of FAssetData and populate PackageNames in the output array
+		// We can index directly from the set because we know its congiguous as we just deserialized it.
+		checkf(SerializedState->GetAssetDataMap().GetMaxIndex() == NumPackages, TEXT("The set needs to be contiguous so we can index into it directly"));
 		ParallelFor(NumPackages,
-			[&AssetRegistryPath, &OutPackageDatas, &AssetDatas, this, bVerifyPackagesExist](int32 AssetIndex)
+			[&](int32 Index)
 			{
-				FConstructPackageData& PackageData = OutPackageDatas[AssetIndex];
-				if (!PackageData.NormalizedFileName.IsNone())
+				const FAssetData& RegistryData = *SerializedState->GetAssetDataMap()[FSetElementId::FromInteger(Index)];
+
+				// If we want to reevaluate (try cooking again) the uncooked packages (packages that were found to be empty when we cooked them before),
+				// then remove the uncooked packages from the set of known packages. Uncooked packages are identified by PackageFlags == 0.
+				if (bReevaluateUncookedPackages && (RegistryData.PackageFlags == 0))
 				{
 					return;
 				}
-				const FName PackageName = PackageData.PackageName;
+
+				const FName PackageName = RegistryData.PackageName;
+				const uint32 NameHash = GetTypeHash(PackageName);
+				FUniquePackageNames& Bucket = UniquePackageNames[NameHash % UNIQUEPACKAGENAMES_BUCKETS];
+				bool bPackageAlreadyAdded;
+				{
+					FWriteScopeLock ScopeLock(Bucket.Lock);
+					Bucket.Names.FindOrAddByHash(NameHash, RegistryData.PackageName, &bPackageAlreadyAdded);
+				}
+				
+				if (bPackageAlreadyAdded)
+				{
+					return;
+				}
+
+				if (FPackageName::GetPackageMountPoint(PackageName.ToString()).IsNone())
+				{
+					// Skip any packages that are not currently mounted; if we tried to find their FileNames below
+					// we would get log spam
+					return;
+				}
+
+				FConstructPackageData& PackageData = OutPackageDatas[Index];
+				PackageData.PackageName = PackageName;
+
+				// For any PackageNames that already have PackageDatas, mark them ahead of the loop to
+				// skip the effort of checking whether they exist on disk inside the loop
+				FPackageData* ExistingPackageData = PackageDatas->FindPackageDataByPackageName(PackageName);
+				if (ExistingPackageData)
+				{
+					PackageData.NormalizedFileName = ExistingPackageData->GetFileName();
+					return;
+				}
 
 				// TODO ICookPackageSplitter: Need to handle GeneratedPackages that exist in the cooked AssetRegistry we are 
 				// reading, but do not exist in WorkspaceDomain and so are not found when we look them up here.
@@ -10494,26 +10497,26 @@ bool UCookOnTheFlyServer::GetAllPackageFilenamesFromAssetRegistry(const FString&
 				if (!PackageFileName.IsNone())
 				{
 					PackageData.NormalizedFileName = PackageFileName;
+					return;
+				}
+
+				if (bVerifyPackagesExist)
+				{
+					UE_LOG(LogCook, Warning, TEXT("Could not resolve package %s from %s"),
+						*PackageName.ToString(), *AssetRegistryPath);
 				}
 				else
 				{
-					if (bVerifyPackagesExist)
+					const bool bContainsMap = !!(RegistryData.PackageFlags & PKG_ContainsMap);
+					PackageFileName = FPackageDatas::LookupFileNameOnDisk(PackageName,
+						false /* bRequireExists */, bContainsMap);
+					if (!PackageFileName.IsNone())
 					{
-						UE_LOG(LogCook, Warning, TEXT("Could not resolve package %s from %s"),
-							*PackageName.ToString(), *AssetRegistryPath);
-					}
-					else
-					{
-						const bool bContainsMap = !!(AssetDatas[AssetIndex]->PackageFlags & PKG_ContainsMap);
-						PackageFileName = FPackageDatas::LookupFileNameOnDisk(PackageName,
-							false /* bRequireExists */, bContainsMap);
-						if (!PackageFileName.IsNone())
-						{
-							PackageData.NormalizedFileName = PackageFileName;
-						}
+						PackageData.NormalizedFileName = PackageFileName;
 					}
 				}
-			});
+			}
+		);
 
 		OutPackageDatas.RemoveAllSwap([](FConstructPackageData& PackageData)
 			{
