SClassPicker: Fix issues with padding and alignment of buttons and text

diff --git a/Engine/Source/Editor/UnrealEd/Private/SClassPickerDialog.cpp b/Engine/Source/Editor/UnrealEd/Private/SClassPickerDialog.cpp
--- a/Engine/Source/Editor/UnrealEd/Private/SClassPickerDialog.cpp
+++ b/Engine/Source/Editor/UnrealEd/Private/SClassPickerDialog.cpp
@@ -269,6 +269,8 @@ TSharedRef<ITableRow> SClassPickerDialog::GenerateListRow(TSharedPtr<FClassPicke
 	UClass* ItemClass = LoadClass<UObject>(NULL, *Obj->ClassName, NULL, LOAD_None, NULL);
 	const FSlateBrush* ItemBrush = FSlateIconFinder::FindIconBrushForClass(ItemClass);
 
+	TSharedRef<SWidget> DocumentationWidget = FEditorClassUtils::GetDocumentationLinkWidget(ItemClass);
+	
 	return 
 	SNew(STableRow< TSharedPtr<FClassPickerDefaults> >, OwnerTable)
 	[
@@ -279,11 +281,12 @@ TSharedRef<ITableRow> SClassPickerDialog::GenerateListRow(TSharedPtr<FClassPicke
 		[
 			SNew(SHorizontalBox)
 			+SHorizontalBox::Slot()
+			.VAlign(VAlign_Center)
 			.FillWidth(0.5f)
 			[
 				SNew(SButton)
 				.OnClicked(this, &SClassPickerDialog::OnDefaultClassPicked, ItemClass)
-				.ToolTip(FEditorClassUtils::GetTooltip(ItemClass))
+				.ToolTipText(Obj->GetName())
 				.Content()
 				[
 					SNew(SHorizontalBox)
@@ -300,24 +303,28 @@ TSharedRef<ITableRow> SClassPickerDialog::GenerateListRow(TSharedPtr<FClassPicke
 					[
 						SNew(STextBlock)
 						.Text(Obj->GetName())
+						.OverflowPolicy(ETextOverflowPolicy::Ellipsis)
 					]
 				]
 			]
 			+SHorizontalBox::Slot()
+			.VAlign(VAlign_Center)
 			.Padding(10.0f, 0.0f, 4.0f, 0.0f)
 			[
 				SNew(STextBlock)
 				.Text(Obj->GetDescription())
+				.ToolTipText(Obj->GetDescription())
 				.AutoWrapText(true)
+				.OverflowPolicy(ETextOverflowPolicy::Ellipsis)
 			]
 			+SHorizontalBox::Slot()
 			.VAlign(VAlign_Center)
 			.AutoWidth()
 			[
 				SNew(SBox)
-				.WidthOverride(60.0f)
+				.WidthOverride(DocumentationWidget != SNullWidget::NullWidget ? 60.0f : 0.0f)
 				[
-					FEditorClassUtils::GetDocumentationLinkWidget(ItemClass)
+					DocumentationWidget
 				]
 			]
 		]
