changed the name of LLM_SCOPE of PackageBuildDependencyTracker to be more descriptive

diff --git a/Engine/Source/Editor/UnrealEd/Private/Cooker/PackageBuildDependencyTracker.cpp b/Engine/Source/Editor/UnrealEd/Private/Cooker/PackageBuildDependencyTracker.cpp
--- a/Engine/Source/Editor/UnrealEd/Private/Cooker/PackageBuildDependencyTracker.cpp
+++ b/Engine/Source/Editor/UnrealEd/Private/Cooker/PackageBuildDependencyTracker.cpp
@@ -100,7 +100,7 @@ void FPackageBuildDependencyTracker::StaticOnObjectHandleRead(UObject* ReadObjec
 			*Referencer.ToString(), *Referenced.ToString());
 	}
 
-	LLM_SCOPE_BYNAME(TEXTVIEW("TObjectPtr"));
+	LLM_SCOPE_BYNAME(TEXTVIEW("PackageBuildDependencyTracker"));
 
 	FBuildDependencyAccessData AccessData{Referenced, AccumulatedScopeData->TargetPlatform};
 	FScopeLock RecordsScopeLock(&Singleton.RecordsLock);
