Cooker: Iterative GeneratedFiles: Fix LooseCookedPackageWriter to delete generated files when requested. Fix it also to delete all file extensions for each package.
#rnx

diff --git a/Engine/Source/Editor/UnrealEd/Private/Cooker/CookTypes.h b/Engine/Source/Editor/UnrealEd/Private/Cooker/CookTypes.h
--- a/Engine/Source/Editor/UnrealEd/Private/Cooker/CookTypes.h
+++ b/Engine/Source/Editor/UnrealEd/Private/Cooker/CookTypes.h
@@ -530,7 +530,7 @@ void LogCookerMessage(const FString& MessageText, EMessageSeverity::Type Severit
 LLM_DECLARE_TAG(Cooker);
 
 inline constexpr uint32 ExpectedMaxNumPlatforms = 32;
-#define REMAPPED_PLUGINS TEXT("RemappedPlugins")
+#define REMAPPED_PLUGINS TEXTVIEW("RemappedPlugins")
 extern float GCookProgressWarnBusyTime;
 
 inline constexpr float TickCookableObjectsFrameTime = .100f;
\ No newline at end of file
diff --git a/Engine/Source/Editor/UnrealEd/Private/Cooker/LooseCookedPackageWriter.cpp b/Engine/Source/Editor/UnrealEd/Private/Cooker/LooseCookedPackageWriter.cpp
--- a/Engine/Source/Editor/UnrealEd/Private/Cooker/LooseCookedPackageWriter.cpp
+++ b/Engine/Source/Editor/UnrealEd/Private/Cooker/LooseCookedPackageWriter.cpp
@@ -32,6 +32,7 @@
 #include "Misc/FileHelper.h"
 #include "Misc/Optional.h"
 #include "Misc/PackageName.h"
+#include "Misc/PackagePath.h"
 #include "Misc/PathViews.h"
 #include "Misc/Paths.h"
 #include "Misc/ScopeLock.h"
@@ -570,7 +571,7 @@ TUniquePtr<FAssetRegistryState> FLooseCookedPackageWriter::LoadPreviousAssetRegi
 		PreviousAssetRegistryFile = FPaths::Combine(MetadataDirectoryPath, GetDevelopmentAssetRegistryFilename());
 	}
 
-	UncookedPathToCookedPath.Reset();
+	PackageNameToCookedFiles.Reset();
 
 	FArrayReader SerializedAssetData;
 	if (!IFileManager::Get().FileExists(*PreviousAssetRegistryFile) || !FFileHelper::LoadFileToArray(SerializedAssetData, *PreviousAssetRegistryFile))
@@ -592,36 +593,32 @@ TUniquePtr<FAssetRegistryState> FLooseCookedPackageWriter::LoadPreviousAssetRegi
 		GetAllCookedFiles();
 		TSet<FName> RemoveFromRegistry;
 		TSet<FName> RemoveFromDisk;
-		RemoveFromDisk.Reserve(UncookedPathToCookedPath.Num());
-		for (TPair<FName, FName>& Pair : UncookedPathToCookedPath)
+		RemoveFromDisk.Reserve(PackageNameToCookedFiles.Num());
+		for (TPair<FName, TArray<FString>>& Pair : PackageNameToCookedFiles)
 		{
 			RemoveFromDisk.Add(Pair.Key);
 		}
+		IAssetRegistry& AssetRegistry = IAssetRegistry::GetChecked();
 		for (const TPair<FName, const FAssetPackageData*>& Pair : PreviousState->GetAssetPackageDataMap())
 		{
 			FName PackageName = Pair.Key;
-			FName UncookedFilename = PackageDatas.GetFileNameByPackageName(PackageName);
+			bool bCurrentPackageExists = AssetRegistry.DoesPackageExistOnDisk(PackageName);
 
 			bool bNoLongerExistsInEditor = false;
 			bool bIsScriptPackage = FPackageName::IsScriptPackage(WriteToString<256>(PackageName));
-			if (UncookedFilename.IsNone())
+			if (!bCurrentPackageExists)
 			{
 				// Script and generated packages do not exist uncooked
 				// Check that the package is not an exception before removing from cooked
 				bool bIsCookedOnly = bIsScriptPackage;
-				bool bIsMap = false;
 				if (!bIsCookedOnly)
 				{
 					for (const FAssetData* AssetData : PreviousState->GetAssetsByPackageName(PackageName))
 					{
 						bIsCookedOnly |= !!(AssetData->PackageFlags & PKG_CookGenerated);
-						bIsMap |= !!(AssetData->PackageFlags & PKG_ContainsMap);
 					}
 				}
 				bNoLongerExistsInEditor = !bIsCookedOnly;
-				UncookedFilename = UE::Cook::FPackageDatas::LookupFileNameOnDisk(PackageName,
-					false /* bRequireExists */, bIsMap);
-
 			}
 			if (bNoLongerExistsInEditor)
 			{
@@ -632,7 +629,7 @@ TUniquePtr<FAssetRegistryState> FLooseCookedPackageWriter::LoadPreviousAssetRegi
 			else
 			{
 				// Keep package on disk if it exists. Keep package in registry if it exists on disk or was a FailedSave.
-				bool bExistsOnDisk = (RemoveFromDisk.Remove(UncookedFilename) == 1); // Remove its RemoveFromDisk entry
+				bool bExistsOnDisk = (RemoveFromDisk.Remove(PackageName) == 1); // Remove its RemoveFromDisk entry
 				const FAssetPackageData* PackageData = Pair.Value;
 				if (!bExistsOnDisk && PackageData->DiskSize >= 0 && !bIsScriptPackage)
 				{
@@ -648,7 +645,7 @@ TUniquePtr<FAssetRegistryState> FLooseCookedPackageWriter::LoadPreviousAssetRegi
 		}
 		if (RemoveFromDisk.Num())
 		{
-			RemoveCookedPackagesByUncookedFilename(RemoveFromDisk.Array());
+			RemoveCookedPackagesByPackageName(RemoveFromDisk.Array(), true /* bRemoveRecords */);
 		}
 	}
 
@@ -663,58 +660,40 @@ FCbObject FLooseCookedPackageWriter::GetOplogAttachment(FName PackageName, FUtf8
 
 void FLooseCookedPackageWriter::RemoveCookedPackages(TArrayView<const FName> PackageNamesToRemove)
 {
-	if (UncookedPathToCookedPath.IsEmpty())
+	if (PackageNameToCookedFiles.IsEmpty())
 	{
 		return;
 	}
-
-	if (PackageNamesToRemove.Num() > 0)
-	{
-		// if we are going to clear the cooked packages it is conceivable that we will recook the packages which we just cooked 
-		// that means it's also conceivable that we will recook the same package which currently has an outstanding async write request
-		UPackage::WaitForAsyncFileWrites();
-
-		auto DeletePackageLambda = [&PackageNamesToRemove, this](int32 PackageIndex)
-		{
-			const FName PackageName = PackageNamesToRemove[PackageIndex];
-			const FName UncookedFileName = PackageDatas.GetFileNameByPackageName(PackageName);
-			if (UncookedFileName.IsNone())
-			{
-				return;
-			}
-			FName* CookedFileName = UncookedPathToCookedPath.Find(UncookedFileName);
-			if (CookedFileName)
-			{
-				TStringBuilder<256> FilePath;
-				CookedFileName->ToString(FilePath);
-				IFileManager::Get().Delete(*FilePath, true, true, true);
-			}
-		};
-		ParallelFor(PackageNamesToRemove.Num(), DeletePackageLambda);
-	}
-
-	// We no longer have a use for UncookedPathToCookedPath, after the RemoveCookedPackages call at the beginning of the cook.
-	UncookedPathToCookedPath.Empty();
+	// if we are going to clear the cooked packages it is conceivable that we will recook the packages which we just cooked 
+	// that means it's also conceivable that we will recook the same package which currently has an outstanding async write request
+	UPackage::WaitForAsyncFileWrites();
+	RemoveCookedPackagesByPackageName(PackageNamesToRemove, false /* bRemoveRecords */);
+	// PackageNameToCookedFiles is no longer used after the RemoveCookedPackages call at the beginning of the cook.
+	PackageNameToCookedFiles.Empty();
 }
 
-void FLooseCookedPackageWriter::RemoveCookedPackagesByUncookedFilename(const TArray<FName>& UncookedFileNamesToRemove)
+void FLooseCookedPackageWriter::RemoveCookedPackagesByPackageName(TArrayView<const FName> PackageNamesToRemove, bool bRemoveRecords)
 {
-	auto DeletePackageLambda = [&UncookedFileNamesToRemove, this](int32 PackageIndex)
+	auto DeletePackageLambda = [&PackageNamesToRemove, this](int32 PackageIndex)
 	{
-		FName UncookedFileName = UncookedFileNamesToRemove[PackageIndex];
-		FName* CookedFileName = UncookedPathToCookedPath.Find(UncookedFileName);
-		if (CookedFileName)
+		FName PackageName = PackageNamesToRemove[PackageIndex];
+		TArray<FString>* CookedFileNames = PackageNameToCookedFiles.Find(PackageName);
+		if (CookedFileNames)
 		{
-			TStringBuilder<256> FilePath;
-			CookedFileName->ToString(FilePath);
-			IFileManager::Get().Delete(*FilePath, true, true, true);
+			for (const FString& CookedFileName : *CookedFileNames)
+			{
+				IFileManager::Get().Delete(*CookedFileName, true, true, true);
+			}
 		}
 	};
-	ParallelFor(UncookedFileNamesToRemove.Num(), DeletePackageLambda);
+	ParallelFor(PackageNamesToRemove.Num(), DeletePackageLambda);
 
-	for (FName UncookedFilename : UncookedFileNamesToRemove)
+	if (bRemoveRecords)
 	{
-		UncookedPathToCookedPath.Remove(UncookedFilename);
+		for (FName PackageName : PackageNamesToRemove)
+		{
+			PackageNameToCookedFiles.Remove(PackageName);
+		}
 	}
 }
 
@@ -832,9 +811,8 @@ public:
 			FStringView Extension(FPathViews::GetExtension(Filename, true /* bIncludeDot */));
 			if (!Extension.IsEmpty())
 			{
-				const TCHAR* ExtensionStr = Extension.GetData();
-				check(ExtensionStr[Extension.Len()] == '\0'); // IsPackageExtension takes a null-terminated TCHAR; we should have it since GetExtension is from the end of the filename
-				if (FPackageName::IsPackageExtension(ExtensionStr))
+				EPackageExtension ExtensionEnum = FPackagePath::ParseExtension(Extension);
+				if (ExtensionEnum != EPackageExtension::Unspecified && ExtensionEnum != EPackageExtension::Custom)
 				{
 					FoundFiles.Add(Filename);
 				}
@@ -859,89 +837,95 @@ void FLooseCookedPackageWriter::GetAllCookedFiles()
 	const FString SandboxProjectDir = FPaths::Combine(OutputPath, FApp::GetProjectName()) + TEXT("/");
 	const FString RelativeRootDir = FPaths::GetRelativePathToRoot();
 	const FString RelativeProjectDir = FPaths::ProjectDir();
-	FString UncookedFilename;
-	UncookedFilename.Reserve(1024);
+	FString ScratchFileName;
+	ScratchFileName.Reserve(1024);
+	FString ScratchPackageName;
+	ScratchPackageName.Reserve(1024);
 
 	for (const FString& CookedFile : CookedFiles)
 	{
-		const FName CookedFName(*CookedFile);
-		const FName UncookedFName = ConvertCookedPathToUncookedPath(
+		const FName PackageName = ConvertCookedPathToPackageName(
 			SandboxRootDir, RelativeRootDir,
 			SandboxProjectDir, RelativeProjectDir,
-			CookedFile, UncookedFilename);
+			CookedFile, ScratchFileName, ScratchPackageName);
 
-		UncookedPathToCookedPath.Add(UncookedFName, CookedFName);
+		if (!PackageName.IsNone())
+		{
+			PackageNameToCookedFiles.FindOrAdd(PackageName).Add(CookedFile);
+		}
 	}
 }
 
-FName FLooseCookedPackageWriter::ConvertCookedPathToUncookedPath(
+FName FLooseCookedPackageWriter::ConvertCookedPathToPackageName(
 	const FString& SandboxRootDir, const FString& RelativeRootDir,
 	const FString& SandboxProjectDir, const FString& RelativeProjectDir,
-	const FString& CookedPath, FString& OutUncookedPath) const
+	const FString& CookedPath, FString& ScratchFileName, FString& ScratchPackageName) const
 {
-	OutUncookedPath.Reset();
+	FString& UncookedFileName = ScratchFileName; 
+	UncookedFileName.Reset();
+	constexpr FStringView RemappedPluginsFolder(REMAPPED_PLUGINS);
+	constexpr FStringView ContentFolder(TEXTVIEW("Content/"));
 
 	// Check for remapped plugins' cooked content
-	if (PluginsToRemap.Num() > 0 && CookedPath.Contains(REMAPPED_PLUGINS))
+	if (PluginsToRemap.Num() > 0 && CookedPath.Contains(RemappedPluginsFolder))
 	{
-		int32 RemappedIndex = CookedPath.Find(REMAPPED_PLUGINS);
+		int32 RemappedIndex = CookedPath.Find(RemappedPluginsFolder);
 		check(RemappedIndex >= 0);
-		static uint32 RemappedPluginStrLen = FCString::Strlen(REMAPPED_PLUGINS);
 		// Snip everything up through the RemappedPlugins/ off so we can find the plugin it corresponds to
-		FString PluginPath = CookedPath.RightChop(RemappedIndex + RemappedPluginStrLen + 1);
+		FString PluginPath = CookedPath.RightChop(RemappedIndex + RemappedPluginsFolder.Len() + 1);
 		// Find the plugin that owns this content
 		for (const TSharedRef<IPlugin>& Plugin : PluginsToRemap)
 		{
 			if (PluginPath.StartsWith(Plugin->GetName()))
 			{
-				OutUncookedPath = Plugin->GetContentDir();
-				static uint32 ContentStrLen = FCString::Strlen(TEXT("Content/"));
+				UncookedFileName = Plugin->GetContentDir();
 				// Chop off the pluginName/Content since it's part of the full path
-				OutUncookedPath /= PluginPath.RightChop(Plugin->GetName().Len() + ContentStrLen);
+				UncookedFileName /= PluginPath.RightChop(Plugin->GetName().Len() + ContentFolder.Len());
 				break;
 			}
 		}
 
-		if (OutUncookedPath.Len() > 0)
-		{
-			return FName(*OutUncookedPath);
-		}
-		// Otherwise fall through to sandbox handling
+		// If UncookedFileName is empty fall through to sandbox handling
 	}
 
-	auto BuildUncookedPath =
-		[&OutUncookedPath](const FString& CookedPath, const FString& CookedRoot, const FString& UncookedRoot)
+	if (UncookedFileName.IsEmpty())
 	{
-		OutUncookedPath.AppendChars(*UncookedRoot, UncookedRoot.Len());
-		OutUncookedPath.AppendChars(*CookedPath + CookedRoot.Len(), CookedPath.Len() - CookedRoot.Len());
-	};
+		auto BuildUncookedPath =
+			[&UncookedFileName](const FString& CookedPath, const FString& CookedRoot, const FString& UncookedRoot)
+		{
+			UncookedFileName.AppendChars(*UncookedRoot, UncookedRoot.Len());
+			UncookedFileName.AppendChars(*CookedPath + CookedRoot.Len(), CookedPath.Len() - CookedRoot.Len());
+		};
 
-	if (CookedPath.StartsWith(SandboxRootDir))
-	{
-		// Optimized CookedPath.StartsWith(SandboxProjectDir) that does not compare all of SandboxRootDir again
-		if (CookedPath.Len() >= SandboxProjectDir.Len() &&
-			0 == FCString::Strnicmp(
-				*CookedPath + SandboxRootDir.Len(),
-				*SandboxProjectDir + SandboxRootDir.Len(),
-				SandboxProjectDir.Len() - SandboxRootDir.Len()))
+		if (CookedPath.StartsWith(SandboxRootDir))
 		{
-			BuildUncookedPath(CookedPath, SandboxProjectDir, RelativeProjectDir);
+			// Optimized CookedPath.StartsWith(SandboxProjectDir) that does not compare all of SandboxRootDir again
+			if (CookedPath.Len() >= SandboxProjectDir.Len() &&
+				0 == FCString::Strnicmp(
+					*CookedPath + SandboxRootDir.Len(),
+					*SandboxProjectDir + SandboxRootDir.Len(),
+					SandboxProjectDir.Len() - SandboxRootDir.Len()))
+			{
+				BuildUncookedPath(CookedPath, SandboxProjectDir, RelativeProjectDir);
+			}
+			else
+			{
+				BuildUncookedPath(CookedPath, SandboxRootDir, RelativeRootDir);
+			}
 		}
 		else
 		{
-			BuildUncookedPath(CookedPath, SandboxRootDir, RelativeRootDir);
+			FString FullCookedFilename = FPaths::ConvertRelativePathToFull(CookedPath);
+			BuildUncookedPath(FullCookedFilename, SandboxRootDir, RelativeRootDir);
 		}
 	}
-	else
+
+	if (!FPackageName::TryConvertFilenameToLongPackageName(UncookedFileName, ScratchPackageName))
 	{
-		FString FullCookedFilename = FPaths::ConvertRelativePathToFull(CookedPath);
-		BuildUncookedPath(FullCookedFilename, SandboxRootDir, RelativeRootDir);
+		return NAME_None;
 	}
 
-	// Convert to a standard filename as required by PackageDatas where this path is used.
-	FPaths::MakeStandardFilename(OutUncookedPath);
-
-	return FName(*OutUncookedPath);
+	return FName(*ScratchPackageName);
 }
 
 EPackageExtension FLooseCookedPackageWriter::BulkDataTypeToExtension(FBulkDataInfo::EType BulkDataType)
diff --git a/Engine/Source/Editor/UnrealEd/Private/Cooker/LooseCookedPackageWriter.h b/Engine/Source/Editor/UnrealEd/Private/Cooker/LooseCookedPackageWriter.h
--- a/Engine/Source/Editor/UnrealEd/Private/Cooker/LooseCookedPackageWriter.h
+++ b/Engine/Source/Editor/UnrealEd/Private/Cooker/LooseCookedPackageWriter.h
@@ -112,20 +112,14 @@ private:
 
 	/* Delete the sandbox directory (asynchronously) in preparation for a clean cook */
 	void DeleteSandboxDirectory();
-	/**
-	* Searches the disk for all the cooked files in the sandbox path provided
-	* Returns a map of the uncooked file path matches to the cooked file path for each package which exists
-	*
-	* @param UncookedpathToCookedPath out Map of the uncooked path matched with the cooked package which exists
-	* @param SandboxRootDir root dir to search for cooked packages in
-	*/
+	/** Searches the disk for all the cooked files in the sandbox. Stores results in PackageNameToCookedFiles. */
 	void GetAllCookedFiles();
 
-	FName ConvertCookedPathToUncookedPath(
+	FName ConvertCookedPathToPackageName(
 		const FString& SandboxRootDir, const FString& RelativeRootDir,
 		const FString& SandboxProjectDir, const FString& RelativeProjectDir,
-		const FString& CookedPath, FString& OutUncookedPath) const;
-	void RemoveCookedPackagesByUncookedFilename(const TArray<FName>& UncookedFileNamesToRemove);
+		const FString& CookedPath, FString& ScratchFileName, FString& ScratchPackageName) const;
+	void RemoveCookedPackagesByPackageName(TArrayView<const FName> PackageNamesToRemove, bool bRemoveRecords);
 	void AsyncSave(FRecord& Record, const FCommitPackageInfo& Info);
 
 	void CollectForSavePackageData(FRecord& Record, FCommitContext& Context);
@@ -146,7 +140,7 @@ private:
 	// If EWriteOptions::ComputeHash is not set, the package will not get added to this.
 	TMap<FName, TRefCountPtr<FPackageHashes>> AllPackageHashes;
 
-	TMap<FName, FName> UncookedPathToCookedPath;
+	TMap<FName, TArray<FString>> PackageNameToCookedFiles;
 	/** CommitPackage can be called in parallel if using recursive save, so we need a lock for shared containers used during CommitPackage */
 	FCriticalSection PackageHashesLock;
 	FString OutputPath;
