[WP Navmesh] Do not use async flag when saving files using source control in WorldPartitionNavigationDataBuilder.
This fixes the "MarkForAdd, Error: Failed mark the file" issue that could occur during build.

diff --git a/Engine/Source/Editor/UnrealEd/Private/WorldPartition/WorldPartitionNavigationDataBuilder.cpp b/Engine/Source/Editor/UnrealEd/Private/WorldPartition/WorldPartitionNavigationDataBuilder.cpp
--- a/Engine/Source/Editor/UnrealEd/Private/WorldPartition/WorldPartitionNavigationDataBuilder.cpp
+++ b/Engine/Source/Editor/UnrealEd/Private/WorldPartition/WorldPartitionNavigationDataBuilder.cpp
@@ -234,7 +234,7 @@ bool UWorldPartitionNavigationDataBuilder::RunInternal(UWorld* World, const FCel
 			}
 		}
 
-		if (!SavePackages(PackagesToSave))
+		if (!SavePackages(PackageHelper, PackagesToSave))
 		{
 			return true;
 		}
@@ -391,7 +391,7 @@ bool UWorldPartitionNavigationDataBuilder::GenerateNavigationData(UWorldPartitio
 	return true;
 }
 
-bool UWorldPartitionNavigationDataBuilder::SavePackages(const TArray<UPackage*>& PackagesToSave) const
+bool UWorldPartitionNavigationDataBuilder::SavePackages(const FPackageSourceControlHelper& PackageHelper, const TArray<UPackage*>& PackagesToSave) const
 {
 	// Save packages
 	TRACE_CPUPROFILER_EVENT_SCOPE(SavingPackages);
@@ -403,7 +403,7 @@ bool UWorldPartitionNavigationDataBuilder::SavePackages(const TArray<UPackage*>&
 		FString PackageFileName = SourceControlHelpers::PackageFilename(Package);
 		FSavePackageArgs SaveArgs;
 		SaveArgs.TopLevelFlags = RF_Standalone;
-		SaveArgs.SaveFlags = SAVE_Async;
+		SaveArgs.SaveFlags = PackageHelper.UseSourceControl() ? ESaveFlags::SAVE_None : ESaveFlags::SAVE_Async;
 		if (!UPackage::SavePackage(Package, nullptr, *PackageFileName, SaveArgs))
 		{
 			UE_LOG(LogWorldPartitionNavigationDataBuilder, Error, TEXT("   Error saving package %s."), *Package->GetName());
@@ -414,7 +414,7 @@ bool UWorldPartitionNavigationDataBuilder::SavePackages(const TArray<UPackage*>&
 	return true;
 }
 
-bool UWorldPartitionNavigationDataBuilder::DeletePackages(FPackageSourceControlHelper& PackageHelper, const TArray<UPackage*>& PackagesToDelete) const
+bool UWorldPartitionNavigationDataBuilder::DeletePackages(const FPackageSourceControlHelper& PackageHelper, const TArray<UPackage*>& PackagesToDelete) const
 {
 	if (!PackagesToDelete.IsEmpty())
 	{
diff --git a/Engine/Source/Editor/UnrealEd/Public/WorldPartition/WorldPartitionNavigationDataBuilder.h b/Engine/Source/Editor/UnrealEd/Public/WorldPartition/WorldPartitionNavigationDataBuilder.h
--- a/Engine/Source/Editor/UnrealEd/Public/WorldPartition/WorldPartitionNavigationDataBuilder.h
+++ b/Engine/Source/Editor/UnrealEd/Public/WorldPartition/WorldPartitionNavigationDataBuilder.h
@@ -24,8 +24,8 @@ protected:
 
 	bool GenerateNavigationData(UWorldPartition* WorldPartition, const FBox& LoadedBounds, const FBox& GeneratingBounds) const;
 
-	bool SavePackages(const TArray<UPackage*>& PackagesToSave) const;
-	bool DeletePackages(FPackageSourceControlHelper& PackageHelper, const TArray<UPackage*>& PackagesToDelete) const;
+	bool SavePackages(const FPackageSourceControlHelper& PackageHelper, const TArray<UPackage*>& PackagesToSave) const;
+	bool DeletePackages(const FPackageSourceControlHelper& PackageHelper, const TArray<UPackage*>& PackagesToDelete) const;
 
 	bool bCleanBuilderPackages = false;
 };
