fix TGA not exporting UDIM
change the TGA exporter to use the Generic path
make TgaImageWrapper write 24 or 32 bit like the legacy exporter
delete the legacy exporter

diff --git a/Engine/Source/Editor/UnrealEd/Classes/Exporters/TextureExporterTGA.h b/Engine/Source/Editor/UnrealEd/Classes/Exporters/TextureExporterTGA.h
--- a/Engine/Source/Editor/UnrealEd/Classes/Exporters/TextureExporterTGA.h
+++ b/Engine/Source/Editor/UnrealEd/Classes/Exporters/TextureExporterTGA.h
@@ -9,19 +9,13 @@
 #include "CoreMinimal.h"
 #include "UObject/ObjectMacros.h"
 #include "Exporters/Exporter.h"
+#include "Exporters/TextureExporterGeneric.h"
 #include "TextureExporterTGA.generated.h"
 
 UCLASS()
-class UNREALED_API UTextureExporterTGA : public UExporter
+class UNREALED_API UTextureExporterTGA : public UTextureExporterGeneric
 {
 	GENERATED_UCLASS_BODY()
-
-
-	//~ Begin UExporter Interface
-	virtual bool SupportsObject(UObject* Object) const override;
-	virtual bool ExportBinary( UObject* Object, const TCHAR* Type, FArchive& Ar, FFeedbackContext* Warn, int32 FileIndex = 0, uint32 PortFlags=0 ) override;
-	//~ End UExporter Interface
+	
+	virtual bool SupportsTexture(UTexture* Texture) const override;
 };
-
-
-
diff --git a/Engine/Source/Editor/UnrealEd/Private/Factories/EditorFactories.cpp b/Engine/Source/Editor/UnrealEd/Private/Factories/EditorFactories.cpp
--- a/Engine/Source/Editor/UnrealEd/Private/Factories/EditorFactories.cpp
+++ b/Engine/Source/Editor/UnrealEd/Private/Factories/EditorFactories.cpp
@@ -4579,6 +4579,8 @@ void UTextureFactory::ParseFromJson(TSharedRef<class FJsonObject> ImportSettings
 
 /*------------------------------------------------------------------------------
 	UTextureExporterPCX implementation.
+	UTextureExporterPCX does not use TextureExporterGeneric because there's no PCX ImageWrapper
+		therefore does not support UDIM and other niceties like TextureExporterGeneric
 ------------------------------------------------------------------------------*/
 UTextureExporterPCX::UTextureExporterPCX(const FObjectInitializer& ObjectInitializer)
 	: Super(ObjectInitializer)
@@ -4587,7 +4589,6 @@ UTextureExporterPCX::UTextureExporterPCX(const FObjectInitializer& ObjectInitial
 	PreferredFormatIndex = 0;
 	FormatExtension.Add(TEXT("PCX"));
 	FormatDescription.Add(TEXT("PCX File"));
-
 }
 
 bool UTextureExporterPCX::SupportsObject(UObject* Object) const
@@ -4600,6 +4601,12 @@ bool UTextureExporterPCX::SupportsObject(UObject* Object) const
 		if (Texture)
 		{
 			bSupportsObject = Texture->Source.GetFormat() == TSF_BGRA8;
+			
+			if (Texture->Source.GetNumBlocks() > 1 )
+			{
+				// does not support UDIM
+				bSupportsObject = false;
+			}		
 		}
 	}
 	return bSupportsObject;
@@ -4663,7 +4670,8 @@ bool UTextureExporterPCX::ExportBinary( UObject* Object, const TCHAR* Type, FArc
 
 /*------------------------------------------------------------------------------
 	UTextureExporterDDS implementation.
-
+	UTextureExporterDDS does use UTextureExporterGeneric, but overrides its methods to do its own writing
+	because it can work on more complex textures, not just simple 2D like the base class
 ------------------------------------------------------------------------------*/
 
 UTextureExporterDDS::UTextureExporterDDS(const FObjectInitializer& ObjectInitializer)
@@ -5159,6 +5167,12 @@ URenderTargetCubeExporterHDR::URenderTargetCubeExporterHDR(const FObjectInitiali
 	SupportedClass = UTextureRenderTargetCube::StaticClass();
 }
 
+//-------------
+// UTextureExporterJPEG does not let you compress data to JPEG
+//	it only writes out existing JPEG data
+// do NOT use UTextureExporterGeneric here, that would go to ImageWrapper and compress to JPEG
+// note: does not support UDIM blocks like UTextureExporterGeneric
+
 UTextureExporterJPEG::UTextureExporterJPEG(const FObjectInitializer& ObjectInitializer)
 	: Super(ObjectInitializer)
 {
@@ -5177,6 +5191,12 @@ bool UTextureExporterJPEG::SupportsObject(UObject* Object) const
 		if (Texture)
 		{
 			// we do NOT do lossy recompression
+			
+			if (Texture->Source.GetNumBlocks() > 1 )
+			{
+				// does not support UDIM
+				return false;
+			}
 
 			// Check it has JPEG BulkData :
 			if ( Texture->Source.GetSourceCompression() == TSCF_JPEG &&
@@ -5216,156 +5236,16 @@ bool UTextureExporterJPEG::ExportBinary( UObject* Object, const TCHAR* Type, FAr
 	UTextureExporterTGA implementation.
 ------------------------------------------------------------------------------*/
 UTextureExporterTGA::UTextureExporterTGA(const FObjectInitializer& ObjectInitializer)
-	: Super(ObjectInitializer)
+	: UTextureExporterGeneric(ObjectInitializer)
 {
-	SupportedClass = UTexture2D::StaticClass();
-	PreferredFormatIndex = 0;
 	FormatExtension.Add(TEXT("TGA"));
 	FormatDescription.Add(TEXT("Targa"));
 }
 
-bool UTextureExporterTGA::SupportsObject(UObject* Object) const
+bool UTextureExporterTGA::SupportsTexture(UTexture* Texture) const
 {
-	bool bSupportsObject = false;
-	if (Super::SupportsObject(Object))
-	{
-		UTexture2D* Texture = Cast<UTexture2D>(Object);
-
-		if (Texture)
-		{
-			bSupportsObject = Texture->Source.GetFormat() == TSF_BGRA8;
-		}
-	}
-	return bSupportsObject;
-}
-
-bool UTextureExporterTGA::ExportBinary( UObject* Object, const TCHAR* Type, FArchive& Ar, FFeedbackContext* Warn, int32 FileIndex, uint32 PortFlags )
-{
-	UTexture2D* Texture = CastChecked<UTexture2D>( Object );
-
-	// could move to TGAImageWrapper and use generic exporter instead
-	//	maybe just leave this as-is
-
-	if (!Texture->Source.IsValid() || (Texture->Source.GetFormat() != TSF_BGRA8 && Texture->Source.GetFormat() != TSF_RGBA16))
-	{
-		return false;
-	}
-
-	// legacy code would claim to support 16 bit, but then just drop the bottom 8 bits
-	// this is no longer used
-	const bool bIsRGBA16 = Texture->Source.GetFormat() == TSF_RGBA16;
-
-	if (bIsRGBA16)
-	{
-		FMessageLog ExportWarning("EditorErrors");
-		FFormatNamedArguments Arguments;
-		Arguments.Add(TEXT("Name"), FText::FromString(Texture->GetName()));
-		ExportWarning.Warning(FText::Format(LOCTEXT("BitDepthTGAWarning", "{Name}: Texture is RGBA16 and cannot be represented at such high bit depth in .tga. Color will be scaled to RGBA8."), Arguments));
-		ExportWarning.Open(EMessageSeverity::Warning);
-	}
-
-	const int32 BytesPerPixel = bIsRGBA16 ? 8 : 4;
-
-	// Detects alpha channel
-	// writes either 24 or 32 bit
-
-	int32 SizeX = Texture->Source.GetSizeX();
-	int32 SizeY = Texture->Source.GetSizeY();
-	TArray64<uint8> RawData;
-	Texture->Source.GetMipData(RawData, 0);
-
-	// If we should export the file with no alpha info.  
-	// If the texture is compressed with no alpha we should definitely not export an alpha channel
-	//  Umm what
-	//bool bExportWithAlpha = !Texture->CompressionNoAlpha;
-	bool bExportWithAlpha = true;
-	if( bExportWithAlpha )
-	{
-		// If the texture isn't compressed with no alpha scan the texture to see if the alpha values are all 255 which means we can skip exporting it.
-		// This is a relatively slow process but we are just exporting textures 
-		bExportWithAlpha = false;
-		const int32 AlphaOffset = bIsRGBA16 ? 7 : 3;
-		for( int32 Y = SizeY - 1; Y >= 0; --Y )
-		{
-			uint8* Color = &RawData[Y * SizeX * BytesPerPixel];
-			for( int32 X = SizeX; X > 0; --X )
-			{
-				// Skip color info
-				Color += AlphaOffset;
-				// Get Alpha value then increment the pointer past it for the next pixel
-				uint8 Alpha = *Color++;
-				if( Alpha != 255 )
-				{
-					// When a texture is imported with no alpha, the alpha bits are set to 255
-					// So if the texture has non 255 alpha values, the texture is a valid alpha channel
-					bExportWithAlpha = true;
-					break;
-				}
-			}
-			if( bExportWithAlpha )
-			{
-				break;
-			}
-		}
-	}
-
-	const uint16 OriginalWidth = IntCastChecked<uint16>(SizeX);
-	const uint16 OriginalHeight = IntCastChecked<uint16>(SizeY);
-
-	FTGAFileHeader TGA;
-	FMemory::Memzero( &TGA, sizeof(TGA) );
-	TGA.ImageTypeCode = 2;
-	TGA.BitsPerPixel = bExportWithAlpha ? 32 : 24 ;
-	TGA.Height = OriginalHeight;
-	TGA.Width = OriginalWidth;
-	Ar.Serialize( &TGA, sizeof(TGA) );
-
-	if( bExportWithAlpha && !bIsRGBA16)
-	{
-		for( int32 Y=0;Y < OriginalHeight;Y++ )
-		{
-			// If we aren't skipping alpha channels we can serialize each line
-			Ar.Serialize( &RawData[ (OriginalHeight - Y - 1) * OriginalWidth * 4 ], OriginalWidth * 4 );
-		}
-	}
-	else
-	{
-		// Serialize each pixel
-		for( int32 Y = OriginalHeight - 1; Y >= 0; --Y )
-		{
-			uint8* Color = &RawData[Y * OriginalWidth * BytesPerPixel];
-			for( int32 X = OriginalWidth; X > 0; --X )
-			{
-				if (bIsRGBA16)
-				{
-					// From RGB to BGR
-					Ar << Color[5];
-					Ar << Color[3];
-					Ar << Color[1];
-					if (bExportWithAlpha)
-					{
-						Ar << Color[7];
-					}
-					Color += 8;
-				}
-				else
-				{
-					Ar << *Color++;
-					Ar << *Color++;
-					Ar << *Color++;
-					// Skip alpha channel since we are exporting with no alpha
-					Color++;
-				}
-			}
-		}
-	}
-
-	FTGAFileFooter Ftr;
-	FMemory::Memzero( &Ftr, sizeof(Ftr) );
-	FMemory::Memcpy( Ftr.Signature, "TRUEVISION-XFILE", 16 );
-	Ftr.TrailingPeriod = '.';
-	Ar.Serialize( &Ftr, sizeof(Ftr) );
-	return true;
+	ETextureSourceFormat TSF = Texture->Source.GetFormat();
+	return TSF == TSF_BGRA8;
 }
 
 /*------------------------------------------------------------------------------
diff --git a/Engine/Source/Runtime/ImageWrapper/Private/Formats/TgaImageWrapper.cpp b/Engine/Source/Runtime/ImageWrapper/Private/Formats/TgaImageWrapper.cpp
--- a/Engine/Source/Runtime/ImageWrapper/Private/Formats/TgaImageWrapper.cpp
+++ b/Engine/Source/Runtime/ImageWrapper/Private/Formats/TgaImageWrapper.cpp
@@ -23,37 +23,87 @@ void FTgaImageWrapper::Compress(int32 Quality)
 {
 	CompressedData.Reset();
 
-	int64 SurfaceBytes = Width * Height * 4;
-	check( RawData.Num() == SurfaceBytes );
+	FImageView Image;
+	if ( ! GetImageViewOfSetRawForCompress(Image) )
+	{
+		SetError(TEXT("No valid image to compress"));
+		return;
+	}
+	
+	if ( Width > UINT16_MAX || Height > UINT16_MAX )
+	{
+		SetError(TEXT("Image is too large to export as TGA"));
+		return;
+	}
+			
+	check( Image.Format == ERawImageFormat::BGRA8 );
+	int64 ImageRowBytes = Width * 4;
+	
+	// write 32-bit if we have a non-all-255 alpha channel, else write 24 bit
+	bool bHasAlpha = FImageCore::DetectAlphaChannel(Image);
+	int64 OutBytesPerPixel = bHasAlpha ? 4 : 3;
+	int64 OutSurfaceBytes = (int64) Width * Height * OutBytesPerPixel;
 
-	int64 CompresedSize = sizeof(FTGAFileHeader) + SurfaceBytes;
+	int64 CompresedSize = sizeof(FTGAFileHeader) + OutSurfaceBytes;
 	CompressedData.SetNum(CompresedSize);
 
 	uint8 * CompressedPtr = CompressedData.GetData();
 	
-	// super lazy TGA writer here :
-	//	always writes 32 bit BGRA
-
+	// http://www.paulbourke.net/dataformats/tga/
+	// https://en.wikipedia.org/wiki/Truevision_TGA
 	FTGAFileHeader header = { };
-
-	header.Width  = INTEL_ORDER16( (uint16) Width );
-	header.Height = INTEL_ORDER16( (uint16) Height );
+	
+	uint16 Width16  = IntCastChecked<uint16>( Width );
+	uint16 Height16 = IntCastChecked<uint16>( Height );
+	header.Width  = INTEL_ORDER16( Width16 );
+	header.Height = INTEL_ORDER16( Height16 );
 	header.ImageTypeCode = 2;
-	header.BitsPerPixel = 32;
-	header.ImageDescriptor = 8;
+	header.BitsPerPixel = OutBytesPerPixel*8;
+
+	// Image descriptor (1 byte): bits 3-0 give the alpha channel depth, bits 5-4 give pixel ordering
+	header.ImageDescriptor = bHasAlpha ? 8 : 0;
 
 	memcpy(CompressedPtr,&header,sizeof(header));
 	CompressedPtr += sizeof(header);
 
 	// write rows BGRA , bottom up :
-	for(int y = Height-1; y>= 0;y--)
+	for(int32 y = Height-1; y>= 0;y--)
 	{
-		int64 RowBytes = Width * 4;
-		const void * From = &RawData[y * RowBytes]; 
-		memcpy(CompressedPtr,From,RowBytes);
-		CompressedPtr += RowBytes;
+		const uint8 * From = &RawData[y * ImageRowBytes]; 
+
+		if ( bHasAlpha )
+		{
+			check( OutBytesPerPixel == 4 );
+			check( RawData.Num() == OutSurfaceBytes );
+
+			memcpy(CompressedPtr,From,ImageRowBytes);
+			CompressedPtr += ImageRowBytes;
+		}
+		else
+		{
+			check( OutBytesPerPixel == 3 );
+			
+			for(int32 x=0;x<Width;x++)
+			{
+				// copy RGB, skip A
+				CompressedPtr[0] = From[0];
+				CompressedPtr[1] = From[1];
+				CompressedPtr[2] = From[2];
+				From += 4;
+				CompressedPtr += 3;
+			}
+		}
 	}
 
+	// TGA file footer is optional :
+	/*
+	FTGAFileFooter Ftr;
+	FMemory::Memzero( &Ftr, sizeof(Ftr) );
+	FMemory::Memcpy( Ftr.Signature, "TRUEVISION-XFILE", 16 );
+	Ftr.TrailingPeriod = '.';
+	Ar.Serialize( &Ftr, sizeof(Ftr) );
+	*/
+
 	check( CompressedPtr == CompressedData.GetData() + CompressedData.Num() );
 }
 
diff --git a/Engine/Source/Runtime/ImageWrapper/Private/ImageWrapperBase.cpp b/Engine/Source/Runtime/ImageWrapper/Private/ImageWrapperBase.cpp
--- a/Engine/Source/Runtime/ImageWrapper/Private/ImageWrapperBase.cpp
+++ b/Engine/Source/Runtime/ImageWrapper/Private/ImageWrapperBase.cpp
@@ -357,6 +357,37 @@ void IImageWrapper::ConvertRawImageFormat(ERawImageFormat::Type RawFormat, ERGBF
 	}
 }
 
+bool FImageWrapperBase::GetImageViewOfSetRawForCompress(FImageView & OutImage) const
+{
+	if ( RawData.IsEmpty() )
+	{
+		return false;
+	}
+	
+	bool bExactMatch;
+	ERawImageFormat::Type RawFormat = GetClosestRawImageFormat(&bExactMatch);
+	
+	// must be bExactMatch, no RB swaps possible, because we just point at the array
+	//	this function will fail if SetRaw is done with RGBA8 rather than BGRA8
+	if ( RawFormat == ERawImageFormat::Invalid || !bExactMatch )
+	{
+		return false;
+	}
+
+	// ImageWrapper RGBFormat doesn't track if pixels are Gamma/sRGB or not
+	//	just assume they are Default for now :
+	EGammaSpace GammaSpace = ERawImageFormat::GetDefaultGammaSpace(RawFormat);
+	
+	OutImage.RawData = (void *) &RawData[0];
+	OutImage.SizeX = Width;
+	OutImage.SizeY = Height;
+	OutImage.NumSlices = 1;
+	OutImage.Format = RawFormat;
+	OutImage.GammaSpace = GammaSpace;
+
+	return true;
+}
+
 bool IImageWrapper::GetRawImage(FImage & OutImage)
 {
 	TArray64<uint8> OutRawData;
diff --git a/Engine/Source/Runtime/ImageWrapper/Private/ImageWrapperBase.h b/Engine/Source/Runtime/ImageWrapper/Private/ImageWrapperBase.h
--- a/Engine/Source/Runtime/ImageWrapper/Private/ImageWrapperBase.h
+++ b/Engine/Source/Runtime/ImageWrapper/Private/ImageWrapperBase.h
@@ -118,6 +118,11 @@ protected:
 	int GetBytesPerPel() const { return GetRGBFormatBytesPerPel(Format,BitDepth); }
 	int GetBytesPerRow() const { return Width * GetBytesPerPel(); }
 
+	// For writers: after SetRaw(), call this to get an ImageView of the raw data that was set
+	//	can return false if the SetRaw does not map to an image format
+	//	pixels point at the RawData array
+	bool GetImageViewOfSetRawForCompress(FImageView & OutImage) const;
+
 	/** Arrays of compressed/raw data */
 	TArray64<uint8> RawData;
 	TArray64<uint8> CompressedData;
@@ -127,6 +132,7 @@ protected:
 	int BitDepth;
 
 	/** Width/Height of the image data */
+	// @todo Oodle : change these to int64
 	int32 Width;
 	int32 Height;
 	
