Updated SearchDirectoryPaths in the source code gather to handle trailing wildcards

diff --git a/Engine/Source/Developer/Localization/Private/LocalizationTargetTypes.cpp b/Engine/Source/Developer/Localization/Private/LocalizationTargetTypes.cpp
--- a/Engine/Source/Developer/Localization/Private/LocalizationTargetTypes.cpp
+++ b/Engine/Source/Developer/Localization/Private/LocalizationTargetTypes.cpp
@@ -76,13 +76,28 @@ bool FGatherTextSearchDirectory::Validate(const bool bIsEngineTarget, FText& Out
 		return false;
 	}
 
-	const FString ResolvedPath = FPaths::IsRelative(Path) ? FPaths::Combine(*FLocalizationGatherPathRootUtil::GetResolvedPathRoot(PathRoot, bIsEngineTarget), *Path) : Path;
+	FString ResolvedPath = FPaths::IsRelative(Path) ? FPaths::Combine(*FLocalizationGatherPathRootUtil::GetResolvedPathRoot(PathRoot, bIsEngineTarget), *Path) : Path;
 	if (ResolvedPath.IsEmpty())
 	{
 		OutError = LOCTEXT("SearchDirectoryEmptyError", "Search directory not specified. Use \".\" to specify the root directory.");
 		return false;
 	}
 
+	{
+		int32 WildcardIndex = INDEX_NONE;
+		if (ResolvedPath.FindChar(TEXT('*'), WildcardIndex))
+		{
+			if (WildcardIndex != ResolvedPath.Len() - 1)
+			{
+				OutError = LOCTEXT("SearchDirectoryInvalidWildcardMustBeTrailingError", "Search directory can only use a trailing wildcard.");
+				return false;
+			}
+
+			// Trim the wildcard from the search path, as it would fail the rest of the validation with it
+			ResolvedPath = FPaths::GetPath(MoveTemp(ResolvedPath));
+		}
+	}
+
 	FText InvalidPathReason;
 	if (!FPaths::ValidatePath(ResolvedPath, &InvalidPathReason))
 	{
diff --git a/Engine/Source/Editor/UnrealEd/Private/Commandlets/GatherTextFromSourceCommandlet.cpp b/Engine/Source/Editor/UnrealEd/Private/Commandlets/GatherTextFromSourceCommandlet.cpp
--- a/Engine/Source/Editor/UnrealEd/Private/Commandlets/GatherTextFromSourceCommandlet.cpp
+++ b/Engine/Source/Editor/UnrealEd/Private/Commandlets/GatherTextFromSourceCommandlet.cpp
@@ -144,7 +144,9 @@ int32 UGatherTextFromSourceCommandlet::Main( const FString& Params )
 	TArray<FString> IncludePathFilters;
 	Algo::Transform(SearchDirectoryPaths, IncludePathFilters, [](const FString& SearchDirectoryPath)
 	{
-		return SearchDirectoryPath / TEXT("*");
+		return SearchDirectoryPath.EndsWith(TEXT("*"), ESearchCase::CaseSensitive)
+			? SearchDirectoryPath								// Already a wildcard
+			: FPaths::Combine(SearchDirectoryPath, TEXT("*"));	// Add a wildcard
 	});
 
 	FGatherTextDelegates::GetAdditionalGatherPaths.Broadcast(GatherManifestHelper->GetTargetName(), IncludePathFilters, ExcludePathFilters);
@@ -153,6 +155,7 @@ int32 UGatherTextFromSourceCommandlet::Main( const FString& Params )
 	TArray<FString> FilesToProcess;
 	{
 		TArray<FString> RootSourceFiles;
+		TSet<FString, FLocKeySetFuncs> ProcessedSearchDirectoryPaths;
 		for (const FString& IncludePathFilter : IncludePathFilters)
 		{
 			FString SearchDirectoryPath = IncludePathFilter;
@@ -162,6 +165,13 @@ int32 UGatherTextFromSourceCommandlet::Main( const FString& Params )
 				SearchDirectoryPath = FPaths::GetPath(MoveTemp(SearchDirectoryPath));
 			}
 
+			bool bAlreadyProcessed = false;
+			ProcessedSearchDirectoryPaths.Add(SearchDirectoryPath, &bAlreadyProcessed);
+			if (bAlreadyProcessed)
+			{
+				continue;
+			}
+
 			for (const FString& UniqueSourceFileSearchFilter : UniqueSourceFileSearchFilters)
 			{
 				IFileManager::Get().FindFilesRecursive(RootSourceFiles, *SearchDirectoryPath, *UniqueSourceFileSearchFilter, true, false, false);
