Updated SearchDirectoryPaths in the source code gather to handle general filter-style wildcards

diff --git a/Engine/Source/Developer/Localization/Private/LocalizationTargetTypes.cpp b/Engine/Source/Developer/Localization/Private/LocalizationTargetTypes.cpp
--- a/Engine/Source/Developer/Localization/Private/LocalizationTargetTypes.cpp
+++ b/Engine/Source/Developer/Localization/Private/LocalizationTargetTypes.cpp
@@ -3,6 +3,7 @@
 #include "LocalizationTargetTypes.h"
 #include "Templates/Casts.h"
 #include "HAL/FileManager.h"
+#include "Misc/AsciiSet.h"
 #include "Misc/FileHelper.h"
 #include "Misc/Paths.h"
 #include "Misc/App.h"
@@ -84,16 +85,11 @@ bool FGatherTextSearchDirectory::Validate(const bool bIsEngineTarget, FText& Out
 	}
 
 	{
-		int32 WildcardIndex = INDEX_NONE;
-		if (ResolvedPath.FindChar(TEXT('*'), WildcardIndex))
+		constexpr FAsciiSet Wildcards("*?");
+		if (const TCHAR* FirstWildcard = FAsciiSet::FindFirstOrEnd(*ResolvedPath, Wildcards); *FirstWildcard != 0)
 		{
-			if (WildcardIndex != ResolvedPath.Len() - 1)
-			{
-				OutError = LOCTEXT("SearchDirectoryInvalidWildcardMustBeTrailingError", "Search directory can only use a trailing wildcard.");
-				return false;
-			}
-
-			// Trim the wildcard from the search path, as it would fail the rest of the validation with it
+			// Trim the wildcard from this search path, as it would fail the rest of the validation with it
+			ResolvedPath = ResolvedPath.Left(UE_PTRDIFF_TO_INT32(FirstWildcard - *ResolvedPath));
 			ResolvedPath = FPaths::GetPath(MoveTemp(ResolvedPath));
 		}
 	}
diff --git a/Engine/Source/Editor/UnrealEd/Private/Commandlets/GatherTextFromSourceCommandlet.cpp b/Engine/Source/Editor/UnrealEd/Private/Commandlets/GatherTextFromSourceCommandlet.cpp
--- a/Engine/Source/Editor/UnrealEd/Private/Commandlets/GatherTextFromSourceCommandlet.cpp
+++ b/Engine/Source/Editor/UnrealEd/Private/Commandlets/GatherTextFromSourceCommandlet.cpp
@@ -2,6 +2,7 @@
 
 #include "Commandlets/GatherTextFromSourceCommandlet.h"
 #include "HAL/FileManager.h"
+#include "Misc/AsciiSet.h"
 #include "Misc/FileHelper.h"
 #include "Misc/Paths.h"
 #include "Misc/ExpressionParserTypes.h"
@@ -144,9 +145,10 @@ int32 UGatherTextFromSourceCommandlet::Main( const FString& Params )
 	TArray<FString> IncludePathFilters;
 	Algo::Transform(SearchDirectoryPaths, IncludePathFilters, [](const FString& SearchDirectoryPath)
 	{
-		return SearchDirectoryPath.EndsWith(TEXT("*"), ESearchCase::CaseSensitive)
+		const TCHAR LastChar = SearchDirectoryPath.Len() > 0 ? SearchDirectoryPath[SearchDirectoryPath.Len() - 1] : 0;
+		return (LastChar == TEXT('*') || LastChar == TEXT('?'))
 			? SearchDirectoryPath								// Already a wildcard
-			: FPaths::Combine(SearchDirectoryPath, TEXT("*"));	// Add a wildcard
+			: FPaths::Combine(SearchDirectoryPath, TEXT("*"));	// Add a catch-all wildcard
 	});
 
 	FGatherTextDelegates::GetAdditionalGatherPaths.Broadcast(GatherManifestHelper->GetTargetName(), IncludePathFilters, ExcludePathFilters);
@@ -158,10 +160,13 @@ int32 UGatherTextFromSourceCommandlet::Main( const FString& Params )
 		TSet<FString, FLocKeySetFuncs> ProcessedSearchDirectoryPaths;
 		for (const FString& IncludePathFilter : IncludePathFilters)
 		{
+			constexpr FAsciiSet Wildcards("*?");
+
 			FString SearchDirectoryPath = IncludePathFilter;
-			if (SearchDirectoryPath.EndsWith(TEXT("*"), ESearchCase::CaseSensitive))
+			if (const TCHAR* FirstWildcard = FAsciiSet::FindFirstOrEnd(*SearchDirectoryPath, Wildcards); *FirstWildcard != 0)
 			{
 				// Trim the wildcard from this search path
+				SearchDirectoryPath = SearchDirectoryPath.Left(UE_PTRDIFF_TO_INT32(FirstWildcard - *SearchDirectoryPath));
 				SearchDirectoryPath = FPaths::GetPath(MoveTemp(SearchDirectoryPath));
 			}
 
