Fix a crash when rebuilding the world partition editor minimap

- CL 20575516 locks the minimap texture for read/write before we start running the minimap rebuild. If the texture ends up being compiled we are then unable to access the texture data as it is locked. This only happens if the compiled texture data is not found in the DDC.
- I've removed the lock from PreRun and instead make a read/write lock right before we actually try to write data to the texture, unlocking it once the write has completed. If debug capture is enabled then I make a readonly for the scope of the capture.
- Note that because the texture payload will be in memory (after being written to) this additional readonly lock does not need to do anything expensive like loading the data from disk or decompressing it, so it should be fast.

diff --git a/Engine/Source/Editor/UnrealEd/Private/WorldPartition/WorldPartitionMiniMapBuilder.cpp b/Engine/Source/Editor/UnrealEd/Private/WorldPartition/WorldPartitionMiniMapBuilder.cpp
--- a/Engine/Source/Editor/UnrealEd/Private/WorldPartition/WorldPartitionMiniMapBuilder.cpp
+++ b/Engine/Source/Editor/UnrealEd/Private/WorldPartition/WorldPartitionMiniMapBuilder.cpp
@@ -66,7 +66,6 @@ bool UWorldPartitionMiniMapBuilder::PreRun(UWorld* World, FPackageSourceControlH
 		WorldMiniMap->MiniMapTexture = Factory->CreateTexture2D(WorldMiniMap, TEXT("MinimapTexture"), RF_NoFlags);
 		WorldMiniMap->MiniMapTexture->Source.Init(MinimapImageSizeX, MinimapImageSizeY, 1, 1, TSF_BGRA8);
 		WorldMiniMap->MiniMapWorldBounds = IterativeWorldBounds;
-		MiniMapSourcePtr = WorldMiniMap->MiniMapTexture->Source.LockMip(0);
 	}
 
 	// Compute world to minimap transform
@@ -137,8 +136,11 @@ bool UWorldPartitionMiniMapBuilder::RunInternal(UWorld* World, const FCellInfo&
 		const uint8* SrcDataPtr = TileTexture->Source.LockMipReadOnly(0);
 		check(SrcDataPtr);
 
+		uint8* const MiniMapDstPtr = WorldMiniMap->MiniMapTexture->Source.LockMip(0);
+		check(MiniMapDstPtr);
+
 		const uint32 DstDataStrideBytes = WorldMiniMap->MiniMapTexture->Source.GetSizeX() * BPP;
-		uint8* const DstDataPtr = MiniMapSourcePtr + (DstMin.Y * DstDataStrideBytes) + (DstMin.X * BPP);
+		uint8* const DstDataPtr = MiniMapDstPtr + (DstMin.Y * DstDataStrideBytes) + (DstMin.X * BPP);
 		check(DstDataPtr);
 
 		for (uint32 RowIdx = 0; RowIdx < CaptureHeightPixels; ++RowIdx)
@@ -163,6 +165,7 @@ bool UWorldPartitionMiniMapBuilder::RunInternal(UWorld* World, const FCellInfo&
 			FFileHelper::CreateBitmap(*MinimapDebugImagePath, TileTexture->Source.GetSizeX(), TileTexture->Source.GetSizeY(), (FColor*)SrcDataPtr);
 		}
 
+		WorldMiniMap->MiniMapTexture->Source.UnlockMip(0);
 		TileTexture->Source.UnlockMip(0);
 	}
 
@@ -186,10 +189,12 @@ bool UWorldPartitionMiniMapBuilder::PostRun(UWorld* World, FPackageSourceControl
 		{
 			const FString DirectoryPath = FPaths::ConvertRelativePathToFull(FPaths::ProjectIntermediateDir() + TEXT("Minimap"));
 			FString MinimapDebugImagePath = DirectoryPath / World->GetName() + "-Minimap.bmp";
-			FFileHelper::CreateBitmap(*MinimapDebugImagePath, WorldMiniMap->MiniMapTexture->Source.GetSizeX(), WorldMiniMap->MiniMapTexture->Source.GetSizeY(), (FColor*)MiniMapSourcePtr);
+
+			const void* MiniMapSourcePtr = WorldMiniMap->MiniMapTexture->Source.LockMipReadOnly(0);
+			FFileHelper::CreateBitmap(*MinimapDebugImagePath, WorldMiniMap->MiniMapTexture->Source.GetSizeX(), WorldMiniMap->MiniMapTexture->Source.GetSizeY(), (const FColor*)MiniMapSourcePtr);
+			WorldMiniMap->MiniMapTexture->Source.UnlockMip(0);
 		}
 
-		WorldMiniMap->MiniMapTexture->Source.UnlockMip(0);
 		WorldMiniMap->MiniMapTexture->PowerOfTwoMode = ETexturePowerOfTwoSetting::PadToPowerOfTwo;	// Required for VTs
 		WorldMiniMap->MiniMapTexture->AdjustMinAlpha = 1.f;
 		WorldMiniMap->MiniMapTexture->LODGroup = TEXTUREGROUP_UI;
diff --git a/Engine/Source/Editor/UnrealEd/Public/WorldPartition/WorldPartitionMiniMapBuilder.h b/Engine/Source/Editor/UnrealEd/Public/WorldPartition/WorldPartitionMiniMapBuilder.h
--- a/Engine/Source/Editor/UnrealEd/Public/WorldPartition/WorldPartitionMiniMapBuilder.h
+++ b/Engine/Source/Editor/UnrealEd/Public/WorldPartition/WorldPartitionMiniMapBuilder.h
@@ -35,7 +35,5 @@ private:
 	int32 MinimapImageSizeX;
 	int32 MinimapImageSizeY;
 
-	uint8* MiniMapSourcePtr;
-
 	bool bDebugCapture;
 };
\ No newline at end of file
