Downgrade the "Cooker has been blocked..." message from a warning to a display message.  Can be restored to a warning when the inter-process communication is moved to its own thread and doesn't suffer delays due to long operations on the game thread.

diff --git a/Engine/Source/Editor/UnrealEd/Private/Cooker/CookDirector.cpp b/Engine/Source/Editor/UnrealEd/Private/Cooker/CookDirector.cpp
--- a/Engine/Source/Editor/UnrealEd/Private/Cooker/CookDirector.cpp
+++ b/Engine/Source/Editor/UnrealEd/Private/Cooker/CookDirector.cpp
@@ -807,7 +807,7 @@ void FCookDirector::SetWorkersStalled(bool bInWorkersStalled)
 		const double CurrentTime = FPlatformTime::Seconds();
 		if (CurrentTime >= WorkersStalledWarnTimeSeconds)
 		{
-			UE_LOG(LogCook, Warning, TEXT("Cooker has been blocked with no results from remote CookWorkers for %.0f seconds."),
+			UE_LOG(LogCook, Display, TEXT("Cooker has been blocked with no results from remote CookWorkers for %.0f seconds."),
 				(float)(CurrentTime - WorkersStalledStartTimeSeconds));
 			WorkersStalledWarnTimeSeconds = CurrentTime + GCookProgressWarnBusyTime;
 		}
