Cherry picked 23677629 from //UE5/Main
Remove polling for registry scan completion and use the correct blocking query function

diff --git a/Engine/Source/Editor/UnrealEd/Private/Commandlets/AssetRegistryGenerator.cpp b/Engine/Source/Editor/UnrealEd/Private/Commandlets/AssetRegistryGenerator.cpp
--- a/Engine/Source/Editor/UnrealEd/Private/Commandlets/AssetRegistryGenerator.cpp
+++ b/Engine/Source/Editor/UnrealEd/Private/Commandlets/AssetRegistryGenerator.cpp
@@ -1233,12 +1233,7 @@ void FAssetRegistryGenerator::Initialize(const TArray<FName> &InStartupPackages,
 	FAssetRegistrySerializationOptions SaveOptions;
 
 	// If the asset registry is still doing it's background scan, we need to wait for it to finish and tick it so that the results are flushed out
-	while (AssetRegistry.IsLoadingAssets())
-	{
-		AssetRegistry.Tick(-1.0f);
-		FThreadHeartBeat::Get().HeartBeat();
-		FPlatformProcess::SleepNoStats(0.0001f);
-	}
+	AssetRegistry.WaitForCompletion();
 
 	ensureMsgf(!AssetRegistry.IsLoadingAssets(), TEXT("Cannot initialize asset registry generator while asset registry is still scanning source assets "));
 
diff --git a/Engine/Source/Editor/UnrealEd/Private/CookOnTheFlyServer.cpp b/Engine/Source/Editor/UnrealEd/Private/CookOnTheFlyServer.cpp
--- a/Engine/Source/Editor/UnrealEd/Private/CookOnTheFlyServer.cpp
+++ b/Engine/Source/Editor/UnrealEd/Private/CookOnTheFlyServer.cpp
@@ -7267,6 +7267,7 @@ void UCookOnTheFlyServer::BlockOnAssetRegistry()
 		TArray<FString> ScanPaths;
 		ScanPaths.Add(FString::Printf(TEXT("/%s/"), *CookByTheBookOptions->DlcName));
 		AssetRegistry->ScanPathsSynchronous(ScanPaths);
+		AssetRegistry->WaitForCompletion();
 	}
 	UE::Cook::FPackageDatas::OnAssetRegistryGenerated(*AssetRegistry);
 
