Adding optional json output file for cook stats.

diff --git a/Engine/Source/Editor/UnrealEd/Private/Cooker/CookProfiling.cpp b/Engine/Source/Editor/UnrealEd/Private/Cooker/CookProfiling.cpp
--- a/Engine/Source/Editor/UnrealEd/Private/Cooker/CookProfiling.cpp
+++ b/Engine/Source/Editor/UnrealEd/Private/Cooker/CookProfiling.cpp
@@ -12,10 +12,14 @@
 #include "CookOnTheSide/CookLog.h"
 #include "CoreGlobals.h"
 #include "DerivedDataBuildRemoteExecutor.h"
+#include "HAL/FileManager.h"
+#include "Misc/CommandLine.h"
 #include "Misc/OutputDevice.h"
 #include "Misc/StringBuilder.h"
 #include "PackageBuildDependencyTracker.h"
+#include "Policies/CondensedJsonPrintPolicy.h"
 #include "Serialization/ArchiveUObject.h"
+#include "Serialization/JsonWriter.h"
 #include "Templates/Casts.h"
 #include "UObject/GCObject.h"
 #include "UObject/MetaData.h"
@@ -1097,6 +1101,27 @@ void LogCookStats(ECookMode::Type CookMode)
 		{
 			UE_LOG(LogCook, Display, TEXT("%s.%s=%s"), *ProfileEntry.Path, *ProfileEntry.Key, *ProfileEntry.Value);
 		}
+
+		FString CookStatsFileName;
+		if (FParse::Value(FCommandLine::Get(), TEXT("-CookStatsFile="), CookStatsFileName))
+		{
+			FString JsonString;
+			TSharedRef<TJsonWriter<TCHAR, TCondensedJsonPrintPolicy<TCHAR>>> JsonWriter = TJsonWriterFactory<TCHAR, TCondensedJsonPrintPolicy<TCHAR> >::Create(&JsonString);
+			JsonWriter->WriteObjectStart();
+			for (const auto& ProfileEntry : CookProfileData)
+			{
+				JsonWriter->WriteObjectStart(ProfileEntry.Key);
+				JsonWriter->WriteValue(TEXT("Path"), ProfileEntry.Path);
+				JsonWriter->WriteValue(TEXT("TimInSec"), ProfileEntry.Value);
+				JsonWriter->WriteObjectEnd();
+			}
+			JsonWriter->WriteObjectEnd();
+			JsonWriter->Close();
+			TUniquePtr<FArchive> JsonFile(IFileManager::Get().CreateFileWriter(*CookStatsFileName));
+			JsonFile->Serialize(TCHAR_TO_ANSI(*JsonString), JsonString.Len());
+			JsonFile->Close();
+		}
+
 	}
 	if (DDCSummaryStats.Num() > 0)
 	{
