[NO NON-PII COMMIT MESSAGE]

diff --git a/Engine/Source/Editor/UnrealEd/Private/EditorEngine.cpp b/Engine/Source/Editor/UnrealEd/Private/EditorEngine.cpp
--- a/Engine/Source/Editor/UnrealEd/Private/EditorEngine.cpp
+++ b/Engine/Source/Editor/UnrealEd/Private/EditorEngine.cpp
@@ -6177,7 +6177,7 @@ void UEditorEngine::DoConvertActors( const TArray<AActor*>& ActorsToConvert, UCl
 				UActorGroupingUtils::SetGroupingActive(bGroupingActiveSaved);
 			}
 
-
+			// Attempt normal spawning if a new actor hasn't been spawned yet via a special case
 			if (!NewActor)
 			{
 				// Set the current level to the level where the convertible actor resides
@@ -6190,17 +6190,20 @@ void UEditorEngine::DoConvertActors( const TArray<AActor*>& ActorsToConvert, UCl
 				const UClass* CommonBaseClass = ActorToConvert->FindNearestCommonBaseClass( ConvertToClass );
 				check ( CommonBaseClass );	
 
-				// Take the old actors location always, not rotation.  If rotation was changed on the source actor, it will be copied below.
-				FVector SpawnLoc = ActorToConvert->GetActorLocation();
-				FRotator SpawnRot = ActorToConvert->GetActorRotation();
+				const FTransform& SpawnTransform = ActorToConvert->GetActorTransform();
 				{
 					FActorSpawnParameters SpawnInfo;
 					SpawnInfo.OverrideLevel = ActorLevel;
 					SpawnInfo.SpawnCollisionHandlingOverride = ESpawnActorCollisionHandlingMethod::AlwaysSpawn;
-					NewActor = World->SpawnActor( ConvertToClass, &SpawnLoc, &SpawnRot, SpawnInfo );
+					SpawnInfo.bDeferConstruction = true;
+					NewActor = World->SpawnActor(ConvertToClass, &SpawnTransform, SpawnInfo);
 
 					if (NewActor)
 					{
+						// Deferred spawning and finishing with !bIsDefaultTransform results in scale being applied for both native and simple construction script created root components
+						constexpr bool bIsDefaultTransform = false;
+						NewActor->FinishSpawning(SpawnTransform, bIsDefaultTransform);
+						
 						// Copy non component properties from the old actor to the new actor
 						for( FProperty* Property = CommonBaseClass->PropertyLink; Property != NULL; Property = Property->PropertyLinkNext )
 						{
