Cooker: Remove invalid call to ::delete FPackageData. PackageDatas have to be freed using the Allocator.
Override operator new and delete to give an error to prevent mistaken calls to new and delete.
#rnx

diff --git a/Engine/Source/Editor/UnrealEd/Private/Cooker/CookPackageData.cpp b/Engine/Source/Editor/UnrealEd/Private/Cooker/CookPackageData.cpp
--- a/Engine/Source/Editor/UnrealEd/Private/Cooker/CookPackageData.cpp
+++ b/Engine/Source/Editor/UnrealEd/Private/Cooker/CookPackageData.cpp
@@ -2379,19 +2379,18 @@ FPackageData& FPackageDatas::CreatePackageData(FName PackageName, FName FileName
 	check(!FileName.IsNone());
 
 	FWriteScopeLock ExistenceWriteLock(ExistenceLock);
-	FPackageData* PackageData = Allocator.NewElement(*this, PackageName, FileName);
 	FPackageData*& ExistingByPackageName = PackageNameToPackageData.FindOrAdd(PackageName);
 	FPackageData*& ExistingByFileName = FileNameToPackageData.FindOrAdd(FileName);
 	if (ExistingByPackageName)
 	{
 		// The other CreatePackageData call should have added the FileName as well
 		check(ExistingByFileName == ExistingByPackageName);
-		delete PackageData;
 		return *ExistingByPackageName;
 	}
 	// If no other CreatePackageData added the PackageName, then they should not have added
 	// the FileName either
 	check(!ExistingByFileName);
+	FPackageData* PackageData = Allocator.NewElement(*this, PackageName, FileName);
 	ExistingByPackageName = PackageData;
 	ExistingByFileName = PackageData;
 	return *PackageData;
diff --git a/Engine/Source/Editor/UnrealEd/Private/Cooker/CookPackageData.h b/Engine/Source/Editor/UnrealEd/Private/Cooker/CookPackageData.h
--- a/Engine/Source/Editor/UnrealEd/Private/Cooker/CookPackageData.h
+++ b/Engine/Source/Editor/UnrealEd/Private/Cooker/CookPackageData.h
@@ -154,6 +154,19 @@ public:
 
 	FPackageData(FPackageDatas& PackageDatas, const FName& InPackageName, const FName& InFileName);
 	~FPackageData();
+	static void* operator new(size_t Size)
+	{
+		checkf(false, TEXT("PackageDatas should be allocated using FPackageDatas.Allocator"));
+		return FMemory::Malloc(Size);
+	}
+	static void* operator new(size_t Size, void* PlacementNewPtr)
+	{
+		return PlacementNewPtr;
+	}
+	static void operator delete(void* Ptr, size_t Size)
+	{
+		checkf(false, TEXT("PackageDatas should be freed using FPackageDatas.Allocator"));
+	}
 
 	/**
 	 * ClearReferences is called on every PackageData before any packageDatas are deleted,
