Add CheckShkAlias command to ShaderPipelineCacheTools commandlet.

Sometimes shader output hashes alias between the shader platforms and/or libraries, and this can be undesirable given that the hash does not include optional data.
This commandlet allows to quickly check multiple *.shk files for presence of the same output hashes and work out backwards what stable keys they map to.

#rnx

diff --git a/Engine/Source/Editor/UnrealEd/Private/Commandlets/ShaderPipelineCacheToolsCommandlet.cpp b/Engine/Source/Editor/UnrealEd/Private/Commandlets/ShaderPipelineCacheToolsCommandlet.cpp
--- a/Engine/Source/Editor/UnrealEd/Private/Commandlets/ShaderPipelineCacheToolsCommandlet.cpp
+++ b/Engine/Source/Editor/UnrealEd/Private/Commandlets/ShaderPipelineCacheToolsCommandlet.cpp
@@ -692,6 +692,98 @@ int32 DumpStableKeysFile(const FString& Token)
 	return 0;
 }
 
+int32 CheckStableKeyAliasing(const TArray<FString>& Tokens)
+{
+	if (Tokens.Num() < 2)
+	{
+		UE_LOG(LogShaderPipelineCacheTools, Error, TEXT("At least two shk files need to be given to check hash aliases between them."));
+		return 1;
+	}
+
+	UE_LOG(LogShaderPipelineCacheTools, Display, TEXT("Loading %d stable shader key files..."), Tokens.Num());
+
+	TArray<TMultiMap<FStableShaderKeyAndValue, FSHAHash>> StableMaps;
+	StableMaps.AddDefaulted(Tokens.Num());
+
+	for (int32 IdxFile = 0; IdxFile < Tokens.Num(); ++IdxFile)
+	{
+		const FStringView File = Tokens[IdxFile];
+		LoadStableShaderKeysMultiple(StableMaps[IdxFile], MakeArrayView(&File, 1));
+	}
+
+	TArray<TMultiMap<FSHAHash, FStableShaderKeyAndValue>> InverseStableMaps;
+	InverseStableMaps.AddDefaulted(StableMaps.Num());
+
+	for (int32 IdxFile = 0; IdxFile < Tokens.Num(); ++IdxFile)
+	{
+		for (const auto& Pair : StableMaps[IdxFile])
+		{
+			FStableShaderKeyAndValue Temp(Pair.Key);
+			InverseStableMaps[IdxFile].Add(Pair.Value, Pair.Key);
+		}
+
+		StableMaps[IdxFile] = TMultiMap<FStableShaderKeyAndValue, FSHAHash>();
+	}
+
+	auto LogAllInFile = [](const FString& FileName, const TMultiMap<FSHAHash, FStableShaderKeyAndValue>& InverseMap, const FSHAHash& Hash)
+	{
+		TArray<FStableShaderKeyAndValue> Values;
+		InverseMap.MultiFind(Hash, Values, false);
+		checkf(!Values.IsEmpty(), TEXT("This function should only be called with a hash known to exist in the inverse map"));
+
+		UE_LOG(LogShaderPipelineCacheTools, Display, TEXT("In file %s, it maps to %d shader(s):"), *FileName, Values.Num());
+		int32 ValuesToLog = 10;
+		for (const FStableShaderKeyAndValue& Value : Values)
+		{
+			UE_LOG(LogShaderPipelineCacheTools, Display, TEXT("    %s"), *Value.ToString());
+			if (ValuesToLog-- <= 0)
+			{
+				UE_LOG(LogShaderPipelineCacheTools, Display, TEXT("    etc (%d total)"), Values.Num());
+				break;
+			}
+		}
+	};
+
+
+	TSet<FSHAHash> ReportedHashes;
+	// check for hashes being the same across files
+	for (int32 IdxFile = 0; IdxFile < Tokens.Num(); ++IdxFile)
+	{
+		for (const auto& Pair : InverseStableMaps[IdxFile])
+		{
+			TSet<FSHAHash> AliasedHashesInThisFile;	// to allow checking for more than one pair of files
+
+			for (int32 IdxOtherFile = IdxFile + 1; IdxOtherFile < Tokens.Num(); ++IdxOtherFile)
+			{
+				if (InverseStableMaps[IdxOtherFile].Contains(Pair.Key) && !ReportedHashes.Contains(Pair.Key))
+				{
+					UE_LOG(LogShaderPipelineCacheTools, Warning, TEXT("Shaderhash %s is contained in both %s and %s files"), *Pair.Key.ToString(), *Tokens[IdxFile], *Tokens[IdxOtherFile]);
+
+					// log it for the current file only if seeing for the first time
+					if (!AliasedHashesInThisFile.Contains(Pair.Key))
+					{
+						// find all and log
+						LogAllInFile(Tokens[IdxFile], InverseStableMaps[IdxFile], Pair.Key);
+					}
+
+					AliasedHashesInThisFile.Add(Pair.Key);		
+
+					LogAllInFile(Tokens[IdxOtherFile], InverseStableMaps[IdxOtherFile], Pair.Key);
+				}
+			}
+
+			ReportedHashes.Append(AliasedHashesInThisFile);
+		}
+	}
+
+	if (ReportedHashes.IsEmpty())
+	{
+		UE_LOG(LogShaderPipelineCacheTools, Display, TEXT("No hash aliases between the files!"));
+	}
+
+	return ReportedHashes.IsEmpty() ? 0 : 1;
+}
+
 void IntersectSets(TSet<FCompactFullName>& Intersect, const TSet<FCompactFullName>& ShaderAssets)
 {
 	if (!Intersect.Num() && ShaderAssets.Num())
@@ -2791,6 +2883,11 @@ int32 UShaderPipelineCacheToolsCommandlet::StaticMain(const FString& Params)
 				}
 			}
 		}
+		else if (Tokens[0] == TEXT("CheckShkAlias") && Tokens.Num() >= 2)
+		{
+			Tokens.RemoveAt(0);
+			return CheckStableKeyAliasing(Tokens);
+		}
 		else if (Tokens[0] == TEXT("Decompress") && Tokens.Num() >= 2)
 		{
 			Tokens.RemoveAt(0);
