Update FileHelpers PrepareWorldsForExplicitSave to occur after final save list has been confirmed

diff --git a/Engine/Source/Editor/UnrealEd/Private/FileHelpers.cpp b/Engine/Source/Editor/UnrealEd/Private/FileHelpers.cpp
--- a/Engine/Source/Editor/UnrealEd/Private/FileHelpers.cpp
+++ b/Engine/Source/Editor/UnrealEd/Private/FileHelpers.cpp
@@ -96,6 +96,10 @@ bool FEditorFileUtils::bSkipExternalObjectSave = false;
 TSet<FString> FEditorFileUtils::PackagesNotSavedDuringSaveAll;
 TSet<FString> FEditorFileUtils::PackagesNotToPromptAnyMore;
 
+namespace EditorFileUtils
+{
+	static bool bIsExplicitSave = false;
+}
 
 static TAutoConsoleVariable<int32> CVarSkipSourceControlCheckForEditablePackages(
 	TEXT("r.Editor.SkipSourceControlCheckForEditablePackages"),
@@ -3211,6 +3215,38 @@ enum class InternalSavePackageResult : int8
 
 static void PrepareSavePackages(const TArray<UPackage*>& PackagesToSave)
 {
+	if (EditorFileUtils::bIsExplicitSave)
+	{
+		// In a given set of packages it can contain at least one World Package (map) and/or at least one Actor package.
+		// If an external actor is being saved but not its world we still want to collect its owning world to pass to PrepareWorldsForExplicitSave
+		// In case there is any validation/extra steps needed for that world based on the add/edit of that Actor
+		// We use a set here to dedupe in case both the actor and its world are included in the dirty packages
+		bool bFoundActorWorld = false;
+		TSet<UWorld*> WorldsToSave;
+		for (UPackage* Package : PackagesToSave)
+		{
+			if (UWorld* WorldToSave = UWorld::FindWorldInPackage(Package))
+			{
+				WorldsToSave.Add(WorldToSave);
+			}
+			else if (!bFoundActorWorld)
+			{
+				// Currently there is only one world associated with saving actors as actors from multiple worlds can't be opened
+				// We can skip checking any further Actor packages once we grab the world off the first discovered
+				if (AActor* ActorToSave = AActor::FindActorInPackage(Package))
+				{
+					WorldsToSave.Add(ActorToSave->GetWorld());
+					bFoundActorWorld = true;
+				}
+			}
+		}
+
+		if (!WorldsToSave.IsEmpty())
+		{
+			FEditorFileUtils::PrepareWorldsForExplicitSave(WorldsToSave.Array());
+		}
+	}
+
 	ResetLoaders(MakeArrayView<UObject*>((UObject**)PackagesToSave.GetData(), PackagesToSave.Num()));
 }
 
@@ -3832,41 +3868,7 @@ bool FEditorFileUtils::SaveDirtyPackages(const bool bPromptUserToSave, const boo
 
 	TArray<UPackage*> PackagesToSave = InternalGetDirtyPackages(bSaveMapPackages, bSaveContentPackages, ShouldIgnorePackageFunction);
 
-	if (bPromptUserToSave)
-	{
-		// In a given set of packages it can contain at least one World Package (map) and/or at least one Actor package.
-		// If an external actor is being saved but not its world we still want to collect its owning world to pass to PrepareWorldsForExplicitSave
-		// In case there is any validation/extra steps needed for that world based on the add/edit of that Actor
-		// We use a set here to dedupe in case both the actor and its world are included in the dirty packages
-		bool bFoundActorWorld = false;
-		TSet<UWorld*> WorldsToSave;
-		for (UPackage* Package : PackagesToSave)
-		{
-			if (UWorld* WorldToSave = UWorld::FindWorldInPackage(Package))
-			{
-				WorldsToSave.Add(WorldToSave);
-			}
-			else if (!bFoundActorWorld)
-			{
-				// Currently there is only one world associated with saving actors as actors from multiple worlds can't be opened
-				// We can skip checking any further Actor packages once we grab the world off the first discovered
-				if (AActor* ActorToSave = AActor::FindActorInPackage(Package))
-				{
-					WorldsToSave.Add(ActorToSave->GetWorld());
-					bFoundActorWorld = true;
-				}
-			}
-		}
-
-		if (!WorldsToSave.IsEmpty())
-		{
-			PrepareWorldsForExplicitSave(WorldsToSave.Array());
-
-			// PrepareWorldsForExplicitSave could have dirtied further packages so refresh our list
-			// We do not go looking for further dirtied worlds however
-			PackagesToSave = InternalGetDirtyPackages(bSaveMapPackages, bSaveContentPackages, ShouldIgnorePackageFunction);
-		}
-	}
+	TGuardValue<bool> IsExplicitSaveGuard(EditorFileUtils::bIsExplicitSave, bPromptUserToSave);
 
 	// Need to track the number of packages we're not ignoring for save.
 	int32 NumPackagesNotIgnored = 0;
@@ -3988,13 +3990,6 @@ bool FEditorFileUtils::SaveCurrentLevel()
 		
 		UPackage* LevelPackage = Level->GetPackage();
 
-		if (UWorld* World = UWorld::FindWorldInPackage(LevelPackage))
-		{
-			// Regardless of if the LevelPackage is dirty or if we are only saving external objects
-			// we still give an opportunity to validate/run extra steps on the world involved in this save before continuing
-			PrepareWorldsForExplicitSave({ World });
-		}
-
 		// Get Packages to save
 		if (!bCheckDirty || LevelPackage->IsDirty() || LevelPackage->HasAnyPackageFlags(PKG_NewlyCreated))
 		{
@@ -4016,8 +4011,11 @@ bool FEditorFileUtils::SaveCurrentLevel()
 
 		if (PackagesToSave.Num())
 		{
+			TGuardValue<bool> IsExplicitSaveGuard(EditorFileUtils::bIsExplicitSave, true);
+
 			// If Level gets saved we don't want it to save its external packages because we've already filtered out the ones that need saving and they are part of the PackagesToSave array (unless level is PKG_NewlyCreated then we should save all actors)
 			TGuardValue<bool> GuardValue(bSkipExternalObjectSave, !LevelPackage->HasAnyPackageFlags(PKG_NewlyCreated));
+
 			const bool bPromptUserToSave = false;
 			const bool bFastSave = false;
 			const bool bCanBeDeclined = false;
