Cooker: Add support for generated packages cooking iteratively.
*) Store the declared dependencies for each generated package in the AssetRegistry. Use these dependencies to detect when a generated package is invalidated.
*) Store a a dependency from the generator package to each generated package in the AssetRegistry. Use thise dependency to trigger the Generator to reexecute generation when a generated file is invalidated.
*) Update FAssetRegistryGenerator::ComputePackageDifferences to not mark generated packages as removed due to being missing from disk. Use the package flags to identify them as generated and not mark them as removed.
*) Improve some generated packages and iterative cook output messages to provide more feedback and more accurate feedback.
*) Fix queuing of generated packages to not assume the packages have not already been cooked.
*) Update UpdateAssetRegistryData to find AssetDatas for generated packages, and to store the new dependencies we create for them.
*) Fix PopulateCookedPackages to find generated packagedatas that don't exist on disk, and to pass the list of generated packages to the generator packagedata for removal if the generated packages have changed.
#rnx

diff --git a/Engine/Source/Editor/UnrealEd/Classes/CookOnTheSide/CookOnTheFlyServer.h b/Engine/Source/Editor/UnrealEd/Classes/CookOnTheSide/CookOnTheFlyServer.h
--- a/Engine/Source/Editor/UnrealEd/Classes/CookOnTheSide/CookOnTheFlyServer.h
+++ b/Engine/Source/Editor/UnrealEd/Classes/CookOnTheSide/CookOnTheFlyServer.h
@@ -1334,7 +1334,11 @@ private:
 	const UE::Cook::FCookSavePackageContext* FindSaveContext(const ITargetPlatform* TargetPlatform) const;
 	/** Allocate a new FCookSavePackageContext and ICookedPackageWriter for the given platform. */
 	UE::Cook::FCookSavePackageContext* CreateSaveContext(const ITargetPlatform* TargetPlatform);
-
+	/**
+	 * Delete files that exist for the package in the PackageWriter from a previous cook. Used when packages are
+	 * detected as invalidated for iterative cook after cook startup (e.g. generated packages).
+	 */
+	void DeleteOutputForPackage(FName PackageName, const ITargetPlatform* TargetPlatform);
 
 	static UCookOnTheFlyServer* ActiveCOTFS;
 	uint32		StatLoadedPackageCount = 0;
@@ -1429,6 +1433,7 @@ private:
 	friend UE::Cook::FAssetRegistryMPCollector;
 	friend UE::Cook::FBeginCookConfigSettings;
 	friend UE::Cook::FCookDirector;
+	friend UE::Cook::FCookGenerationInfo;
 	friend UE::Cook::FCookWorkerClient;
 	friend UE::Cook::FCookWorkerServer;
 	friend UE::Cook::FGeneratorPackage;
diff --git a/Engine/Source/Editor/UnrealEd/Private/Commandlets/AssetRegistryGenerator.cpp b/Engine/Source/Editor/UnrealEd/Private/Commandlets/AssetRegistryGenerator.cpp
--- a/Engine/Source/Editor/UnrealEd/Private/Commandlets/AssetRegistryGenerator.cpp
+++ b/Engine/Source/Editor/UnrealEd/Private/Commandlets/AssetRegistryGenerator.cpp
@@ -4,6 +4,7 @@
 
 #include "Algo/Sort.h"
 #include "Algo/StableSort.h"
+#include "Algo/Unique.h"
 #include "AssetRegistry/AssetRegistryModule.h"
 #include "AssetRegistry/CookTagList.h"
 #include "CollectionManagerModule.h"
@@ -34,6 +35,7 @@
 #include "ProfilingDebugging/ScopedTimers.h"
 #include "Serialization/ArrayReader.h"
 #include "Serialization/ArrayWriter.h"
+#include "Serialization/CompactBinarySerialization.h"
 #include "Serialization/JsonReader.h"
 #include "Serialization/JsonReader.h"
 #include "Serialization/JsonSerializer.h"
@@ -950,23 +952,32 @@ void FAssetRegistryGenerator::SetPreviousAssetRegistry(TUniquePtr<FAssetRegistry
 		for (const TPair<FName, const FAssetPackageData*>& Pair : PreviousPackageDataMap)
 		{
 			FName PackageName = Pair.Key;
-			if (!State.GetAssetPackageData(PackageName))
+			if (FPackageName::IsScriptPackage(WriteToString<256>(PackageName)))
 			{
 				continue;
 			}
+			FIterativelySkippedPackageUpdateData& UpdateData = PreviousPackagesToUpdate.FindOrAdd(PackageName);
+			bool bGenerated = false;
+
 			TArrayView<FAssetData const * const> PreviousAssetDatas = InPreviousState->GetAssetsByPackageName(PackageName);
-			TPair<TArray<FAssetData>, FAssetPackageData>& UpdateData = PreviousPackagesToUpdate.FindOrAdd(PackageName);
 			if (PreviousAssetDatas.Num() > 0)
 			{
-				TArray<FAssetData>& UpdateAssetDatas = UpdateData.Get<0>();
-				UpdateAssetDatas.Reserve(PreviousAssetDatas.Num());
+				UpdateData.AssetDatas.Reserve(PreviousAssetDatas.Num());
 				for (const FAssetData* AssetData : PreviousAssetDatas)
 				{
-					UpdateAssetDatas.Add(FAssetData(*AssetData));
+					bGenerated |= (AssetData->PackageFlags & PKG_CookGenerated) != 0;
+					UpdateData.AssetDatas.Emplace(*AssetData);
 				}
 			}
-			const FAssetPackageData& PreviousPackageData = *Pair.Value;
-			UpdateData.Get<1>() = PreviousPackageData;
+			UpdateData.PackageData = *Pair.Value;
+
+			// Keep the dependencies and referencers of generated packages
+			if (bGenerated)
+			{
+				FAssetIdentifier PackageIdentifier(PackageName);
+				InPreviousState->GetDependencies(PackageIdentifier, UpdateData.PackageDependencies, UE::AssetRegistry::EDependencyCategory::Package);
+				InPreviousState->GetReferencers(PackageIdentifier, UpdateData.PackageReferencers, UE::AssetRegistry::EDependencyCategory::Package);
+			}
 		}
 	}
 }
@@ -1095,16 +1106,22 @@ FAssetPackageData* FAssetRegistryGenerator::GetAssetPackageData(const FName& Pac
 
 void FAssetRegistryGenerator::UpdateKeptPackages()
 {
-	for (TPair<FName, TPair<TArray<FAssetData>, FAssetPackageData>>& Pair : PreviousPackagesToUpdate)
+	for (TPair<FName, FIterativelySkippedPackageUpdateData>& Pair : PreviousPackagesToUpdate)
 	{
-		for (const FAssetData& PreviousAssetData : Pair.Value.Get<0>())
+		for (const FAssetData& PreviousAssetData : Pair.Value.AssetDatas)
 		{
 			State.UpdateAssetData(PreviousAssetData);
 		}
-		FAssetPackageData& PackageData = *State.CreateOrGetAssetPackageData(Pair.Key);
-		FAssetPackageData& PreviousPackageData = Pair.Value.Get<1>();
-		PackageData.CookedHash = PreviousPackageData.CookedHash;
-		PackageData.DiskSize = PreviousPackageData.DiskSize;
+		*State.CreateOrGetAssetPackageData(Pair.Key) = Pair.Value.PackageData;
+
+		if (!Pair.Value.PackageDependencies.IsEmpty())
+		{
+			State.AddDependencies(FAssetIdentifier(Pair.Key), Pair.Value.PackageDependencies);
+		}
+		if (!Pair.Value.PackageReferencers.IsEmpty())
+		{
+			State.AddReferencers(FAssetIdentifier(Pair.Key), Pair.Value.PackageReferencers);
+		}
 	}
 }
 
@@ -1247,7 +1264,8 @@ void FAssetRegistryGenerator::Initialize(const TArray<FName> &InStartupPackages,
 	FGameDelegates::Get().GetAssignLayerChunkDelegate() = FAssignLayerChunkDelegate::CreateStatic(AssignLayerChunkDelegate);
 }
 
-void FAssetRegistryGenerator::ComputePackageDifferences(const FComputeDifferenceOptions& Options, const FAssetRegistryState& PreviousState, FAssetRegistryDifference& OutDifference)
+void FAssetRegistryGenerator::ComputePackageDifferences(const FComputeDifferenceOptions& Options,
+	const FAssetRegistryState& PreviousState, FAssetRegistryDifference& OutDifference)
 {
 	TArray<FName> ModifiedScriptPackages;
 	const TMap<FName, const FAssetPackageData*>& PreviousAssetPackageDataMap = PreviousState.GetAssetPackageDataMap();
@@ -1273,12 +1291,16 @@ void FAssetRegistryGenerator::ComputePackageDifferences(const FComputeDifference
 			}
 			else
 			{
-				OutDifference.IdenticalCookedPackages.Add(PackageName);
+				// Do not record identical script packages since we do not cook them
+				if (!FPackageName::IsScriptPackage(WriteToString<256>(PackageName)))
+				{
+					OutDifference.IdenticalCookedPackages.Add(PackageName);
+				}
 			}
 		}
 		else
 		{
-			if (FPackageName::IsScriptPackage(PackageName.ToString()))
+			if (FPackageName::IsScriptPackage(WriteToString<256>(PackageName)))
 			{
 				ModifiedScriptPackages.Add(PackageName);
 			}
@@ -1295,7 +1317,85 @@ void FAssetRegistryGenerator::ComputePackageDifferences(const FComputeDifference
 		const FAssetPackageData* CurrentPackageData = State.GetAssetPackageData(PackageName);
 		if (!CurrentPackageData)
 		{
-			OutDifference.RemovedPackages.Add(PackageName);
+			// If it's a generated package, never mark it as removed (that can only be handled by the generator)
+			// Mark it as modified if any of its dependencies are modified.
+			TConstArrayView<const FAssetData*> PreviousAssets = PreviousState.GetAssetsByPackageName(PackageName);
+			bool bGenerated = !PreviousAssets.IsEmpty() && (PreviousAssets[0]->PackageFlags & PKG_CookGenerated) != 0;
+			if (bGenerated)
+			{
+				bool bModified = true;
+
+				TArray<FAssetIdentifier> Dependencies;
+				TArray<FAssetIdentifier> Referencers;
+				PreviousState.GetDependencies(FAssetIdentifier(PackageName), Dependencies, UE::AssetRegistry::EDependencyCategory::Package);
+				PreviousState.GetReferencers(FAssetIdentifier(PackageName), Referencers, UE::AssetRegistry::EDependencyCategory::Package);
+				FName GeneratorName = Referencers.Num() == 1 ? Referencers[0].PackageName : NAME_None;
+				const FAssetPackageData* GeneratorData = !GeneratorName.IsNone() ? State.GetAssetPackageData(GeneratorName) : nullptr;
+				if (!GeneratorData)
+				{
+					UE_LOG(LogAssetRegistryGenerator, Error,
+						TEXT("Generator could not be found for package %s. This package will be missing from iterative cook results. Recook with a full cook."),
+						*PackageName.ToString());
+				}
+				else
+				{
+					FGeneratorPackageInfo& Info = OutDifference.PreviousGeneratorPackages.FindOrAdd(GeneratorName);
+					Info.Generated.Add(PackageName);
+
+					bModified = false;
+					for (const FAssetIdentifier& Dependency : Dependencies)
+					{
+						const FAssetPackageData* PreviousDependencyData = PreviousState.GetAssetPackageData(Dependency.PackageName);
+						const FAssetPackageData* CurrentDependencyData = State.GetAssetPackageData(Dependency.PackageName);
+						PRAGMA_DISABLE_DEPRECATION_WARNINGS
+						if (!PreviousDependencyData || !CurrentDependencyData ||
+							PreviousDependencyData->PackageGuid != CurrentDependencyData->PackageGuid)
+						{
+							bModified = true;
+							break;
+						}
+						PRAGMA_ENABLE_DEPRECATION_WARNINGS
+					}
+
+					Info.bGeneratedModified |= bModified;
+				}
+				if (bModified)
+				{
+					OutDifference.ModifiedPackages.Add(PackageName);
+				}
+				else
+				{
+					OutDifference.IdenticalCookedPackages.Add(PackageName);
+				}
+			}
+			else
+			{
+				OutDifference.RemovedPackages.Add(PackageName);
+			}
+		}
+	}
+
+	for (TPair<FName, FGeneratorPackageInfo>& Pair : OutDifference.PreviousGeneratorPackages)
+	{
+		FName GeneratorName = Pair.Key;
+		FGeneratorPackageInfo& Info = Pair.Value;
+		if (OutDifference.RemovedPackages.Contains(GeneratorName))
+		{
+			// If a Generator package has been removed, remove all generated
+			for (FName Generated : Info.Generated)
+			{
+				OutDifference.ModifiedPackages.Remove(Generated);
+				OutDifference.IdenticalCookedPackages.Remove(Generated);
+				OutDifference.RemovedPackages.Add(Generated);
+			}
+		}
+		else if (Info.bGeneratedModified)
+		{
+			// Even if not recursing modifications, if a GeneratedPackage has been modified, we need to recook the
+			// generator package, because the generator may change generation depending on modifications to generated.
+			OutDifference.IdenticalCookedPackages.Remove(GeneratorName);
+			OutDifference.IdenticalUncookedPackages.Remove(GeneratorName);
+			OutDifference.ModifiedPackages.Add(GeneratorName);
 		}
 	}
 
@@ -1313,12 +1413,20 @@ void FAssetRegistryGenerator::ComputePackageDifferences(const FComputeDifference
 		{
 			FName ModifiedPackage = ModifiedPackagesToRecurse[RecurseIndex];
 			TArray<FAssetIdentifier> Referencers;
-			State.GetReferencers(ModifiedPackage, Referencers, UE::AssetRegistry::EDependencyCategory::Package, UE::AssetRegistry::EDependencyQuery::Hard);
+			// Read referencers from the previous state (to find referencers of generated packages) and of the current state
+			// (to find referencers of editoronly packages)
+			PreviousState.GetReferencers(ModifiedPackage, Referencers, UE::AssetRegistry::EDependencyCategory::Package,
+				UE::AssetRegistry::EDependencyQuery::Hard);
+			State.GetReferencers(ModifiedPackage, Referencers, UE::AssetRegistry::EDependencyCategory::Package,
+				UE::AssetRegistry::EDependencyQuery::Hard);
+			Algo::Sort(Referencers, [](FAssetIdentifier& A, FAssetIdentifier& B) { return A.FastLess(B); });
+			Referencers.SetNum(Algo::Unique(Referencers));
 
 			for (const FAssetIdentifier& Referencer : Referencers)
 			{
 				FName ReferencerPackage = Referencer.PackageName;
-				if (!OutDifference.ModifiedPackages.Contains(ReferencerPackage) && (OutDifference.IdenticalCookedPackages.Contains(ReferencerPackage) || OutDifference.IdenticalUncookedPackages.Contains(ReferencerPackage)))
+				if (!OutDifference.ModifiedPackages.Contains(ReferencerPackage) &&
+					(OutDifference.IdenticalCookedPackages.Contains(ReferencerPackage) || OutDifference.IdenticalUncookedPackages.Contains(ReferencerPackage)))
 				{
 					// Remove from identical list
 					OutDifference.IdenticalCookedPackages.Remove(ReferencerPackage);
@@ -1332,17 +1440,61 @@ void FAssetRegistryGenerator::ComputePackageDifferences(const FComputeDifference
 	}
 }
 
-void FAssetRegistryGenerator::ComputePackageRemovals(const FAssetRegistryState& PreviousState, TArray<FName>& RemovedPackages)
+void FAssetRegistryGenerator::ComputePackageRemovals(const FAssetRegistryState& PreviousState, TArray<FName>& RemovedPackages,
+	TMap<FName, FGeneratorPackageInfo>& PreviousGeneratorPackages)
 {
+	PreviousGeneratorPackages.Reset();
+	TSet<FName> RemovedPackageSet;
+
 	for (const TPair<FName, const FAssetPackageData*>& PackagePair : PreviousState.GetAssetPackageDataMap())
 	{
 		FName PackageName = PackagePair.Key;
 		const FAssetPackageData* CurrentPackageData = State.GetAssetPackageData(PackageName);
 		if (!CurrentPackageData)
 		{
-			RemovedPackages.Add(PackageName);
+			// If it's a generated package, never mark it as removed (that can only be handled by the generator)
+			// Mark it as modified if its generator or any of its dependencies are modified.
+			TConstArrayView<const FAssetData*> PreviousAssets = PreviousState.GetAssetsByPackageName(PackageName);
+			bool bGenerated = !PreviousAssets.IsEmpty() && (PreviousAssets[0]->PackageFlags & PKG_CookGenerated);
+			if (bGenerated)
+			{
+				TArray<FAssetIdentifier> Referencers;
+				PreviousState.GetReferencers(FAssetIdentifier(PackageName), Referencers, UE::AssetRegistry::EDependencyCategory::Package);
+				FName GeneratorName = Referencers.Num() == 1 ? Referencers[0].PackageName : NAME_None;
+				const FAssetPackageData* GeneratorData = !GeneratorName.IsNone() ? State.GetAssetPackageData(GeneratorName) : nullptr;
+				if (!GeneratorData)
+				{
+					UE_LOG(LogAssetRegistryGenerator, Error,
+						TEXT("Generator could not be found for package %s. This package will be missing from iterative cook results. Recook with a full cook."),
+						*PackageName.ToString());
+				}
+				else
+				{
+					PreviousGeneratorPackages.FindOrAdd(GeneratorName).Generated.Add(PackageName);
+				}
+			}
+			else
+			{
+				RemovedPackageSet.Add(PackageName);
+			}
+		}
+	}
+
+	for (TMap<FName, FGeneratorPackageInfo>::TIterator Iter(PreviousGeneratorPackages); Iter; ++Iter)
+	{
+		// If a Generator package has been removed, remove all generated
+		FName GeneratorName = Iter->Key;
+		if (RemovedPackageSet.Contains(GeneratorName))
+		{
+			for (FName Generated : Iter->Value.Generated)
+			{
+				RemovedPackageSet.Add(Generated);
+			}
+			Iter.RemoveCurrent();
 		}
 	}
+
+	RemovedPackages = RemovedPackageSet.Array();
 }
 
 void FAssetRegistryGenerator::FinalizeChunkIDs(const TSet<FName>& InCookedPackages,
@@ -2533,10 +2685,8 @@ const FAssetData* FAssetRegistryGenerator::CreateOrFindAssetData(UObject& Object
 }
 
 void FAssetRegistryGenerator::UpdateAssetRegistryData(const UPackage& Package,
-	FSavePackageResultStruct& SavePackageResult,
-	FCookTagList&& InArchiveCookTagList,
-	bool bIncludeOnlyDiskAssets
-	)
+	FSavePackageResultStruct& SavePackageResult, FCookTagList&& InArchiveCookTagList, bool bIncludeOnlyDiskAssets,
+	TOptional<FAssetPackageData>&& OverrideAssetPackageData, TOptional<TArray<FAssetDependency>>&& OverridePackageDependencies)
 {
 	LLM_SCOPE_BYTAG(Cooker_GeneratedAssetRegistry);
 	const FName PackageName = Package.GetFName();
@@ -2558,6 +2708,16 @@ void FAssetRegistryGenerator::UpdateAssetRegistryData(const UPackage& Package,
 	// Create a record for assets that were created during PostLoad or cooking and are not in the global assetregistry
 	CreateOrFindAssetDatas(Package);
 
+	// Set Overrides for generated packages
+	if (OverrideAssetPackageData)
+	{
+		*AssetPackageData = MoveTemp(*OverrideAssetPackageData);
+	}
+	if (OverridePackageDependencies)
+	{
+		SetOverridePackageDependencies(PackageName, *OverridePackageDependencies);
+	}
+
 	if (bSaveSucceeded)
 	{
 		// Migrate cook tags over
@@ -2585,6 +2745,19 @@ void FAssetRegistryGenerator::UpdateAssetRegistryData(const UPackage& Package,
 		TEXT("Trying to update asset package flags in package '%s' that does not exist"), *PackageName.ToString());
 }
 
+void FAssetRegistryGenerator::SetOverridePackageDependencies(FName PackageName, TConstArrayView<FAssetDependency> OverridePackageDependencies)
+{
+	TArray<FAssetDependency, TInlineAllocator<10>> DependencyStructs;
+	DependencyStructs.Reserve(OverridePackageDependencies.Num());
+#if DO_CHECK
+	for (FAssetDependency Dependency : OverridePackageDependencies)
+	{
+		check(Dependency.Category == UE::AssetRegistry::EDependencyCategory::Package);
+	}
+#endif
+	State.SetDependencies(FAssetIdentifier(PackageName), OverridePackageDependencies, UE::AssetRegistry::EDependencyCategory::Package);
+}
+
 void FAssetRegistryGenerator::UpdateAssetRegistryData(UE::Cook::FMPCollectorServerMessageContext& Context,
 	UE::Cook::FAssetRegistryPackageMessage&& Message)
 {
@@ -2599,8 +2772,17 @@ void FAssetRegistryGenerator::UpdateAssetRegistryData(UE::Cook::FMPCollectorServ
 		State.UpdateAssetData(MoveTemp(AssetData), true /* bCreateIfNotExists */);
 	}
 	FAssetPackageData* AssetPackageData = GetAssetPackageData(PackageName);
+	if (Message.OverrideAssetPackageData)
+	{
+		*AssetPackageData = MoveTemp(*Message.OverrideAssetPackageData);
+	}
 	AssetPackageData->DiskSize = Message.DiskSize;
 
+	if (Message.OverridePackageDependencies)
+	{
+		SetOverridePackageDependencies(PackageName, *Message.OverridePackageDependencies);
+	}
+
 	CookTagsToAdd.Append(MoveTemp(Message.CookTags));
 }
 
@@ -2614,7 +2796,8 @@ FAssetRegistryReporterRemote::FAssetRegistryReporterRemote(FCookWorkerClient& In
 }
 
 void FAssetRegistryReporterRemote::UpdateAssetRegistryData(FPackageData& PackageData, const UPackage& Package,
-	FSavePackageResultStruct& SavePackageResult, FCookTagList&& InArchiveCookTagList, bool bIncludeOnlyDiskAssets)
+	FSavePackageResultStruct& SavePackageResult, FCookTagList&& InArchiveCookTagList, bool bIncludeOnlyDiskAssets,
+	TOptional<FAssetPackageData>&& OverrideAssetPackageData, TOptional<TArray<FAssetDependency>>&& OverridePackageDependencies)
 {
 	uint32 NewPackageFlags = 0;
 	int64 DiskSize = -1;
@@ -2638,6 +2821,8 @@ void FAssetRegistryReporterRemote::UpdateAssetRegistryData(FPackageData& Package
 	Message.TargetPlatform = TargetPlatform;
 	Message.PackageFlags = NewPackageFlags;
 	Message.DiskSize = DiskSize;
+	Message.OverrideAssetPackageData = MoveTemp(OverrideAssetPackageData);
+	Message.OverridePackageDependencies = MoveTemp(OverridePackageDependencies);
 
 	// Add to the message all the AssetDatas in the package from the global AssetRegistry
 	IAssetRegistry::Get()->GetAssetsByPackageName(PackageName, Message.AssetDatas,
@@ -2685,6 +2870,22 @@ void FAssetRegistryPackageMessage::Write(FCbWriter& Writer) const
 
 	}
 	Writer.EndArray();
+	if (OverrideAssetPackageData)
+	{
+		Writer.BeginObject("P");
+		{
+			// Currently we only replicate Guid and ImportedClasses, since these are the only fields set by generated pacakges
+			PRAGMA_DISABLE_DEPRECATION_WARNINGS;
+			Writer << "G" << OverrideAssetPackageData->PackageGuid;
+			PRAGMA_ENABLE_DEPRECATION_WARNINGS;
+			Writer << "C" << OverrideAssetPackageData->ImportedClasses;
+		}
+		Writer.EndObject();
+	}
+	if (OverridePackageDependencies)
+	{
+		Writer << "D" << *OverridePackageDependencies;
+	}
 	// We cast TMap<FSoftObjectPath,                     ValueType>
 	//      to TMap<FSoftObjectPathSerializationWrapper, ValueType>
 	// to workaround FSoftObjectPath's implicit constructor. See comment in CompactBinaryTCP.h
@@ -2714,6 +2915,34 @@ bool FAssetRegistryPackageMessage::TryRead(FCbObjectView Object)
 			return false;
 		}
 	}
+	OverrideAssetPackageData.Reset();
+	FCbFieldView OverrideAssetPackageDataField = Object["P"];
+	if (OverrideAssetPackageDataField.HasValue())
+	{
+		OverrideAssetPackageData.Emplace();
+		// Currently we only replicate Guid and ImportedClasses, since these are the only fields set by generated pacakges
+		PRAGMA_DISABLE_DEPRECATION_WARNINGS;
+		if (!LoadFromCompactBinary(OverrideAssetPackageDataField["G"], OverrideAssetPackageData->PackageGuid))
+		{
+			return false;
+		}
+		PRAGMA_ENABLE_DEPRECATION_WARNINGS;
+		if (!LoadFromCompactBinary(OverrideAssetPackageDataField["C"], OverrideAssetPackageData->ImportedClasses))
+		{
+			return false;
+		}
+	}
+
+	OverridePackageDependencies.Reset();
+	FCbFieldView OverridePackageDependenciesField = Object["D"];
+	if (OverridePackageDependenciesField.HasValue())
+	{
+		OverridePackageDependencies.Emplace();
+		if (!LoadFromCompactBinary(OverridePackageDependenciesField, *OverridePackageDependencies))
+		{
+			return false;
+		}
+	}
 
 	static_assert(std::is_same_v<decltype(CookTags)::KeyType, FSoftObjectPath>, "Expected KeyType of CookTags to be FSoftObjectPath");
 	if (!LoadFromCompactBinary(Object["T"], reinterpret_cast<TMap<FSoftObjectPathSerializationWrapper, decltype(CookTags)::ValueType>&>(CookTags)))
diff --git a/Engine/Source/Editor/UnrealEd/Private/CookOnTheFlyServer.cpp b/Engine/Source/Editor/UnrealEd/Private/CookOnTheFlyServer.cpp
--- a/Engine/Source/Editor/UnrealEd/Private/CookOnTheFlyServer.cpp
+++ b/Engine/Source/Editor/UnrealEd/Private/CookOnTheFlyServer.cpp
@@ -2935,14 +2935,15 @@ UE::Cook::EPollStatus UCookOnTheFlyServer::ConditionalCreateGeneratorPackage(UE:
 		return EPollStatus::Error;
 	}
 
-	UE_LOG(LogCook, Display, TEXT("Splitting Package %s with class %s acting on object %s."),
-		*PackageData.GetPackageName().ToString(), *Splitter->GetSplitDataClass()->GetName(), *SplitDataObject->GetFullName());
+	UE_LOG(LogCook, Display, TEXT("Splitting Package %s with splitter %s acting on object %s."),
+		*PackageData.GetPackageName().ToString(), *Splitter->GetSplitterDebugName(), *SplitDataObject->GetFullName());
 
 	// Create instance of CookPackageSplitter class
 	ICookPackageSplitter* SplitterInstance = Splitter->CreateInstance(SplitDataObject);
 	if (!SplitterInstance)
 	{
-		UE_LOG(LogCook, Error, TEXT("Error instantiating Cook Package Splitter for object %s."), *SplitDataObject->GetFullName());
+		UE_LOG(LogCook, Error, TEXT("Error instantiating Cook Package Splitter %s for object %s."),
+			*Splitter->GetSplitterDebugName(), *SplitDataObject->GetFullName());
 		return EPollStatus::Error;
 	}
 
@@ -3009,7 +3010,6 @@ UE::Cook::EPollStatus UCookOnTheFlyServer::QueueGeneratedPackages(UE::Cook::FGen
 		for (const FCookGenerationInfo& ChildInfo: Generator.GetPackagesToGenerate())
 		{
 			FPackageData* ChildPackageData = ChildInfo.PackageData;
-			ChildPackageData->ClearCookedPlatformData();
 			QueueDiscoveredPackageData(*ChildPackageData, FInstigator(EInstigator::GeneratedPackage, OwnerName));
 		}
 		Info.SetSaveStateComplete(FCookGenerationInfo::ESaveState::QueueGeneratedPackages);
@@ -5634,10 +5634,86 @@ void FSaveCookedPackageContext::FinishPlatform()
 		// For most types, our contract is that we use the data in the asset registry that was updated on load and is in
 		// the disk asset storage of the global asset registry. Most types do not take significant action during cooking
 		// so we do not need to allow them to update their tags after they finish loading.
-		// But generator packages do a lot of work during generation and may need to update their tags, so read the
-		// latest tags off of their asset by setting bIncludeOnlyDiskAssets=true.
-		const bool bIncludeOnlyDiskAssets = PackageData.GetGeneratorPackage() == nullptr;
-		Reporter.UpdateAssetRegistryData(PackageData, *Package, SavePackageResult, MoveTemp(*ArchiveCookContext.GetCookTagList()), bIncludeOnlyDiskAssets);
+		bool bIncludeOnlyDiskAssets = true;
+		TOptional<TArray<FAssetDependency>> OverridePackageDependencies;
+		TOptional<FAssetPackageData> OverrideAssetPackageData;
+		FGeneratorPackage* GeneratorPackage;
+		if (GeneratorPackage = PackageData.GetGeneratorPackage(); GeneratorPackage)
+		{
+			// Generator packages do a lot of work during generation and may need to update their tags, so read the
+			// latest tags off of their asset by setting bIncludeOnlyDiskAssets=true.
+			bIncludeOnlyDiskAssets = false;
+			OverridePackageDependencies.Emplace();
+
+			// Set override dependencies equal to the global AssetRegistry dependencies plus a dependency on
+			// each generated package.
+			COTFS.AssetRegistry->GetDependencies(Package->GetFName(), *OverridePackageDependencies,
+				UE::AssetRegistry::EDependencyCategory::Package);
+			OverridePackageDependencies->Reserve(GeneratorPackage->GetPackagesToGenerate().Num());
+			for (FCookGenerationInfo& GeneratedInfo : GeneratorPackage->GetPackagesToGenerate())
+			{
+				FAssetDependency& Dependency = OverridePackageDependencies->Emplace_GetRef();
+				Dependency.AssetId = FAssetIdentifier(GeneratedInfo.PackageData->GetPackageName());
+				Dependency.Category = UE::AssetRegistry::EDependencyCategory::Package;
+				Dependency.Properties = UE::AssetRegistry::EDependencyProperty::Game;
+			}
+		}
+		else if (GeneratorPackage = PackageData.GetGeneratedOwner(); GeneratorPackage)
+		{
+			FCookGenerationInfo* GeneratedInfo = GeneratorPackage->FindInfo(PackageData);
+			if (!GeneratedInfo)
+			{
+				UE_LOG(LogCook, Error, TEXT("GeneratedInfo missing for package %s."), *PackageData.GetPackageName().ToString());
+			}
+			else
+			{
+				// Generated packages do not exist on disk, so read their assetdata from the in-memory package
+				bIncludeOnlyDiskAssets = false;
+				OverridePackageDependencies.Emplace();
+				OverridePackageDependencies->Reserve(GeneratedInfo->Dependencies.Num());
+
+				// There should be no package dependencies present for the package from the global assetregistry
+				// because it is newly created. Add on the dependencies declared for it from the CookPackageSplitter.
+				for (FName PackageDependency : GeneratedInfo->Dependencies)
+				{
+					FAssetDependency& Dependency = OverridePackageDependencies->Emplace_GetRef();
+					Dependency.AssetId = FAssetIdentifier(PackageDependency);
+					Dependency.Category = UE::AssetRegistry::EDependencyCategory::Package;
+					// COOKPACKAGESPLITTERTODO: CookPackageSplitter needs to specify the DependencyProperties
+					Dependency.Properties = UE::AssetRegistry::EDependencyProperty::Hard | UE::AssetRegistry::EDependencyProperty::Game;
+				}
+
+				// Update the AssetPackageData for each requested platform with Guid and ImportedClasses
+				TSet<UClass*> PackageClasses;
+				ForEachObjectWithPackage(Package, [&PackageClasses, this](UObject* Object)
+					{
+						UClass* Class = Object->GetClass();
+						if (!Class->IsInPackage(Package)) // Imported classes list does not include classes in the package
+						{
+							PackageClasses.Add(Object->GetClass());
+						}
+						return true;
+					});
+				TArray<FName> ImportedClasses;
+				ImportedClasses.Reserve(PackageClasses.Num());
+				for (UClass* Class : PackageClasses)
+				{
+					TStringBuilder<256> ClassPath;
+					Class->GetPathName(nullptr, ClassPath);
+					ImportedClasses.Add(FName(ClassPath));
+				}
+				ImportedClasses.Sort(FNameLexicalLess());
+
+				OverrideAssetPackageData.Emplace();
+				PRAGMA_DISABLE_DEPRECATION_WARNINGS;
+				OverrideAssetPackageData->PackageGuid = GeneratedInfo->Guid;
+				PRAGMA_ENABLE_DEPRECATION_WARNINGS;
+				OverrideAssetPackageData->ImportedClasses = ImportedClasses;
+			}
+		}
+		Reporter.UpdateAssetRegistryData(PackageData, *Package, SavePackageResult,
+			MoveTemp(*ArchiveCookContext.GetCookTagList()), bIncludeOnlyDiskAssets,
+			MoveTemp(OverrideAssetPackageData), MoveTemp(OverridePackageDependencies));
 	}
 
 	// If not retrying, mark the package as cooked, either successfully or with failure
@@ -7086,12 +7162,13 @@ void UCookOnTheFlyServer::PopulateCookedPackages(TArrayView<const ITargetPlatfor
 
 		TUniquePtr<FAssetRegistryState> PreviousAssetRegistry(PackageWriter.LoadPreviousAssetRegistry());
 		int32 NumPreviousPackages = PreviousAssetRegistry ? PreviousAssetRegistry->GetNumPackages() : 0;
-		UE_LOG(LogCook, Display, TEXT("Found '%d' cooked package(s) in package store"), NumPreviousPackages);
 		if (NumPreviousPackages == 0)
 		{
+			UE_LOG(LogCook, Display, TEXT("Found %d cooked package(s) in package store."), NumPreviousPackages);
 			continue;
 		}
 
+		TMap<FName, FAssetRegistryGenerator::FGeneratorPackageInfo> PreviousGeneratorPackages;
 		if (bHybridIterativeEnabled)
 		{
 			// HybridIterative does the equivalent operation of bRecurseModifications=true and bRecurseScriptModifications=true,
@@ -7102,8 +7179,10 @@ void UCookOnTheFlyServer::PopulateCookedPackages(TArrayView<const ITargetPlatfor
 			// check this for all packages in the previous cook rather than just the currently referenced
 			// set because when packages are removed the references to them are usually also removed.
 			TArray<FName> TombstonePackages;
-			PlatformAssetRegistry.ComputePackageRemovals(*PreviousAssetRegistry, TombstonePackages);
+			PlatformAssetRegistry.ComputePackageRemovals(*PreviousAssetRegistry, TombstonePackages,
+				PreviousGeneratorPackages);
 			PackageWriter.RemoveCookedPackages(TombstonePackages);
+			UE_LOG(LogCook, Display, TEXT("Found %d cooked package(s) in package store."), NumPreviousPackages);
 		}
 		else
 		{
@@ -7117,6 +7196,13 @@ void UCookOnTheFlyServer::PopulateCookedPackages(TArrayView<const ITargetPlatfor
 			Options.bRecurseScriptModifications = !IsCookFlagSet(ECookInitializationFlags::IgnoreScriptPackagesOutOfDate);
 			FAssetRegistryGenerator::FAssetRegistryDifference Difference;
 			PlatformAssetRegistry.ComputePackageDifferences(Options, *PreviousAssetRegistry, Difference);
+			PreviousGeneratorPackages = MoveTemp(Difference.PreviousGeneratorPackages);
+
+			TSet<FName> GeneratedPackages;
+			for (TPair<FName, FAssetRegistryGenerator::FGeneratorPackageInfo>& Pair : PreviousGeneratorPackages)
+			{
+				GeneratedPackages.Append(Pair.Value.Generated);
+			}
 
 			TArray<FName> PackagesToRemove;
 			PackagesToRemove.Reserve(Difference.ModifiedPackages.Num() + Difference.RemovedPackages.Num() + Difference.IdenticalUncookedPackages.Num());
@@ -7132,8 +7218,11 @@ void UCookOnTheFlyServer::PopulateCookedPackages(TArrayView<const ITargetPlatfor
 			{
 				PackagesToRemove.Add(PackageToRemove);
 			}
-
-			UE_LOG(LogCook, Display, TEXT("Keeping '%d' and removing '%d' cooked package(s)"), Difference.IdenticalCookedPackages.Num(), PackagesToRemove.Num());
+			int32 NumUncooked = Difference.IdenticalUncookedPackages.Num();
+			int32 NumCooked = FMath::Max(0, NumPreviousPackages - NumUncooked);
+			UE_LOG(LogCook, Display, TEXT("Found %d cooked package(s) in package store."), NumCooked);
+			UE_LOG(LogCook, Display, TEXT("Keeping %d and removing %d cooked package(s)."),
+				Difference.IdenticalCookedPackages.Num(), PackagesToRemove.Num() - NumUncooked);
 
 			PackageWriter.RemoveCookedPackages(PackagesToRemove);
 
@@ -7145,7 +7234,8 @@ void UCookOnTheFlyServer::PopulateCookedPackages(TArrayView<const ITargetPlatfor
 			for (const FName& IdenticalPackage : Difference.IdenticalCookedPackages)
 			{
 				// Mark this package as cooked so that we don't unnecessarily try to cook it again
-				FPackageData* PackageData = PackageDatas->TryAddPackageDataByPackageName(IdenticalPackage);
+				bool bRequireExists = !GeneratedPackages.Contains(IdenticalPackage);
+				FPackageData* PackageData = PackageDatas->TryAddPackageDataByPackageName(IdenticalPackage, bRequireExists);
 				if (PackageData)
 				{
 					PackageData->SetPlatformCooked(TargetPlatform, true /* bSucceeded */);
@@ -7181,6 +7271,19 @@ void UCookOnTheFlyServer::PopulateCookedPackages(TArrayView<const ITargetPlatfor
 			}
 		}
 
+		for (TPair<FName, FAssetRegistryGenerator::FGeneratorPackageInfo> Pair : PreviousGeneratorPackages)
+		{
+			FPackageData* Generator = PackageDatas->TryAddPackageDataByPackageName(Pair.Key);
+			if (!Generator)
+			{
+				UE_LOG(LogCook, Warning,
+					TEXT("Previous cook results returned a record for generator package %s, but that package can no longer be found; its generated packages will not be removed from cook results. Run a full cook to remove them."),
+					*Pair.Key.ToString());
+				continue;
+			}
+			Generator->CreateGeneratorPackage(nullptr, nullptr).SetPreviousGeneratedPackages(MoveTemp(Pair.Value.Generated));
+		}
+
 		PlatformAssetRegistry.SetPreviousAssetRegistry(MoveTemp(PreviousAssetRegistry));
 	}
 }
@@ -7883,28 +7986,10 @@ FString UCookOnTheFlyServer::ConvertToFullSandboxPath( const FString &FileName,
 	FString Result;
 	if (bForWrite)
 	{
-		// Ideally this would be in the Sandbox File but it can't access the project or plugin
-		if (PluginsToRemap.Num() > 0)
+		if (FLooseCookedPackageWriter::TryConvertUncookedFilenameToCookedRemappedPluginFilename(
+			FileName, PluginsToRemap, SandboxFile->GetSandboxDirectory(), Result))
 		{
-			// Handle remapping of plugins
-			for (TSharedRef<IPlugin> Plugin : PluginsToRemap)
-			{
-				// If these match, then this content is part of plugin that gets remapped when packaged/staged
-				if (FileName.StartsWith(Plugin->GetContentDir()))
-				{
-					FString SearchFor;
-					SearchFor /= Plugin->GetName() / TEXT("Content");
-					int32 FoundAt = FileName.Find(SearchFor, ESearchCase::IgnoreCase, ESearchDir::FromEnd);
-					check(FoundAt != -1);
-					// Strip off everything but <PluginName/Content/<remaing path to file>
-					FString SnippedOffPath = FileName.RightChop(FoundAt);
-					// Put this is in <sandbox path>/RemappedPlugins/<PluginName>/Content/<remaing path to file>
-					FString RemappedPath = SandboxFile->GetSandboxDirectory();
-					RemappedPath /= REMAPPED_PLUGINS;
-					Result = RemappedPath / SnippedOffPath;
-					return Result;
-				}
-			}
+			return Result;
 		}
 		Result = SandboxFile->ConvertToAbsolutePathForExternalAppForWrite(*FileName);
 	}
@@ -8643,8 +8728,9 @@ void UCookOnTheFlyServer::PrintFinishStats()
 	const FPlatformMemoryStats MemStats = FPlatformMemory::GetStats();
 	UE_LOG(LogCook, Display, TEXT("Peak Used virtual %u MiB Peak Used physical %u MiB"), MemStats.PeakUsedVirtual / 1024 / 1024, MemStats.PeakUsedPhysical / 1024 / 1024);
 
-	COOK_STAT(UE_LOG(LogCook, Display, TEXT("Packages Cooked: %d, Packages Iteratively Skipped: %d, Total Packages: %d"),
-		DetailedCookStats::NumPackagesSavedForCook, DetailedCookStats::NumPackagesIterativelySkipped, PackageDatas->GetNumCooked()));
+	COOK_STAT(UE_LOG(LogCook, Display, TEXT("Packages Cooked: %d, Packages Iteratively Skipped: %d, Packages Skipped by Platform: %d, Total Packages: %d"),
+		DetailedCookStats::NumPackagesSavedForCook, DetailedCookStats::NumPackagesIterativelySkipped, PackageDatas->GetNumFailedCooked(),
+		PackageDatas->GetNumCooked()));
 }
 
 void UCookOnTheFlyServer::PrintDetailedCookStats()
@@ -9113,7 +9199,7 @@ void UCookOnTheFlyServer::LoadBeginCookIterativeFlagsLocal(FBeginCookContext& Be
 			}
 			else if (!ArePreviousCookSettingsCompatible(PlatformContext.CurrentCookSettings, TargetPlatform))
 			{
-				UE_LOG(LogCook, Display, TEXT("Cook invalidated for platform %s because cook DLC settings have changed, clearing all cooked content"), *TargetPlatform->PlatformName());
+				UE_LOG(LogCook, Display, TEXT("Cook invalidated for platform %s because cook DLC settings have changed, clearing all cooked content."), *TargetPlatform->PlatformName());
 				bIterativeAllowed = false;
 			}
 			else
@@ -9123,18 +9209,18 @@ void UCookOnTheFlyServer::LoadBeginCookIterativeFlagsLocal(FBeginCookContext& Be
 				{
 					if (!IsCookFlagSet(ECookInitializationFlags::IgnoreIniSettingsOutOfDate))
 					{
-						UE_LOG(LogCook, Display, TEXT("Cook invalidated for platform %s ini settings don't match from last cook, clearing all cooked content"), *TargetPlatform->PlatformName());
+						UE_LOG(LogCook, Display, TEXT("Cook invalidated for platform %s because ini settings have changed, clearing all cooked content."), *TargetPlatform->PlatformName());
 						bIterativeAllowed = false;
 					}
 					else
 					{
-						UE_LOG(LogCook, Display, TEXT("Inisettings were out of date for platform %s but we are going with it anyway because IgnoreIniSettingsOutOfDate is set"), *TargetPlatform->PlatformName());
+						UE_LOG(LogCook, Display, TEXT("Inisettings were out of date for platform %s but we are going with it anyway because IgnoreIniSettingsOutOfDate is set."), *TargetPlatform->PlatformName());
 						bIterativeAllowed = true;
 					}
 				}
 				else
 				{
-					UE_LOG(LogCook, Display, TEXT("Keeping cooked content for platform %s and cooking iteratively"), *TargetPlatform->PlatformName());
+					UE_LOG(LogCook, Display, TEXT("Keeping cooked content for platform %s and cooking iteratively."), *TargetPlatform->PlatformName());
 					bIterativeAllowed = true;
 				}
 			}
@@ -9332,6 +9418,11 @@ UE::Cook::FCookSavePackageContext* UCookOnTheFlyServer::CreateSaveContext(const
 
 }
 
+void UCookOnTheFlyServer::DeleteOutputForPackage(FName PackageName, const ITargetPlatform* TargetPlatform)
+{
+	FindOrCreatePackageWriter(TargetPlatform).RemoveCookedPackages(TArrayView<FName>(&PackageName, 1));
+}
+
 void UCookOnTheFlyServer::FinalizePackageStore()
 {
 	UE_SCOPED_HIERARCHICAL_COOKTIMER(FinalizePackageStore);
@@ -11039,7 +11130,9 @@ uint32 UCookOnTheFlyServer::CookFullLoadAndSave()
 						FAssetPackageData* AssetPackageData = nullptr;
 						{
 							FScopeLock AssetGeneratorScopeLock(&AssetGeneratorLock);
-							Generator.UpdateAssetRegistryData(*Package, SaveResult, MoveTemp(*CookContext.GetCookTagList()));
+							Generator.UpdateAssetRegistryData(*Package, SaveResult, MoveTemp(*CookContext.GetCookTagList()),
+								false /* bIncludeDiskOnlyAssets */, TOptional<FAssetPackageData>(),
+								TOptional<TArray<FAssetDependency>>());
 							AssetPackageData = Generator.GetAssetPackageData(Package->GetFName());
 						}
 						check(AssetPackageData);
diff --git a/Engine/Source/Editor/UnrealEd/Private/Cooker/CookPackageData.cpp b/Engine/Source/Editor/UnrealEd/Private/Cooker/CookPackageData.cpp
--- a/Engine/Source/Editor/UnrealEd/Private/Cooker/CookPackageData.cpp
+++ b/Engine/Source/Editor/UnrealEd/Private/Cooker/CookPackageData.cpp
@@ -328,7 +328,7 @@ void FPackageData::SetPlatformCooked(const ITargetPlatform* TargetPlatform, bool
 	}
 	if (bModified && !bHasAnyOthers)
 	{
-		PackageDatas.GetMonitor().OnFirstCookedPlatformAdded(*this);
+		PackageDatas.GetMonitor().OnFirstCookedPlatformAdded(*this, bSucceeded);
 	}
 }
 
@@ -1337,20 +1337,33 @@ void FPackageData::SetGeneratedOwner(FGeneratorPackage* InGeneratedOwner)
 	GeneratedOwner = InGeneratedOwner;
 }
 
-UE::Cook::FGeneratorPackage* FPackageData::CreateGeneratorPackage(const UObject* InSplitDataObject,
+UE::Cook::FGeneratorPackage* FPackageData::GetGeneratorPackage() const
+{
+	UE::Cook::FGeneratorPackage* Result = GeneratorPackage.Get();
+	if (Result && Result->IsInitialized())
+	{
+		return Result;
+	}
+	return nullptr;
+}
+
+UE::Cook::FGeneratorPackage& FPackageData::CreateGeneratorPackage(const UObject* InSplitDataObject,
 	ICookPackageSplitter* InCookPackageSplitterInstance)
 {
-	if (!GetGeneratorPackage())
+	if (!GeneratorPackage)
 	{
 		GeneratorPackage.Reset(new UE::Cook::FGeneratorPackage(*this, InSplitDataObject,
 			InCookPackageSplitterInstance));
 	}
 	else
 	{
+		GeneratorPackage->InitializeSave(InSplitDataObject, InCookPackageSplitterInstance);
 		// The earlier exit from SaveState should have reset the progress back to StartGeneratorSave or earlier
-		check(GetGeneratorPackage()->GetOwnerInfo().GetSaveState() <= FCookGenerationInfo::ESaveState::StartPopulate);
+		check(GeneratorPackage->GetOwnerInfo().GetSaveState() <= FCookGenerationInfo::ESaveState::StartPopulate);
 	}
-	return GetGeneratorPackage();
+	FGeneratorPackage* Result = GeneratorPackage.Get();
+	check(Result);
+	return *Result;
 }
 
 FConstructPackageData FPackageData::CreateConstructData()
@@ -1388,17 +1401,53 @@ namespace UE::Cook
 FGeneratorPackage::FGeneratorPackage(UE::Cook::FPackageData& InOwner, const UObject* InSplitDataObject,
 	ICookPackageSplitter* InCookPackageSplitterInstance)
 : OwnerInfo(InOwner, true /* bInGenerated */)
-, SplitDataObjectName(*InSplitDataObject->GetFullName())
 {
-	check(InCookPackageSplitterInstance);
-	CookPackageSplitterInstance.Reset(InCookPackageSplitterInstance);
-	SetOwnerPackage(InOwner.GetPackage());
+	InitializeSave(InSplitDataObject, InCookPackageSplitterInstance);
+}
+
+void FGeneratorPackage::InitializeSave(const UObject* InSplitDataObject,
+	ICookPackageSplitter* InCookPackageSplitterInstance)
+{
+	if (InCookPackageSplitterInstance)
+	{
+		if (IsInitialized())
+		{
+			check(CookPackageSplitterInstance.Get() == InCookPackageSplitterInstance);
+			check(SplitDataObjectName == FName(*InSplitDataObject->GetFullName()));
+		}
+		else
+		{
+			CookPackageSplitterInstance.Reset(InCookPackageSplitterInstance);
+			SplitDataObjectName = *InSplitDataObject->GetFullName();
+			SetOwnerPackage(GetOwner().GetPackage());
+		}
+	}
 }
 
 FGeneratorPackage::~FGeneratorPackage()
 {
+	if (!IsInitialized())
+	{
+		return;
+	}
+
 	ConditionalNotifyCompletion(ICookPackageSplitter::ETeardown::Canceled);
 	ClearGeneratedPackages();
+	if (!PreviousGeneratedPackages.IsEmpty())
+	{
+		FPackageData& PackageData = GetOwner();
+		UCookOnTheFlyServer& COTFS =PackageData.GetPackageDatas().GetCookOnTheFlyServer();
+		TArray<const ITargetPlatform*, TInlineAllocator<1>> RequestedPlatforms;
+		PackageData.GetRequestedPlatforms(RequestedPlatforms);
+		for (FName PreviousPackageName : PreviousGeneratedPackages)
+		{
+			for (const ITargetPlatform* Platform : RequestedPlatforms)
+			{
+				COTFS.DeleteOutputForPackage(PreviousPackageName, Platform);
+			}
+		}
+		PreviousGeneratedPackages.Empty();
+	}
 }
 
 void FGeneratorPackage::ConditionalNotifyCompletion(ICookPackageSplitter::ETeardown Status)
@@ -1412,6 +1461,7 @@ void FGeneratorPackage::ConditionalNotifyCompletion(ICookPackageSplitter::ETeard
 
 void FGeneratorPackage::ClearGeneratedPackages()
 {
+	check(IsInitialized());
 	for (FCookGenerationInfo& Info: PackagesToGenerate)
 	{
 		if (Info.PackageData)
@@ -1425,12 +1475,19 @@ void FGeneratorPackage::ClearGeneratedPackages()
 
 bool FGeneratorPackage::TryGenerateList(UObject* OwnerObject, FPackageDatas& PackageDatas)
 {
+	check(IsInitialized());
+
 	FPackageData& OwnerPackageData = GetOwner();
 	UPackage* LocalOwnerPackage = OwnerPackageData.GetPackage();
 	check(LocalOwnerPackage);
 	TArray<ICookPackageSplitter::FGeneratedPackage> GeneratorDatas =
 		CookPackageSplitterInstance->GetGenerateList(LocalOwnerPackage, OwnerObject);
 	PackagesToGenerate.Reset(GeneratorDatas.Num());
+	TArray<const ITargetPlatform*, TInlineAllocator<1>> RequestedPlatforms;
+	OwnerPackageData.GetRequestedPlatforms(RequestedPlatforms);
+	UCookOnTheFlyServer& COTFS = GetOwner().GetPackageDatas().GetCookOnTheFlyServer();
+	bool bHybridIterativeEnabled = COTFS.bHybridIterativeEnabled;
+
 	for (ICookPackageSplitter::FGeneratedPackage& SplitterData : GeneratorDatas)
 	{
 		if (!SplitterData.GetCreateAsMap().IsSet())
@@ -1470,9 +1527,39 @@ bool FGeneratorPackage::TryGenerateList(UObject* OwnerObject, FPackageDatas& Pac
 		GeneratedInfo.RelativePath = MoveTemp(SplitterData.RelativePath);
 		GeneratedInfo.GeneratedRootPath = MoveTemp(SplitterData.GeneratedRootPath);
 		GeneratedInfo.Dependencies = MoveTemp(SplitterData.Dependencies);
+		Algo::Sort(GeneratedInfo.Dependencies, FNameLexicalLess());
+		GeneratedInfo.Dependencies.SetNum(Algo::Unique(GeneratedInfo.Dependencies));
 		GeneratedInfo.SetIsCreateAsMap(bCreateAsMap);
 		PackageData->SetGeneratedOwner(this);
 		PackageData->SetWorkerAssignmentConstraint(FWorkerId::Local());
+
+		// Create the Guid from the GenerationHash and Dependencies
+		GeneratedInfo.CreateGuid();
+
+		if (!bHybridIterativeEnabled)
+		{
+			GeneratedInfo.IterativeCookValidateOrClear(*this, RequestedPlatforms);
+		}
+	}
+	if (!PreviousGeneratedPackages.IsEmpty())
+	{
+		TSet<FName> CurrentGeneratedPackageNames;
+		CurrentGeneratedPackageNames.Reserve(PackagesToGenerate.Num());
+		for (const FCookGenerationInfo& Info : PackagesToGenerate)
+		{
+			CurrentGeneratedPackageNames.Add(Info.PackageData->GetPackageName());
+		}
+		for (FName PreviousPackageName : PreviousGeneratedPackages)
+		{
+			if (!CurrentGeneratedPackageNames.Contains(PreviousPackageName))
+			{
+				for (const ITargetPlatform* TargetPlatform : RequestedPlatforms)
+				{
+					COTFS.DeleteOutputForPackage(PreviousPackageName, TargetPlatform);
+				}
+			}
+		}
+		PreviousGeneratedPackages.Empty();
 	}
 	RemainingToPopulate = GeneratorDatas.Num() + 1; // GeneratedPackaged plus one for the Generator
 	return true;
@@ -1480,6 +1567,8 @@ bool FGeneratorPackage::TryGenerateList(UObject* OwnerObject, FPackageDatas& Pac
 
 FCookGenerationInfo* FGeneratorPackage::FindInfo(const FPackageData& PackageData)
 {
+	check(IsInitialized());
+
 	if (&PackageData == &GetOwner())
 	{
 		return &OwnerInfo;
@@ -1500,6 +1589,7 @@ const FCookGenerationInfo* FGeneratorPackage::FindInfo(const FPackageData& Packa
 
 UObject* FGeneratorPackage::FindSplitDataObject() const
 {
+	check(IsInitialized());
 	FString ObjectPath = GetSplitDataObjectName().ToString();
 
 	// SplitDataObjectName is a FullObjectPath; strip off the leading <ClassName> in
@@ -1515,6 +1605,11 @@ UObject* FGeneratorPackage::FindSplitDataObject() const
 void FGeneratorPackage::PreGarbageCollect(FCookGenerationInfo& Info, TArray<UObject*>& GCKeepObjects,
 	TArray<UPackage*>& GCKeepPackages, TArray<FPackageData*>& GCKeepPackageDatas, bool& bOutShouldDemote)
 {
+	if (!IsInitialized())
+	{
+		return;
+	}
+
 	bOutShouldDemote = false;
 	check(Info.PackageData); // Caller validates this is non-null
 	if (Info.GetSaveState() > FCookGenerationInfo::ESaveState::CallPopulate)
@@ -1561,6 +1656,11 @@ void FGeneratorPackage::PreGarbageCollect(FCookGenerationInfo& Info, TArray<UObj
 
 void FGeneratorPackage::PostGarbageCollect()
 {
+	if (!IsInitialized())
+	{
+		return;
+	}
+
 	FPackageData& Owner = GetOwner();
 	if (Owner.GetState() == EPackageState::Save)
 	{
@@ -1618,10 +1718,11 @@ void FGeneratorPackage::PostGarbageCollect()
 UPackage* FGeneratorPackage::CreateGeneratedUPackage(FCookGenerationInfo& GeneratedInfo,
 	const UPackage* InOwnerPackage, const TCHAR* GeneratedPackageName)
 {
+	check(IsInitialized());
 	UPackage* GeneratedPackage = CreatePackage(GeneratedPackageName);
-	PRAGMA_DISABLE_DEPRECATION_WARNINGS
-	GeneratedPackage->SetGuid(InOwnerPackage->GetGuid());
-	PRAGMA_ENABLE_DEPRECATION_WARNINGS
+	PRAGMA_DISABLE_DEPRECATION_WARNINGS;
+	GeneratedPackage->SetGuid(GeneratedInfo.Guid);
+	PRAGMA_ENABLE_DEPRECATION_WARNINGS;
 	GeneratedPackage->SetPersistentGuid(InOwnerPackage->GetPersistentGuid());
 	GeneratedPackage->SetPackageFlags(PKG_CookGenerated);
 	GeneratedInfo.SetHasCreatedPackage(true);
@@ -1629,11 +1730,13 @@ UPackage* FGeneratorPackage::CreateGeneratedUPackage(FCookGenerationInfo& Genera
 	{
 		GeneratedPackage->SetLoadedByEditorPropertiesOnly(false);
 	}
+
 	return GeneratedPackage;
 }
 
 void FGeneratorPackage::SetPackageSaved(FCookGenerationInfo& Info, FPackageData& PackageData)
 {
+	check(IsInitialized());
 	if (Info.HasSaved())
 	{
 		return;
@@ -1649,11 +1752,13 @@ void FGeneratorPackage::SetPackageSaved(FCookGenerationInfo& Info, FPackageData&
 
 bool FGeneratorPackage::IsComplete() const
 {
+	check(IsInitialized());
 	return RemainingToPopulate == 0;
 }
 
 void FGeneratorPackage::ResetSaveState(FCookGenerationInfo& Info, UPackage* Package, EReleaseSaveReason ReleaseSaveReason)
 {
+	check(IsInitialized());
 	if (Info.GetSaveState() > FCookGenerationInfo::ESaveState::CallPopulate)
 	{
 		UObject* SplitObject = FindSplitDataObject();
@@ -1727,6 +1832,10 @@ void FGeneratorPackage::ResetSaveState(FCookGenerationInfo& Info, UPackage* Pack
 
 void FGeneratorPackage::UpdateSaveAfterGarbageCollect(const FPackageData& PackageData, bool& bInOutDemote)
 {
+	if (!IsInitialized())
+	{
+		return;
+	}
 	FCookGenerationInfo* Info = FindInfo(PackageData);
 	if (!Info)
 	{
@@ -1938,6 +2047,46 @@ void FBeginCacheObjects::EndRound(int32 NumPlatforms)
 	}
 }
 
+void FCookGenerationInfo::CreateGuid()
+{
+	// COOKPACKAGESPLITTERTODO: Add GenerationHash
+	FBlake3 Blake3;
+	IAssetRegistry& AssetRegistry = IAssetRegistry::GetChecked();
+	for (FName DependencyName : Dependencies)
+	{
+		TOptional<FAssetPackageData> DependencyData = AssetRegistry.GetAssetPackageDataCopy(DependencyName);
+		if (DependencyData)
+		{
+			PRAGMA_DISABLE_DEPRECATION_WARNINGS;
+			Blake3.Update(&DependencyData->PackageGuid, sizeof(DependencyData->PackageGuid));
+			PRAGMA_ENABLE_DEPRECATION_WARNINGS;
+		}
+	}
+	FBlake3Hash GeneratedHash = Blake3.Finalize();
+	const uint32* HashInts = reinterpret_cast<const uint32*>(GeneratedHash.GetBytes());
+	Guid = FGuid(HashInts[0], HashInts[1], HashInts[2], HashInts[3]);
+}
+
+void FCookGenerationInfo::IterativeCookValidateOrClear(FGeneratorPackage& Generator,
+	TConstArrayView<const ITargetPlatform*> RequestedPlatforms)
+{
+	// COOKPACKAGESPLITTERTODO: Test whether the Hash from the generator has changed
+	constexpr bool bInvalidated = false;
+	if (bInvalidated)
+	{
+		for (const ITargetPlatform* TargetPlatform : RequestedPlatforms)
+		{
+			FPackageData::FPlatformData* PlatformData = PackageData->FindPlatformData(TargetPlatform);
+			if (PlatformData && PlatformData->bCookAttempted)
+			{
+				PackageData->ClearCookProgress(TargetPlatform);
+				UCookOnTheFlyServer& COTFS = Generator.GetOwner().GetPackageDatas().GetCookOnTheFlyServer();
+				COTFS.DeleteOutputForPackage(PackageData->GetPackageName(), TargetPlatform);
+			}
+		}
+	}
+}
+
 //////////////////////////////////////////////////////////////////////////
 // FPendingCookedPlatformData
 
@@ -2091,9 +2240,14 @@ int32 FPackageDataMonitor::GetNumInProgress() const
 	return NumInProgress;
 }
 
-int32 FPackageDataMonitor::GetNumCooked() const
+int32 FPackageDataMonitor::GetNumFailedCooked() const
 {
-	return NumCooked;
+	return NumFailedCooked;
+}
+
+int32 FPackageDataMonitor::GetNumSuccessfulCooked() const
+{
+	return NumSuccessfulCooked;
 }
 
 void FPackageDataMonitor::OnInProgressChanged(FPackageData& PackageData, bool bInProgress)
@@ -2108,12 +2262,14 @@ void FPackageDataMonitor::OnPreloadAllocatedChanged(FPackageData& PackageData, b
 	check(NumPreloadAllocated >= 0);
 }
 
-void FPackageDataMonitor::OnFirstCookedPlatformAdded(FPackageData& PackageData)
+void FPackageDataMonitor::OnFirstCookedPlatformAdded(FPackageData& PackageData, bool bSuccessful)
 {
 	if (!PackageData.GetMonitorIsCooked())
 	{
-		++NumCooked;
 		PackageData.SetMonitorIsCooked(true);
+		PackageData.SetMonitorIsSuccessfulCooked(bSuccessful);
+		int32& Stat = (bSuccessful ? NumSuccessfulCooked : NumFailedCooked);
+		++Stat;
 	}
 }
 
@@ -2121,8 +2277,10 @@ void FPackageDataMonitor::OnLastCookedPlatformRemoved(FPackageData& PackageData)
 {
 	if (PackageData.GetMonitorIsCooked())
 	{
-		--NumCooked;
+		int32& Stat = (PackageData.GetMonitorIsSuccessfulCooked() ? NumSuccessfulCooked : NumFailedCooked);
+		--Stat;
 		PackageData.SetMonitorIsCooked(false);
+		PackageData.SetMonitorIsSuccessfulCooked(false);
 	}
 }
 
@@ -2618,7 +2776,17 @@ FPackageData* FPackageDatas::UpdateFileName(FName PackageName)
 
 int32 FPackageDatas::GetNumCooked()
 {
-	return Monitor.GetNumCooked();
+	return Monitor.GetNumSuccessfulCooked() + Monitor.GetNumFailedCooked();
+}
+
+int32 FPackageDatas::GetNumSuccessfulCooked()
+{
+	return Monitor.GetNumSuccessfulCooked();
+}
+
+int32 FPackageDatas::GetNumFailedCooked()
+{
+	return Monitor.GetNumFailedCooked();
 }
 
 void FPackageDatas::GetCookedPackagesForPlatform(const ITargetPlatform* Platform, TArray<FPackageData*>& CookedPackages,
diff --git a/Engine/Source/Editor/UnrealEd/Private/Cooker/CookPackageData.h b/Engine/Source/Editor/UnrealEd/Private/Cooker/CookPackageData.h
--- a/Engine/Source/Editor/UnrealEd/Private/Cooker/CookPackageData.h
+++ b/Engine/Source/Editor/UnrealEd/Private/Cooker/CookPackageData.h
@@ -179,6 +179,8 @@ public:
 	FPackageData& operator=(const FPackageData& other) = delete;
 	FPackageData& operator=(FPackageData&& other) = delete;
 
+	FPackageDatas& GetPackageDatas() const { return PackageDatas; }
+
 	/** Return a copy of Package->GetName(). It is derived from the FileName if necessary, and is never modified. */
 	const FName& GetPackageName() const;
 	/**
@@ -459,6 +461,8 @@ public:
 	/** Get/Set the PackageDataMonitor flag that counts whether this PackageData has finished cooking. */
 	bool GetMonitorIsCooked() const { return static_cast<bool>(bMonitorIsCooked); }
 	void SetMonitorIsCooked(bool Value) { bMonitorIsCooked = Value != 0; }
+	bool GetMonitorIsSuccessfulCooked() const { return static_cast<bool>(bMonitorIsSuccessfulCooked); }
+	void SetMonitorIsSuccessfulCooked(bool Value) { bMonitorIsSuccessfulCooked = Value != 0; }
 
 	/** Remove all request data about the given platform from all fields in this PackageData. */
 	void OnRemoveSessionPlatform(const ITargetPlatform* Platform);
@@ -470,12 +474,12 @@ public:
 	void RemapTargetPlatforms(const TMap<ITargetPlatform*, ITargetPlatform*>& Remap);
 
 	/** Create and return a GeneratorPackage helper object for a given CookPackageSplitter. */
-	FGeneratorPackage* CreateGeneratorPackage(const UObject* InSplitDataObject,
+	FGeneratorPackage& CreateGeneratorPackage(const UObject* InSplitDataObject,
 		ICookPackageSplitter* InCookPackageSplitterInstance);
 	/** Destroy GeneratorPackage helper object. */
 	void DestroyGeneratorPackage() { GeneratorPackage.Reset(); }
 	/** Get GeneratorPackage helper object. */
-	FGeneratorPackage* GetGeneratorPackage() const { return GeneratorPackage.Get(); }
+	FGeneratorPackage* GetGeneratorPackage() const;
 	/** Get whether the PackageData has done any necessary Generator steps and is ready for BeginCache calls. */
 	bool HasCompletedGeneration() const { return static_cast<bool>(bCompletedGeneration); }
 	/** Set whether the PackageData has done any necessary Generator steps and is ready for BeginCache calls. */
@@ -639,6 +643,7 @@ private:
 	uint32 bCookedPlatformDataCalled : 1;
 	uint32 bCookedPlatformDataComplete : 1;
 	uint32 bMonitorIsCooked : 1;
+	uint32 bMonitorIsSuccessfulCooked : 1;
 	uint32 bInitializedGeneratorSave : 1;
 	uint32 bCompletedGeneration : 1;
 	uint32 bGenerated : 1;
@@ -738,8 +743,19 @@ public:
 		KeepReferencedPackages.Append(InKeepReferencedPackages);
 	}
 
+	/** Create the Guid for this generated package, based on dependencies and GenerationHash. */
+	void CreateGuid();
+
+	/**
+	 * If the package has iterative results from a previous cook that were not invalidated by dependency changes,
+	 * test whether they need to be invalidated now and clear the results if so. Only called in legacy iterative;
+	 * incremental cooks handle invalidation by querying the TargetDomainDigest during the RequestCluster.
+	 */
+	void IterativeCookValidateOrClear(FGeneratorPackage& Generator, TConstArrayView<const ITargetPlatform*> RequestedPlatforms);
+
 public:
 	FBeginCacheObjects BeginCacheObjects;
+	FGuid Guid;
 	FString RelativePath;
 	FString GeneratedRootPath;
 	TArray<FName> Dependencies;
@@ -766,6 +782,8 @@ public:
 	FGeneratorPackage(UE::Cook::FPackageData& InOwner, const UObject* InSplitDataObject,
 		ICookPackageSplitter* InCookPackageSplitterInstance);
 	~FGeneratorPackage();
+	void InitializeSave(const UObject* InSplitDataObject, ICookPackageSplitter* InCookPackageSplitterInstance);
+	bool IsInitialized() const { return CookPackageSplitterInstance.IsValid(); }
 	/** Clear references to owned generated packages, and mark those packages as orphaned */
 	void ClearGeneratedPackages();
 
@@ -773,9 +791,9 @@ public:
 	bool TryGenerateList(UObject* OwnerObject, FPackageDatas& PackageDatas);
 
 	/** Accessor for the packages to generate */
-	TArrayView<UE::Cook::FCookGenerationInfo> GetPackagesToGenerate() { return PackagesToGenerate; }
+	TArrayView<UE::Cook::FCookGenerationInfo> GetPackagesToGenerate() { check(IsInitialized()); return PackagesToGenerate; }
 	/** Return the GenerationInfo used to save the Generator's UPackage */
-	UE::Cook::FCookGenerationInfo& GetOwnerInfo() { return OwnerInfo; }
+	UE::Cook::FCookGenerationInfo& GetOwnerInfo() { check(IsInitialized()); return OwnerInfo; }
 	/** Return owner FPackageData. */
 	UE::Cook::FPackageData& GetOwner() { return *OwnerInfo.PackageData; }
 	/** Return the GenerationInfo for the given PackageData, or null if not found. */
@@ -785,7 +803,7 @@ public:
 	/** Return CookPackageSplitter. */
 	ICookPackageSplitter* GetCookPackageSplitterInstance() const { return CookPackageSplitterInstance.Get(); }
 	/** Return the SplitDataObject's FullObjectPath. */
-	const FName GetSplitDataObjectName() const { return SplitDataObjectName; }
+	const FName GetSplitDataObjectName() const { check(IsInitialized()); return SplitDataObjectName; }
 	/**
 	 * Find again the split object from its name, or return null if no longer in memory.
 	 * It may have been GC'd and reloaded since the last time we used it.
@@ -794,7 +812,7 @@ public:
 
 	void ResetSaveState(FCookGenerationInfo& Info, UPackage* Package, UE::Cook::EReleaseSaveReason ReleaseSaveReason);
 
-	int32& GetNextPopulateIndex() { return NextPopulateIndex; }
+	int32& GetNextPopulateIndex() { check(IsInitialized()); return NextPopulateIndex; }
 
 	/** Callbacks during garbage collection */
 	void PreGarbageCollect(FCookGenerationInfo& Info, TArray<UObject*>& GCKeepObjects,
@@ -813,6 +831,7 @@ public:
 
 	UPackage* GetOwnerPackage() const { return OwnerPackage.Get(); };
 	void SetOwnerPackage(UPackage* InPackage) { OwnerPackage = InPackage; }
+	void SetPreviousGeneratedPackages(TConstArrayView<FName> Packages) { PreviousGeneratedPackages = Packages; }
 
 private:
 	void ConditionalNotifyCompletion(ICookPackageSplitter::ETeardown Status);
@@ -826,6 +845,7 @@ private:
 	/** Recorded list of packages to generate from the splitter, and data we need about them */
 	TArray<FCookGenerationInfo> PackagesToGenerate;
 	TWeakObjectPtr<UPackage> OwnerPackage;
+	TArray<FName> PreviousGeneratedPackages;
 
 	int32 NextPopulateIndex = 0;
 	int32 RemainingToPopulate = 0;
@@ -924,7 +944,8 @@ public:
 	int32 GetNumInProgress() const;
 	int32 GetNumPreloadAllocated() const;
 	/** Report the number of packages that have cooked any platform. Used by CookCommandlet progress reporting. */
-	int32 GetNumCooked() const;
+	int32 GetNumSuccessfulCooked() const;
+	int32 GetNumFailedCooked() const;
 	/**
 	 * Report the number of FPackageData that are currently marked as urgent.
 	 * Used to check if a Pump function needs to exit to handle urgent PackageData in other states.
@@ -941,7 +962,7 @@ public:
 	void OnInProgressChanged(FPackageData& PackageData, bool bInProgress);
 	void OnPreloadAllocatedChanged(FPackageData& PackageData, bool bPreloadAllocated);
 	/** Callback called from FPackageData when it has set a platform to CookAttempted=true and it does not have any others cooked. */
-	void OnFirstCookedPlatformAdded(FPackageData& PackageData);
+	void OnFirstCookedPlatformAdded(FPackageData& PackageData, bool bSuccessful);
 	/** Callback called from FPackageData when it has set a platform to CookAttempted=false and it does not have any others cooked. */
 	void OnLastCookedPlatformRemoved(FPackageData& PackageData);
 	/** Callback called from FPackageData when it has changed its urgency. */
@@ -954,7 +975,8 @@ private:
 	void TrackUrgentRequests(EPackageState State, int32 Delta);
 
 	int32 NumInProgress = 0;
-	int32 NumCooked = 0;
+	int32 NumFailedCooked = 0;
+	int32 NumSuccessfulCooked = 0;
 	int32 NumPreloadAllocated = 0;
 	int32 NumUrgentInState[static_cast<uint32>(EPackageState::Count)];
 };
@@ -1193,6 +1215,8 @@ public:
 
 	/** Report the number of packages that have cooked any platform. Used by cook commandlet progress reporting. */
 	int32 GetNumCooked();
+	int32 GetNumSuccessfulCooked();
+	int32 GetNumFailedCooked();
 	/**
 	 * Append to CookedPackages all packages that have cooked any platform, either successfully if
 	 * bGetSuccessfulCookedPackages is true and/or unsuccessfully if bGetFailedCookedPackages is true.
diff --git a/Engine/Source/Editor/UnrealEd/Private/Cooker/LooseCookedPackageWriter.cpp b/Engine/Source/Editor/UnrealEd/Private/Cooker/LooseCookedPackageWriter.cpp
--- a/Engine/Source/Editor/UnrealEd/Private/Cooker/LooseCookedPackageWriter.cpp
+++ b/Engine/Source/Editor/UnrealEd/Private/Cooker/LooseCookedPackageWriter.cpp
@@ -46,6 +46,7 @@
 #include "Serialization/CompactBinaryWriter.h"
 #include "Serialization/LargeMemoryWriter.h"
 #include "Serialization/MemoryWriter.h"
+#include "String/Find.h"
 #include "Tasks/Task.h"
 #include "Templates/RefCounting.h"
 #include "Templates/UnrealTemplate.h"
@@ -662,8 +663,10 @@ void FLooseCookedPackageWriter::RemoveCookedPackages(TArrayView<const FName> Pac
 {
 	if (PackageNameToCookedFiles.IsEmpty())
 	{
+		FindAndDeleteCookedFilesForPackages(PackageNamesToRemove);
 		return;
 	}
+
 	// if we are going to clear the cooked packages it is conceivable that we will recook the packages which we just cooked 
 	// that means it's also conceivable that we will recook the same package which currently has an outstanding async write request
 	UPackage::WaitForAsyncFileWrites();
@@ -856,6 +859,57 @@ void FLooseCookedPackageWriter::GetAllCookedFiles()
 	}
 }
 
+void FLooseCookedPackageWriter::FindAndDeleteCookedFilesForPackages(TConstArrayView<FName> PackageNames)
+{
+	const FString& SandboxRootDir = OutputPath;
+	const FString SandboxProjectDir = FPaths::Combine(OutputPath, FApp::GetProjectName()) + TEXT("/");
+	const FString RelativeRootDir = FPaths::GetRelativePathToRoot();
+	const FString RelativeProjectDir = FPaths::ProjectDir();
+
+	for (FName PackageName : PackageNames)
+	{
+		FString CookedFileName = ConvertPackageNameToCookedPath(
+			SandboxRootDir, RelativeRootDir,
+			SandboxProjectDir, RelativeProjectDir,
+			WriteToString<256>(PackageName));
+		if (CookedFileName.IsEmpty())
+		{
+			continue;
+		}
+		FStringView ParentDir;
+		FStringView BaseName;
+		FStringView Extension;
+		FPathViews::Split(CookedFileName, ParentDir, BaseName, Extension);
+
+		IPlatformFile& PlatformFile = FPlatformFileManager::Get().GetPlatformFile();
+		TArray<FString, TInlineAllocator<3>> FilesToRemove;
+		PlatformFile.IterateDirectory(*WriteToString<1024>(ParentDir),
+			[BaseName, ParentDir, &FilesToRemove](const TCHAR* FoundFullPath, bool bDirectory)
+			{
+				FStringView FoundParentDir;
+				FStringView FoundBaseName;
+				FStringView FoundExtension;
+				FPathViews::Split(FoundFullPath, FoundParentDir, FoundBaseName, FoundExtension);
+				if (FoundBaseName == BaseName)
+				{
+					if (FoundParentDir.IsEmpty())
+					{
+						FilesToRemove.Add(FPaths::ConvertRelativePathToFull(FString(ParentDir), FString(FoundFullPath)));
+					}
+					else
+					{
+						FilesToRemove.Add(FString(FoundFullPath));
+					}
+				}
+				return true;
+			});
+		for (const FString& FileName : FilesToRemove)
+		{
+			PlatformFile.DeleteFile(*FileName);
+		}
+	}
+}
+
 FName FLooseCookedPackageWriter::ConvertCookedPathToPackageName(
 	const FString& SandboxRootDir, const FString& RelativeRootDir,
 	const FString& SandboxProjectDir, const FString& RelativeProjectDir,
@@ -928,6 +982,79 @@ FName FLooseCookedPackageWriter::ConvertCookedPathToPackageName(
 	return FName(*ScratchPackageName);
 }
 
+bool FLooseCookedPackageWriter::TryConvertUncookedFilenameToCookedRemappedPluginFilename(
+	FStringView FileName, TConstArrayView<TSharedRef<IPlugin>> InPluginsToRemap, FStringView SandboxDirectory,
+	FString& OutCookedFileName)
+{
+	// Ideally this would be in the Sandbox File but it can't access the project or plugin
+	if (InPluginsToRemap.IsEmpty())
+	{
+		return false;
+	}
+
+	for (const TSharedRef<IPlugin>& Plugin : InPluginsToRemap)
+	{
+		// If these match, then this content is part of plugin that gets remapped when packaged/staged
+		if (FileName.StartsWith(Plugin->GetContentDir()))
+		{
+			FString SearchFor;
+			SearchFor /= Plugin->GetName() / TEXT("Content");
+			int32 FoundAt = UE::String::FindLast(FileName, SearchFor, ESearchCase::IgnoreCase);
+			check(FoundAt != INDEX_NONE);
+			// Strip off everything but <PluginName/Content/<remaing path to file>
+			FStringView SnippedOffPath = FileName.RightChop(FoundAt);
+			// Put this is in <sandbox path>/RemappedPlugins/<PluginName>/Content/<remaing path to file>
+			constexpr FStringView RemappedPluginsDirName(REMAPPED_PLUGINS);
+			OutCookedFileName.Reserve(SandboxDirectory.Len() + RemappedPluginsDirName.Len() + SnippedOffPath.Len() + 2);
+			OutCookedFileName = SandboxDirectory;
+			OutCookedFileName /= REMAPPED_PLUGINS;
+			OutCookedFileName /= SnippedOffPath;
+			return true;
+		}
+	}
+	return false;
+}
+
+FString FLooseCookedPackageWriter::ConvertPackageNameToCookedPath(
+	const FString& SandboxRootDir, const FString& RelativeRootDir,
+	const FString& SandboxProjectDir, const FString& RelativeProjectDir,
+	FStringView PackageName) const
+{
+	FString UncookedFileName;
+	if (!FPackageName::TryConvertLongPackageNameToFilename(FString(PackageName), UncookedFileName))
+	{
+		return TEXT("");
+	}
+
+	FString CookedFileName;
+	if (TryConvertUncookedFilenameToCookedRemappedPluginFilename(UncookedFileName, PluginsToRemap,
+		SandboxRootDir, CookedFileName))
+	{
+		return CookedFileName;
+	}
+
+	auto BuildCookedPath =
+		[&CookedFileName](const FString& UncookedPath, const FString& CookedRoot, const FString& UncookedRoot)
+	{
+		CookedFileName.AppendChars(*CookedRoot, CookedRoot.Len());
+		CookedFileName.AppendChars(*UncookedPath + UncookedRoot.Len(), UncookedPath.Len() - UncookedRoot.Len());
+	};
+
+	if (UncookedFileName.StartsWith(RelativeProjectDir))
+	{
+		BuildCookedPath(UncookedFileName, SandboxProjectDir, RelativeProjectDir);
+	}
+	else if (UncookedFileName.StartsWith(RelativeRootDir))
+	{
+		BuildCookedPath(UncookedFileName, SandboxRootDir, RelativeRootDir);
+	}
+	else
+	{
+		CookedFileName.Empty();
+	}
+	return CookedFileName;
+}
+
 EPackageExtension FLooseCookedPackageWriter::BulkDataTypeToExtension(FBulkDataInfo::EType BulkDataType)
 {
 	switch (BulkDataType)
diff --git a/Engine/Source/Editor/UnrealEd/Private/Cooker/LooseCookedPackageWriter.h b/Engine/Source/Editor/UnrealEd/Private/Cooker/LooseCookedPackageWriter.h
--- a/Engine/Source/Editor/UnrealEd/Private/Cooker/LooseCookedPackageWriter.h
+++ b/Engine/Source/Editor/UnrealEd/Private/Cooker/LooseCookedPackageWriter.h
@@ -71,6 +71,10 @@ public:
 	virtual FPackageWriterRecords::FPackage* ConstructRecord() override;
 
 	static EPackageExtension BulkDataTypeToExtension(FBulkDataInfo::EType BulkDataType);
+	static bool TryConvertUncookedFilenameToCookedRemappedPluginFilename(
+		FStringView FileName, TConstArrayView<TSharedRef<IPlugin>> InPluginsToRemap,
+		FStringView SandboxDirectory, FString& OutCookedFileName);
+
 private:
 
 	/** Version of the superclass's per-package record that includes our class-specific data. */
@@ -114,11 +118,17 @@ private:
 	void DeleteSandboxDirectory();
 	/** Searches the disk for all the cooked files in the sandbox. Stores results in PackageNameToCookedFiles. */
 	void GetAllCookedFiles();
+	void FindAndDeleteCookedFilesForPackages(TConstArrayView<FName> PackageNames);
 
 	FName ConvertCookedPathToPackageName(
 		const FString& SandboxRootDir, const FString& RelativeRootDir,
 		const FString& SandboxProjectDir, const FString& RelativeProjectDir,
 		const FString& CookedPath, FString& ScratchFileName, FString& ScratchPackageName) const;
+	FString ConvertPackageNameToCookedPath(
+		const FString& SandboxRootDir, const FString& RelativeRootDir,
+		const FString& SandboxProjectDir, const FString& RelativeProjectDir,
+		FStringView PackageName) const;
+
 	void RemoveCookedPackagesByPackageName(TArrayView<const FName> PackageNamesToRemove, bool bRemoveRecords);
 	void AsyncSave(FRecord& Record, const FCommitPackageInfo& Info);
 
diff --git a/Engine/Source/Editor/UnrealEd/Public/Commandlets/AssetRegistryGenerator.h b/Engine/Source/Editor/UnrealEd/Public/Commandlets/AssetRegistryGenerator.h
--- a/Engine/Source/Editor/UnrealEd/Public/Commandlets/AssetRegistryGenerator.h
+++ b/Engine/Source/Editor/UnrealEd/Public/Commandlets/AssetRegistryGenerator.h
@@ -49,11 +49,6 @@ public:
 
 	const ITargetPlatform* GetTargetPlatform() const { return TargetPlatform; }
 
-	/** 
-	 * Sets asset registry from a previous run that is used for iterative cooking
-	 */
-	void SetPreviousAssetRegistry(TUniquePtr<FAssetRegistryState>&& PreviousState);
-
 	/**
 	 * Options when computing the differences between current and previous state.
 	 */
@@ -65,6 +60,13 @@ public:
 		bool bRecurseScriptModifications;
 	};
 
+	/** Info about a GeneratorPackage (see ICookPackageSplitter) loaded from previous iterative cooks. */
+	struct FGeneratorPackageInfo
+	{
+		TArray<FName> Generated;
+		bool bGeneratedModified = false;
+	};
+
 	/**
 	 * Differences between the current and the previous state.
 	 */
@@ -80,8 +82,13 @@ public:
 		TSet<FName> IdenticalCookedPackages;
 		/** IdenticalUncookedPackages list of uncooked packages that have not changed.These were filtered out by platform or editor only */
 		TSet<FName> IdenticalUncookedPackages;
+		/** Set of Generator packages and the packages they generated. */
+		TMap<FName, FGeneratorPackageInfo> PreviousGeneratorPackages;
 	};
 
+	/** Sets asset registry from a previous run that is used for iterative cooking. */
+	void SetPreviousAssetRegistry(TUniquePtr<FAssetRegistryState>&& PreviousState);
+
 	/**
 	 * Computes differences between the current asset registry state and the provided previous state.
 	 *
@@ -92,7 +99,8 @@ public:
 	void ComputePackageDifferences(const FComputeDifferenceOptions& Options, const FAssetRegistryState& PreviousState, FAssetRegistryDifference& OutDifference);
 
 	/** Computes just the list of packages in the PreviousState that no longer exist in the current state. */
-	void ComputePackageRemovals(const FAssetRegistryState& PreviousState, TArray<FName>& RemovedPackages);
+	void ComputePackageRemovals(const FAssetRegistryState& PreviousState, TArray<FName>& RemovedPackages,
+		TMap<FName, FGeneratorPackageInfo>& PreviousGeneratorPackages);
 
 	/**
 	 * FinalizeChunkIDs 
@@ -197,14 +205,18 @@ public:
 
 	/**
 	 * Attempts to update the metadata for a package in an asset registry generator.
-	 * This is only called for CookByTheBooks.
+	 * This is only called for CookByTheBook.
 	 *
 	 * @param Package The package to update info on
 	 * @param SavePackageResult The metadata to associate with the given package name
 	 * @param bIncludeOnlyDiskAssets Include only disk assets or else also enumerate memory assets
 	 */
-	void UpdateAssetRegistryData(const UPackage& Package, FSavePackageResultStruct& SavePackageResult, FCookTagList&& InArchiveCookTagList, bool bIncludeOnlyDiskAssets = true);
-	void UpdateAssetRegistryData(UE::Cook::FMPCollectorServerMessageContext& Context, UE::Cook::FAssetRegistryPackageMessage&& Message);
+	void UpdateAssetRegistryData(const UPackage& Package, FSavePackageResultStruct& SavePackageResult,
+		FCookTagList&& InArchiveCookTagList, bool bIncludeOnlyDiskAssets,
+		TOptional<FAssetPackageData>&& OverrideAssetPackageData,
+		TOptional<TArray<FAssetDependency>>&& OverridePackageDependencies);
+	void UpdateAssetRegistryData(UE::Cook::FMPCollectorServerMessageContext& Context,
+		UE::Cook::FAssetRegistryPackageMessage&& Message);
 
 	/**
 	 * Check config to see whether chunk assignments use the AssetManager. If so, run the once-per-process construction
@@ -236,6 +248,8 @@ private:
 	 */
 	void UpdateKeptPackages();
 
+	void SetOverridePackageDependencies(FName PackageName, TConstArrayView<FAssetDependency> OverridePackageDependencies);
+
 	static void InitializeUseAssetManager();
 
 	/** State of the asset registry that is being built for this platform */
@@ -247,7 +261,14 @@ private:
 	 */
 	TMap<FSoftObjectPath, TArray<TPair<FName, FString>>> CookTagsToAdd;
 
-	TMap<FName, TPair<TArray<FAssetData>, FAssetPackageData>> PreviousPackagesToUpdate;
+	struct FIterativelySkippedPackageUpdateData
+	{
+		TArray<FAssetData> AssetDatas;
+		FAssetPackageData PackageData;
+		TArray<FAssetDependency> PackageDependencies;
+		TArray<FAssetDependency> PackageReferencers;
+	};
+	TMap<FName, FIterativelySkippedPackageUpdateData> PreviousPackagesToUpdate;
 	/** List of packages that were loaded at startup */
 	TSet<FName> StartupPackages;
 	/** List of packages that were successfully cooked */
@@ -435,7 +456,9 @@ public:
 	virtual ~IAssetRegistryReporter() {}
 
 	virtual void UpdateAssetRegistryData(FPackageData& PackageData, const UPackage& Package,
-		FSavePackageResultStruct& SavePackageResult, FCookTagList&& InArchiveCookTagList, bool bIncludeOnlyDiskAssets = true) = 0;
+		FSavePackageResultStruct& SavePackageResult, FCookTagList&& InArchiveCookTagList,
+		bool bIncludeOnlyDiskAssets, TOptional<FAssetPackageData>&& OverrideAssetPackageData, 
+		TOptional<TArray<FAssetDependency>>&& OverridePackageDependencies) = 0;
 
 };
 
@@ -448,9 +471,12 @@ public:
 	}
 
 	virtual void UpdateAssetRegistryData(FPackageData& PackageData, const UPackage& Package,
-		FSavePackageResultStruct& SavePackageResult, FCookTagList&& InArchiveCookTagList, bool bIncludeOnlyDiskAssets = true) override
+		FSavePackageResultStruct& SavePackageResult, FCookTagList&& InArchiveCookTagList,
+		bool bIncludeOnlyDiskAssets, TOptional<FAssetPackageData>&& OverrideAssetPackageData,
+		TOptional<TArray<FAssetDependency>>&& OverridePackageDependencies) override
 	{
-		Generator.UpdateAssetRegistryData(Package, SavePackageResult, MoveTemp(InArchiveCookTagList), bIncludeOnlyDiskAssets);
+		Generator.UpdateAssetRegistryData(Package, SavePackageResult, MoveTemp(InArchiveCookTagList),
+			bIncludeOnlyDiskAssets, MoveTemp(OverrideAssetPackageData), MoveTemp(OverridePackageDependencies));
 	}
 
 private:
@@ -463,7 +489,9 @@ public:
 	FAssetRegistryReporterRemote(FCookWorkerClient& InClient, const ITargetPlatform* InTargetPlatform);
 
 	virtual void UpdateAssetRegistryData(FPackageData& PackageData, const UPackage& Package,
-		FSavePackageResultStruct& SavePackageResult, FCookTagList&& InArchiveCookTagList, bool bIncludeOnlyDiskAssets = true) override;
+		FSavePackageResultStruct& SavePackageResult, FCookTagList&& InArchiveCookTagList,
+		bool bIncludeOnlyDiskAssets, TOptional<FAssetPackageData>&& OverrideAssetPackageData, 
+		TOptional<TArray<FAssetDependency>>&& OverridePackageDependencies) override;
 
 private:
 	FCookWorkerClient& Client;
@@ -483,6 +511,8 @@ public:
 	FName PackageName;
 	const ITargetPlatform* TargetPlatform;
 	TArray<FAssetData> AssetDatas;
+	TOptional<FAssetPackageData> OverrideAssetPackageData;
+	TOptional<TArray<FAssetDependency>> OverridePackageDependencies;
 	TMap<FSoftObjectPath, TArray<TPair<FName, FString>>> CookTags;
 	uint32 PackageFlags = 0;
 	int64 DiskSize = -1;
diff --git a/Engine/Source/Editor/UnrealEd/Public/CookPackageSplitter.h b/Engine/Source/Editor/UnrealEd/Public/CookPackageSplitter.h
--- a/Engine/Source/Editor/UnrealEd/Public/CookPackageSplitter.h
+++ b/Engine/Source/Editor/UnrealEd/Public/CookPackageSplitter.h
@@ -27,6 +27,15 @@ template <typename FuncType> class TFunctionRef;
 class ICookPackageSplitter
 {
 public:
+	// Static API functions - these static functions are referenced by REGISTER_COOKPACKAGE_SPLITTER
+	// before creating an instance of the class.
+	/** Return whether the CookPackageSplitter subclass should handle the given SplitDataClass instance. */
+	static bool ShouldSplit(UObject* SplitData) { return false; }
+	/** Return DebugName for this SplitterClass in cook log messages. */
+	static FString GetSplitterDebugName() { return TEXT("<NoNameSpecified>"); }
+
+
+	// Virtual API functions - functions called from the cooker after creating the splitter.
 	virtual ~ICookPackageSplitter() {}
 
 	enum class ETeardown
@@ -37,6 +46,12 @@ public:
 	/** Do teardown actions after all packages have saved, or when the cook is cancelled. Always called before destruction. */
 	virtual void Teardown(ETeardown Status) {}
 
+	/**
+	 * If true, this splitter forces the Generator package objects it needs to remain referenced, and the cooker
+	 * should expect them to still be in memory after a garbage collect so long as the splitter is alive.
+	 */
+	virtual bool UseInternalReferenceToAvoidGarbageCollect() { return false; }
+
 	/** Data sent to the cooker to describe each desired generated package */
 	struct FGeneratedPackage
 	{
@@ -49,18 +64,6 @@ public:
 	private:
 		TOptional<bool> bCreateAsMap;
 	};
-	
-	/** 
-	 * Return whether the CookPackageSplitter subclass should handle the given SplitDataClass instance. 
-	 * Note that this is a static function referenced by macros, not part of the virtual api.
-	 */
-	static bool ShouldSplit(UObject* SplitData) { return false; }
-	
-	/**
-	 * If true, this splitter forces the Generator package objects it needs to remain referenced, and the cooker
-	 * should expect them to still be in memory after a garbage collect so long as the splitter is alive.
-	 */
-	virtual bool UseInternalReferenceToAvoidGarbageCollect() { return false; }
 
 	/** Return the list of packages to generate. */
 	virtual TArray<FGeneratedPackage> GetGenerateList(const UPackage* OwnerPackage, const UObject* OwnerObject) = 0;
@@ -216,6 +219,7 @@ public:
 	virtual UClass* GetSplitDataClass() const = 0;
 	virtual bool ShouldSplitPackage(UObject* Object) const = 0;
 	virtual ICookPackageSplitter* CreateInstance(UObject* Object) const = 0;
+	virtual FString GetSplitterDebugName() const = 0;
 
 	static void ForEach(TFunctionRef<void(FRegisteredCookPackageSplitter*)> Func);
 
@@ -243,6 +247,7 @@ class PREPROCESSOR_JOIN(PREPROCESSOR_JOIN(SplitterClass, SplitDataClass), _Regis
 	virtual UClass* GetSplitDataClass() const override { return SplitDataClass::StaticClass(); } \
 	virtual bool ShouldSplitPackage(UObject* Object) const override { return SplitterClass::ShouldSplit(Object); } \
 	virtual ICookPackageSplitter* CreateInstance(UObject* SplitData) const override { return new SplitterClass(); } \
+	virtual FString GetSplitterDebugName() const override { return SplitterClass::GetSplitterDebugName(); } \
 }; \
 namespace PREPROCESSOR_JOIN(SplitterClass, SplitDataClass) \
 { \
diff --git a/Engine/Source/Runtime/AssetRegistry/Private/AssetRegistry.cpp b/Engine/Source/Runtime/AssetRegistry/Private/AssetRegistry.cpp
--- a/Engine/Source/Runtime/AssetRegistry/Private/AssetRegistry.cpp
+++ b/Engine/Source/Runtime/AssetRegistry/Private/AssetRegistry.cpp
@@ -31,6 +31,8 @@
 #include "Misc/TrackedActivity.h"
 #include "PackageReader.h"
 #include "Serialization/ArrayReader.h"
+#include "Serialization/CompactBinarySerialization.h"
+#include "Serialization/CompactBinaryWriter.h"
 #include "Templates/UnrealTemplate.h"
 #include "UObject/ConstructorHelpers.h"
 #include "UObject/CoreRedirects.h"
@@ -6714,6 +6716,53 @@ bool FAssetRegistryDependencyOptions::GetManageQuery(UE::AssetRegistry::FDepende
 	return false;
 }
 
+void FAssetDependency::WriteCompactBinary(FCbWriter& Writer) const
+{
+	Writer.BeginArray();
+	Writer << AssetId;
+	static_assert(sizeof(uint8) >= sizeof(Category));
+	Writer.AddInteger((uint8)Category);
+	static_assert(sizeof(uint8) >= sizeof(Properties));
+	Writer.AddInteger((uint8)Properties);
+	Writer.EndArray();
+}
+
+bool LoadFromCompactBinary(FCbFieldView Field, FAssetDependency& Dependency)
+{
+	FCbArrayView ArrayField = Field.AsArrayView();
+	if (ArrayField.Num() < 3)
+	{
+		Dependency = FAssetDependency();
+		return false;
+	}
+	FCbFieldViewIterator Iter = ArrayField.CreateViewIterator();
+	if (!LoadFromCompactBinary(Iter++, Dependency.AssetId))
+	{
+		Dependency = FAssetDependency();
+		return false;
+	}
+	uint8 Value;
+	if (LoadFromCompactBinary(Iter++, Value))
+	{
+		Dependency.Category = (UE::AssetRegistry::EDependencyCategory)Value;
+	}
+	else
+	{
+		Dependency = FAssetDependency();
+		return false;
+	}
+	if (LoadFromCompactBinary(Iter++, Value))
+	{
+		Dependency.Properties = (UE::AssetRegistry::EDependencyProperty)Value;
+	}
+	else
+	{
+		Dependency = FAssetDependency();
+		return false;
+	}
+	return true;
+}
+
 namespace UE::AssetRegistry
 {
 
diff --git a/Engine/Source/Runtime/AssetRegistry/Private/AssetRegistryState.cpp b/Engine/Source/Runtime/AssetRegistry/Private/AssetRegistryState.cpp
--- a/Engine/Source/Runtime/AssetRegistry/Private/AssetRegistryState.cpp
+++ b/Engine/Source/Runtime/AssetRegistry/Private/AssetRegistryState.cpp
@@ -452,7 +452,7 @@ void FAssetRegistryState::InitializeFromExisting(const FAssetDataMap& AssetDataM
 					TargetNode->SetIsDependencyListSorted(InCategory, false);
 					TargetNode->AddDependency(TargetDependency, InCategory, InFlags);
 					TargetDependency->SetIsReferencersSorted(false);
-					TargetDependency->AddReferencer(TargetDependency);
+					TargetDependency->AddReferencer(TargetNode);
 				}
 			});
 			TargetNode->SetIsDependenciesInitialized(true);
@@ -505,7 +505,9 @@ void FAssetRegistryState::PruneAssetData(const TSet<FName>& RequiredPackages, co
 			bRemoveAssetData = true;
 		}
 		else if (Options.bFilterAssetDataWithNoTags && AssetData->TagsAndValues.Num() == 0 &&
-			!FPackageName::IsLocalizedPackage(WriteToString<256>(AssetData->PackageName)))
+			!FPackageName::IsLocalizedPackage(WriteToString<256>(AssetData->PackageName)) &&
+			// TODO: Add a package flag for PKG_CookGenerator and check that here as well.
+			!(AssetData->PackageFlags & PKG_CookGenerated))
 		{
 			bRemoveAssetData = true;
 			bRemoveDependencyData = Options.bFilterDependenciesWithNoTags;
@@ -997,7 +999,7 @@ bool FAssetRegistryState::GetReferencers(const FAssetIdentifier& AssetIdentifier
 		TArray<FDependsNode*> DependencyNodes;
 		Node->GetReferencers(DependencyNodes, Category, Flags);
 
-		OutReferencers.Reserve(DependencyNodes.Num());
+		OutReferencers.Reserve(OutReferencers.Num() + DependencyNodes.Num());
 		for (FDependsNode* DependencyNode : DependencyNodes)
 		{
 			OutReferencers.Add(DependencyNode->GetIdentifier());
@@ -1034,6 +1036,107 @@ bool FAssetRegistryState::GetReferencers(const FAssetIdentifier& AssetIdentifier
 	}
 }
 
+void FAssetRegistryState::ClearDependencies(const FAssetIdentifier& AssetIdentifier,
+	UE::AssetRegistry::EDependencyCategory Category)
+{
+	FDependsNode* ReferencerNode = FindDependsNode(AssetIdentifier);
+	if (!ReferencerNode)
+	{
+		return;
+	}
+
+	TArray<FDependsNode*> OldDependencies;
+	ReferencerNode->GetDependencies(OldDependencies);
+	ReferencerNode->ClearDependencies(Category);
+
+	for (FDependsNode* DependencyNode : OldDependencies)
+	{
+		if (!ReferencerNode->ContainsDependency(DependencyNode))
+		{
+			DependencyNode->RemoveReferencer(ReferencerNode);
+		}
+	}
+}
+
+void FAssetRegistryState::AddDependencies(const FAssetIdentifier& AssetIdentifier,
+	TConstArrayView<FAssetDependency> Dependencies)
+{
+	if (Dependencies.IsEmpty())
+	{
+		return;
+	}
+	FDependsNode* ReferencerNode = CreateOrFindDependsNode(AssetIdentifier);
+	for (const FAssetDependency& Dependency : Dependencies)
+	{
+		FDependsNode* DependencyNode = CreateOrFindDependsNode(Dependency.AssetId);
+		ReferencerNode->AddDependency(DependencyNode, Dependency.Category, Dependency.Properties);
+		DependencyNode->AddReferencer(ReferencerNode);
+	}
+}
+
+void FAssetRegistryState::SetDependencies(const FAssetIdentifier& AssetIdentifier,
+	TConstArrayView<FAssetDependency> Dependencies, UE::AssetRegistry::EDependencyCategory Category)
+{
+	for (const FAssetDependency& Dependency : Dependencies)
+	{
+		checkf(!(Dependency.Category & ~Category), TEXT("Input dependency has category %d which is outside of the requested categories %d."),
+			(int32)Dependency.Category, (int32)Category);
+	}
+
+	ClearDependencies(AssetIdentifier, Category);
+	AddDependencies(AssetIdentifier, Dependencies);
+}
+
+void FAssetRegistryState::ClearReferencers(const FAssetIdentifier& AssetIdentifier,
+	UE::AssetRegistry::EDependencyCategory Category)
+{
+	FDependsNode* DependencyNode = FindDependsNode(AssetIdentifier);
+	if (!DependencyNode)
+	{
+		return;
+	}
+
+	TArray<FDependsNode*> OldExisting;
+	DependencyNode->GetReferencers(OldExisting, Category);
+	for (FDependsNode* ReferencerNode : OldExisting)
+	{
+		ReferencerNode->RemoveDependency(DependencyNode, Category);
+		if (!ReferencerNode->ContainsDependency(DependencyNode))
+		{
+			DependencyNode->RemoveReferencer(ReferencerNode);
+		}
+	}
+}
+
+void FAssetRegistryState::AddReferencers(const FAssetIdentifier& AssetIdentifier,
+	TConstArrayView<FAssetDependency> Referencers)
+{
+	if (Referencers.IsEmpty())
+	{
+		return;
+	}
+	FDependsNode* DependencyNode = CreateOrFindDependsNode(AssetIdentifier);
+	for (const FAssetDependency& Referencer : Referencers)
+	{
+		FDependsNode* ReferencerNode = CreateOrFindDependsNode(Referencer.AssetId);
+		ReferencerNode->AddDependency(DependencyNode, Referencer.Category, Referencer.Properties);
+		DependencyNode->AddReferencer(ReferencerNode);
+	}
+}
+
+void FAssetRegistryState::SetReferencers(const FAssetIdentifier& AssetIdentifier,
+	TConstArrayView<FAssetDependency> Referencers, UE::AssetRegistry::EDependencyCategory Category)
+{
+	for (const FAssetDependency& Referencer : Referencers)
+	{
+		checkf(!(Referencer.Category & ~Category), TEXT("Input referencer has category %d which is outside of the requested categories %d."),
+			(int32)Referencer.Category, (int32)Category);
+	}
+
+	ClearReferencers(AssetIdentifier, Category);
+	AddReferencers(AssetIdentifier, Referencers);
+}
+
 bool FAssetRegistryState::Serialize(FArchive& Ar, const FAssetRegistrySerializationOptions& Options)
 {
 	return Ar.IsSaving() ? Save(Ar, Options) : Load(Ar, FAssetRegistryLoadOptions(Options));
diff --git a/Engine/Source/Runtime/AssetRegistry/Public/AssetRegistry/AssetRegistryState.h b/Engine/Source/Runtime/AssetRegistry/Public/AssetRegistry/AssetRegistryState.h
--- a/Engine/Source/Runtime/AssetRegistry/Public/AssetRegistry/AssetRegistryState.h
+++ b/Engine/Source/Runtime/AssetRegistry/Public/AssetRegistry/AssetRegistryState.h
@@ -280,7 +280,7 @@ public:
 	UE_DEPRECATED(4.26, "Use GetDependencies that takes a UE::AssetRegistry::EDependencyCategory instead")
 	bool GetDependencies(const FAssetIdentifier& AssetIdentifier, TArray<FAssetIdentifier>& OutDependencies, EAssetRegistryDependencyType::Type InDependencyType) const;
 	/**
-	 * Gets a list of packages and searchable names that are referenced by the supplied package or name. (On disk references ONLY)
+	 * Appends a list of packages and searchable names that are referenced by the supplied package or name. (On disk references ONLY)
 	 *
 	 * @param AssetIdentifier	the name of the package/name for which to gather dependencies
 	 * @param OutDependencies	a list of things that are referenced by AssetIdentifier
@@ -293,7 +293,7 @@ public:
 	UE_DEPRECATED(4.26, "Use GetReferencers that takes a UE::AssetRegistry::EDependencyCategory instead")
 	bool GetReferencers(const FAssetIdentifier& AssetIdentifier, TArray<FAssetIdentifier>& OutReferencers, EAssetRegistryDependencyType::Type InReferenceType = EAssetRegistryDependencyType::All) const;
 	/**
-	 * Gets a list of packages and searchable names that reference the supplied package or name. (On disk references ONLY)
+	 * Appends a list of packages and searchable names that reference the supplied package or name. (On disk references ONLY)
 	 *
 	 * @param AssetIdentifier	the name of the package/name for which to gather dependencies
 	 * @param OutReferencers	a list of things that reference AssetIdentifier
@@ -460,6 +460,39 @@ public:
 	/** Removes the asset data from the lookup maps */
 	void RemoveAssetData(FAssetData* AssetData, bool bRemoveDependencyData, bool& bOutRemovedAssetData, bool& bOutRemovedPackageData);
 
+	/**
+	 * Clear all dependencies of the given category from the given AssetIdentifier (e.g. package).
+	 * Also clears the referencer link from each of the dependencies.
+	 */
+	void ClearDependencies(const FAssetIdentifier& AssetIdentifier, UE::AssetRegistry::EDependencyCategory Category);
+	/**
+	 * Add the given dependencies to the given AssetIdentifier (e.g. package).
+	 * Also adds a referencer link on each of the dependencies.
+	 */
+	void AddDependencies(const FAssetIdentifier& AssetIdentifier, TConstArrayView<FAssetDependency> Dependencies);
+	/**
+	 * Clears existing dependencies of the given Category(s) and assigns the input Dependencies. Gives an error if 
+	 * any elements of Dependencies are outside of the Category(s).
+	 */
+	void SetDependencies(const FAssetIdentifier& AssetIdentifier, TConstArrayView<FAssetDependency> Dependencies,
+		UE::AssetRegistry::EDependencyCategory Category = UE::AssetRegistry::EDependencyCategory::All);
+	/**
+	 * Clear all referencers of the given category from the given AssetIdentifier (e.g. package).
+	 * Also clears the dependency link from each of the referencers.
+	 */
+	void ClearReferencers(const FAssetIdentifier& AssetIdentifier, UE::AssetRegistry::EDependencyCategory Category);
+	/**
+	 * Add a dependency on the given AssetIdentifier (e.g. package) from each of the Referencers.
+	 * Also adds a referencer link to each referencer on the AssetIdentifer's node.
+	 */
+	void AddReferencers(const FAssetIdentifier& AssetIdentifier, TConstArrayView<FAssetDependency> Referencers);
+	/**
+	 * Clears existing referencers of the given Category(s) and assigns the input Referencers. Gives an error if
+	 * any elements of Referencers are outside of the Category(s).
+	 */
+	void SetReferencers(const FAssetIdentifier& AssetIdentifier, TConstArrayView<FAssetDependency> Referencers,
+		UE::AssetRegistry::EDependencyCategory Category = UE::AssetRegistry::EDependencyCategory::All);
+
 	/** Resets to default state */
 	void Reset();
 
diff --git a/Engine/Source/Runtime/AssetRegistry/Public/AssetRegistry/IAssetRegistry.h b/Engine/Source/Runtime/AssetRegistry/Public/AssetRegistry/IAssetRegistry.h
--- a/Engine/Source/Runtime/AssetRegistry/Public/AssetRegistry/IAssetRegistry.h
+++ b/Engine/Source/Runtime/AssetRegistry/Public/AssetRegistry/IAssetRegistry.h
@@ -24,6 +24,8 @@
 
 class FArchive;
 class FAssetRegistryState;
+class FCbFieldView;
+class FCbWriter;
 class FDependsNode;
 struct FARFilter;
 struct FARCompiledFilter;
@@ -97,6 +99,17 @@ struct FAssetDependency
 	{
 		return AssetId == Other.AssetId && Category == Other.Category && Properties == Other.Properties;
 	}
+
+	ASSETREGISTRY_API void WriteCompactBinary(FCbWriter& Writer) const;
+private:
+	friend FCbWriter& operator<<(FCbWriter& Writer, const FAssetDependency& Dependency)
+	{
+		// Hidden friend function needs to be inline, but call a subfunction to hide the implementation
+		Dependency.WriteCompactBinary(Writer);
+		return Writer;
+	}
+	// Load Cannot be inline because we need to hide implementation and copy-by-value is invalid without definition
+	ASSETREGISTRY_API friend bool LoadFromCompactBinary(FCbFieldView Field, FAssetDependency& Dependency);
 };
 
 UINTERFACE(MinimalApi, BlueprintType, meta = (CannotImplementInterfaceInBlueprint))
diff --git a/Engine/Source/Runtime/CoreUObject/Private/AssetRegistry/AssetData.cpp b/Engine/Source/Runtime/CoreUObject/Private/AssetRegistry/AssetData.cpp
--- a/Engine/Source/Runtime/CoreUObject/Private/AssetRegistry/AssetData.cpp
+++ b/Engine/Source/Runtime/CoreUObject/Private/AssetRegistry/AssetData.cpp
@@ -28,6 +28,53 @@ UE_IMPLEMENT_STRUCT("/Script/CoreUObject", AssetData);
 const FGuid FAssetRegistryVersion::GUID(0x717F9EE7, 0xE9B0493A, 0x88B39132, 0x1B388107);
 FCustomVersionRegistration GRegisterAssetRegistryVersion(FAssetRegistryVersion::GUID, FAssetRegistryVersion::LatestVersion, TEXT("AssetRegistry"));
 
+void FAssetIdentifier::WriteCompactBinary(FCbWriter& Writer) const
+{
+	Writer.BeginArray();
+	FName PrimaryAssetTypeName = (FName)PrimaryAssetType;
+	Writer << PrimaryAssetTypeName;
+	Writer << PackageName;
+	if (!ObjectName.IsNone())
+	{
+		Writer << ObjectName;
+	}
+	Writer.EndArray();
+}
+
+bool LoadFromCompactBinary(FCbFieldView Field, FAssetIdentifier& Identifier)
+{
+	FCbArrayView ArrayView = Field.AsArrayView();
+	if (ArrayView.Num() < 2)
+	{
+		Identifier = FAssetIdentifier();
+		return false;
+	}
+	FCbFieldViewIterator Iter = ArrayView.CreateViewIterator();
+	FName PrimaryAssetTypeName;
+	if (LoadFromCompactBinary(Iter++, PrimaryAssetTypeName))
+	{
+		Identifier.PrimaryAssetType = PrimaryAssetTypeName;
+	}
+	else
+	{
+		Identifier = FAssetIdentifier();
+		return false;
+	}
+	if (!LoadFromCompactBinary(Iter++, Identifier.PackageName))
+	{
+		return false;
+	}
+	if (ArrayView.Num() >= 3)
+	{
+		if (!LoadFromCompactBinary(Iter++, Identifier.ObjectName))
+		{
+			Identifier = FAssetIdentifier();
+			return false;
+		}
+	}
+	return true;
+}
+
 namespace UE::AssetRegistry::Private
 {
 	FAssetPathParts SplitIntoOuterPathAndAssetName(FStringView InObjectPath)
diff --git a/Engine/Source/Runtime/CoreUObject/Public/AssetRegistry/AssetIdentifier.h b/Engine/Source/Runtime/CoreUObject/Public/AssetRegistry/AssetIdentifier.h
--- a/Engine/Source/Runtime/CoreUObject/Public/AssetRegistry/AssetIdentifier.h
+++ b/Engine/Source/Runtime/CoreUObject/Public/AssetRegistry/AssetIdentifier.h
@@ -4,6 +4,9 @@
 
 #include "UObject/PrimaryAssetId.h"
 
+class FCbFieldView;
+class FCbWriter;
+
 /** A structure defining a thing that can be reference by something else in the asset registry. Represents either a package of a primary asset id */
 struct FAssetIdentifier
 {
@@ -203,5 +206,33 @@ struct FAssetIdentifier
 		}
 		return ValueName.LexicalLess(Other.ValueName);
 	}
+
+	bool FastLess(const FAssetIdentifier& Other) const
+	{
+		if (PrimaryAssetType != Other.PrimaryAssetType)
+		{
+			return PrimaryAssetType.FastLess(Other.PrimaryAssetType);
+		}
+		if (PackageName != Other.PackageName)
+		{
+			return PackageName.FastLess(Other.PackageName);
+		}
+		if (ObjectName != Other.ObjectName)
+		{
+			return ObjectName.FastLess(Other.ObjectName);
+		}
+		return ValueName.FastLess(Other.ValueName);
+	}
+
+	COREUOBJECT_API void WriteCompactBinary(FCbWriter& Writer) const;
+private:
+	friend FCbWriter& operator<<(FCbWriter& Writer, const FAssetIdentifier& Identifier)
+	{
+		// Hidden friend function needs to be inline, but call a subfunction to hide the implementation
+		Identifier.WriteCompactBinary(Writer);
+		return Writer;
+	}
+	// Load Cannot be inline because we need to hide implementation and copy-by-value is invalid without definition
+	COREUOBJECT_API friend bool LoadFromCompactBinary(FCbFieldView Field, FAssetIdentifier& Identifier);
 };
 
diff --git a/Engine/Source/Runtime/CoreUObject/Public/UObject/PrimaryAssetId.h b/Engine/Source/Runtime/CoreUObject/Public/UObject/PrimaryAssetId.h
--- a/Engine/Source/Runtime/CoreUObject/Public/UObject/PrimaryAssetId.h
+++ b/Engine/Source/Runtime/CoreUObject/Public/UObject/PrimaryAssetId.h
@@ -87,6 +87,11 @@ struct FPrimaryAssetType
 		return Name.LexicalLess(Other.Name);
 	}
 
+	bool FastLess(const FPrimaryAssetType& Other) const
+	{
+		return Name.FastLess(Other.Name);
+	}
+
 private:
 	friend struct Z_Construct_UScriptStruct_FPrimaryAssetType_Statics;
 
diff --git a/Engine/Source/Runtime/Engine/Private/WorldPartition/WorldPartitionCookPackageSplitter.h b/Engine/Source/Runtime/Engine/Private/WorldPartition/WorldPartitionCookPackageSplitter.h
--- a/Engine/Source/Runtime/Engine/Private/WorldPartition/WorldPartitionCookPackageSplitter.h
+++ b/Engine/Source/Runtime/Engine/Private/WorldPartition/WorldPartitionCookPackageSplitter.h
@@ -14,6 +14,7 @@ class FWorldPartitionCookPackageSplitter : public FGCObject, public ICookPackage
 public:
 	//~ Begin of ICookPackageSplitter
 	static bool ShouldSplit(UObject* SplitData);
+	static FString GetSplitterDebugName() { return TEXT("FWorldPartitionCookPackageSplitter"); }
 
 	FWorldPartitionCookPackageSplitter();
 	virtual ~FWorldPartitionCookPackageSplitter();
