Downgrade innocuous MPCook warning to a display until we have fixed its unexpected occurrence.
#rnx

diff --git a/Engine/Source/Editor/UnrealEd/Private/Cooker/CookWorkerServer.cpp b/Engine/Source/Editor/UnrealEd/Private/Cooker/CookWorkerServer.cpp
--- a/Engine/Source/Editor/UnrealEd/Private/Cooker/CookWorkerServer.cpp
+++ b/Engine/Source/Editor/UnrealEd/Private/Cooker/CookWorkerServer.cpp
@@ -630,7 +630,7 @@ void FCookWorkerServer::RecordResults(FPackageResultsMessage& Message)
 		}
 		if (PendingPackages.Remove(PackageData) != 1)
 		{
-			UE_LOG(LogCook, Warning, TEXT("CookWorkerServer %d received FPackageResultsMessage for package %s which is not a pending package. Ignoring it."),
+			UE_LOG(LogCook, Display, TEXT("CookWorkerServer %d received FPackageResultsMessage for package %s which is not a pending package. Ignoring it."),
 				ProfileId, *Result.GetPackageName().ToString());
 			continue;
 		}
