Prevent physics asset to be imported by using the class rule permission framework.

#rnx

diff --git a/Engine/Source/Editor/UnrealEd/Private/Fbx/FbxMainImport.cpp b/Engine/Source/Editor/UnrealEd/Private/Fbx/FbxMainImport.cpp
--- a/Engine/Source/Editor/UnrealEd/Private/Fbx/FbxMainImport.cpp
+++ b/Engine/Source/Editor/UnrealEd/Private/Fbx/FbxMainImport.cpp
@@ -41,6 +41,7 @@
 #include "AssetToolsModule.h"
 #include "IAssetTools.h"
 #include "Misc/NamePermissionList.h"
+#include "PhysicsEngine/PhysicsAsset.h"
 
 DEFINE_LOG_CATEGORY(LogFbx);
 
@@ -50,10 +51,6 @@ DEFINE_LOG_CATEGORY(LogFbx);
 namespace UnFbx
 {
 
-static bool GDisableAutomaticPhysicsAssetCreation = false;
-static FAutoConsoleVariableRef DisableAutomaticPhysicsAssetCreationCVar(TEXT("FbxImport.DisableAutomaticPhysicsAssetCreation"),
-	GDisableAutomaticPhysicsAssetCreation, TEXT("Prevents physics assets from being created automatically by FBX import (False: disabled, True: enabled"));
-
 TSharedPtr<FFbxImporter> FFbxImporter::StaticInstance;
 
 TSharedPtr<FFbxImporter> FFbxImporter::StaticPreviewInstance;
@@ -132,7 +129,7 @@ FBXImportOptions* GetImportOptions( UnFbx::FFbxImporter* FbxImporter, UFbxImport
 			ImportUI->PhysicsAsset = NULL;
 		}
 
-		if (GDisableAutomaticPhysicsAssetCreation)
+		if (!FbxImporter->CanImportClass(UPhysicsAsset::StaticClass()))
 		{
 			ImportUI->bCreatePhysicsAsset = false;
 		}
