PR #10015: Fix problem with Blueprint Debugging 'StepOver' command stopping in incorrect blueprints (Contributed by [REDACTED_PERSON])

diff --git a/Engine/Source/Editor/UnrealEd/Private/Kismet2/KismetDebugUtilities.cpp b/Engine/Source/Editor/UnrealEd/Private/Kismet2/KismetDebugUtilities.cpp
--- a/Engine/Source/Editor/UnrealEd/Private/Kismet2/KismetDebugUtilities.cpp
+++ b/Engine/Source/Editor/UnrealEd/Private/Kismet2/KismetDebugUtilities.cpp
@@ -624,27 +624,40 @@ void FKismetDebugUtilities::CheckBreakConditions(UEdGraphNode* NodeStoppedAt, bo
 				InOutBreakExecution = Data.TargetGraphStackDepth >= ScriptStack.Num();
 				if(InOutBreakExecution && Data.TargetGraphStackDepth == ScriptStack.Num())
 				{
-					// we're at the same stack depth, don't break if we've entered a different graph, but do break if we left the 
-					// graph that we were trying to step over..
-					const FFrame* CurrentFrame = ScriptStack.Last();
-					if(CurrentFrame->Object)
+					// If we have Data.TargetGraphNodes.Num() > 0, see if we can find a BlueprintNode matching a TargetGraphNode by iterating
+					// up the Blueprint class hierarchy of our CurrentFrame->Object calling FindSourceNodeFromCodeLocation at each level.
+					if (Data.TargetGraphNodes.Num() > 0)
 					{
-						if(UBlueprintGeneratedClass* BPGC = Cast<UBlueprintGeneratedClass>(CurrentFrame->Object->GetClass()))
+						// we're at the same stack depth, don't break if we've entered a different graph, but do break if we left the 
+						// graph that we were trying to step over..
+						const FFrame* CurrentFrame = ScriptStack.Last();
+						if (CurrentFrame->Object)
 						{
-							UEdGraphNode* BlueprintNode = BPGC->DebugData.FindSourceNodeFromCodeLocation(CurrentFrame->Node, BreakpointOffset, true);
-							if(Data.TargetGraphNodes.Num() == 0 || Data.TargetGraphNodes.Contains(BlueprintNode))
+							UClass* BPClass = nullptr;
+							while (true)
 							{
-								InOutBreakExecution = true;
-							}
-							else
-							{
-								InOutBreakExecution = false; // nowhere to stop
+								BPClass = (BPClass == nullptr) ? CurrentFrame->Object->GetClass() : BPClass->GetSuperClass();
+								if (BPClass == nullptr)
+								{
+									InOutBreakExecution = false;
+									break;
+								}
+
+								const UBlueprintGeneratedClass* BPGC = Cast<UBlueprintGeneratedClass>(BPClass);
+								if (BPGC == nullptr)
+								{
+									InOutBreakExecution = false;
+									break;
+								}
+
+								const UEdGraphNode* BlueprintNode = BPGC->DebugData.FindSourceNodeFromCodeLocation(CurrentFrame->Node, BreakpointOffset, true);
+
+								if (Data.TargetGraphNodes.Contains(BlueprintNode))
+								{
+									break;
+								}
 							}
 						}
-						else
-						{
-							InOutBreakExecution = false;
-						}
 					}
 				}
 			}
