DerivedDataCache: Added check for redirected packages that are known to be removed. Previously warning *Failed to find soft reference* was displayed for packages that were added to ini files with PackageRedirects(removed=true).

diff --git a/Engine/Source/Editor/UnrealEd/Private/Commandlets/DerivedDataCacheCommandlet.cpp b/Engine/Source/Editor/UnrealEd/Private/Commandlets/DerivedDataCacheCommandlet.cpp
--- a/Engine/Source/Editor/UnrealEd/Private/Commandlets/DerivedDataCacheCommandlet.cpp
+++ b/Engine/Source/Editor/UnrealEd/Private/Commandlets/DerivedDataCacheCommandlet.cpp
@@ -30,13 +30,13 @@ DerivedDataCacheCommandlet.cpp: Commandlet for DDC maintenence
 #include "PackageHelperFunctions.h"
 #include "Settings/ProjectPackagingSettings.h"
 #include "ShaderCompiler.h"
+#include "UObject/CoreRedirects.h"
 #include "UObject/Package.h"
 #include "UObject/UObjectHash.h"
 #include "UObject/UObjectIterator.h"
 #include "WorldPartition/WorldPartition.h"
 #include "WorldPartition/WorldPartitionHelpers.h"
 #include "WorldPartition/WorldPartitionSubsystem.h"
-
 DEFINE_LOG_CATEGORY_STATIC(LogDerivedDataCacheCommandlet, Log, All);
 
 class UDerivedDataCacheCommandlet::FObjectReferencer : public FGCObject
@@ -844,6 +844,15 @@ int32 UDerivedDataCacheCommandlet::Main( const FString& Params )
 					{
 						PackagesToProcess.Add(SoftRefName);
 						FString SoftRefFilename;
+
+						FCoreRedirectObjectName CoreRedirectName;
+						CoreRedirectName.PackageName = SoftRefName;
+						// Packages that are specifically identified by PackageRedirects as removed need to be skipped.
+						if (FCoreRedirects::IsKnownMissing(ECoreRedirectFlags::Type_Package, CoreRedirectName))
+						{
+							continue;
+						}
+
 						if (FPackageName::DoesPackageExist(SoftRefName.ToString(), &SoftRefFilename))
 						{
 							UE_LOG(LogDerivedDataCacheCommandlet, Log, TEXT("Queueing soft reference '%s' for later processing"), *SoftRefName.ToString());
