Backout CL 24196898

#rnx

diff --git a/Engine/Source/Editor/UnrealEd/Private/LevelEditorViewport.cpp b/Engine/Source/Editor/UnrealEd/Private/LevelEditorViewport.cpp
--- a/Engine/Source/Editor/UnrealEd/Private/LevelEditorViewport.cpp
+++ b/Engine/Source/Editor/UnrealEd/Private/LevelEditorViewport.cpp
@@ -1801,18 +1801,6 @@ void FTrackingTransaction::Cancel()
 	{
 		ScopedTransaction->Cancel();
 	}
-
-	// If the transaction is cancelled, reset the package dirty states to reflect their original state
-	for (const TTuple<UPackage*, bool>& InitialPackageDirtyState : InitialPackageDirtyStates)
-	{
-		UPackage* Package = InitialPackageDirtyState.Key;
-		const bool bInitialDirtyState = InitialPackageDirtyState.Value;
-		if (!bInitialDirtyState && Package->IsDirty())
-		{
-			Package->SetDirtyFlag(false);
-		}
-	}
-
 	End();
 }
 
@@ -2179,7 +2167,6 @@ void FLevelEditorViewportClient::BeginCameraMovement(bool bHasMovement)
 					FCoreUObjectDelegates::OnPreObjectPropertyChanged.Broadcast(ActorLock, PropertyChain);
 				}
 			}
-			TrackingTransaction.Cancel();
 			bIsCameraMoving = true;
 		}
 	}
