Asset Thumbnails: Wait for rendering commands to finish before drawing the actual thumbnail

diff --git a/Engine/Source/Editor/UnrealEd/Private/AssetThumbnail.cpp b/Engine/Source/Editor/UnrealEd/Private/AssetThumbnail.cpp
--- a/Engine/Source/Editor/UnrealEd/Private/AssetThumbnail.cpp
+++ b/Engine/Source/Editor/UnrealEd/Private/AssetThumbnail.cpp
@@ -1209,6 +1209,11 @@ bool FAssetThumbnailPool::LoadThumbnail(TSharedRef<FThumbnailInfo> ThumbnailInfo
 							}
 						});
 
+					// // We have to wait for the render command to finish since thumbnail rendering cannot be done on the GPU currently
+					FRenderCommandFence Fence;
+					Fence.BeginFence();
+					Fence.Wait();
+
 					//@todo: this should be done on the GPU only but it is not supported by thumbnail tools yet
 					ThumbnailTools::RenderThumbnail(
 						Asset,
