Add thumbnail for round box brush
Add thumbnail for slate brushes with UI material

diff --git a/Engine/Source/Editor/UnrealEd/Classes/ThumbnailRendering/SlateBrushThumbnailRenderer.h b/Engine/Source/Editor/UnrealEd/Classes/ThumbnailRendering/SlateBrushThumbnailRenderer.h
--- a/Engine/Source/Editor/UnrealEd/Classes/ThumbnailRendering/SlateBrushThumbnailRenderer.h
+++ b/Engine/Source/Editor/UnrealEd/Classes/ThumbnailRendering/SlateBrushThumbnailRenderer.h
@@ -24,7 +24,13 @@ public:
 
 	// Begin UThumbnailRenderer Object
 	UNREALED_API virtual void Draw(UObject* Object, int32 X, int32 Y, uint32 Width, uint32 Height, FRenderTarget*, FCanvas* Canvas, bool bAdditionalViewFamily) override;
+	virtual EThumbnailRenderFrequency GetThumbnailRenderFrequency(UObject* Object) const override;
 	// End UThumbnailRenderer Object
 
+	void CreateThumbnailAsImage(uint32 Width, uint32 Height, FRenderTarget* RenderTarget, FSlateBrush& Brush);
+	void CreateTextureThumbnailOnCanvas(int32 X, int32 Y, uint32 Width, uint32 Height, FRenderTarget* RenderTarget, FSlateBrush& Brush, FCanvas* Canvas, UTexture2D* Texture);
+
+private:
+	bool bIsLastFrequencyRealTime;
 };
 
diff --git a/Engine/Source/Editor/UnrealEd/Private/ThumbnailRendering/SlateBrushThumbnailRenderer.cpp b/Engine/Source/Editor/UnrealEd/Private/ThumbnailRendering/SlateBrushThumbnailRenderer.cpp
--- a/Engine/Source/Editor/UnrealEd/Private/ThumbnailRendering/SlateBrushThumbnailRenderer.cpp
+++ b/Engine/Source/Editor/UnrealEd/Private/ThumbnailRendering/SlateBrushThumbnailRenderer.cpp
@@ -11,9 +11,15 @@
 // FPreviewScene derived helpers for rendering
 #include "CanvasItem.h"
 #include "CanvasTypes.h"
+// helpers for rendering RoundedBox thumbnail as an image widget
+#include "Slate/WidgetRenderer.h"
+#include "Input/HittestGrid.h"
+#include "Widgets/Images/SImage.h"
+#include "Widgets/SVirtualWindow.h"
 
 USlateBrushThumbnailRenderer::USlateBrushThumbnailRenderer(const FObjectInitializer& ObjectInitializer)
 	: Super(ObjectInitializer)
+	, bIsLastFrequencyRealTime(false)
 {
 }
 
@@ -24,35 +30,35 @@ void USlateBrushThumbnailRenderer::Draw(UObject* Object, int32 X, int32 Y, uint3
 	{
 		FSlateBrush Brush = SlateBrushAsset->Brush;
 		UTexture2D* Texture = Cast<UTexture2D>( Brush.GetResourceObject() );
+		UMaterialInterface* Material = Cast<UMaterialInterface>(Brush.GetResourceObject());
 
-		// Draw the background checkboard pattern
-		const int32 CheckerDensity = 8;
-		auto Checker = ToRawPtr(UThumbnailManager::Get().CheckerboardTexture);
-		Canvas->DrawTile(
-			0.0f, 0.0f, Width, Height,							// Dimensions
-			0.0f, 0.0f, CheckerDensity, CheckerDensity,			// UVs
-			FLinearColor::White, Checker->GetResource());			// Tint & Texture
+		if (bIsLastFrequencyRealTime)
+		{
+			bIsLastFrequencyRealTime = false;
+		}
 
 		if (Texture)
 		{
 			switch(Brush.DrawAs)
 			{
 			case ESlateBrushDrawType::Image:
-				{
-					FCanvasTileItem CanvasTile( FVector2D( X, Y ), Texture->GetResource(), FVector2D( Width,Height ), Brush.TintColor.GetSpecifiedColor() );
-					CanvasTile.BlendMode = SE_BLEND_Translucent;
-					CanvasTile.Draw( Canvas );
-				}
-				break;
 			case ESlateBrushDrawType::Border:
+			case ESlateBrushDrawType::NoDrawType:
 				{
-					FCanvasTileItem CanvasTile( FVector2D( X, Y ), Texture->GetResource(), FVector2D( Width,Height ), Brush.TintColor.GetSpecifiedColor() );
-					CanvasTile.BlendMode = SE_BLEND_Translucent;
-					CanvasTile.Draw( Canvas );
+					CreateTextureThumbnailOnCanvas(X, Y, Width, Height, RenderTarget, Brush, Canvas, Texture);
 				}
 				break;
 			case ESlateBrushDrawType::Box:
 				{
+					// Draw the checkerboard background
+					const int32 CheckerDensity = 8;
+					auto Checker = ToRawPtr(UThumbnailManager::Get().CheckerboardTexture);
+					Canvas->DrawTile(
+						0.0f, 0.0f, Width, Height,							// Dimensions
+						0.0f, 0.0f, CheckerDensity, CheckerDensity,			// UVs
+						FLinearColor::White, Checker->GetResource(),		// Tint & Texture
+						true);
+					
 					float NaturalWidth = Texture->GetSurfaceWidth();
 					float NaturalHeight = Texture->GetSurfaceHeight();
 
@@ -169,17 +175,125 @@ void USlateBrushThumbnailRenderer::Draw(UObject* Object, int32 X, int32 Y, uint3
 					}
 				}
 				break;
-			case ESlateBrushDrawType::NoDrawType:
+			case ESlateBrushDrawType::RoundedBox:
 				{
-					FCanvasTileItem CanvasTile( FVector2D( X, Y ), Texture->GetResource(), FVector2D( Width,Height ), Brush.TintColor.GetSpecifiedColor() );
-					CanvasTile.BlendMode = SE_BLEND_Translucent;
-					CanvasTile.Draw( Canvas );
+					CreateThumbnailAsImage(Width, Height, RenderTarget, Brush);
 				}
 				break;
 			default:
 
+				check(false);
+			}
+		}
+		else if (Material && Material->GetMaterial() && Material->GetMaterial()->IsUIMaterial())
+		{
+			switch (Brush.DrawAs)
+			{
+			case ESlateBrushDrawType::Image:
+			case ESlateBrushDrawType::Border:
+			case ESlateBrushDrawType::Box:
+			case ESlateBrushDrawType::NoDrawType:
+			case ESlateBrushDrawType::RoundedBox:
+			{
+				bIsLastFrequencyRealTime = true;
+				CreateThumbnailAsImage(Width, Height, RenderTarget, Brush);
+			}
+			break;
+			default:
+
 				check(false);
 			}
 		}
 	}
+	
+}
+
+EThumbnailRenderFrequency USlateBrushThumbnailRenderer::GetThumbnailRenderFrequency(UObject* Object) const
+{ 
+	if (USlateBrushAsset* SlateBrushAsset = Cast<USlateBrushAsset>(Object))
+	{
+		FSlateBrush Brush = SlateBrushAsset->Brush;
+		if (UMaterialInterface* Material = Cast<UMaterialInterface>(Brush.GetResourceObject()))
+		{
+			if (Material->GetMaterial() && Material->GetMaterial()->IsUIMaterial())
+			{
+				return EThumbnailRenderFrequency::Realtime;
+			}
+		}
+	}
+
+	// If we are switching from Realtime to OnAssetSave, we need to cache the realtime thumbnail first to use it until the next time we save the asset.
+	// So we set the frequency to Once for the very first thumbnail rendering call after we have switched from Realtime.
+	return bIsLastFrequencyRealTime ? EThumbnailRenderFrequency::Once : EThumbnailRenderFrequency::OnAssetSave;
+
+}
+void USlateBrushThumbnailRenderer::CreateThumbnailAsImage(uint32 Width, uint32 Height, FRenderTarget* RenderTarget, FSlateBrush& Brush)
+{
+	const bool bUseGammaCorrection = true;
+	FWidgetRenderer* WidgetRenderer = new FWidgetRenderer(bUseGammaCorrection, true);
+
+	UTexture2D* CheckerboardTexture = UThumbnailManager::Get().CheckerboardTexture;
+	FSlateBrush CheckerboardBrush;
+	CheckerboardBrush.SetResourceObject(CheckerboardTexture);
+	CheckerboardBrush.ImageSize = FVector2D(CheckerboardTexture->GetSizeX(), CheckerboardTexture->GetSizeY());
+	CheckerboardBrush.Tiling = ESlateBrushTileType::Both;
+
+	const FVector2D DrawSize((float)Width, (float)Height);
+	const float DeltaTime = 0.f;
+
+	if (Brush.GetResourceObject())
+	{
+		TSharedRef<SWidget> Thumbnail =
+			SNew(SOverlay)
+
+			// Checkerboard
+			+ SOverlay::Slot()
+			[
+				SNew(SImage)
+				.Image(&CheckerboardBrush)
+			]
+
+		+ SOverlay::Slot()
+			[
+				SNew(SImage)
+				.Image(&Brush)
+			];
+
+		WidgetRenderer->DrawWidget(RenderTarget, Thumbnail, DrawSize, DeltaTime);
+	}
+	else
+	{
+		TSharedRef<SWidget> Thumbnail =
+			SNew(SOverlay)
+
+			// Checkerboard
+			+ SOverlay::Slot()
+			[
+				SNew(SImage)
+				.Image(&CheckerboardBrush)
+			];
+
+		WidgetRenderer->DrawWidget(RenderTarget, Thumbnail, DrawSize, DeltaTime);
+	}
+
+	if (WidgetRenderer)
+	{
+		BeginCleanup(WidgetRenderer);
+		WidgetRenderer = nullptr;
+	}
+}
+
+void USlateBrushThumbnailRenderer::CreateTextureThumbnailOnCanvas(int32 X, int32 Y, uint32 Width, uint32 Height, FRenderTarget* RenderTarget, FSlateBrush& Brush, FCanvas* Canvas, UTexture2D* Texture)
+{
+	// Draw the background checkboard pattern
+	const int32 CheckerDensity = 8;
+	auto Checker = ToRawPtr(UThumbnailManager::Get().CheckerboardTexture);
+	Canvas->DrawTile(
+		0.0f, 0.0f, Width, Height,							// Dimensions
+		0.0f, 0.0f, CheckerDensity, CheckerDensity,			// UVs
+		FLinearColor::White, Checker->GetResource(),		// Tint & Texture
+		true);												// Alpha Blend
+	FCanvasTileItem CanvasTile(FVector2D(X, Y), Texture->GetResource(), FVector2D(Width, Height), Brush.TintColor.GetSpecifiedColor());
+	CanvasTile.BlendMode = SE_BLEND_Translucent;
+	CanvasTile.Draw(Canvas);
 }
