Non-unity error fix in slate brush thumbnail renderer.

diff --git a/Engine/Source/Editor/UnrealEd/Classes/ThumbnailRendering/SlateBrushThumbnailRenderer.h b/Engine/Source/Editor/UnrealEd/Classes/ThumbnailRendering/SlateBrushThumbnailRenderer.h
--- a/Engine/Source/Editor/UnrealEd/Classes/ThumbnailRendering/SlateBrushThumbnailRenderer.h
+++ b/Engine/Source/Editor/UnrealEd/Classes/ThumbnailRendering/SlateBrushThumbnailRenderer.h
@@ -14,6 +14,7 @@
 
 class FCanvas;
 class FRenderTarget;
+struct FSlateBrush;
 
 UCLASS(config=Editor,MinimalAPI)
 class USlateBrushThumbnailRenderer : public UDefaultSizedThumbnailRenderer
diff --git a/Engine/Source/Editor/UnrealEd/Private/ThumbnailRendering/SlateBrushThumbnailRenderer.cpp b/Engine/Source/Editor/UnrealEd/Private/ThumbnailRendering/SlateBrushThumbnailRenderer.cpp
--- a/Engine/Source/Editor/UnrealEd/Private/ThumbnailRendering/SlateBrushThumbnailRenderer.cpp
+++ b/Engine/Source/Editor/UnrealEd/Private/ThumbnailRendering/SlateBrushThumbnailRenderer.cpp
@@ -8,6 +8,8 @@
 #include "ThumbnailRendering/ThumbnailManager.h"
 #include "Slate/SlateBrushAsset.h"
 #include "TextureResource.h"
+#include "Materials/MaterialInterface.h"
+#include "Materials/Material.h"
 // FPreviewScene derived helpers for rendering
 #include "CanvasItem.h"
 #include "CanvasTypes.h"
@@ -185,23 +187,27 @@ void USlateBrushThumbnailRenderer::Draw(UObject* Object, int32 X, int32 Y, uint3
 				check(false);
 			}
 		}
-		else if (Material && Material->GetMaterial() && Material->GetMaterial()->IsUIMaterial())
+		else if (Material)
 		{
-			switch (Brush.DrawAs)
+			UMaterial* Mat = Material->GetMaterial();
+			if (Mat && Mat->IsUIMaterial())
 			{
-			case ESlateBrushDrawType::Image:
-			case ESlateBrushDrawType::Border:
-			case ESlateBrushDrawType::Box:
-			case ESlateBrushDrawType::NoDrawType:
-			case ESlateBrushDrawType::RoundedBox:
-			{
-				bIsLastFrequencyRealTime = true;
-				CreateThumbnailAsImage(Width, Height, RenderTarget, Brush);
-			}
-			break;
-			default:
+				switch (Brush.DrawAs)
+				{
+				case ESlateBrushDrawType::Image:
+				case ESlateBrushDrawType::Border:
+				case ESlateBrushDrawType::Box:
+				case ESlateBrushDrawType::NoDrawType:
+				case ESlateBrushDrawType::RoundedBox:
+				{
+					bIsLastFrequencyRealTime = true;
+					CreateThumbnailAsImage(Width, Height, RenderTarget, Brush);
+				}
+				break;
+				default:
 
-				check(false);
+					check(false);
+				}
 			}
 		}
 	}
@@ -215,7 +221,8 @@ EThumbnailRenderFrequency USlateBrushThumbnailRenderer::GetThumbnailRenderFreque
 		FSlateBrush Brush = SlateBrushAsset->Brush;
 		if (UMaterialInterface* Material = Cast<UMaterialInterface>(Brush.GetResourceObject()))
 		{
-			if (Material->GetMaterial() && Material->GetMaterial()->IsUIMaterial())
+			UMaterial* Mat = Material->GetMaterial();
+			if (Mat && Mat->IsUIMaterial())
 			{
 				return EThumbnailRenderFrequency::Realtime;
 			}
