Fix Non-unity error in USlateBrushThumbnailRenderer

diff --git a/Engine/Source/Editor/UnrealEd/Classes/ThumbnailRendering/SlateBrushThumbnailRenderer.h b/Engine/Source/Editor/UnrealEd/Classes/ThumbnailRendering/SlateBrushThumbnailRenderer.h
--- a/Engine/Source/Editor/UnrealEd/Classes/ThumbnailRendering/SlateBrushThumbnailRenderer.h
+++ b/Engine/Source/Editor/UnrealEd/Classes/ThumbnailRendering/SlateBrushThumbnailRenderer.h
@@ -14,6 +14,7 @@
 
 class FCanvas;
 class FRenderTarget;
+class UTexture2D;
 struct FSlateBrush;
 
 UCLASS(config=Editor,MinimalAPI)
