Good code reduction, no real reason not to unify them.

PR #9565: Unify Copy/Move name collisions with Duplicate (Contributed by [REDACTED_PERSON])

diff --git a/Engine/Source/Editor/UnrealEd/Private/ObjectTools.cpp b/Engine/Source/Editor/UnrealEd/Private/ObjectTools.cpp
--- a/Engine/Source/Editor/UnrealEd/Private/ObjectTools.cpp
+++ b/Engine/Source/Editor/UnrealEd/Private/ObjectTools.cpp
@@ -3698,77 +3698,8 @@ namespace ObjectTools
 		// If the target package already exists, check for name clashes and find a unique name
 		if ( bUniqueDefaultName )
 		{
-			UPackage* NewPackage = FindPackage(NULL, *PackageName);
-
-			if ( NewPackage )
-			{
-				NewPackage->FullyLoad();
-			}
-			else
-			{
-				FString PackageFilename;
-				if ( FPackageName::DoesPackageExist(PackageName, &PackageFilename) )
-				{
-					NewPackage = LoadPackage(NULL, *PackageFilename, LOAD_None);
-				}
-			}
-
-			if (NewPackage)
-			{
-				FString ObjectPrefix = ObjectName;
-				int32 Suffix = 2;
-
-				// Check if this is already a copied object name and increment it if it is
-				FString LeftSplit;
-				FString RightSplit;
-				if( ObjectName.Split( "_", &LeftSplit, &RightSplit, ESearchCase::CaseSensitive, ESearchDir::FromEnd ) == true )
-				{
-					bool bOnlyNumeric = true;
-					for( int index = 0; index < RightSplit.Len(); index++ )
-					{
-						if( FChar::IsDigit(RightSplit[index] ) == false )
-						{
-							bOnlyNumeric = false;
-							break;
-						}
-					}
-					if( bOnlyNumeric == true )
-					{
-						Suffix = FCString::Atoi(*RightSplit) + 1;
-						ObjectPrefix = LeftSplit;
-					}
-				}
-
-				// If the package and object names were equal before, ensure that the generated names are also equal
-				const FString PackageShortName = FPackageName::GetLongPackageAssetName(*PackageName);
-				const FString PackagePath = FPackageName::GetLongPackagePath(*PackageName);
-				FString PackagePrefix = (ObjectName == PackageShortName) ? (PackagePath / ObjectPrefix) : PackageName;
-
-				for (; NewPackage && StaticFindObjectFast(NULL, NewPackage, FName(*ObjectName)); Suffix++)
-				{
-					// DlgName exists in DlgPackage - generate a new one with a numbered suffix
-					ObjectName = FString::Printf(TEXT("%s_%d"), *ObjectPrefix, Suffix);
-
-					// Don't change the package name if we encounter an object name clash when moving to a legacy package
-					{
-						PackageName = FString::Printf(TEXT("%s_%d"), *PackagePrefix, Suffix);
-						NewPackage = FindPackage(NULL, *PackageName);
-
-						if ( NewPackage )
-						{
-							NewPackage->FullyLoad();
-						}
-						else
-						{
-							FString PackageFilename;
-							if ( FPackageName::DoesPackageExist(PackageName, &PackageFilename) )
-							{
-								NewPackage = LoadPackage(NULL, *PackageFilename, LOAD_None);
-							}
-						}
-					}
-				}
-			}
+			FAssetToolsModule& AssetToolsModule = FModuleManager::LoadModuleChecked<FAssetToolsModule>("AssetTools");
+			AssetToolsModule.Get().CreateUniqueAssetName(*PackageName, TEXT(""), PackageName, ObjectName);
 		}
 
 		if( !InOutInfo.bOkToAll && InOutInfo.bPromptForRenameOnConflict )
