HLOD builder - Fixed formatted commandlet output that is messing with the output progress display in editor

diff --git a/Engine/Source/Editor/UnrealEd/Private/WorldPartition/WorldPartitionHLODsBuilder.cpp b/Engine/Source/Editor/UnrealEd/Private/WorldPartition/WorldPartitionHLODsBuilder.cpp
--- a/Engine/Source/Editor/UnrealEd/Private/WorldPartition/WorldPartitionHLODsBuilder.cpp
+++ b/Engine/Source/Editor/UnrealEd/Private/WorldPartition/WorldPartitionHLODsBuilder.cpp
@@ -491,7 +491,7 @@ bool UWorldPartitionHLODsBuilder::BuildHLODActors()
 		UPackage* ActorPackage = HLODActor->GetPackage();
 		if (ActorPackage->IsDirty())
 		{
-			UE_LOG(LogWorldPartitionHLODsBuilder, Display, TEXT("[%d / %d] HLOD actor %s was modified, saving..."), CurrentActor + 1, HLODActorsToBuild.Num(), *HLODActor->GetActorLabel());
+			UE_LOG(LogWorldPartitionHLODsBuilder, Display, TEXT("HLOD actor %s was modified, saving..."), *HLODActor->GetActorLabel());
 
 			bool bSaved = SourceControlHelper->Save(ActorPackage);
 			if (!bSaved)
