Added editor performance setting to control "r.VSyncEditor"

diff --git a/Engine/Source/Editor/UnrealEd/Classes/Editor/EditorPerformanceSettings.h b/Engine/Source/Editor/UnrealEd/Classes/Editor/EditorPerformanceSettings.h
--- a/Engine/Source/Editor/UnrealEd/Classes/Editor/EditorPerformanceSettings.h
+++ b/Engine/Source/Editor/UnrealEd/Classes/Editor/EditorPerformanceSettings.h
@@ -57,6 +57,10 @@ class UEditorPerformanceSettings : public UDeveloperSettings
 	UPROPERTY(EditAnywhere, config, Category = EditorPerformance, meta = (DisplayName = "Enable Scalability Warning Indicator"))
 	uint32 bEnableScalabilityWarningIndicator : 1;
 	
+	/** Should VSync be enabled in editor? */
+	UPROPERTY(EditAnywhere, config, Category=EditorPerformance, meta=(DisplayName="Enable VSync", ConsoleVariable="r.VSyncEditor"))
+	uint32 bEnableVSync : 1;
+
 	/** 
 	 * By default the editor will adjust scene scaling (quality) for high DPI in order to ensure consistent performance with very large render targets.
 	 * Enabling this will disable automatic adjusting and render at the full resolution of the viewport
diff --git a/Engine/Source/Editor/UnrealEd/Private/EditorPerformanceSettings.cpp b/Engine/Source/Editor/UnrealEd/Private/EditorPerformanceSettings.cpp
--- a/Engine/Source/Editor/UnrealEd/Private/EditorPerformanceSettings.cpp
+++ b/Engine/Source/Editor/UnrealEd/Private/EditorPerformanceSettings.cpp
@@ -42,6 +42,11 @@ void UEditorPerformanceSettings::PostInitProperties()
 	CVarEditorViewportHighDPI->Set(bDisplayHighDPIViewports != 0, ECVF_SetByProjectSetting);
 	CVarEditorViewportOverrideGameScreenPercentage->Set(bOverridePIEScreenPercentage != 0, ECVF_SetByProjectSetting);
 
+	if (FProperty* EnableVSyncProperty = GetClass()->FindPropertyByName(GET_MEMBER_NAME_CHECKED(UEditorPerformanceSettings, bEnableVSync)))
+	{
+		ExportValuesToConsoleVariables(EnableVSyncProperty);
+	}
+
 	UEditorPerformanceProjectSettings::ExportResolutionValuesToConsoleVariables();
 }
 
