Fix for editables arrays of verse sub objects being broken after reinstancing. Because the array is empty by default any elements added through the editor won't be considered DefaultSubObjects so they will be missed by some portions of the reinstancing code. So after reinstancing the elements copied to the new parent object will still use the old VERSE_DEAD class of the subobject. This change ensures non DefaultSubObjects aren't skipped.

diff --git a/Engine/Source/Editor/UnrealEd/Private/Kismet2/KismetReinstanceUtilities.cpp b/Engine/Source/Editor/UnrealEd/Private/Kismet2/KismetReinstanceUtilities.cpp
--- a/Engine/Source/Editor/UnrealEd/Private/Kismet2/KismetReinstanceUtilities.cpp
+++ b/Engine/Source/Editor/UnrealEd/Private/Kismet2/KismetReinstanceUtilities.cpp
@@ -2301,9 +2301,9 @@ void FBlueprintCompileReinstancer::ReplaceInstancesOfClass_Inner(const TMap<UCla
 				{
 					UObject* OldObject = ObjectsToReplace[OldObjIndex];
 
-					// Skipping any object that outer is going to be replaced
+					// Skipping any default sub object that outer is going to be replaced
 					// This isn't needed for the actor loop as the only outer for actor is a level
-					if (InOldToNewClassMap.Contains(OldObject->GetOuter()->GetClass()))
+					if ((OldObject->IsDefaultSubobject() || OldObject->HasAnyFlags(RF_DefaultSubObject)) && InOldToNewClassMap.Contains(OldObject->GetOuter()->GetClass()))
 					{
 						continue;
 					}
