Allow for pre-booting the cooker and injecting one or more plugins to be cooked after the pre-boot has happened.

diff --git a/Engine/Source/Developer/DesktopPlatform/Internal/LockFile.h b/Engine/Source/Developer/DesktopPlatform/Internal/LockFile.h
new file mode 100644
--- /dev/null
+++ b/Engine/Source/Developer/DesktopPlatform/Internal/LockFile.h
@@ -0,0 +1,11 @@
+// Copyright Epic Games, Inc. All Rights Reserved.
+
+#pragma once
+
+#include "Containers/UnrealString.h"
+#include "HAL/Platform.h"
+
+struct DESKTOPPLATFORM_API FLockFile
+{
+	static bool TryReadAndClear(const TCHAR* FileName, FString& OutContents);
+};
diff --git a/Engine/Source/Developer/DesktopPlatform/Private/LockFile.cpp b/Engine/Source/Developer/DesktopPlatform/Private/LockFile.cpp
new file mode 100644
--- /dev/null
+++ b/Engine/Source/Developer/DesktopPlatform/Private/LockFile.cpp
@@ -0,0 +1,110 @@
+// Copyright Epic Games, Inc. All Rights Reserved.
+
+#include "LockFile.h"
+#include "Containers/ArrayView.h"
+#include "Containers/StringFwd.h"
+#include "Memory/SharedBuffer.h"
+#include "Misc/FileHelper.h"
+#include "Misc/PathViews.h"
+#include "Misc/ScopeExit.h"
+#include "Misc/StringBuilder.h"
+
+
+#if PLATFORM_WINDOWS
+#	include "Windows/AllowWindowsPlatformTypes.h"
+#	include <shellapi.h>
+#	include <synchapi.h>
+#	include <Windows.h>
+#	include "Windows/HideWindowsPlatformTypes.h"
+#endif
+
+#if PLATFORM_UNIX || PLATFORM_MAC
+#	include <sys/file.h>
+#	include <sys/sem.h>
+#endif
+
+bool
+FLockFile::TryReadAndClear(const TCHAR* FileName, FString& OutContents)
+{
+	OutContents.Reset();
+#if PLATFORM_WINDOWS
+	uint32 Access = GENERIC_READ | DELETE;
+	uint32 WinFlags = 0; // Exclusive access
+	uint32 Disposition = OPEN_EXISTING;
+
+	TStringBuilder<MAX_PATH> FullFileNameBuilder;
+	FPathViews::ToAbsolutePath(FStringView(FileName), FullFileNameBuilder);
+	for (TCHAR& Char : MakeArrayView(FullFileNameBuilder))
+	{
+		if (Char == TEXT('/'))
+		{
+			Char = TEXT('\\');
+		}
+	}
+	if (FullFileNameBuilder.Len() >= MAX_PATH)
+	{
+		FullFileNameBuilder.Prepend(TEXTVIEW("\\\\?\\"));
+	}
+	HANDLE Handle = CreateFileW(FullFileNameBuilder.ToString(), Access, WinFlags, NULL, Disposition, FILE_ATTRIBUTE_NORMAL, NULL);
+	if (Handle != INVALID_HANDLE_VALUE)
+	{
+		ON_SCOPE_EXIT { CloseHandle(Handle); };
+		LARGE_INTEGER LI;
+		if (GetFileSizeEx(Handle, &LI))
+		{
+			checkf(LI.QuadPart == LI.u.LowPart, TEXT("File exceeds supported 2GB limit."));
+			int32 FileSize32 = LI.u.LowPart;
+			FUniqueBuffer FileBytes = FUniqueBuffer::Alloc(FileSize32);
+			DWORD ReadBytes = 0;
+			if (ReadFile(Handle, FileBytes.GetData(), FileSize32, &ReadBytes, NULL) && (ReadBytes == FileSize32))
+			{
+				FFileHelper::BufferToString(OutContents, reinterpret_cast<const uint8*>(FileBytes.GetData()), FileSize32);
+				// Set the 'DeleteFile' disposition on the handle to delete it on close
+				FILE_DISPOSITION_INFO DispositionInfo { 1 };
+				SetFileInformationByHandle(Handle, FileDispositionInfo, &DispositionInfo, sizeof(DispositionInfo));
+				return true;
+			}
+		}
+	}
+	return false;
+#elif PLATFORM_UNIX || PLATFORM_MAC
+	TAnsiStringBuilder<256> FilePath;
+	FilePath << FileName;
+	int32 Fd = open(FilePath.ToString(), O_RDONLY);
+	if (Fd < 0)
+	{
+		return false;
+	}
+
+	int32 LockRet = flock(Fd, LOCK_EX | LOCK_NB);
+	if (LockRet < 0)
+	{
+		close(Fd);
+		return false;
+	}
+
+	ON_SCOPE_EXIT
+	{
+		flock(Fd, LOCK_UN);
+		close(Fd);
+	};
+
+	struct stat Stat;
+	fstat(Fd, &Stat);
+	uint64 FileSize = uint64(Stat.st_size);
+
+	bool bSuccess = false;
+	FUniqueBuffer FileBytes = FUniqueBuffer::Alloc(FileSize);
+	if (read(Fd, FileBytes.GetData(), FileSize) == FileSize)
+	{
+		FFileHelper::BufferToString(OutContents, reinterpret_cast<const uint8*>(FileBytes.GetData()), FileSize);
+	}
+
+	unlink(FilePath.ToString());
+
+	return true;
+#else
+	#error Not implemented for this platform
+	return false;
+#endif
+}
\ No newline at end of file
diff --git a/Engine/Source/Editor/UnrealEd/Classes/CookOnTheSide/CookOnTheFlyServer.h b/Engine/Source/Editor/UnrealEd/Classes/CookOnTheSide/CookOnTheFlyServer.h
--- a/Engine/Source/Editor/UnrealEd/Classes/CookOnTheSide/CookOnTheFlyServer.h
+++ b/Engine/Source/Editor/UnrealEd/Classes/CookOnTheSide/CookOnTheFlyServer.h
@@ -11,6 +11,7 @@
 #include "INetworkFileSystemModule.h"
 #include "IPlatformFileSandboxWrapper.h"
 #include "Misc/EnumClassFlags.h"
+#include "Misc/Optional.h"
 #include "Templates/Function.h"
 #include "Templates/RefCounting.h"
 #include "Templates/UniquePtr.h"
@@ -899,6 +900,9 @@ private:
 
 	//////////////////////////////////////////////////////////////////////////
 	// cook by the book specific functions
+	const FCookByTheBookStartupOptions& BlockOnPrebootCookGate(bool& bOutAbortCook,
+		const FCookByTheBookStartupOptions& CookByTheBookStartupOptions,
+		TOptional<FCookByTheBookStartupOptions>& ModifiedStartupOptions);
 	/** Construct the on-stack data for StartCook functions, based on the arguments to StartCookByTheBook */
 	FBeginCookContext CreateBeginCookByTheBookContext(const FCookByTheBookStartupOptions& StartupOptions);
 	/** Set the PlatformManager's session platforms and finish filling out the BeginContext for StartCookByTheBook */
diff --git a/Engine/Source/Editor/UnrealEd/Private/CookOnTheFlyServer.cpp b/Engine/Source/Editor/UnrealEd/Private/CookOnTheFlyServer.cpp
--- a/Engine/Source/Editor/UnrealEd/Private/CookOnTheFlyServer.cpp
+++ b/Engine/Source/Editor/UnrealEd/Private/CookOnTheFlyServer.cpp
@@ -78,6 +78,7 @@
 #include "Internationalization/PackageLocalizationManager.h"
 #include "IPAddress.h"
 #include "LocalizationChunkDataGenerator.h"
+#include "LockFile.h"
 #include "Logging/MessageLog.h"
 #include "Logging/TokenizedMessage.h"
 #include "Materials/Material.h"
@@ -121,6 +122,7 @@
 #include "PipelineCacheChunkDataGenerator.h"
 #include "String/Find.h"
 #include "String/ParseLines.h"
+#include "String/ParseTokens.h"
 #include "TargetDomain/TargetDomainUtils.h"
 #include "UnrealEdGlobals.h"
 #include "UObject/ArchiveCookContext.h"
@@ -259,6 +261,8 @@ const FString& GetAssetRegistryFilename()
 	return AssetRegistryFilename;
 }
 
+static void ConditionalWaitOnCommandFile(FStringView GateName, TFunctionRef<void (FStringView)> CommandHandler = [](FStringView){});
+
 /**
  * Uses the FMessageLog to log a message
  * 
@@ -9512,6 +9516,18 @@ void UCookOnTheFlyServer::DiscoverPlatformSpecificNeverCookPackages(
 
 void UCookOnTheFlyServer::StartCookByTheBook( const FCookByTheBookStartupOptions& CookByTheBookStartupOptions )
 {
+	TOptional<FCookByTheBookStartupOptions> ModifiedStartupOptions;
+	bool bAbort = false;
+
+	const FCookByTheBookStartupOptions& EffectiveStartupOptions = BlockOnPrebootCookGate(bAbort,
+		CookByTheBookStartupOptions,
+		ModifiedStartupOptions);
+
+	if (bAbort)
+	{
+		return;
+	}
+
 	UE_SCOPED_COOKTIMER(StartCookByTheBook);
 	LLM_SCOPE_BYTAG(Cooker);
 	check(IsInGameThread());
@@ -9519,7 +9535,7 @@ void UCookOnTheFlyServer::StartCookByTheBook( const FCookByTheBookStartupOptions
 
 	// Initialize systems and settings that the rest of StartCookByTheBook depends on
 	// Functions in this section are ordered and can depend on the functions before them
-	FBeginCookContext BeginContext = CreateBeginCookByTheBookContext(CookByTheBookStartupOptions);
+	FBeginCookContext BeginContext = CreateBeginCookByTheBookContext(EffectiveStartupOptions);
 	BlockOnAssetRegistry();
 	CreateSandboxFile(BeginContext);
 	LoadBeginCookConfigSettings(BeginContext);
@@ -9555,6 +9571,63 @@ void UCookOnTheFlyServer::StartCookByTheBook( const FCookByTheBookStartupOptions
 	RegisterCookByTheBookDelegates();
 }
 
+const UCookOnTheFlyServer::FCookByTheBookStartupOptions& UCookOnTheFlyServer::BlockOnPrebootCookGate(bool& bOutAbortCook,
+	const FCookByTheBookStartupOptions& CookByTheBookStartupOptions,
+	TOptional<FCookByTheBookStartupOptions>& ModifiedStartupOptions)
+{
+	const FCookByTheBookStartupOptions* EffectiveStartupOptions = &CookByTheBookStartupOptions;
+	bOutAbortCook = false;
+	
+	if (IsCookByTheBookMode() && !IsCookingInEditor())
+	{
+		ConditionalWaitOnCommandFile(TEXTVIEW("Cook"),
+			[this, &CookByTheBookStartupOptions, &ModifiedStartupOptions, &EffectiveStartupOptions, &bOutAbortCook] (FStringView CommandContents)
+		{
+			UE::String::ParseTokensMultiple(CommandContents, { ' ', '\t', '\r', '\n' },
+				[this, &CookByTheBookStartupOptions, &ModifiedStartupOptions, &EffectiveStartupOptions, &bOutAbortCook] (FStringView Token)
+				{
+					const FStringView PackageToken(TEXT("-Package="));
+
+					if (Token == TEXT("-CookAbort"))
+					{
+						UE_LOG(LogCook, Display, TEXT("Received CookAbort token in cook command file, exiting cook."));
+						bOutAbortCook = true;
+					}
+					else if (Token == TEXT("-FastExit"))
+					{
+						FCommandLine::Append(TEXT(" -FastExit"));
+					}
+					else if (Token.StartsWith(PackageToken))
+					{
+						if (!ModifiedStartupOptions.IsSet())
+						{
+							ModifiedStartupOptions.Emplace(CookByTheBookStartupOptions);
+							EffectiveStartupOptions = &ModifiedStartupOptions.GetValue();
+						}
+
+						UE::String::ParseTokens(Token.RightChop(PackageToken.Len()), TEXT('+'),
+							[&ModifiedStartupOptions](FStringView PackageToken)
+							{
+								ModifiedStartupOptions->CookMaps.Add(FString(PackageToken));
+
+								UE_LOG(LogCook, Display, TEXT("Received Package token in cook command file, adding package '%.*s' to cook workload."), PackageToken.Len(), PackageToken.GetData());
+							}, UE::String::EParseTokensOptions::SkipEmpty);
+					}
+					else if (UAssetManager::IsValid() && UAssetManager::Get().HandleCookCommand(Token))
+					{
+						// Do nothing - handled by the asset manager
+					}
+					else
+					{
+						UE_LOG(LogCook, Warning, TEXT("Ignoring unknown/unsupported token in cook command file: %.*s"), Token.Len(), Token.GetData());
+					}
+				}, UE::String::EParseTokensOptions::SkipEmpty);
+		});
+	}
+
+	return *EffectiveStartupOptions;
+}
+
 FBeginCookContext UCookOnTheFlyServer::CreateBeginCookByTheBookContext(const FCookByTheBookStartupOptions& StartupOptions)
 {
 	FBeginCookContext BeginContext(*this);
@@ -11376,4 +11449,38 @@ void UCookOnTheFlyServer::RegisterLocalizationChunkDataGenerator()
 	}
 }
 
+static
+void ConditionalWaitOnCommandFile(FStringView GateName, TFunctionRef<void (FStringView)> CommandHandler)
+{
+	WriteToString<128> ArgPrefix(TEXTVIEW("-"), GateName, TEXTVIEW("WaitOnCommandFile="));
+
+	FString WaitOnCommandFile;
+	if (!FParse::Value(FCommandLine::Get(), *ArgPrefix, WaitOnCommandFile))
+	{
+		return;
+	}
+
+	uint64 WaitStartTime = FPlatformTime::Cycles64();
+	uint64 LastMessageTime = WaitStartTime;
+	FString CommandContents;
+	if (!FLockFile::TryReadAndClear(*WaitOnCommandFile, CommandContents))
+	{
+		UE_LOG(LogCook, Display, TEXT("Waiting for %.*s command file at %s..."), GateName.Len(), GateName.GetData(), *WaitOnCommandFile);
+
+		while (!FLockFile::TryReadAndClear(*WaitOnCommandFile, CommandContents))
+		{
+			uint64 LoopTime = FPlatformTime::Cycles64();
+			if (FPlatformTime::ToSeconds64(LoopTime - LastMessageTime) > 60.f)
+			{
+				double TimeSinceWaitStartTime = FPlatformTime::ToSeconds64(LoopTime - WaitStartTime);
+				UE_LOG(LogCook, Display, TEXT("Waited %.1fs for %.*s command file at %s..."), TimeSinceWaitStartTime, GateName.Len(), GateName.GetData(), *WaitOnCommandFile);
+				LastMessageTime = LoopTime;
+			}
+			FPlatformProcess::Sleep(20.f/1000.f);
+		}
+	}
+
+	CommandHandler(CommandContents);
+}
+
 #undef LOCTEXT_NAMESPACE
diff --git a/Engine/Source/Runtime/Engine/Classes/Engine/AssetManager.h b/Engine/Source/Runtime/Engine/Classes/Engine/AssetManager.h
--- a/Engine/Source/Runtime/Engine/Classes/Engine/AssetManager.h
+++ b/Engine/Source/Runtime/Engine/Classes/Engine/AssetManager.h
@@ -625,6 +625,17 @@ public:
 	  * Gathers information about which assets the game wishes to encrypt into named groups
 	  */
 	virtual void GetContentEncryptionConfig(FContentEncryptionConfig& OutContentEncryptionConfig) {}
+
+	/** 
+	  * Processes a command token as part of cooking.
+	 *
+	 * @param Token - The token to process.
+	 * @return True if the token was processed, false otherwise.
+	 */
+	virtual bool HandleCookCommand(FStringView Token)
+	{
+		return false;
+	}
 #endif
 
 protected:
