Cooker: Iterative Cook and WorldPartition:
Extend CookPackageSplitter API to have the splitter specify Dependency categories and properties, and to specify a GenerationHash that can invalidate a generated package due to changes in the generator without requiring a change in the the dependencies.
Extend WorldPartitionSplitterto specify the dependencies for its generated pacakges.
Update Cooker to track in the cooked AssetRegistry the ExternalActor packages that are dependencies for WorldPartition maps.
#rnx

diff --git a/Engine/Source/Developer/IoStoreUtilities/Private/IoStoreUtilities.cpp b/Engine/Source/Developer/IoStoreUtilities/Private/IoStoreUtilities.cpp
--- a/Engine/Source/Developer/IoStoreUtilities/Private/IoStoreUtilities.cpp
+++ b/Engine/Source/Developer/IoStoreUtilities/Private/IoStoreUtilities.cpp
@@ -3072,7 +3072,7 @@ static void AddChunkInfoToAssetRegistry(TMap<FPackageId, TArray<FIoStoreTocChunk
 
 	for (const TPair<FName, const FAssetPackageData*>& AssetPackage : AssetPackageMap)
 	{
-		if (AssetPackage.Value->DiskSize == -1)
+		if (AssetPackage.Value->DiskSize < 0)
 		{
 			// No data on disk!
 			continue;
diff --git a/Engine/Source/Editor/UnrealEd/Classes/CookOnTheSide/CookOnTheFlyServer.h b/Engine/Source/Editor/UnrealEd/Classes/CookOnTheSide/CookOnTheFlyServer.h
--- a/Engine/Source/Editor/UnrealEd/Classes/CookOnTheSide/CookOnTheFlyServer.h
+++ b/Engine/Source/Editor/UnrealEd/Classes/CookOnTheSide/CookOnTheFlyServer.h
@@ -715,7 +715,7 @@ public:
 	 * Clear platforms' explored flags for all PackageDatas and optionally clear the cookresult flags.
 	 *
 	 * @param TargetPlatforms List of pairs with targetplatforms to reset and bool bResetResults indicating whether
-	 *        the platform should clear bCookAttempted in addition to clearing bExplored.
+	 *        the platform should clear CookResults in addition to clearing IsExplored.
 	 */
 	void ResetCook(TConstArrayView<TPair<const ITargetPlatform*, bool>> TargetPlatforms);
 
@@ -1271,6 +1271,8 @@ private:
 
 	/** Cook (save) a package and process the results */
 	void SaveCookedPackage(UE::Cook::FSaveCookedPackageContext& Context);
+	/** Helper for save of generator packages using ExternalActors: Record special-case ExternalActor packages for iterative builds. */
+	void RecordNeverCookDependencies(TConstArrayView<FName> NeverCookDependencies);
 	friend class UE::Cook::FSaveCookedPackageContext;
 
 	/**
diff --git a/Engine/Source/Editor/UnrealEd/Private/Commandlets/AssetRegistryGenerator.cpp b/Engine/Source/Editor/UnrealEd/Private/Commandlets/AssetRegistryGenerator.cpp
--- a/Engine/Source/Editor/UnrealEd/Private/Commandlets/AssetRegistryGenerator.cpp
+++ b/Engine/Source/Editor/UnrealEd/Private/Commandlets/AssetRegistryGenerator.cpp
@@ -1264,22 +1264,45 @@ void FAssetRegistryGenerator::Initialize(const TArray<FName> &InStartupPackages,
 	FGameDelegates::Get().GetAssignLayerChunkDelegate() = FAssignLayerChunkDelegate::CreateStatic(AssignLayerChunkDelegate);
 }
 
+static UE::Cook::ECookResult DiskSizeToCookResult(int64 DiskSize)
+{
+	if (DiskSize >= 0)
+	{
+		return UE::Cook::ECookResult::Succeeded;
+	}
+	switch (DiskSize)
+	{
+	case -2: return UE::Cook::ECookResult::NeverCookPlaceholder;
+	default: return UE::Cook::ECookResult::Failed;
+	}
+}
+
+static int64 CookResultToDiskSize(UE::Cook::ECookResult CookResult)
+{
+	switch (CookResult)
+	{
+	case UE::Cook::ECookResult::Succeeded: return 0;
+	case UE::Cook::ECookResult::NeverCookPlaceholder: return -2;
+	default: return -1;
+	}
+}
+
 void FAssetRegistryGenerator::ComputePackageDifferences(const FComputeDifferenceOptions& Options,
 	const FAssetRegistryState& PreviousState, FAssetRegistryDifference& OutDifference)
 {
 	TArray<FName> ModifiedScriptPackages;
 	const TMap<FName, const FAssetPackageData*>& PreviousAssetPackageDataMap = PreviousState.GetAssetPackageDataMap();
+	OutDifference.Packages.Reserve(PreviousAssetPackageDataMap.Num());
 
 	for (const TPair<FName, const FAssetPackageData*>& PackagePair : State.GetAssetPackageDataMap())
 	{
 		FName PackageName = PackagePair.Key;
 		const FAssetPackageData* CurrentPackageData = PackagePair.Value;
-
 		const FAssetPackageData* PreviousPackageData = PreviousAssetPackageDataMap.FindRef(PackageName);
 
 		if (!PreviousPackageData)
 		{
-			OutDifference.NewPackages.Add(PackageName);
+			// A package that was not explored in the previous cook. No need to record it
 		}
 		PRAGMA_DISABLE_DEPRECATION_WARNINGS
 		else if (CurrentPackageData->PackageGuid == PreviousPackageData->PackageGuid)
@@ -1287,26 +1310,46 @@ void FAssetRegistryGenerator::ComputePackageDifferences(const FComputeDifference
 		{
 			if (PreviousPackageData->DiskSize < 0)
 			{
-				OutDifference.IdenticalUncookedPackages.Add(PackageName);
+				UE::Cook::ECookResult CookResult = DiskSizeToCookResult(PreviousPackageData->DiskSize);
+				if (CookResult == UE::Cook::ECookResult::NeverCookPlaceholder)
+				{
+					OutDifference.Packages.Add(PackageName, EDifference::IdenticalNeverCookPlaceholder);
+				}
+				else
+				{
+					OutDifference.Packages.Add(PackageName, EDifference::IdenticalUncooked);
+				}
+			}
+			else if (FPackageName::IsScriptPackage(WriteToString<256>(PackageName)))
+			{
+				OutDifference.Packages.Add(PackageName, EDifference::IdenticalScript);
 			}
 			else
 			{
-				// Do not record identical script packages since we do not cook them
-				if (!FPackageName::IsScriptPackage(WriteToString<256>(PackageName)))
-				{
-					OutDifference.IdenticalCookedPackages.Add(PackageName);
-				}
+				OutDifference.Packages.Add(PackageName, EDifference::IdenticalCooked);
 			}
 		}
 		else
 		{
-			if (FPackageName::IsScriptPackage(WriteToString<256>(PackageName)))
+			if (PreviousPackageData->DiskSize < 0)
 			{
-				ModifiedScriptPackages.Add(PackageName);
+				UE::Cook::ECookResult CookResult = DiskSizeToCookResult(PreviousPackageData->DiskSize);
+				if (CookResult == UE::Cook::ECookResult::NeverCookPlaceholder)
+				{
+					OutDifference.Packages.Add(PackageName, EDifference::ModifiedNeverCookPlaceholder);
+				}
+				else
+				{
+					OutDifference.Packages.Add(PackageName, EDifference::ModifiedUncooked);
+				}
+			}
+			else if (FPackageName::IsScriptPackage(WriteToString<256>(PackageName)))
+			{
+				OutDifference.Packages.Add(PackageName, EDifference::ModifiedScript);
 			}
 			else
 			{
-				OutDifference.ModifiedPackages.Add(PackageName);
+				OutDifference.Packages.Add(PackageName, EDifference::ModifiedCooked);
 			}
 		}
 	}
@@ -1314,141 +1357,178 @@ void FAssetRegistryGenerator::ComputePackageDifferences(const FComputeDifference
 	for (const TPair<FName, const FAssetPackageData*>& PackagePair : PreviousAssetPackageDataMap)
 	{
 		FName PackageName = PackagePair.Key;
+		const FAssetPackageData* PreviousPackageData = PackagePair.Value;
 		const FAssetPackageData* CurrentPackageData = State.GetAssetPackageData(PackageName);
+
 		if (!CurrentPackageData)
 		{
-			// If it's a generated package, never mark it as removed (that can only be handled by the generator)
-			// Mark it as modified if any of its dependencies are modified.
+			// If it's a generated package, exclude it from the results list and do not remove it.
+			// It will be evaluated for identical/modified/removed only if the generator package is cooked,
+			// during the generator's process step.
 			TConstArrayView<const FAssetData*> PreviousAssets = PreviousState.GetAssetsByPackageName(PackageName);
 			bool bGenerated = !PreviousAssets.IsEmpty() && (PreviousAssets[0]->PackageFlags & PKG_CookGenerated) != 0;
 			if (bGenerated)
 			{
-				bool bModified = true;
-
-				TArray<FAssetIdentifier> Dependencies;
+				// Keep track of all generators and their list of generated
 				TArray<FAssetIdentifier> Referencers;
-				PreviousState.GetDependencies(FAssetIdentifier(PackageName), Dependencies, UE::AssetRegistry::EDependencyCategory::Package);
 				PreviousState.GetReferencers(FAssetIdentifier(PackageName), Referencers, UE::AssetRegistry::EDependencyCategory::Package);
 				FName GeneratorName = Referencers.Num() == 1 ? Referencers[0].PackageName : NAME_None;
 				const FAssetPackageData* GeneratorData = !GeneratorName.IsNone() ? State.GetAssetPackageData(GeneratorName) : nullptr;
 				if (!GeneratorData)
 				{
-					UE_LOG(LogAssetRegistryGenerator, Error,
-						TEXT("Generator could not be found for package %s. This package will be missing from iterative cook results. Recook with a full cook."),
-						*PackageName.ToString());
+					// Mark it as removed; it will be regenerated when the Generator cooks
+					OutDifference.Packages.Add(PackageName, EDifference::RemovedCooked);
 				}
 				else
 				{
-					FGeneratorPackageInfo& Info = OutDifference.PreviousGeneratorPackages.FindOrAdd(GeneratorName);
-					Info.Generated.Add(PackageName);
-
-					bModified = false;
-					for (const FAssetIdentifier& Dependency : Dependencies)
+					FGeneratorPackageInfo& Info = OutDifference.GeneratorPackages.FindOrAdd(GeneratorName);
+					PRAGMA_DISABLE_DEPRECATION_WARNINGS;
+					Info.Generated.Add(PackageName, PreviousPackageData->PackageGuid);
+					PRAGMA_ENABLE_DEPRECATION_WARNINGS;
+				}
+			}
+			else
+			{
+				if (PreviousPackageData->DiskSize < 0)
+				{
+					UE::Cook::ECookResult CookResult = DiskSizeToCookResult(PreviousPackageData->DiskSize);
+					if (CookResult == UE::Cook::ECookResult::NeverCookPlaceholder)
 					{
-						const FAssetPackageData* PreviousDependencyData = PreviousState.GetAssetPackageData(Dependency.PackageName);
-						const FAssetPackageData* CurrentDependencyData = State.GetAssetPackageData(Dependency.PackageName);
-						PRAGMA_DISABLE_DEPRECATION_WARNINGS
-						if (!PreviousDependencyData || !CurrentDependencyData ||
-							PreviousDependencyData->PackageGuid != CurrentDependencyData->PackageGuid)
-						{
-							bModified = true;
-							break;
-						}
-						PRAGMA_ENABLE_DEPRECATION_WARNINGS
+						OutDifference.Packages.Add(PackageName, EDifference::RemovedNeverCookPlaceholder);
+					}
+					else
+					{
+						OutDifference.Packages.Add(PackageName, EDifference::RemovedUncooked);
 					}
-
-					Info.bGeneratedModified |= bModified;
 				}
-				if (bModified)
+				else if (FPackageName::IsScriptPackage(WriteToString<256>(PackageName)))
 				{
-					OutDifference.ModifiedPackages.Add(PackageName);
+					OutDifference.Packages.Add(PackageName, EDifference::RemovedScript);
 				}
 				else
 				{
-					OutDifference.IdenticalCookedPackages.Add(PackageName);
+					OutDifference.Packages.Add(PackageName, EDifference::RemovedCooked);
 				}
 			}
-			else
-			{
-				OutDifference.RemovedPackages.Add(PackageName);
-			}
 		}
 	}
 
-	for (TPair<FName, FGeneratorPackageInfo>& Pair : OutDifference.PreviousGeneratorPackages)
+	// If a Generator package has been removed, remove all its generated packages
+	for (TMap<FName, FGeneratorPackageInfo>::TIterator Iter(OutDifference.GeneratorPackages); Iter; ++Iter)
 	{
-		FName GeneratorName = Pair.Key;
-		FGeneratorPackageInfo& Info = Pair.Value;
-		if (OutDifference.RemovedPackages.Contains(GeneratorName))
+		FName GeneratorName = Iter->Key;
+		EDifference* GeneratorDifference = OutDifference.Packages.Find(GeneratorName);
+		if (GeneratorDifference && *GeneratorDifference == EDifference::RemovedCooked)
 		{
-			// If a Generator package has been removed, remove all generated
-			for (FName Generated : Info.Generated)
+			for (const TPair<FName, FGuid>& Generated : Iter->Value.Generated)
 			{
-				OutDifference.ModifiedPackages.Remove(Generated);
-				OutDifference.IdenticalCookedPackages.Remove(Generated);
-				OutDifference.RemovedPackages.Add(Generated);
+				OutDifference.Packages.Add(Generated.Key, EDifference::RemovedCooked);
 			}
-		}
-		else if (Info.bGeneratedModified)
-		{
-			// Even if not recursing modifications, if a GeneratedPackage has been modified, we need to recook the
-			// generator package, because the generator may change generation depending on modifications to generated.
-			OutDifference.IdenticalCookedPackages.Remove(GeneratorName);
-			OutDifference.IdenticalUncookedPackages.Remove(GeneratorName);
-			OutDifference.ModifiedPackages.Add(GeneratorName);
+			Iter.RemoveCurrent();
 		}
 	}
 
 	if (Options.bRecurseModifications)
 	{
-		// Recurse modified packages to their dependencies. This is needed because we only compare package guids
-		TArray<FName> ModifiedPackagesToRecurse = OutDifference.ModifiedPackages.Array();
+		const FStringView ExternalActorsFolderName(ULevel::GetExternalActorsFolderName());
 
-		if (Options.bRecurseScriptModifications)
+		// Recurse modified packages to their dependencies. This is needed because we only compare package guids
+		TArray<FName> VisitStack;
+		for (TPair<FName, EDifference>& Pair : OutDifference.Packages)
 		{
-			ModifiedPackagesToRecurse.Append(ModifiedScriptPackages);
+			if (Pair.Value == EDifference::ModifiedCooked || Pair.Value == EDifference::ModifiedUncooked || Pair.Value == EDifference::ModifiedNeverCookPlaceholder
+				|| (Pair.Value == EDifference::ModifiedScript && Options.bRecurseScriptModifications))
+			{
+				VisitStack.Add(Pair.Key);
+			}
 		}
 
-		for (int32 RecurseIndex = 0; RecurseIndex < ModifiedPackagesToRecurse.Num(); RecurseIndex++)
+		TSet<FName> Visited;
+		Visited.Reserve(State.GetNumPackages());
+		for (FName PackageName : VisitStack)
+		{
+			Visited.Add(PackageName);
+		}
+		while (!VisitStack.IsEmpty())
 		{
-			FName ModifiedPackage = ModifiedPackagesToRecurse[RecurseIndex];
+			FName ModifiedPackage = VisitStack.Pop(false /* bAllowShrinking */);
+			FString ModifiedPackageLeafName;
+
+			// Read referencers from the current state. If there are referencers in the old state that are not in the
+			// current state, then they must have changed themselves and so are already in the modified set.
 			TArray<FAssetIdentifier> Referencers;
-			// Read referencers from the previous state (to find referencers of generated packages) and of the current state
-			// (to find referencers of editoronly packages)
-			PreviousState.GetReferencers(ModifiedPackage, Referencers, UE::AssetRegistry::EDependencyCategory::Package,
-				UE::AssetRegistry::EDependencyQuery::Hard);
 			State.GetReferencers(ModifiedPackage, Referencers, UE::AssetRegistry::EDependencyCategory::Package,
-				UE::AssetRegistry::EDependencyQuery::Hard);
-			Algo::Sort(Referencers, [](FAssetIdentifier& A, FAssetIdentifier& B) { return A.FastLess(B); });
-			Referencers.SetNum(Algo::Unique(Referencers));
+				UE::AssetRegistry::EDependencyQuery::Hard | UE::AssetRegistry::EDependencyQuery::Build);
 
 			for (const FAssetIdentifier& Referencer : Referencers)
 			{
-				FName ReferencerPackage = Referencer.PackageName;
-				if (!OutDifference.ModifiedPackages.Contains(ReferencerPackage) &&
-					(OutDifference.IdenticalCookedPackages.Contains(ReferencerPackage) || OutDifference.IdenticalUncookedPackages.Contains(ReferencerPackage)))
+				FName ReferencerPackageName = Referencer.PackageName;
+				// TODO: Replace this workaround for ExternalActors with a modification to the ExternalActor Packages' 
+				// dependencies. External actors have an import dependency (hard, build, game) on their Map package because
+				// the map package is their outer. But unless they have some other use of the Map package, we do not want to
+				// mark them as modified even if their map package is modified. Doing so would mark all actors in the map as
+				// modified anytime one of them changed.
+				// Workaround: Detect external actors by naming convention and suppress their reference to the map package.
+				WriteToString<256> ReferencerPackageNameStr(ReferencerPackageName);
+				int32 ExternalActorsFolderIndex = ReferencerPackageNameStr.ToView().Find(ExternalActorsFolderName);
+				if (ExternalActorsFolderIndex != INDEX_NONE)
+				{
+					if (ModifiedPackageLeafName.IsEmpty())
+					{
+						WriteToString<256> ModifiedPackageNameStr(ModifiedPackage);
+						ModifiedPackageLeafName = FPathViews::GetCleanFilename(ModifiedPackageNameStr);
+					}
+					FStringView GeneratedRelativePath = ReferencerPackageNameStr.ToView().RightChop(ExternalActorsFolderIndex + ExternalActorsFolderName.Len());
+					if (GeneratedRelativePath.Contains(ModifiedPackageLeafName))
+					{
+						// Suppress this reference; External actors should not be marked dirty if only their map package is dirty
+						continue;
+					}
+				}
+
+				int32 ReferencerPackageHash = GetTypeHash(ReferencerPackageName);
+				bool bAlreadyVisited;
+				Visited.AddByHash(ReferencerPackageHash, ReferencerPackageName, &bAlreadyVisited);
+				if (bAlreadyVisited)
 				{
-					// Remove from identical list
-					OutDifference.IdenticalCookedPackages.Remove(ReferencerPackage);
-					OutDifference.IdenticalUncookedPackages.Remove(ReferencerPackage);
+					continue;
+				}
 
-					OutDifference.ModifiedPackages.Add(ReferencerPackage);
-					ModifiedPackagesToRecurse.Add(ReferencerPackage);
+				EDifference* ReferencerDifference = OutDifference.Packages.FindByHash(ReferencerPackageHash, ReferencerPackageName);
+				if (!ReferencerDifference)
+				{
+					// The referencer is not in the packages explored during the previous cook, so none of the previous cook
+					// packages had it as a dependency, so we do not need to follow its referencers.
+					continue;
 				}
+
+				// Convert this Referencer to modified. We do not need to handle Removed differences, because we found
+				// it in the current state's dependency tree and so it can not be a Removed difference.
+				switch (*ReferencerDifference)
+				{
+				case EDifference::IdenticalCooked: *ReferencerDifference = EDifference::ModifiedCooked; break;
+				case EDifference::IdenticalUncooked: *ReferencerDifference = EDifference::ModifiedUncooked; break;
+				case EDifference::IdenticalNeverCookPlaceholder: *ReferencerDifference = EDifference::ModifiedNeverCookPlaceholder; break;
+				case EDifference::IdenticalScript: *ReferencerDifference = EDifference::ModifiedScript; break;
+				default: break;
+				}
+				VisitStack.Add(ReferencerPackageName);
 			}
 		}
 	}
 }
 
-void FAssetRegistryGenerator::ComputePackageRemovals(const FAssetRegistryState& PreviousState, TArray<FName>& RemovedPackages,
-	TMap<FName, FGeneratorPackageInfo>& PreviousGeneratorPackages)
+void FAssetRegistryGenerator::ComputePackageRemovals(const FAssetRegistryState& PreviousState, TArray<FName>& OutRemovedPackages,
+	TMap<FName, FGeneratorPackageInfo>& OutGeneratorPackages, int32& OutNumNeverCookPlaceHolderPackages)
 {
-	PreviousGeneratorPackages.Reset();
+	OutGeneratorPackages.Reset();
+	OutNumNeverCookPlaceHolderPackages = 0;
 	TSet<FName> RemovedPackageSet;
 
 	for (const TPair<FName, const FAssetPackageData*>& PackagePair : PreviousState.GetAssetPackageDataMap())
 	{
 		FName PackageName = PackagePair.Key;
+		const FAssetPackageData* PreviousPackageData = PackagePair.Value;
 		const FAssetPackageData* CurrentPackageData = State.GetAssetPackageData(PackageName);
 		if (!CurrentPackageData)
 		{
@@ -1464,13 +1544,14 @@ void FAssetRegistryGenerator::ComputePackageRemovals(const FAssetRegistryState&
 				const FAssetPackageData* GeneratorData = !GeneratorName.IsNone() ? State.GetAssetPackageData(GeneratorName) : nullptr;
 				if (!GeneratorData)
 				{
-					UE_LOG(LogAssetRegistryGenerator, Error,
-						TEXT("Generator could not be found for package %s. This package will be missing from iterative cook results. Recook with a full cook."),
-						*PackageName.ToString());
+					// Mark it as removed; it will be regenerated when the Generator cooks
+					RemovedPackageSet.Add(PackageName);
 				}
 				else
 				{
-					PreviousGeneratorPackages.FindOrAdd(GeneratorName).Generated.Add(PackageName);
+					PRAGMA_DISABLE_DEPRECATION_WARNINGS;
+					OutGeneratorPackages.FindOrAdd(GeneratorName).Generated.Add(PackageName, PreviousPackageData->PackageGuid);
+					PRAGMA_ENABLE_DEPRECATION_WARNINGS;
 				}
 			}
 			else
@@ -1478,23 +1559,31 @@ void FAssetRegistryGenerator::ComputePackageRemovals(const FAssetRegistryState&
 				RemovedPackageSet.Add(PackageName);
 			}
 		}
+		if (PreviousPackageData->DiskSize < 0)
+		{
+			UE::Cook::ECookResult CookResult = DiskSizeToCookResult(PreviousPackageData->DiskSize);
+			if (CookResult == UE::Cook::ECookResult::NeverCookPlaceholder)
+			{
+				++OutNumNeverCookPlaceHolderPackages;
+			}
+		}
 	}
 
-	for (TMap<FName, FGeneratorPackageInfo>::TIterator Iter(PreviousGeneratorPackages); Iter; ++Iter)
+	// If a Generator package has been removed, remove all its generated packages
+	for (TMap<FName, FGeneratorPackageInfo>::TIterator Iter(OutGeneratorPackages); Iter; ++Iter)
 	{
-		// If a Generator package has been removed, remove all generated
 		FName GeneratorName = Iter->Key;
 		if (RemovedPackageSet.Contains(GeneratorName))
 		{
-			for (FName Generated : Iter->Value.Generated)
+			for (const TPair<FName, FGuid>& Generated : Iter->Value.Generated)
 			{
-				RemovedPackageSet.Add(Generated);
+				RemovedPackageSet.Add(Generated.Key);
 			}
 			Iter.RemoveCurrent();
 		}
 	}
 
-	RemovedPackages = RemovedPackageSet.Array();
+	OutRemovedPackages = RemovedPackageSet.Array();
 }
 
 void FAssetRegistryGenerator::FinalizeChunkIDs(const TSet<FName>& InCookedPackages,
@@ -1521,6 +1610,9 @@ void FAssetRegistryGenerator::FinalizeChunkIDs(const TSet<FName>& InCookedPackag
 	CookedPackages = InCookedPackages;
 	DevelopmentOnlyPackages = InDevelopmentOnlyPackages;
 
+	// Possibly apply previous AssetData and AssetPackageData for packages kept from a previous cook
+	UpdateKeptPackages();
+
 	TSet<FName> AllPackages;
 	AllPackages.Append(CookedPackages);
 	AllPackages.Append(DevelopmentOnlyPackages);
@@ -1530,11 +1622,17 @@ void FAssetRegistryGenerator::FinalizeChunkIDs(const TSet<FName>& InCookedPackag
 	AssetRegistry.InitializeSerializationOptions(DevelopmentSaveOptions, TargetPlatform->IniPlatformName(), UE::AssetRegistry::ESerializationTarget::ForDevelopment);
 	State.PruneAssetData(AllPackages, TSet<FName>(), DevelopmentSaveOptions);
 
-	// Mark development only packages as explicitly -1 size to indicate it was not cooked
+	// Development only packages should have been marked via UpdateAssetRegistryData with a negative size indicating why they were not cooked 
 	for (FName DevelopmentOnlyPackage : DevelopmentOnlyPackages)
 	{
 		FAssetPackageData* PackageData = State.CreateOrGetAssetPackageData(DevelopmentOnlyPackage);
-		PackageData->DiskSize = -1;
+		if (PackageData->DiskSize >= 0)
+		{
+			UE_LOG(LogAssetRegistryGenerator, Warning,
+				TEXT("Package %s is a development-only package but was not marked with DiskSize explaining why it was not cooked. Marking it as a generic cook failure."),
+				*DevelopmentOnlyPackage.ToString());
+			PackageData->DiskSize = -1;
+		}
 	}
 
 	// Copy ExplicitChunkIDs and other data from the AssetRegistry into the maps we use during finalization
@@ -1714,9 +1812,6 @@ bool FAssetRegistryGenerator::SaveAssetRegistry(const FString& SandboxPath, bool
 		SaveOptions.DisableFilters();
 	}
 
-	// Possibly apply previous AssetData and AssetPackageData for packages kept from a previous cook
-	UpdateKeptPackages();
-
 	UpdateCollectionAssetData();
 
 	if (DevelopmentSaveOptions.bSerializeAssetRegistry)
@@ -2684,17 +2779,17 @@ const FAssetData* FAssetRegistryGenerator::CreateOrFindAssetData(UObject& Object
 	return AssetData;
 }
 
-void FAssetRegistryGenerator::UpdateAssetRegistryData(const UPackage& Package,
-	FSavePackageResultStruct& SavePackageResult, FCookTagList&& InArchiveCookTagList, bool bIncludeOnlyDiskAssets,
-	TOptional<FAssetPackageData>&& OverrideAssetPackageData, TOptional<TArray<FAssetDependency>>&& OverridePackageDependencies)
+void FAssetRegistryGenerator::UpdateAssetRegistryData(FName PackageName, const UPackage* Package,
+	UE::Cook::ECookResult CookResult, FSavePackageResultStruct* SavePackageResult, FCookTagList&& InArchiveCookTagList,
+	bool bIncludeOnlyDiskAssets, TOptional<FAssetPackageData>&& OverrideAssetPackageData,
+	TOptional<TArray<FAssetDependency>>&& OverridePackageDependencies)
 {
 	LLM_SCOPE_BYTAG(Cooker_GeneratedAssetRegistry);
-	const FName PackageName = Package.GetFName();
 	PreviousPackagesToUpdate.Remove(PackageName);
 
 	uint32 NewPackageFlags = 0;
 	FAssetPackageData* AssetPackageData = GetAssetPackageData(PackageName);
-	bool bSaveSucceeded = SavePackageResult.IsSuccessful();
+	bool bSaveSucceeded = CookResult == UE::Cook::ECookResult::Succeeded;
 
 	// Copy latest data for all Assets in the package into the cooked registry. This should be done even
 	// if not successful so that editor-only packages are recorded as well
@@ -2706,7 +2801,10 @@ void FAssetRegistryGenerator::UpdateAssetRegistryData(const UPackage& Package,
 		State.UpdateAssetData(MoveTemp(AssetData), true /* bCreateIfNotExists */);
 	}
 	// Create a record for assets that were created during PostLoad or cooking and are not in the global assetregistry
-	CreateOrFindAssetDatas(Package);
+	if (Package)
+	{
+		CreateOrFindAssetDatas(*Package);
+	}
 
 	// Set Overrides for generated packages
 	if (OverrideAssetPackageData)
@@ -2720,14 +2818,15 @@ void FAssetRegistryGenerator::UpdateAssetRegistryData(const UPackage& Package,
 
 	if (bSaveSucceeded)
 	{
+		check(SavePackageResult);
 		// Migrate cook tags over
 		CookTagsToAdd.Append(MoveTemp(InArchiveCookTagList.ObjectToTags));
 		InArchiveCookTagList.Reset();
 
 		// Set the PackageFlags to the recorded value from SavePackage
-		NewPackageFlags = SavePackageResult.SerializedPackageFlags;
+		NewPackageFlags = SavePackageResult->SerializedPackageFlags;
 
-		AssetPackageData->DiskSize = SavePackageResult.TotalFileSize;
+		AssetPackageData->DiskSize = SavePackageResult->TotalFileSize;
 
 		// The CookedHash is assigned during CookByTheBookFinished		
 	}
@@ -2736,8 +2835,9 @@ void FAssetRegistryGenerator::UpdateAssetRegistryData(const UPackage& Package,
 		// Set the package flags to zero to indicate that the package failed to save
 		NewPackageFlags = 0;
 
-		// Set DiskSize (previous value was disksize in the WorkspaceDomain) to -1 to indicate the cooked file does not exist
-		AssetPackageData->DiskSize = -1;
+		// Set DiskSize (previous value was disksize in the WorkspaceDomain) to a negative number to indicate the
+		// cooked file does not exist. The magnitude of the negative value indicates CookResult.
+		AssetPackageData->DiskSize = CookResultToDiskSize(CookResult);
 	}
 
 	const bool bUpdated = UpdateAssetPackageFlags(PackageName, NewPackageFlags);
@@ -2747,8 +2847,6 @@ void FAssetRegistryGenerator::UpdateAssetRegistryData(const UPackage& Package,
 
 void FAssetRegistryGenerator::SetOverridePackageDependencies(FName PackageName, TConstArrayView<FAssetDependency> OverridePackageDependencies)
 {
-	TArray<FAssetDependency, TInlineAllocator<10>> DependencyStructs;
-	DependencyStructs.Reserve(OverridePackageDependencies.Num());
 #if DO_CHECK
 	for (FAssetDependency Dependency : OverridePackageDependencies)
 	{
@@ -2795,27 +2893,29 @@ FAssetRegistryReporterRemote::FAssetRegistryReporterRemote(FCookWorkerClient& In
 {
 }
 
-void FAssetRegistryReporterRemote::UpdateAssetRegistryData(FPackageData& PackageData, const UPackage& Package,
-	FSavePackageResultStruct& SavePackageResult, FCookTagList&& InArchiveCookTagList, bool bIncludeOnlyDiskAssets,
-	TOptional<FAssetPackageData>&& OverrideAssetPackageData, TOptional<TArray<FAssetDependency>>&& OverridePackageDependencies)
+void FAssetRegistryReporterRemote::UpdateAssetRegistryData(FName PackageName, const UPackage* Package,
+	UE::Cook::ECookResult CookResult, FSavePackageResultStruct* SavePackageResult, FCookTagList&& InArchiveCookTagList,
+	bool bIncludeOnlyDiskAssets, TOptional<FAssetPackageData>&& OverrideAssetPackageData,
+	TOptional<TArray<FAssetDependency>>&& OverridePackageDependencies)
 {
 	uint32 NewPackageFlags = 0;
 	int64 DiskSize = -1;
-	bool bSaveSucceeded = SavePackageResult.IsSuccessful();
+	bool bSaveSucceeded = CookResult == UE::Cook::ECookResult::Succeeded;
 	if (bSaveSucceeded)
 	{
-		NewPackageFlags = SavePackageResult.SerializedPackageFlags;
-		DiskSize = SavePackageResult.TotalFileSize;
+		check(SavePackageResult);
+		NewPackageFlags = SavePackageResult->SerializedPackageFlags;
+		DiskSize = SavePackageResult->TotalFileSize;
 	}
 	else
 	{
 		// Set the package flags to zero to indicate that the package failed to save
 		NewPackageFlags = 0;
-		// Set DiskSize (previous value was disksize in the WorkspaceDomain) to -1 to indicate the cooked file does not exist
-		DiskSize = -1;
+		// Set DiskSize (previous value was disksize in the WorkspaceDomain) to a negative number to indicate the
+		// cooked file does not exist. The magnitude of the negative value indicates CookResult.
+		DiskSize = CookResultToDiskSize(CookResult);
 	}
 
-	FName PackageName = PackageData.GetPackageName();
 	FAssetRegistryPackageMessage Message;
 	Message.PackageName = PackageName;
 	Message.TargetPlatform = TargetPlatform;
@@ -2834,17 +2934,20 @@ void FAssetRegistryReporterRemote::UpdateAssetRegistryData(FPackageData& Package
 	{
 		ExistingAssets.Add(AssetData.AssetName);
 	}
-	ForEachObjectWithOuter(&Package, [&Message, &ExistingAssets](UObject* const Object)
-		{
-			if (Object->IsAsset())
+	if (Package)
+	{
+		ForEachObjectWithOuter(Package, [&Message, &ExistingAssets](UObject* const Object)
 			{
-				FName AssetName = Object->GetFName();
-				if (!ExistingAssets.Contains(AssetName))
+				if (Object->IsAsset())
 				{
-					Message.AssetDatas.Emplace(Object, true /* bAllowBlueprintClass */);
+					FName AssetName = Object->GetFName();
+					if (!ExistingAssets.Contains(AssetName))
+					{
+						Message.AssetDatas.Emplace(Object, true /* bAllowBlueprintClass */);
+					}
 				}
-			}
-		}, /*bIncludeNestedObjects*/ false);
+			}, /*bIncludeNestedObjects*/ false);
+	}
 
 	// Add the cooktags that were recorded during serialization
 	Message.CookTags = MoveTemp(InArchiveCookTagList.ObjectToTags);
@@ -2979,8 +3082,8 @@ void FAssetRegistryMPCollector::ClientTickPackage(FMPCollectorClientTickPackageC
 
 		if (!RegistryReporter.PackageUpdateMessages.RemoveAndCopyValue(Context.GetPackageName(), Message))
 		{
-			// For a failed package, UpdateAssetRegistryData will never have been called. Silently skip it and do not send a message.
-			if (ContextPlatformData.bSuccessful)
+			// For a failed package, UpdateAssetRegistryData might never have been called. Silently skip it and do not send a message.
+			if (ContextPlatformData.CookResults == ECookResult::Succeeded)
 			{
 				// For a successful package, UpdateAssetRegistryData should have been called
 				UE_CLOG(!bLoggedWarning, LogAssetRegistryGenerator, Warning,
diff --git a/Engine/Source/Editor/UnrealEd/Private/Commandlets/AssetSizeQueryCommandlet.cpp b/Engine/Source/Editor/UnrealEd/Private/Commandlets/AssetSizeQueryCommandlet.cpp
--- a/Engine/Source/Editor/UnrealEd/Private/Commandlets/AssetSizeQueryCommandlet.cpp
+++ b/Engine/Source/Editor/UnrealEd/Private/Commandlets/AssetSizeQueryCommandlet.cpp
@@ -117,7 +117,7 @@ int32 UAssetSizeQueryCommandlet::Main(const FString& FullCommandLine)
 	const TMap<FName, const FAssetPackageData*>& PackageDatas = AssetRegistry.GetAssetPackageDataMap();
 	for (const TPair<FName, const FAssetPackageData*>& Pair : PackageDatas)
 	{
-		if (Pair.Value->DiskSize != -1)
+		if (Pair.Value->DiskSize >= 0)
 		{
 			TotalDiskSize += Pair.Value->DiskSize;
 		}
diff --git a/Engine/Source/Editor/UnrealEd/Private/Commandlets/DiffAssetRegistriesCommandlet.cpp b/Engine/Source/Editor/UnrealEd/Private/Commandlets/DiffAssetRegistriesCommandlet.cpp
--- a/Engine/Source/Editor/UnrealEd/Private/Commandlets/DiffAssetRegistriesCommandlet.cpp
+++ b/Engine/Source/Editor/UnrealEd/Private/Commandlets/DiffAssetRegistriesCommandlet.cpp
@@ -475,7 +475,7 @@ void UDiffAssetRegistriesCommandlet::ConsistencyCheck(const FString& OldPath, co
 	{
 		const FAssetPackageData* Data = PackageMap[Package];
 
-		if (!GuidModified.Contains(Package))
+		if (!GuidModified.Contains(Package) && Data->DiskSize >= 0)
 		{
 			++Changes;
 			ChangeBytes += Data->DiskSize;
diff --git a/Engine/Source/Editor/UnrealEd/Private/CookOnTheFlyServer.cpp b/Engine/Source/Editor/UnrealEd/Private/CookOnTheFlyServer.cpp
--- a/Engine/Source/Editor/UnrealEd/Private/CookOnTheFlyServer.cpp
+++ b/Engine/Source/Editor/UnrealEd/Private/CookOnTheFlyServer.cpp
@@ -1621,7 +1621,7 @@ void UCookOnTheFlyServer::UpdateDisplay(UE::Cook::FTickStackData& StackData, boo
 		return;
 	}
 
-	const int32 CookedPackagesCount = PackageDatas->GetNumCooked();
+	const int32 CookedPackagesCount = PackageDatas->GetNumCooked() - PackageDatas->GetNumCooked(ECookResult::NeverCookPlaceholder);
 	const int32 CookPendingCount = WorkerRequests->GetNumExternalRequests() + PackageDatas->GetMonitor().GetNumInProgress();
 	if (bForceDisplay ||
 		(DeltaProgressDisplayTime >= GCookProgressUpdateTime && CookPendingCount != 0 &&
@@ -2633,7 +2633,7 @@ void UCookOnTheFlyServer::LoadPackageInQueue(UE::Cook::FPackageData& PackageData
 			PackageData.GetRequestedPlatforms(RequestedPlatforms);
 			QueueDiscoveredPackageData(OtherPackageData, FInstigator(PackageData.GetInstigator()));
 
-			PackageData.SetPlatformsCooked(PlatformManager->GetSessionPlatforms(), true);
+			PackageData.SetPlatformsCooked(PlatformManager->GetSessionPlatforms(), ECookResult::Succeeded);
 			RejectPackageToLoad(PackageData, TEXT("is redirected to another filename"), ESuppressCookReason::Redirected);
 			return;
 		}
@@ -2708,7 +2708,7 @@ void UCookOnTheFlyServer::RejectPackageToLoad(UE::Cook::FPackageData& PackageDat
 	// make sure this package doesn't exist
 	for (const TPair<const ITargetPlatform*, UE::Cook::FPackageData::FPlatformData>& Pair : PackageData.GetPlatformDatas())
 	{
-		if (!Pair.Value.bRequested)
+		if (!Pair.Value.IsRequested())
 		{
 			continue;
 		}
@@ -4144,7 +4144,7 @@ void UCookOnTheFlyServer::PumpSaves(UE::Cook::FTickStackData& StackData, uint32
 			{
 				check(PackageData.HasPrepareSaveFailed()); // Should have been set by PrepareSave; we rely on this for cleanup
 				ReleaseCookedPlatformData(PackageData, EReleaseSaveReason::AbortSave);
-				PackageData.SetPlatformsCooked(PlatformsForPackage, false /* bSucceeded */);
+				PackageData.SetPlatformsCooked(PlatformsForPackage, ECookResult::Failed);
 				DemoteToIdle(PackageData, ESendFlags::QueueAdd, ESuppressCookReason::SaveError);
 				++OutNumPushed;
 				continue;
@@ -5559,7 +5559,7 @@ void FSaveCookedPackageContext::SetupPlatform(const ITargetPlatform* InTargetPla
 	FPackageData::FPlatformData& PlatformData = PackageData.FindOrAddPlatformData(TargetPlatform);
 	uint32 PlatformSaveFlagsMask = SAVE_AllowTimeout;
 	SaveFlags = (SaveFlags & ~PlatformSaveFlagsMask);
-	if (!PlatformData.bSaveTimedOut)
+	if (!PlatformData.IsSaveTimedOut())
 	{
 		// If we timedout before, do not allow another timeout, otherwise do allow it
 		SaveFlags |= SAVE_AllowTimeout;
@@ -5581,6 +5581,7 @@ void FSaveCookedPackageContext::FinishPlatform()
 	TRACE_CPUPROFILER_EVENT_SCOPE(FSaveCookedPackageContext::FinishPlatform);
 
 	bool bSuccessful = SavePackageResult.IsSuccessful();
+	ECookResult CookResult = bSuccessful ? ECookResult::Succeeded : ECookResult::Failed;
 
 	if (bPlatformSetupSuccessful)
 	{
@@ -5661,6 +5662,12 @@ void FSaveCookedPackageContext::FinishPlatform()
 				Dependency.Category = UE::AssetRegistry::EDependencyCategory::Package;
 				Dependency.Properties = UE::AssetRegistry::EDependencyProperty::Game;
 			}
+
+			if (!bHasFirstPlatformResults)
+			{
+				GeneratorPackage->FetchNeverCookDependencies();
+				COTFS.RecordNeverCookDependencies(GeneratorPackage->GetNeverCookDependencies());
+			}
 		}
 		else if (GeneratorPackage = PackageData.GetGeneratedOwner(); GeneratorPackage)
 		{
@@ -5673,19 +5680,9 @@ void FSaveCookedPackageContext::FinishPlatform()
 			{
 				// Generated packages do not exist on disk, so read their assetdata from the in-memory package
 				bIncludeOnlyDiskAssets = false;
-				OverridePackageDependencies.Emplace();
-				OverridePackageDependencies->Reserve(GeneratedInfo->Dependencies.Num());
-
 				// There should be no package dependencies present for the package from the global assetregistry
 				// because it is newly created. Add on the dependencies declared for it from the CookPackageSplitter.
-				for (FName PackageDependency : GeneratedInfo->Dependencies)
-				{
-					FAssetDependency& Dependency = OverridePackageDependencies->Emplace_GetRef();
-					Dependency.AssetId = FAssetIdentifier(PackageDependency);
-					Dependency.Category = UE::AssetRegistry::EDependencyCategory::Package;
-					// COOKPACKAGESPLITTERTODO: CookPackageSplitter needs to specify the DependencyProperties
-					Dependency.Properties = UE::AssetRegistry::EDependencyProperty::Hard | UE::AssetRegistry::EDependencyProperty::Game;
-				}
+				OverridePackageDependencies.Emplace(GeneratedInfo->PackageDependencies);
 
 				// Update the AssetPackageData for each requested platform with Guid and ImportedClasses
 				TSet<UClass*> PackageClasses;
@@ -5715,7 +5712,7 @@ void FSaveCookedPackageContext::FinishPlatform()
 				OverrideAssetPackageData->ImportedClasses = ImportedClasses;
 			}
 		}
-		Reporter.UpdateAssetRegistryData(PackageData, *Package, SavePackageResult,
+		Reporter.UpdateAssetRegistryData(Package->GetFName(), Package, CookResult, &SavePackageResult,
 			MoveTemp(*ArchiveCookContext.GetCookTagList()), bIncludeOnlyDiskAssets,
 			MoveTemp(OverrideAssetPackageData), MoveTemp(OverridePackageDependencies));
 	}
@@ -5724,7 +5721,7 @@ void FSaveCookedPackageContext::FinishPlatform()
 	bool bIsRetryErrorCode = IsRetryErrorCode(SavePackageResult.Result);
 	if (!bIsRetryErrorCode)
 	{
-		PackageData.SetPlatformCooked(TargetPlatform, bSuccessful);
+		PackageData.SetPlatformCooked(TargetPlatform, CookResult);
 	}
 
 	// Update flags used to determine garbage collection.
@@ -5756,7 +5753,7 @@ void FSaveCookedPackageContext::FinishPlatform()
 	}
 	if (SavePackageResult.Result == ESavePackageResult::Timeout)
 	{
-		PackageData.FindOrAddPlatformData(TargetPlatform).bSaveTimedOut = true;
+		PackageData.FindOrAddPlatformData(TargetPlatform).SetSaveTimedOut(true);
 		bHasTimeOut = true;
 	}
 
@@ -5827,7 +5824,6 @@ void FSaveCookedPackageContext::FinishPackage()
 #endif
 			}
 		}
-		COOK_STAT(++DetailedCookStats::NumPackagesSavedForCook);
 	}
 	else if (bReferencedOnlyByEditorOnlyData)
 	{
@@ -5837,6 +5833,39 @@ void FSaveCookedPackageContext::FinishPackage()
 
 } // namespace UE::Cook
 
+void UCookOnTheFlyServer::RecordNeverCookDependencies(TConstArrayView<FName> NeverCookDependencies)
+{
+	using namespace UE::Cook;
+
+	if (IsCookWorkerMode())
+	{
+		// The dependencies will be replicated to the CookDirectory during ReportPromoteToSaveComplete
+		return;
+	}
+
+	// External actors are a special case in the cooker, and they are only referenced through the
+	// WorldPartitionCookPackageSplitter. They are marked as NeverCook, but we need to add them
+	// to the cook results so we can detect whether they change in iterative cooks. Call a function
+	// on the splitter to get a list of them and add them.
+	for (FName DependencyName : NeverCookDependencies)
+	{
+		FPackageData* DependencyData = PackageDatas->TryAddPackageDataByPackageName(DependencyName);
+		if (DependencyData)
+		{
+			for (const ITargetPlatform* TargetPlatform : PlatformManager->GetSessionPlatforms())
+			{
+				IAssetRegistryReporter& Reporter = *PlatformManager->GetPlatformData(TargetPlatform)->RegistryReporter;
+
+				DependencyData->SetPlatformCooked(TargetPlatform, ECookResult::NeverCookPlaceholder);
+				Reporter.UpdateAssetRegistryData(DependencyName, nullptr /* Package */,
+					ECookResult::NeverCookPlaceholder, nullptr /* SavePackageResult */, FCookTagList(nullptr),
+					true /* bIncludeOnlyDiskAssets */, TOptional<FAssetPackageData>(),
+					TOptional<TArray<FAssetDependency>>());
+			}
+		}
+	}
+}
+
 void UCookOnTheFlyServer::Initialize( ECookMode::Type DesiredCookMode, ECookInitializationFlags InCookFlags, const FString &InOutputDirectoryOverride )
 {
 	TRACE_CPUPROFILER_EVENT_SCOPE(UCookOnTheFlyServer::Initialize);
@@ -7149,6 +7178,7 @@ FAsyncIODelete& UCookOnTheFlyServer::GetAsyncIODelete()
 void UCookOnTheFlyServer::PopulateCookedPackages(TArrayView<const ITargetPlatform* const> TargetPlatforms)
 {
 	using namespace UE::Cook;
+	using EDifference = FAssetRegistryGenerator::EDifference;
 	TRACE_CPUPROFILER_EVENT_SCOPE(UCookOnTheFlyServer::PopulateCookedPackages);
 	checkf(!IsCookWorkerMode(), TEXT("Calling PopulateCookedPackages should be impossible in a CookWorker."));
 
@@ -7183,9 +7213,14 @@ void UCookOnTheFlyServer::PopulateCookedPackages(TArrayView<const ITargetPlatfor
 			// check this for all packages in the previous cook rather than just the currently referenced
 			// set because when packages are removed the references to them are usually also removed.
 			TArray<FName> TombstonePackages;
+			int32 NumNeverCookPlaceHolderPackages;
 			PlatformAssetRegistry.ComputePackageRemovals(*PreviousAssetRegistry, TombstonePackages,
-				PreviousGeneratorPackages);
+				PreviousGeneratorPackages, NumNeverCookPlaceHolderPackages);
 			PackageWriter.RemoveCookedPackages(TombstonePackages);
+
+			// Do not show NeverCookPlaceholder packages in the package counts
+			NumPreviousPackages -= NumNeverCookPlaceHolderPackages;
+			NumPreviousPackages = FMath::Max(0, NumPreviousPackages);
 			UE_LOG(LogCook, Display, TEXT("Found %d cooked package(s) in package store."), NumPreviousPackages);
 		}
 		else
@@ -7200,79 +7235,128 @@ void UCookOnTheFlyServer::PopulateCookedPackages(TArrayView<const ITargetPlatfor
 			Options.bRecurseScriptModifications = !IsCookFlagSet(ECookInitializationFlags::IgnoreScriptPackagesOutOfDate);
 			FAssetRegistryGenerator::FAssetRegistryDifference Difference;
 			PlatformAssetRegistry.ComputePackageDifferences(Options, *PreviousAssetRegistry, Difference);
-			PreviousGeneratorPackages = MoveTemp(Difference.PreviousGeneratorPackages);
-
-			TSet<FName> GeneratedPackages;
-			for (TPair<FName, FAssetRegistryGenerator::FGeneratorPackageInfo>& Pair : PreviousGeneratorPackages)
-			{
-				GeneratedPackages.Append(Pair.Value.Generated);
-			}
+			PreviousGeneratorPackages = MoveTemp(Difference.GeneratorPackages);
 
+			TArray<FName> IdenticalCooked;
 			TArray<FName> PackagesToRemove;
-			PackagesToRemove.Reserve(Difference.ModifiedPackages.Num() + Difference.RemovedPackages.Num() + Difference.IdenticalUncookedPackages.Num());
-			for (FName PackageToRemove : Difference.ModifiedPackages)
-			{
-				PackagesToRemove.Add(PackageToRemove);
-			}
-			for (FName PackageToRemove : Difference.RemovedPackages)
-			{
-				PackagesToRemove.Add(PackageToRemove);
-			}
-			for (FName PackageToRemove : Difference.IdenticalUncookedPackages)
-			{
-				PackagesToRemove.Add(PackageToRemove);
-			}
-			int32 NumUncooked = Difference.IdenticalUncookedPackages.Num();
-			int32 NumCooked = FMath::Max(0, NumPreviousPackages - NumUncooked);
-			UE_LOG(LogCook, Display, TEXT("Found %d cooked package(s) in package store."), NumCooked);
-			UE_LOG(LogCook, Display, TEXT("Keeping %d and removing %d cooked package(s)."),
-				Difference.IdenticalCookedPackages.Num(), PackagesToRemove.Num() - NumUncooked);
+			int32 DeferredEvaluateGeneratedNum = 0;
+			int32 IdenticalCookedNum = 0;
+			int32 ModifiedCookedNum = 0;
+			int32 RemovedCookedNum = 0;
 
-			PackageWriter.RemoveCookedPackages(PackagesToRemove);
-
-			if (bFirstPlatform)
-			{
-				COOK_STAT(DetailedCookStats::NumPackagesIterativelySkipped += Difference.IdenticalCookedPackages.Num());
-				bFirstPlatform = false;
-			}
-			for (const FName& IdenticalPackage : Difference.IdenticalCookedPackages)
+			auto AddCookedPackage =
+				[this, TargetPlatform](const FName PackageName, ECookResult CookResult, bool bRequireExists)
+				{
+					FPackageData* PackageData = PackageDatas->TryAddPackageDataByPackageName(PackageName, bRequireExists);
+					if (PackageData)
+					{
+						PackageData->SetPlatformCooked(TargetPlatform, CookResult);
+						PackageData->FindOrAddPlatformData(TargetPlatform).SetIterativelySkipped(true);
+					}
+				};
+			auto AddIdenticalCookedPackage =
+				[&AddCookedPackage, &IdenticalCooked](const FName PackageName, bool bRequireExists)
+				{
+					IdenticalCooked.Add(PackageName);
+					AddCookedPackage(PackageName, ECookResult::Succeeded, bRequireExists);
+					// Declare the package to the EDLCookInfo verification so we don't warn about missing exports from it
+					UE::SavePackageUtilities::EDLCookInfoAddIterativelySkippedPackage(PackageName);
+				};
+			bool bCookByTheBook = IsCookByTheBookMode();
+			auto AddRequestForModifiedPackage = [this, TargetPlatform, bCookByTheBook](const FName PackageName)
 			{
-				// Mark this package as cooked so that we don't unnecessarily try to cook it again
-				bool bRequireExists = !GeneratedPackages.Contains(IdenticalPackage);
-				FPackageData* PackageData = PackageDatas->TryAddPackageDataByPackageName(IdenticalPackage, bRequireExists);
-				if (PackageData)
+				if (bCookByTheBook)
 				{
-					PackageData->SetPlatformCooked(TargetPlatform, true /* bSucceeded */);
+					// cook on the fly will queue packages when it needs them, but for cook by the book we force cook the modified files
+					// so that the output set of packages is up to date (even if the user is currently cooking only a subset)
+					FPackageData* PackageData = PackageDatas->TryAddPackageDataByPackageName(PackageName);
+					if (PackageData)
+					{
+						WorkerRequests->AddStartCookByTheBookRequest(FFilePlatformRequest(PackageData->GetFileName(),
+							EInstigator::IterativeCook, TConstArrayView<const ITargetPlatform*>{ TargetPlatform }));
+					}
 				}
 
-				// Declare the package to the EDLCookInfo verification so we don't warn about missing exports from it
-				UE::SavePackageUtilities::EDLCookInfoAddIterativelySkippedPackage(IdenticalPackage);
-			}
-			PackageWriter.MarkPackagesUpToDate(Difference.IdenticalCookedPackages.Array());
-			for (FName UncookedPackage : Difference.IdenticalUncookedPackages)
+			};
+
+			// Add CookedPackages for any identical packages, delete from disk any modified packages
+			// For legacy paranoia, also delete from disk any packages that were marked as uncooked.
+			for (const TPair<FName, EDifference>& Pair : Difference.Packages)
 			{
-				// Mark this failed-to-cook package as cooked so that we don't unnecessarily try to cook it again
-				FPackageData* PackageData = PackageDatas->TryAddPackageDataByPackageName(UncookedPackage);
-				if (PackageData)
+				FName PackageName = Pair.Key;
+				switch (Pair.Value)
 				{
-					ensure(!PackageData->HasAnyCookedPlatforms({ TargetPlatform }, /* bIncludeFailed */ false));
-					PackageData->SetPlatformCooked(TargetPlatform, false /* bSucceeded */);
+				case EDifference::IdenticalCooked:
+					AddIdenticalCookedPackage(PackageName, true /* bRequireExists */);
+					break;
+				case EDifference::ModifiedCooked:
+					AddRequestForModifiedPackage(PackageName);
+					PackagesToRemove.Add(PackageName);
+					++ModifiedCookedNum;
+					break;
+				case EDifference::RemovedCooked:
+					PackagesToRemove.Add(PackageName);
+					++RemovedCookedNum;
+					break;
+				case EDifference::IdenticalUncooked:
+					AddCookedPackage(PackageName, ECookResult::Failed, true /* bRequireExists */);
+					PackagesToRemove.Add(PackageName);
+					break;
+				case EDifference::ModifiedUncooked:
+					PackagesToRemove.Add(PackageName);
+					break;
+				case EDifference::RemovedUncooked:
+					PackagesToRemove.Add(PackageName);
+					break;
+				case EDifference::IdenticalNeverCookPlaceholder:
+					AddCookedPackage(PackageName, ECookResult::NeverCookPlaceholder, true /* bRequireExists */);
+					PackagesToRemove.Add(PackageName);
+					break;
+				case EDifference::ModifiedNeverCookPlaceholder:
+					PackagesToRemove.Add(PackageName);
+					break;
+				case EDifference::RemovedNeverCookPlaceholder:
+					PackagesToRemove.Add(PackageName);
+					break;
+				default:
+					break;
 				}
 			}
-			if (IsCookByTheBookMode())
+
+			// Add as identical any generated packages from any identical generator, because we will skip cooking
+			// the generator and therefore will skip cooking the generated. Count the number of generated packages from
+			// modified generators and report that we will evaluate them later.
+			for (TMap<FName, FAssetRegistryGenerator::FGeneratorPackageInfo>::TIterator Iter(PreviousGeneratorPackages);
+				Iter; ++Iter)
 			{
-				for (const FName& RemovePackageName : Difference.ModifiedPackages)
+				FName Generator = Iter->Key;
+				EDifference* GeneratorDifference = Difference.Packages.Find(Generator);
+				if (GeneratorDifference && *GeneratorDifference == EDifference::IdenticalCooked)
 				{
-					// cook on the fly will requeue this package when it wants it, but for cook by the book we force cook the modified file
-					// so that the output set of packages is up to date (even if the user is currently cooking only a subset)
-					FPackageData* PackageData = PackageDatas->TryAddPackageDataByPackageName(RemovePackageName);
-					if (PackageData)
+					for (const TPair<FName, FGuid>& GeneratedPair : Iter->Value.Generated)
 					{
-						WorkerRequests->AddStartCookByTheBookRequest(FFilePlatformRequest(PackageData->GetFileName(),
-							EInstigator::IterativeCook, TConstArrayView<const ITargetPlatform*>{ TargetPlatform }));
+						AddIdenticalCookedPackage(GeneratedPair.Key, false /* bRequireExists */);
 					}
+					Iter.RemoveCurrent();
+				}
+				else
+				{
+					DeferredEvaluateGeneratedNum += Iter->Value.Generated.Num();
 				}
 			}
+
+			int32 CookedNum = IdenticalCooked.Num() + ModifiedCookedNum + RemovedCookedNum + DeferredEvaluateGeneratedNum;
+			UE_LOG(LogCook, Display, TEXT("Found %d cooked package(s) in package store."), CookedNum);
+			UE_LOG(LogCook, Display, TEXT("Keeping %d. Recooking %d. Removing %d. %d generated packages to be evaluated for iterative skipping later."),
+				IdenticalCooked.Num(), ModifiedCookedNum, RemovedCookedNum, DeferredEvaluateGeneratedNum);
+			if (bFirstPlatform)
+			{
+				COOK_STAT(DetailedCookStats::NumPackagesIterativelySkipped += IdenticalCooked.Num());
+				bFirstPlatform = false;
+			}
+
+			PackageWriter.MarkPackagesUpToDate(IdenticalCooked);
+			PackageWriter.RemoveCookedPackages(PackagesToRemove);
 		}
 
 		for (TPair<FName, FAssetRegistryGenerator::FGeneratorPackageInfo> Pair : PreviousGeneratorPackages)
@@ -8512,10 +8596,7 @@ void UCookOnTheFlyServer::CookByTheBookFinished()
 				const FName& PlatformName = FName(*TargetPlatform->PlatformName());
 				FString PlatformNameString = PlatformName.ToString();
 
-				PackageDatas->GetCookedPackagesForPlatform(TargetPlatform, CookedPackageDatas, false, /* include successful */ true);
-
-				// ignore any packages which failed to cook
-				PackageDatas->GetCookedPackagesForPlatform(TargetPlatform, IgnorePackageDatas, /* include failed */ true, false);
+				PackageDatas->GetCookedPackagesForPlatform(TargetPlatform, CookedPackageDatas, IgnorePackageDatas);
 
 				bool bForceNoFilterAssetsFromAssetRegistry = false;
 
@@ -8718,6 +8799,8 @@ void UCookOnTheFlyServer::ShutdownCookSession()
 
 void UCookOnTheFlyServer::PrintFinishStats()
 {
+	using namespace UE::Cook;
+
 	const float TotalCookTime = (float)(FPlatformTime::Seconds() - CookByTheBookOptions->CookStartTime);
 	if (IsCookByTheBookMode())
 	{
@@ -8733,8 +8816,10 @@ void UCookOnTheFlyServer::PrintFinishStats()
 	UE_LOG(LogCook, Display, TEXT("Peak Used virtual %u MiB Peak Used physical %u MiB"), MemStats.PeakUsedVirtual / 1024 / 1024, MemStats.PeakUsedPhysical / 1024 / 1024);
 
 	COOK_STAT(UE_LOG(LogCook, Display, TEXT("Packages Cooked: %d, Packages Iteratively Skipped: %d, Packages Skipped by Platform: %d, Total Packages: %d"),
-		DetailedCookStats::NumPackagesSavedForCook, DetailedCookStats::NumPackagesIterativelySkipped, PackageDatas->GetNumFailedCooked(),
-		PackageDatas->GetNumCooked()));
+		PackageDatas->GetNumCooked(ECookResult::Succeeded) - DetailedCookStats::NumPackagesIterativelySkipped,
+		DetailedCookStats::NumPackagesIterativelySkipped,
+		PackageDatas->GetNumCooked(ECookResult::Failed),
+		PackageDatas->GetNumCooked() - PackageDatas->GetNumCooked(ECookResult::NeverCookPlaceholder)));
 }
 
 void UCookOnTheFlyServer::PrintDetailedCookStats()
@@ -8762,7 +8847,7 @@ TMap<FName, TSet<FName>> UCookOnTheFlyServer::BuildMapDependencyGraph(const ITar
 	TMap<FName, TSet<FName>> MapDependencyGraph;
 
 	TArray<UE::Cook::FPackageData*> PlatformCookedPackages;
-	PackageDatas->GetCookedPackagesForPlatform(TargetPlatform, PlatformCookedPackages, /* include failed */ true, /* include successful */ true);
+	PackageDatas->GetCookedPackagesForPlatform(TargetPlatform, PlatformCookedPackages, PlatformCookedPackages);
 
 	// assign chunks for all the map packages
 	for (const UE::Cook::FPackageData* const CookedPackage : PlatformCookedPackages)
@@ -8963,12 +9048,11 @@ void UCookOnTheFlyServer::ResetCook(TConstArrayView<TPair<const ITargetPlatform*
 			if (PlatformData)
 			{
 				bool bResetResults = Pair.Value;
-				PlatformData->bExplored = false;
+				PlatformData->SetExplored(false);
 				if (bResetResults)
 				{
-					PlatformData->bCookAttempted = false;
-					PlatformData->bCookSucceeded = false;
-					PlatformData->bSaveTimedOut = false;
+					PlatformData->SetSaveTimedOut(false);
+					PackageData->ClearCookProgress(TargetPlatform);
 				}
 			}
 		}
@@ -9317,8 +9401,6 @@ void UCookOnTheFlyServer::BeginCookSandbox(FBeginCookContext& BeginContext)
 		}
 	}
 
-	COOK_STAT(DetailedCookStats::NumPackagesSavedForCook = 0);
-
 #if OUTPUT_COOKTIMING
 	FString PlatformNames;
 	for (const ITargetPlatform* Target : BeginContext.TargetPlatforms)
@@ -10186,7 +10268,7 @@ void UCookOnTheFlyServer::RecordDLCPackagesFromBaseGame(FBeginCookContext& Begin
 		{
 			if (UE::Cook::FPackageData* PackageData = PackageDatas->TryAddPackageDataByFileName(FName(*AssetPath)))
 			{
-				PackageData->SetPlatformsCooked(BeginContext.TargetPlatforms, true /* Succeeded */);
+				PackageData->SetPlatformsCooked(BeginContext.TargetPlatforms, UE::Cook::ECookResult::Succeeded);
 			}
 			else
 			{
@@ -10468,7 +10550,7 @@ void UCookOnTheFlyServer::MaybeMarkPackageAsAlreadyLoaded(UPackage *Package)
 			FString Platforms;
 			for (const TPair<const ITargetPlatform*, UE::Cook::FPackageData::FPlatformData>& Pair : PackageData->GetPlatformDatas())
 			{
-				if (Pair.Value.bCookAttempted)
+				if (Pair.Value.IsCookAttempted())
 				{
 					Platforms += TEXT(" ");
 					Platforms += Pair.Key->PlatformName();
@@ -11137,7 +11219,7 @@ uint32 UCookOnTheFlyServer::CookFullLoadAndSave()
 
 				uint32 OriginalPackageFlags = Package->GetPackageFlags();
 
-				TArray<bool> SavePackageSuccessPerPlatform;
+				TArray<UE::Cook::ECookResult> SavePackageSuccessPerPlatform;
 				SavePackageSuccessPerPlatform.SetNum(TargetPlatforms.Num());
 				for (int32 PlatformIndex = 0; PlatformIndex < TargetPlatforms.Num(); ++PlatformIndex)
 				{
@@ -11203,7 +11285,10 @@ uint32 UCookOnTheFlyServer::CookFullLoadAndSave()
 						FAssetPackageData* AssetPackageData = nullptr;
 						{
 							FScopeLock AssetGeneratorScopeLock(&AssetGeneratorLock);
-							Generator.UpdateAssetRegistryData(*Package, SaveResult, MoveTemp(*CookContext.GetCookTagList()),
+							UE::Cook::ECookResult CookResult = SaveResult == ESavePackageResult::Success ?
+								UE::Cook::ECookResult::Succeeded : UE::Cook::ECookResult::Failed;
+							Generator.UpdateAssetRegistryData(Package->GetFName(), Package, CookResult, &SaveResult,
+								MoveTemp(*CookContext.GetCookTagList()),
 								false /* bIncludeDiskOnlyAssets */, TOptional<FAssetPackageData>(),
 								TOptional<TArray<FAssetDependency>>());
 							AssetPackageData = Generator.GetAssetPackageData(Package->GetFName());
@@ -11233,22 +11318,22 @@ uint32 UCookOnTheFlyServer::CookFullLoadAndSave()
 
 						if (SaveResult != ESavePackageResult::ReferencedOnlyByEditorOnlyData)
 						{
-							SavePackageSuccessPerPlatform[PlatformIndex] = true;
+							SavePackageSuccessPerPlatform[PlatformIndex] = UE::Cook::ECookResult::Succeeded;
 						}
 						else
 						{
-							SavePackageSuccessPerPlatform[PlatformIndex] = false;
+							SavePackageSuccessPerPlatform[PlatformIndex] = UE::Cook::ECookResult::Failed;
 						}
 					}
 					else
 					{
-						SavePackageSuccessPerPlatform[PlatformIndex] = false;
+						SavePackageSuccessPerPlatform[PlatformIndex] = UE::Cook::ECookResult::Failed;
 					}
 				}
 
 				PackageData.SetPlatformsCooked(TargetPlatforms, SavePackageSuccessPerPlatform);
 
-				if (SavePackageSuccessPerPlatform.Contains(false))
+				if (SavePackageSuccessPerPlatform.Contains(UE::Cook::ECookResult::Failed))
 				{
 					PackageTracker->UncookedEditorOnlyPackages.Add(PackageName);
 				}
diff --git a/Engine/Source/Editor/UnrealEd/Private/Cooker/CookMPCollector.h b/Engine/Source/Editor/UnrealEd/Private/Cooker/CookMPCollector.h
--- a/Engine/Source/Editor/UnrealEd/Private/Cooker/CookMPCollector.h
+++ b/Engine/Source/Editor/UnrealEd/Private/Cooker/CookMPCollector.h
@@ -48,7 +48,7 @@ public:
 	struct FPlatformData
 	{
 		const ITargetPlatform* TargetPlatform = nullptr;
-		bool bSuccessful = false;
+		ECookResult CookResults = ECookResult::NotAttempted;
 	};
 	TConstArrayView<const ITargetPlatform*> GetPlatforms() const { return Platforms; }
 	TConstArrayView<FPlatformData> GetPlatformDatas() const { return PlatformDatas; }
diff --git a/Engine/Source/Editor/UnrealEd/Private/Cooker/CookPackageData.cpp b/Engine/Source/Editor/UnrealEd/Private/Cooker/CookPackageData.cpp
--- a/Engine/Source/Editor/UnrealEd/Private/Cooker/CookPackageData.cpp
+++ b/Engine/Source/Editor/UnrealEd/Private/Cooker/CookPackageData.cpp
@@ -51,7 +51,8 @@ static FAutoConsoleVariableRef CVarPollAsyncPeriod(
 //////////////////////////////////////////////////////////////////////////
 // FPackageData
 FPackageData::FPlatformData::FPlatformData()
-	: bRequested(false), bCookAttempted(false), bCookSucceeded(false), bExplored(false), bSaveTimedOut(false), bCookable(true)
+	: bRequested(0), bExplored(0), bSaveTimedOut(0), bCookable(1), bIterativelySkipped(0)
+	, CookResults((uint8)ECookResult::NotAttempted)
 {
 }
 
@@ -60,7 +61,8 @@ FPackageData::FPackageData(FPackageDatas& PackageDatas, const FName& InPackageNa
 	, Instigator(EInstigator::NotYetRequested), bIsUrgent(0)
 	, bIsVisited(0), bIsPreloadAttempted(0)
 	, bIsPreloaded(0), bHasSaveCache(0), bPrepareSaveFailed(0), bPrepareSaveRequiresGC(0)
-	, bCookedPlatformDataStarted(0), bCookedPlatformDataCalled(0), bCookedPlatformDataComplete(0), bMonitorIsCooked(0)
+	, bCookedPlatformDataStarted(0), bCookedPlatformDataCalled(0), bCookedPlatformDataComplete(0)
+	, MonitorCookResult((uint8)ECookResult::NotAttempted)
 	, bInitializedGeneratorSave(0), bCompletedGeneration(0), bGenerated(0), bKeepReferencedDuringGC(0)
 {
 	SetState(EPackageState::Idle);
@@ -102,7 +104,7 @@ int32 FPackageData::GetNumRequestedPlatforms() const
 	int32 Result = 0;
 	for (const TPair<const ITargetPlatform*, FPlatformData>& Pair : PlatformDatas)
 	{
-		Result += Pair.Value.bRequested ? 1 : 0;
+		Result += Pair.Value.IsRequested() ? 1 : 0;
 	}
 	return Result;
 }
@@ -111,7 +113,7 @@ void FPackageData::SetPlatformsRequested(TConstArrayView<const ITargetPlatform*>
 {
 	for (const ITargetPlatform* TargetPlatform : TargetPlatforms)
 	{
-		PlatformDatas.FindOrAdd(TargetPlatform).bRequested = true;
+		PlatformDatas.FindOrAdd(TargetPlatform).SetRequested(true);
 	}
 }
 
@@ -119,7 +121,7 @@ void FPackageData::ClearRequestedPlatforms()
 {
 	for (TPair<const ITargetPlatform*, FPlatformData>& Pair : PlatformDatas)
 	{
-		Pair.Value.bRequested = false;
+		Pair.Value.SetRequested(false);
 	}
 }
 
@@ -137,7 +139,7 @@ bool FPackageData::HasAllRequestedPlatforms(const TArrayView<const ITargetPlatfo
 	for (const ITargetPlatform* QueryPlatform : Platforms)
 	{
 		const FPlatformData* PlatformData = PlatformDatas.Find(QueryPlatform);
-		if (!PlatformData || !PlatformData->bRequested)
+		if (!PlatformData || !PlatformData->IsRequested())
 		{
 			return false;
 		}
@@ -149,7 +151,8 @@ bool FPackageData::AreAllRequestedPlatformsCooked(bool bAllowFailedCooks) const
 {
 	for (const TPair<const ITargetPlatform*, FPlatformData>& Pair : PlatformDatas)
 	{
-		if (Pair.Value.bRequested && (!Pair.Value.bCookAttempted || (!bAllowFailedCooks && !Pair.Value.bCookSucceeded)))
+		if (Pair.Value.IsRequested() &&
+			(!Pair.Value.IsCookAttempted() || (!bAllowFailedCooks && !Pair.Value.IsCookSucceeded())))
 		{
 			return false;
 		}
@@ -161,7 +164,7 @@ bool FPackageData::AreAllRequestedPlatformsExplored() const
 {
 	for (const TPair<const ITargetPlatform*, FPlatformData>& Pair : PlatformDatas)
 	{
-		if (Pair.Value.bRequested && !Pair.Value.bExplored)
+		if (Pair.Value.IsRequested() && !Pair.Value.IsExplored())
 		{
 			return false;
 		}
@@ -183,7 +186,7 @@ bool FPackageData::HasAllExploredPlatforms(const TArrayView<const ITargetPlatfor
 	for (const ITargetPlatform* QueryPlatform : Platforms)
 	{
 		const FPlatformData* PlatformData = FindPlatformData(QueryPlatform);
-		if (!PlatformData || !PlatformData->bExplored)
+		if (!PlatformData || !PlatformData->IsExplored())
 		{
 			return false;
 		}
@@ -202,7 +205,7 @@ bool FPackageData::CanCookForPlatforms() const
 	{
 		for (const TPair<const ITargetPlatform*, FPlatformData>& Pair : PlatformDatas)
 		{
-			if (Pair.Value.bCookable)
+			if (Pair.Value.IsCookable())
 			{
 				return true;
 			}
@@ -279,56 +282,62 @@ void FPackageData::ClearInProgressData()
 }
 
 void FPackageData::SetPlatformsCooked(const TConstArrayView<const ITargetPlatform*> TargetPlatforms,
-	const TConstArrayView<bool> Succeeded)
+	const TConstArrayView<ECookResult> Result)
 {
-	check(TargetPlatforms.Num() == Succeeded.Num());
+	check(TargetPlatforms.Num() == Result.Num());
 	for (int32 n = 0; n < TargetPlatforms.Num(); ++n)
 	{
-		SetPlatformCooked(TargetPlatforms[n], Succeeded[n]);
+		SetPlatformCooked(TargetPlatforms[n], Result[n]);
 	}
 }
 
 void FPackageData::SetPlatformsCooked(const TConstArrayView<const ITargetPlatform*> TargetPlatforms,
-	bool bSucceeded)
+	ECookResult Result)
 {
 	for (const ITargetPlatform* TargetPlatform : TargetPlatforms)
 	{
-		SetPlatformCooked(TargetPlatform, bSucceeded);
+		SetPlatformCooked(TargetPlatform, Result);
 	}
 }
 
-void FPackageData::SetPlatformCooked(const ITargetPlatform* TargetPlatform, bool bSucceeded)
+void FPackageData::SetPlatformCooked(const ITargetPlatform* TargetPlatform, ECookResult CookResult)
 {
-	bool bHasAnyOthers = false;
-	bool bModified = false;
+	bool bNewCookAttemptedValue = (CookResult != ECookResult::NotAttempted);
+	bool bModifiedCookAttempted = false;
+	bool bHasAnyOtherCookAttempted = false;
 	bool bExists = false;
 	for (TPair<const ITargetPlatform*, FPlatformData>& Pair : PlatformDatas)
 	{
 		if (Pair.Key == TargetPlatform)
 		{
 			bExists = true;
-			bModified = bModified | (Pair.Value.bCookAttempted == false);
-			Pair.Value.bCookAttempted = true;
-			Pair.Value.bCookSucceeded = bSucceeded;
+			bModifiedCookAttempted = bModifiedCookAttempted | (Pair.Value.IsCookAttempted() != bNewCookAttemptedValue);
+			Pair.Value.SetCookResults(CookResult);
 			// Clear the SaveTimedOut when get a cook result, in case we save again later and need to allow retry again
-			Pair.Value.bSaveTimedOut = false;
+			Pair.Value.SetSaveTimedOut(false);
 		}
 		else
 		{
-			bHasAnyOthers = bHasAnyOthers | (Pair.Value.bCookAttempted != false);
+			bHasAnyOtherCookAttempted = bHasAnyOtherCookAttempted | Pair.Value.IsCookAttempted();
 		}
 	}
-	if (!bExists)
+	if (!bExists && bNewCookAttemptedValue)
 	{
 		FPlatformData& Value = PlatformDatas.FindOrAdd(TargetPlatform);
-		Value.bCookAttempted = true;
-		Value.bCookSucceeded = bSucceeded;
-		Value.bSaveTimedOut = false;
-		bModified = true;
+		Value.SetCookResults(CookResult);
+		Value.SetSaveTimedOut(false);
+		bModifiedCookAttempted = true;
 	}
-	if (bModified && !bHasAnyOthers)
+	if (bModifiedCookAttempted && !bHasAnyOtherCookAttempted)
 	{
-		PackageDatas.GetMonitor().OnFirstCookedPlatformAdded(*this, bSucceeded);
+		if (bNewCookAttemptedValue)
+		{
+			PackageDatas.GetMonitor().OnFirstCookedPlatformAdded(*this, CookResult);
+		}
+		else
+		{
+			PackageDatas.GetMonitor().OnLastCookedPlatformRemoved(*this);
+		}
 	}
 }
 
@@ -345,11 +354,10 @@ void FPackageData::ClearCookProgress()
 	bool bModifiedCookAttempted = false;
 	for (TPair<const ITargetPlatform*, FPlatformData>& Pair : PlatformDatas)
 	{
-		bModifiedCookAttempted = bModifiedCookAttempted | (Pair.Value.bCookAttempted != false);
-		Pair.Value.bCookAttempted = false;
-		Pair.Value.bCookSucceeded = false;
-		Pair.Value.bExplored = false;
-		Pair.Value.bSaveTimedOut = false;
+		bModifiedCookAttempted = bModifiedCookAttempted | Pair.Value.IsCookAttempted();
+		Pair.Value.SetCookResults(ECookResult::NotAttempted);
+		Pair.Value.SetExplored(false);
+		Pair.Value.SetSaveTimedOut(false);
 	}
 	if (bModifiedCookAttempted)
 	{
@@ -365,15 +373,14 @@ void FPackageData::ClearCookProgress(const ITargetPlatform* TargetPlatform)
 	{
 		if (Pair.Key == TargetPlatform)
 		{
-			bModifiedCookAttempted = bModifiedCookAttempted | (Pair.Value.bCookAttempted != false);
-			Pair.Value.bCookAttempted = false;
-			Pair.Value.bCookSucceeded = false;
-			Pair.Value.bExplored = false;
-			Pair.Value.bSaveTimedOut = false;
+			bModifiedCookAttempted = bModifiedCookAttempted | Pair.Value.IsCookAttempted();
+			Pair.Value.SetCookResults(ECookResult::NotAttempted);
+			Pair.Value.SetExplored(false);
+			Pair.Value.SetSaveTimedOut(false);
 		}
 		else
 		{
-			bHasAnyOthers = bHasAnyOthers | (Pair.Value.bCookAttempted != false);
+			bHasAnyOthers = bHasAnyOthers | Pair.Value.IsCookAttempted();
 		}
 	}
 	if (bModifiedCookAttempted && !bHasAnyOthers)
@@ -405,7 +412,7 @@ const FPackageData::FPlatformData* FPackageData::FindPlatformData(const ITargetP
 bool FPackageData::HasAnyCookedPlatform() const
 {
 	return Algo::AnyOf(PlatformDatas,
-		[](const TPair<const ITargetPlatform*, FPlatformData>& Pair) { return Pair.Value.bCookAttempted; });
+		[](const TPair<const ITargetPlatform*, FPlatformData>& Pair) { return Pair.Value.IsCookAttempted(); });
 }
 
 bool FPackageData::HasAnyCookedPlatforms(const TArrayView<const ITargetPlatform* const>& Platforms,
@@ -451,17 +458,17 @@ bool FPackageData::HasAllCookedPlatforms(const TArrayView<const ITargetPlatform*
 bool FPackageData::HasCookedPlatform(const ITargetPlatform* Platform, bool bIncludeFailed) const
 {
 	ECookResult Result = GetCookResults(Platform);
-	return (Result == ECookResult::Succeeded) | ((Result == ECookResult::Failed) & (bIncludeFailed != 0));
+	return (Result == ECookResult::Succeeded) | ((Result != ECookResult::NotAttempted) & (bIncludeFailed != 0));
 }
 
 ECookResult FPackageData::GetCookResults(const ITargetPlatform* Platform) const
 {
 	const FPlatformData* PlatformData = PlatformDatas.Find(Platform);
-	if (PlatformData && PlatformData->bCookAttempted)
+	if (PlatformData)
 	{
-		return PlatformData->bCookSucceeded ? ECookResult::Succeeded : ECookResult::Failed;
+		return PlatformData->GetCookResults();
 	}
-	return ECookResult::Unseen;
+	return ECookResult::NotAttempted;
 }
 
 UPackage* FPackageData::GetPackage() const
@@ -1358,8 +1365,11 @@ UE::Cook::FGeneratorPackage& FPackageData::CreateGeneratorPackage(const UObject*
 	else
 	{
 		GeneratorPackage->InitializeSave(InSplitDataObject, InCookPackageSplitterInstance);
-		// The earlier exit from SaveState should have reset the progress back to StartGeneratorSave or earlier
-		check(GeneratorPackage->GetOwnerInfo().GetSaveState() <= FCookGenerationInfo::ESaveState::StartPopulate);
+		if (GeneratorPackage->IsInitialized())
+		{
+			// The earlier exit from SaveState should have reset the progress back to StartGeneratorSave or earlier
+			check(GeneratorPackage->GetOwnerInfo().GetSaveState() <= FCookGenerationInfo::ESaveState::StartPopulate);
+		}
 	}
 	FGeneratorPackage* Result = GeneratorPackage.Get();
 	check(Result);
@@ -1433,21 +1443,6 @@ FGeneratorPackage::~FGeneratorPackage()
 
 	ConditionalNotifyCompletion(ICookPackageSplitter::ETeardown::Canceled);
 	ClearGeneratedPackages();
-	if (!PreviousGeneratedPackages.IsEmpty())
-	{
-		FPackageData& PackageData = GetOwner();
-		UCookOnTheFlyServer& COTFS =PackageData.GetPackageDatas().GetCookOnTheFlyServer();
-		TArray<const ITargetPlatform*, TInlineAllocator<1>> RequestedPlatforms;
-		PackageData.GetRequestedPlatforms(RequestedPlatforms);
-		for (FName PreviousPackageName : PreviousGeneratedPackages)
-		{
-			for (const ITargetPlatform* Platform : RequestedPlatforms)
-			{
-				COTFS.DeleteOutputForPackage(PreviousPackageName, Platform);
-			}
-		}
-		PreviousGeneratedPackages.Empty();
-	}
 }
 
 void FGeneratorPackage::ConditionalNotifyCompletion(ICookPackageSplitter::ETeardown Status)
@@ -1488,6 +1483,11 @@ bool FGeneratorPackage::TryGenerateList(UObject* OwnerObject, FPackageDatas& Pac
 	UCookOnTheFlyServer& COTFS = GetOwner().GetPackageDatas().GetCookOnTheFlyServer();
 	bool bHybridIterativeEnabled = COTFS.bHybridIterativeEnabled;
 
+	int32 NumIterativeModified = 0; 
+	int32 NumIterativeRemoved = 0;
+	int32 NumIterativePrevious = PreviousGeneratedPackages.Num();
+	TArray<FName> IdenticalGenerated;
+
 	for (ICookPackageSplitter::FGeneratedPackage& SplitterData : GeneratorDatas)
 	{
 		if (!SplitterData.GetCreateAsMap().IsSet())
@@ -1526,9 +1526,20 @@ bool FGeneratorPackage::TryGenerateList(UObject* OwnerObject, FPackageDatas& Pac
 		FCookGenerationInfo& GeneratedInfo = PackagesToGenerate.Emplace_GetRef(*PackageData, false /* bInGenerator */);
 		GeneratedInfo.RelativePath = MoveTemp(SplitterData.RelativePath);
 		GeneratedInfo.GeneratedRootPath = MoveTemp(SplitterData.GeneratedRootPath);
-		GeneratedInfo.Dependencies = MoveTemp(SplitterData.Dependencies);
-		Algo::Sort(GeneratedInfo.Dependencies, FNameLexicalLess());
-		GeneratedInfo.Dependencies.SetNum(Algo::Unique(GeneratedInfo.Dependencies));
+		GeneratedInfo.PackageDependencies = MoveTemp(SplitterData.PackageDependencies);
+		for (TArray<FAssetDependency>::TIterator Iter(GeneratedInfo.PackageDependencies); Iter; ++Iter)
+		{
+			if (Iter->Category != UE::AssetRegistry::EDependencyCategory::Package)
+			{
+				UE_LOG(LogCook, Error,
+					TEXT("PackageSplitter specified a dependency with category %d rather than category Package. Dependency will be ignored. Splitter=%s, Generated=%s."),
+					(int32)Iter->Category, *this->GetSplitDataObjectName().ToString(), *PackageName);
+				Iter.RemoveCurrent();
+			}
+		}
+		Algo::Sort(GeneratedInfo.PackageDependencies,
+			[](const FAssetDependency& A, const FAssetDependency& B) { return A.LexicalLess(B); });
+		GeneratedInfo.PackageDependencies.SetNum(Algo::Unique(GeneratedInfo.PackageDependencies));
 		GeneratedInfo.SetIsCreateAsMap(bCreateAsMap);
 		PackageData->SetGeneratedOwner(this);
 		PackageData->SetWorkerAssignmentConstraint(FWorkerId::Local());
@@ -1536,35 +1547,57 @@ bool FGeneratorPackage::TryGenerateList(UObject* OwnerObject, FPackageDatas& Pac
 		// Create the Guid from the GenerationHash and Dependencies
 		GeneratedInfo.CreateGuid();
 
-		if (!bHybridIterativeEnabled)
+		FGuid PreviousGuid;
+		if (PreviousGeneratedPackages.RemoveAndCopyValue(PackageFName, PreviousGuid) && !bHybridIterativeEnabled)
 		{
-			GeneratedInfo.IterativeCookValidateOrClear(*this, RequestedPlatforms);
+			bool bIdentical;
+			GeneratedInfo.IterativeCookValidateOrClear(*this, RequestedPlatforms, PreviousGuid, bIdentical);
+			if (bIdentical)
+			{
+				IdenticalGenerated.Add(PackageFName);
+			}
+			else
+			{
+				++NumIterativeModified;
+			}
 		}
 	}
 	if (!PreviousGeneratedPackages.IsEmpty())
 	{
-		TSet<FName> CurrentGeneratedPackageNames;
-		CurrentGeneratedPackageNames.Reserve(PackagesToGenerate.Num());
-		for (const FCookGenerationInfo& Info : PackagesToGenerate)
+		NumIterativeRemoved = PreviousGeneratedPackages.Num();
+		for (TPair<FName, FGuid>& Pair : PreviousGeneratedPackages)
 		{
-			CurrentGeneratedPackageNames.Add(Info.PackageData->GetPackageName());
-		}
-		for (FName PreviousPackageName : PreviousGeneratedPackages)
-		{
-			if (!CurrentGeneratedPackageNames.Contains(PreviousPackageName))
+			for (const ITargetPlatform* TargetPlatform : RequestedPlatforms)
 			{
-				for (const ITargetPlatform* TargetPlatform : RequestedPlatforms)
-				{
-					COTFS.DeleteOutputForPackage(PreviousPackageName, TargetPlatform);
-				}
+				COTFS.DeleteOutputForPackage(Pair.Key, TargetPlatform);
 			}
 		}
 		PreviousGeneratedPackages.Empty();
 	}
+	if (NumIterativePrevious > 0 && !bHybridIterativeEnabled)
+	{
+		UE_LOG(LogCook, Display, TEXT("Found %d cooked package(s) in package store for generator package %s."),
+			NumIterativePrevious, *WriteToString<256>(GetOwner().GetPackageName()));
+		UE_LOG(LogCook, Display, TEXT("Keeping %d. Recooking %d. Removing %d."),
+			IdenticalGenerated.Num(), NumIterativeModified, NumIterativeRemoved);
+		COOK_STAT(DetailedCookStats::NumPackagesIterativelySkipped += IdenticalGenerated.Num());
+
+		for (const ITargetPlatform* TargetPlatform : RequestedPlatforms)
+		{
+			COTFS.FindOrCreatePackageWriter(TargetPlatform).MarkPackagesUpToDate(IdenticalGenerated);
+		}
+	}
+
 	RemainingToPopulate = GeneratorDatas.Num() + 1; // GeneratedPackaged plus one for the Generator
 	return true;
 }
 
+void FGeneratorPackage::FetchNeverCookDependencies()
+{
+	check(IsInitialized());
+	NeverCookDependencies = CookPackageSplitterInstance->GetNeverCookDependencies();
+}
+
 FCookGenerationInfo* FGeneratorPackage::FindInfo(const FPackageData& PackageData)
 {
 	check(IsInitialized());
@@ -2049,12 +2082,12 @@ void FBeginCacheObjects::EndRound(int32 NumPlatforms)
 
 void FCookGenerationInfo::CreateGuid()
 {
-	// COOKPACKAGESPLITTERTODO: Add GenerationHash
 	FBlake3 Blake3;
+	Blake3.Update(&GenerationHash, sizeof(GenerationHash));
 	IAssetRegistry& AssetRegistry = IAssetRegistry::GetChecked();
-	for (FName DependencyName : Dependencies)
+	for (const FAssetDependency& Dependency : PackageDependencies)
 	{
-		TOptional<FAssetPackageData> DependencyData = AssetRegistry.GetAssetPackageDataCopy(DependencyName);
+		TOptional<FAssetPackageData> DependencyData = AssetRegistry.GetAssetPackageDataCopy(Dependency.AssetId.PackageName);
 		if (DependencyData)
 		{
 			PRAGMA_DISABLE_DEPRECATION_WARNINGS;
@@ -2068,21 +2101,51 @@ void FCookGenerationInfo::CreateGuid()
 }
 
 void FCookGenerationInfo::IterativeCookValidateOrClear(FGeneratorPackage& Generator,
-	TConstArrayView<const ITargetPlatform*> RequestedPlatforms)
+	TConstArrayView<const ITargetPlatform*> RequestedPlatforms, const FGuid& PreviousGuid, bool& bOutIdentical)
 {
-	// COOKPACKAGESPLITTERTODO: Test whether the Hash from the generator has changed
-	constexpr bool bInvalidated = false;
-	if (bInvalidated)
+	UCookOnTheFlyServer& COTFS = Generator.GetOwner().GetPackageDatas().GetCookOnTheFlyServer();
+	bOutIdentical = PreviousGuid == this->Guid;
+	if (bOutIdentical)
 	{
-		for (const ITargetPlatform* TargetPlatform : RequestedPlatforms)
+		// If not directly modified, mark it as indirectly modified if any of its dependencies
+		// were detected as modified during PopulateCookedPackages.
+		for (const FAssetDependency& Dependency : this->PackageDependencies)
 		{
-			FPackageData::FPlatformData* PlatformData = PackageData->FindPlatformData(TargetPlatform);
-			if (PlatformData && PlatformData->bCookAttempted)
+			FPackageData* DependencyData = COTFS.PackageDatas->FindPackageDataByPackageName(Dependency.AssetId.PackageName);
+			if (!DependencyData)
+			{
+				bOutIdentical = false;
+				break;
+			}
+			for (const ITargetPlatform* TargetPlatform : RequestedPlatforms)
 			{
-				PackageData->ClearCookProgress(TargetPlatform);
-				UCookOnTheFlyServer& COTFS = Generator.GetOwner().GetPackageDatas().GetCookOnTheFlyServer();
-				COTFS.DeleteOutputForPackage(PackageData->GetPackageName(), TargetPlatform);
+				FPackageData::FPlatformData* DependencyPlatformData = DependencyData->FindPlatformData(TargetPlatform);
+				if (!DependencyPlatformData || !DependencyPlatformData->IsIterativelySkipped())
+				{
+					bOutIdentical = false;
+					break;
+				}
 			}
+			if (!bOutIdentical)
+			{
+				break;
+			}
+		}
+	}
+
+	if (bOutIdentical)
+	{
+		for (const ITargetPlatform* TargetPlatform : RequestedPlatforms)
+		{
+			PackageData->SetPlatformCooked(TargetPlatform, ECookResult::Succeeded);
+			PackageData->FindOrAddPlatformData(TargetPlatform).SetIterativelySkipped(true);
+		}
+	}
+	else
+	{
+		for (const ITargetPlatform* TargetPlatform : RequestedPlatforms)
+		{
+			COTFS.DeleteOutputForPackage(PackageData->GetPackageName(), TargetPlatform);
 		}
 	}
 }
@@ -2240,14 +2303,9 @@ int32 FPackageDataMonitor::GetNumInProgress() const
 	return NumInProgress;
 }
 
-int32 FPackageDataMonitor::GetNumFailedCooked() const
+int32 FPackageDataMonitor::GetNumCooked(ECookResult Result) const
 {
-	return NumFailedCooked;
-}
-
-int32 FPackageDataMonitor::GetNumSuccessfulCooked() const
-{
-	return NumSuccessfulCooked;
+	return NumCooked[(uint8)Result];
 }
 
 void FPackageDataMonitor::OnInProgressChanged(FPackageData& PackageData, bool bInProgress)
@@ -2262,25 +2320,23 @@ void FPackageDataMonitor::OnPreloadAllocatedChanged(FPackageData& PackageData, b
 	check(NumPreloadAllocated >= 0);
 }
 
-void FPackageDataMonitor::OnFirstCookedPlatformAdded(FPackageData& PackageData, bool bSuccessful)
+void FPackageDataMonitor::OnFirstCookedPlatformAdded(FPackageData& PackageData, ECookResult CookResult)
 {
-	if (!PackageData.GetMonitorIsCooked())
+	check(CookResult != ECookResult::NotAttempted);
+	if (PackageData.GetMonitorCookResult() == ECookResult::NotAttempted)
 	{
-		PackageData.SetMonitorIsCooked(true);
-		PackageData.SetMonitorIsSuccessfulCooked(bSuccessful);
-		int32& Stat = (bSuccessful ? NumSuccessfulCooked : NumFailedCooked);
-		++Stat;
+		PackageData.SetMonitorCookResult(CookResult);
+		++NumCooked[(uint8)CookResult];
 	}
 }
 
 void FPackageDataMonitor::OnLastCookedPlatformRemoved(FPackageData& PackageData)
 {
-	if (PackageData.GetMonitorIsCooked())
+	ECookResult CookResult = PackageData.GetMonitorCookResult();
+	if (CookResult != ECookResult::NotAttempted)
 	{
-		int32& Stat = (PackageData.GetMonitorIsSuccessfulCooked() ? NumSuccessfulCooked : NumFailedCooked);
-		--Stat;
-		PackageData.SetMonitorIsCooked(false);
-		PackageData.SetMonitorIsSuccessfulCooked(false);
+		--NumCooked[(uint8)CookResult];
+		PackageData.SetMonitorCookResult(ECookResult::NotAttempted);
 	}
 }
 
@@ -2732,7 +2788,7 @@ void FPackageDatas::AddExistingPackageDatasForPlatform(TConstArrayView<FConstruc
 
 			PackageData = NewPackageData;
 		}
-		PackageData->SetPlatformCooked(TargetPlatform, true /* Succeeded */);
+		PackageData->SetPlatformCooked(TargetPlatform, ECookResult::Succeeded);
 	}
 }
 
@@ -2776,30 +2832,29 @@ FPackageData* FPackageDatas::UpdateFileName(FName PackageName)
 
 int32 FPackageDatas::GetNumCooked()
 {
-	return Monitor.GetNumSuccessfulCooked() + Monitor.GetNumFailedCooked();
-}
-
-int32 FPackageDatas::GetNumSuccessfulCooked()
-{
-	return Monitor.GetNumSuccessfulCooked();
+	int32 Count = 0;
+	for (uint8 CookResult = 0; CookResult < (uint8)ECookResult::Count; ++CookResult)
+	{
+		Count += Monitor.GetNumCooked((ECookResult)CookResult);
+	}
+	return Count;
 }
 
-int32 FPackageDatas::GetNumFailedCooked()
+int32 FPackageDatas::GetNumCooked(ECookResult CookResult)
 {
-	return Monitor.GetNumFailedCooked();
+	return Monitor.GetNumCooked(CookResult);
 }
 
-void FPackageDatas::GetCookedPackagesForPlatform(const ITargetPlatform* Platform, TArray<FPackageData*>& CookedPackages,
-	bool bGetFailedCookedPackages, bool bGetSuccessfulCookedPackages)
+void FPackageDatas::GetCookedPackagesForPlatform(const ITargetPlatform* Platform, TArray<FPackageData*>& SucceededPackages,
+	TArray<FPackageData*>& FailedPackages)
 {
 	LockAndEnumeratePackageDatas(
-	[Platform, &CookedPackages, bGetFailedCookedPackages, bGetSuccessfulCookedPackages](FPackageData* PackageData)
+	[Platform, &SucceededPackages, &FailedPackages](FPackageData* PackageData)
 	{
 		ECookResult CookResults = PackageData->GetCookResults(Platform);
-		if (((CookResults == ECookResult::Succeeded) & (bGetSuccessfulCookedPackages != 0)) |
-			((CookResults == ECookResult::Failed) & (bGetFailedCookedPackages != 0)))
+		if (CookResults != ECookResult::NotAttempted)
 		{
-			CookedPackages.Add(PackageData);
+			(CookResults == ECookResult::Succeeded ? SucceededPackages : FailedPackages).Add(PackageData);
 		}
 	});
 }
diff --git a/Engine/Source/Editor/UnrealEd/Private/Cooker/CookPackageData.h b/Engine/Source/Editor/UnrealEd/Private/Cooker/CookPackageData.h
--- a/Engine/Source/Editor/UnrealEd/Private/Cooker/CookPackageData.h
+++ b/Engine/Source/Editor/UnrealEd/Private/Cooker/CookPackageData.h
@@ -144,12 +144,39 @@ public:
 	struct FPlatformData
 	{
 		FPlatformData();
-		bool bRequested : 1;
-		bool bCookAttempted : 1;
-		bool bCookSucceeded : 1;
-		bool bExplored : 1;
-		bool bSaveTimedOut : 1;
-		bool bCookable : 1;
+
+		/** Requested from StartCookByTheBook, CookOnTheFly, or as a transitive dependency. */
+		bool IsRequested() const { return bRequested != 0; }
+		void SetRequested(bool bValue) { bRequested = (uint32) bValue; }
+
+		/** The package has passed through a CookRequestCluster and its transitive dependencies were added. */
+		bool IsExplored() const { return bExplored != 0; }
+		void SetExplored(bool bValue) { bExplored = (uint32) bValue; }
+
+		/** UPackage::Save was called on the package, but timedout and needs to be retried. */
+		bool IsSaveTimedOut() const { return bSaveTimedOut != 0; }
+		void SetSaveTimedOut(bool bValue) { bSaveTimedOut = (uint32) bValue; }
+
+		/** True if the package has either not passed through a CookRequestCluster, or it has and it was found cookable. */
+		bool IsCookable() const { return bCookable != 0; }
+		void SetCookable(bool bValue) { bCookable = (uint32) bValue; }
+
+		/** The package was found to be unmodified in the current iterative cook. */
+		bool IsIterativelySkipped() const { return bIterativelySkipped != 0; }
+		void SetIterativelySkipped(bool bValue) { bIterativelySkipped = (uint32) bValue; }
+
+		ECookResult GetCookResults() const { return (ECookResult) CookResults; }
+		bool IsCookAttempted() const { return CookResults != (uint32) ECookResult::NotAttempted; }
+		bool IsCookSucceeded() const { return CookResults == (uint32) ECookResult::Succeeded; }
+		void SetCookResults(ECookResult Value) { CookResults = (uint32) Value; }
+
+	private:
+		uint32 bRequested : 1;
+		uint32 bExplored : 1;
+		uint32 bSaveTimedOut : 1;
+		uint32 bCookable : 1;
+		uint32 bIterativelySkipped : 1;
+		uint32 CookResults : (int)ECookResult::NumBits;
 	};
 
 	FPackageData(FPackageDatas& PackageDatas, const FName& InPackageName, const FName& InFileName);
@@ -197,7 +224,7 @@ public:
 		OutPlatforms.Reset(PlatformDatas.Num());
 		for (const TPair<const ITargetPlatform*, FPlatformData>& Pair : PlatformDatas)
 		{
-			if (Pair.Value.bRequested)
+			if (Pair.Value.IsRequested())
 			{
 				OutPlatforms.Add(Pair.Key);
 			}
@@ -266,9 +293,9 @@ public:
 	 * In version that takes two arrays, TargetPlatforms and Succeeded must be the same length.
 	 */
 	void SetPlatformsCooked(const TConstArrayView<const ITargetPlatform*> TargetPlatforms,
-		const TConstArrayView<bool> Succeeded);
-	void SetPlatformsCooked(const TConstArrayView<const ITargetPlatform*> TargetPlatforms, bool bSucceeded);
-	void SetPlatformCooked(const ITargetPlatform* TargetPlatform, bool bSucceeded);
+		const TConstArrayView<ECookResult> Succeeded);
+	void SetPlatformsCooked(const TConstArrayView<const ITargetPlatform*> TargetPlatforms, ECookResult Result);
+	void SetPlatformCooked(const ITargetPlatform* TargetPlatform, ECookResult Result);
 	/**
 	 * FindOrAdd each TargetPlatform and set its flags: CookAttempted=false.
 	 * In Version that takes no TargetPlatform, CookAttempted is cleared from all existing platforms.
@@ -308,7 +335,7 @@ public:
 	bool HasCookedPlatform(const ITargetPlatform* Platform, bool bIncludeFailed) const;
 	/**
 	 * Return the CookResult for the given platform.  If the platform has not been cooked,
-	 * returns ECookResult::Unseen, otherwise returns Succeeded or Failed depending on its success flag.
+	 * returns ECookResult::NotAttempted, otherwise returns whatever result has been set.
 	 */
 	ECookResult GetCookResults(const ITargetPlatform* Platform) const;
 
@@ -319,7 +346,7 @@ public:
 		OutPlatforms.Reset(PlatformDatas.Num());
 		for (const TPair<const ITargetPlatform*, FPlatformData>& Pair : PlatformDatas)
 		{
-			if (Pair.Value.bRequested && !Pair.Value.bCookAttempted)
+			if (Pair.Value.IsRequested() && !Pair.Value.IsCookAttempted())
 			{
 				OutPlatforms.Add(Pair.Key);
 			}
@@ -458,11 +485,9 @@ public:
 	 */
 	void ClearCookedPlatformData();
 
-	/** Get/Set the PackageDataMonitor flag that counts whether this PackageData has finished cooking. */
-	bool GetMonitorIsCooked() const { return static_cast<bool>(bMonitorIsCooked); }
-	void SetMonitorIsCooked(bool Value) { bMonitorIsCooked = Value != 0; }
-	bool GetMonitorIsSuccessfulCooked() const { return static_cast<bool>(bMonitorIsSuccessfulCooked); }
-	void SetMonitorIsSuccessfulCooked(bool Value) { bMonitorIsSuccessfulCooked = Value != 0; }
+	/** Get/Set the PackageDataMonitor flag that counts whether this PackageData has finished cooking and with what result. */
+	ECookResult GetMonitorCookResult() const { return (ECookResult)MonitorCookResult; }
+	void SetMonitorCookResult(ECookResult Value) { MonitorCookResult = (uint8) Value; }
 
 	/** Remove all request data about the given platform from all fields in this PackageData. */
 	void OnRemoveSessionPlatform(const ITargetPlatform* Platform);
@@ -642,8 +667,7 @@ private:
 	uint32 bCookedPlatformDataStarted : 1;
 	uint32 bCookedPlatformDataCalled : 1;
 	uint32 bCookedPlatformDataComplete : 1;
-	uint32 bMonitorIsCooked : 1;
-	uint32 bMonitorIsSuccessfulCooked : 1;
+	uint32 MonitorCookResult : (int) ECookResult::NumBits;
 	uint32 bInitializedGeneratorSave : 1;
 	uint32 bCompletedGeneration : 1;
 	uint32 bGenerated : 1;
@@ -751,14 +775,16 @@ public:
 	 * test whether they need to be invalidated now and clear the results if so. Only called in legacy iterative;
 	 * incremental cooks handle invalidation by querying the TargetDomainDigest during the RequestCluster.
 	 */
-	void IterativeCookValidateOrClear(FGeneratorPackage& Generator, TConstArrayView<const ITargetPlatform*> RequestedPlatforms);
+	void IterativeCookValidateOrClear(FGeneratorPackage& Generator,
+		TConstArrayView<const ITargetPlatform*> RequestedPlatforms, const FGuid& PreviousGuid, bool& bOutIdentical);
 
 public:
 	FBeginCacheObjects BeginCacheObjects;
 	FGuid Guid;
 	FString RelativePath;
 	FString GeneratedRootPath;
-	TArray<FName> Dependencies;
+	FBlake3Hash GenerationHash;
+	TArray<FAssetDependency> PackageDependencies;
 	FPackageData* PackageData = nullptr;
 	TArray<UPackage*> KeepReferencedPackages;
 private:
@@ -789,6 +815,8 @@ public:
 
 	/** Call the Splitter's GetGenerateList and create the PackageDatas */
 	bool TryGenerateList(UObject* OwnerObject, FPackageDatas& PackageDatas);
+	/** Call the Splitter's GetNeverCookDependencies and store them on this. */
+	void FetchNeverCookDependencies();
 
 	/** Accessor for the packages to generate */
 	TArrayView<UE::Cook::FCookGenerationInfo> GetPackagesToGenerate() { check(IsInitialized()); return PackagesToGenerate; }
@@ -831,7 +859,10 @@ public:
 
 	UPackage* GetOwnerPackage() const { return OwnerPackage.Get(); };
 	void SetOwnerPackage(UPackage* InPackage) { OwnerPackage = InPackage; }
-	void SetPreviousGeneratedPackages(TConstArrayView<FName> Packages) { PreviousGeneratedPackages = Packages; }
+	void SetPreviousGeneratedPackages(TMap<FName, FGuid>&& Packages) { PreviousGeneratedPackages = MoveTemp(Packages); }
+
+	TConstArrayView<FName> GetNeverCookDependencies() { check(IsInitialized()); return NeverCookDependencies; }
+	TArray<FName> ReleaseNeverCookDependencies() { TArray<FName> Result = MoveTemp(NeverCookDependencies); return Result; }
 
 private:
 	void ConditionalNotifyCompletion(ICookPackageSplitter::ETeardown Status);
@@ -845,7 +876,8 @@ private:
 	/** Recorded list of packages to generate from the splitter, and data we need about them */
 	TArray<FCookGenerationInfo> PackagesToGenerate;
 	TWeakObjectPtr<UPackage> OwnerPackage;
-	TArray<FName> PreviousGeneratedPackages;
+	TMap<FName, FGuid> PreviousGeneratedPackages;
+	TArray<FName> NeverCookDependencies;
 
 	int32 NextPopulateIndex = 0;
 	int32 RemainingToPopulate = 0;
@@ -944,8 +976,7 @@ public:
 	int32 GetNumInProgress() const;
 	int32 GetNumPreloadAllocated() const;
 	/** Report the number of packages that have cooked any platform. Used by CookCommandlet progress reporting. */
-	int32 GetNumSuccessfulCooked() const;
-	int32 GetNumFailedCooked() const;
+	int32 GetNumCooked(ECookResult CookResult) const;
 	/**
 	 * Report the number of FPackageData that are currently marked as urgent.
 	 * Used to check if a Pump function needs to exit to handle urgent PackageData in other states.
@@ -962,7 +993,7 @@ public:
 	void OnInProgressChanged(FPackageData& PackageData, bool bInProgress);
 	void OnPreloadAllocatedChanged(FPackageData& PackageData, bool bPreloadAllocated);
 	/** Callback called from FPackageData when it has set a platform to CookAttempted=true and it does not have any others cooked. */
-	void OnFirstCookedPlatformAdded(FPackageData& PackageData, bool bSuccessful);
+	void OnFirstCookedPlatformAdded(FPackageData& PackageData, ECookResult CookResult);
 	/** Callback called from FPackageData when it has set a platform to CookAttempted=false and it does not have any others cooked. */
 	void OnLastCookedPlatformRemoved(FPackageData& PackageData);
 	/** Callback called from FPackageData when it has changed its urgency. */
@@ -975,8 +1006,7 @@ private:
 	void TrackUrgentRequests(EPackageState State, int32 Delta);
 
 	int32 NumInProgress = 0;
-	int32 NumFailedCooked = 0;
-	int32 NumSuccessfulCooked = 0;
+	int32 NumCooked[(uint8)ECookResult::Count]{};
 	int32 NumPreloadAllocated = 0;
 	int32 NumUrgentInState[static_cast<uint32>(EPackageState::Count)];
 };
@@ -1215,14 +1245,14 @@ public:
 
 	/** Report the number of packages that have cooked any platform. Used by cook commandlet progress reporting. */
 	int32 GetNumCooked();
-	int32 GetNumSuccessfulCooked();
-	int32 GetNumFailedCooked();
+	int32 GetNumCooked(ECookResult CookResult);
 	/**
-	 * Append to CookedPackages all packages that have cooked any platform, either successfully if
-	 * bGetSuccessfulCookedPackages is true and/or unsuccessfully if bGetFailedCookedPackages is true.
+	 * Append to SucceededPackages all packages that have cooked any platform with success, and to
+	 * FailedPackages all packages that have cooked any platform with other results. The output variables are permitted
+	 * to point to the same array.
 	 */
-	void GetCookedPackagesForPlatform(const ITargetPlatform* Platform, TArray<FPackageData*>& CookedPackages,
-		bool bGetFailedCookedPackages, bool bGetSuccessfulCookedPackages);
+	void GetCookedPackagesForPlatform(const ITargetPlatform* Platform, TArray<FPackageData*>& SucceededPackages,
+		TArray<FPackageData*>& FailedPackages);
 
 	/**
 	 * Delete all PackageDatas and free all other memory used by this FPackageDatas.
diff --git a/Engine/Source/Editor/UnrealEd/Private/Cooker/CookProfiling.cpp b/Engine/Source/Editor/UnrealEd/Private/Cooker/CookProfiling.cpp
--- a/Engine/Source/Editor/UnrealEd/Private/Cooker/CookProfiling.cpp
+++ b/Engine/Source/Editor/UnrealEd/Private/Cooker/CookProfiling.cpp
@@ -246,7 +246,6 @@ namespace DetailedCookStats
 	int32 PeakSaveQueueSize = 0;
 	uint32 NumPreloadedDependencies = 0;
 	uint32 NumPackagesIterativelySkipped = 0;
-	uint32 NumPackagesSavedForCook = 0;
 	FCookStatsManager::FAutoRegisterCallback RegisterCookOnTheFlyServerStats([](FCookStatsManager::AddStatFuncRef AddStat)
 		{
 			AddStat(TEXT("Package.Load"), FCookStatsManager::CreateKeyValueArray(TEXT("NumPreloadedDependencies"), NumPreloadedDependencies));
diff --git a/Engine/Source/Editor/UnrealEd/Private/Cooker/CookProfiling.h b/Engine/Source/Editor/UnrealEd/Private/Cooker/CookProfiling.h
--- a/Engine/Source/Editor/UnrealEd/Private/Cooker/CookProfiling.h
+++ b/Engine/Source/Editor/UnrealEd/Private/Cooker/CookProfiling.h
@@ -156,8 +156,6 @@ extern bool IsUnversioned;
 // Stats tracked through FAutoRegisterCallback
 extern uint32 NumPreloadedDependencies;
 extern uint32 NumPackagesIterativelySkipped;
-/** Tracks the number of packages cooked - counted only once even if multiple platforms cooked. */
-extern uint32 NumPackagesSavedForCook;
 extern int32 PeakRequestQueueSize;
 extern int32 PeakLoadQueueSize;
 extern int32 PeakSaveQueueSize;
diff --git a/Engine/Source/Editor/UnrealEd/Private/Cooker/CookRequestCluster.cpp b/Engine/Source/Editor/UnrealEd/Private/Cooker/CookRequestCluster.cpp
--- a/Engine/Source/Editor/UnrealEd/Private/Cooker/CookRequestCluster.cpp
+++ b/Engine/Source/Editor/UnrealEd/Private/Cooker/CookRequestCluster.cpp
@@ -744,7 +744,7 @@ void FRequestCluster::FGraphSearch::VisitVertex(const FVertexData& VertexData)
 				if (IsIterativeEnabled(PackageName))
 				{
 					bFoundCachedTargetDomain = true;
-					PackageData->SetPlatformCooked(TargetPlatform, true);
+					PackageData->SetPlatformCooked(TargetPlatform, ECookResult::Succeeded);
 					PackageWriter->MarkPackagesUpToDate({ PackageName });
 					// Declare the package to the EDLCookInfo verification so we don't warn about missing exports from it
 					UE::SavePackageUtilities::EDLCookInfoAddIterativelySkippedPackage(PackageName);
@@ -827,8 +827,8 @@ void FRequestCluster::FGraphSearch::VisitVertex(const FVertexData& VertexData)
 	for (const ITargetPlatform* TargetPlatform : Cluster.Platforms)
 	{
 		FPackageData::FPlatformData& PlatformData = PackageData->FindOrAddPlatformData(TargetPlatform);
-		PlatformData.bExplored = true;
-		PlatformData.bCookable = VertexData.bCookable;
+		PlatformData.SetExplored(true);
+		PlatformData.SetCookable(VertexData.bCookable);
 	}
 	bool bAlreadyCooked = PackageData->AreAllRequestedPlatformsCooked(true /* bAllowFailedCooks */);
 
diff --git a/Engine/Source/Editor/UnrealEd/Private/Cooker/CookTypes.h b/Engine/Source/Editor/UnrealEd/Private/Cooker/CookTypes.h
--- a/Engine/Source/Editor/UnrealEd/Private/Cooker/CookTypes.h
+++ b/Engine/Source/Editor/UnrealEd/Private/Cooker/CookTypes.h
@@ -9,6 +9,7 @@
 #include "CookOnTheSide/CookOnTheFlyServer.h" // ECookTickFlags
 #include "DerivedDataRequestOwner.h"
 #include "HAL/LowLevelMemTracker.h"
+#include "HAL/PlatformMath.h"
 #include "HAL/Platform.h"
 #include "Logging/TokenizedMessage.h"
 #include "Misc/AssertionMacros.h"
@@ -102,12 +103,18 @@ namespace UE::Cook
 	};
 
 	/* The Result of a Cook */
-	enum class ECookResult
+	enum class ECookResult : uint8
 	{
-		Unseen,		/* The package has not finished cooking, or if it was previously cooked its result was removed due to e.g. modification of the package. */
-		Succeeded,  /* The package was saved with success. */
-		Failed,     /* The package was processed but failed to load or save. */
-		Skipped     /* For reporting the ECookResults specific to a request: the package was skipped due to e.g. already being cooked or being in NeverCook packages. */
+		/* CookResults have not yet been set */
+		NotAttempted,
+		/* The package was saved with success. */
+		Succeeded,
+		/* The package was processed but SavePackage failed. */
+		Failed,
+		/** The package is a NeverCook package that needs to be added to cookresults for dependency tracking. */
+		NeverCookPlaceholder,
+		Count,
+		NumBits= FPlatformMath::ConstExprCeilLogTwo(ECookResult::Count),
 	};
 
 	/** Return type for functions called reentrantly that can succeed,fail,or be incomplete */
diff --git a/Engine/Source/Editor/UnrealEd/Private/Cooker/CookWorkerClient.cpp b/Engine/Source/Editor/UnrealEd/Private/Cooker/CookWorkerClient.cpp
--- a/Engine/Source/Editor/UnrealEd/Private/Cooker/CookWorkerClient.cpp
+++ b/Engine/Source/Editor/UnrealEd/Private/Cooker/CookWorkerClient.cpp
@@ -191,6 +191,10 @@ void FCookWorkerClient::ReportPromoteToSaveComplete(FPackageData& PackageData)
 	Result->SetPackageName(PackageName);
 	Result->SetSuppressCookReason(ESuppressCookReason::InvalidSuppressCookReason);
 	Result->SetPlatforms(OrderedSessionPlatforms);
+	if (FGeneratorPackage* Generator = PackageData.GetGeneratorPackage(); Generator)
+	{
+		Result->SetNeverCookDependencies(Generator->ReleaseNeverCookDependencies());
+	}
 
 	int32 NumPlatforms = OrderedSessionPlatforms.Num();
 	for (int32 PlatformIndex = 0; PlatformIndex < NumPlatforms; ++PlatformIndex)
@@ -198,7 +202,7 @@ void FCookWorkerClient::ReportPromoteToSaveComplete(FPackageData& PackageData)
 		ITargetPlatform* TargetPlatform = OrderedSessionPlatforms[PlatformIndex];
 		FPackageRemoteResult::FPlatformResult& PlatformResults = Result->GetPlatforms()[PlatformIndex];
 		FPackageData::FPlatformData& PackagePlatformData = PackageData.FindOrAddPlatformData(TargetPlatform);
-		PlatformResults.SetSuccessful(PackagePlatformData.bCookSucceeded);
+		PlatformResults.SetCookResults(PackagePlatformData.GetCookResults());
 	}
 
 	ReportPackageMessage(PackageName, MoveTemp(ResultOwner));
@@ -213,7 +217,7 @@ void FCookWorkerClient::ReportPackageMessage(FName PackageName, TUniquePtr<FPack
 	for (FPackageRemoteResult::FPlatformResult& PlatformResult : Result->GetPlatforms())
 	{
 		ContextPlatformDatas.Add(FMPCollectorClientTickPackageContext::FPlatformData
-			{ PlatformResult.GetPlatform(), PlatformResult.IsSuccessful() });
+			{ PlatformResult.GetPlatform(), PlatformResult.GetCookResults() });
 	}
 	FMPCollectorClientTickPackageContext Context;
 	Context.PackageName = PackageName;
@@ -690,11 +694,12 @@ void FCookWorkerClient::AssignPackages(FAssignPackagesMessage& Message)
 		}
 
 		// We do not want CookWorkers to explore dependencies in CookRequestCluster because the Director did it already.
-		// Mmark the PackageDatas we get from the Director as already explored.
+		// Mark the PackageDatas we get from the Director as already explored.
 		for (const ITargetPlatform* TargetPlatform : OrderedSessionPlatforms)
 		{
 			FPackageData::FPlatformData& PlatformData = PackageData.FindOrAddPlatformData(TargetPlatform);
-			PlatformData.bExplored = true;
+			PlatformData.SetExplored(true);
+			PlatformData.SetCookable(true);
 		}
 
 		PackageData.SetRequestData(OrderedSessionPlatforms , false /* bInIsUrgent */, FCompletionCallback(),
diff --git a/Engine/Source/Editor/UnrealEd/Private/Cooker/CookWorkerServer.cpp b/Engine/Source/Editor/UnrealEd/Private/Cooker/CookWorkerServer.cpp
--- a/Engine/Source/Editor/UnrealEd/Private/Cooker/CookWorkerServer.cpp
+++ b/Engine/Source/Editor/UnrealEd/Private/Cooker/CookWorkerServer.cpp
@@ -654,9 +654,10 @@ void FCookWorkerServer::RecordResults(FPackageResultsMessage& Message)
 			{
 				ITargetPlatform* TargetPlatform = OrderedSessionPlatforms[PlatformIndex];
 				FPackageRemoteResult::FPlatformResult& PlatformResult = Result.GetPlatforms()[PlatformIndex];
-				PackageData->SetPlatformCooked(TargetPlatform, PlatformResult.IsSuccessful());
+				PackageData->SetPlatformCooked(TargetPlatform, PlatformResult.GetCookResults());
 				HandleReceivedPackagePlatformMessages(*PackageData, TargetPlatform, PlatformResult.ReleaseMessages());
 			}
+			COTFS.RecordNeverCookDependencies(Result.GetNeverCookDependencies());
 			if (Result.IsReferencedOnlyByEditorOnlyData())
 			{
 				COTFS.PackageTracker->UncookedEditorOnlyPackages.AddUnique(Result.GetPackageName());
diff --git a/Engine/Source/Editor/UnrealEd/Private/Cooker/PackageResultsMessage.cpp b/Engine/Source/Editor/UnrealEd/Private/Cooker/PackageResultsMessage.cpp
--- a/Engine/Source/Editor/UnrealEd/Private/Cooker/PackageResultsMessage.cpp
+++ b/Engine/Source/Editor/UnrealEd/Private/Cooker/PackageResultsMessage.cpp
@@ -142,11 +142,15 @@ void FPackageResultsMessage::Write(FCbWriter& Writer) const
 			for (const FPackageRemoteResult::FPlatformResult& PlatformResult : Result.Platforms)
 			{
 				Writer.BeginObject();
-				Writer << "S" << PlatformResult.bSuccessful;
+				Writer << "S" << (uint8) PlatformResult.GetCookResults();
 				WriteMessagesArray(Writer, PlatformResult.Messages);
 				Writer.EndObject();
 			}
 			Writer.EndArray();
+			if (!Result.NeverCookDependencies.IsEmpty())
+			{
+				Writer << "NCD" << Result.NeverCookDependencies;
+			}
 		}
 		Writer.EndObject();
 	}
@@ -198,13 +202,26 @@ bool FPackageResultsMessage::TryRead(FCbObjectView Object)
 			FPackageRemoteResult::FPlatformResult& PlatformResult = Result.Platforms.Emplace_GetRef();
 			FCbObjectView PlatformObject = PlatformField.AsObjectView();
 
-			PlatformResult.bSuccessful = PlatformObject["S"].AsBool();
+			uint8 CookResultsInt = PlatformObject["S"].AsUInt8();
+			if (CookResultsInt >= (uint8)ECookResult::Count)
+			{
+				return false;
+			}
+			PlatformResult.SetCookResults((ECookResult)CookResultsInt);
 
 			if (!TryReadMessagesArray(PlatformObject, PlatformResult.Messages))
 			{
 				return false;
 			}
 		}
+		FCbFieldView NeverCookDependenciesField = ResultObject["NCD"];
+		if (NeverCookDependenciesField.IsArray())
+		{
+			if (!LoadFromCompactBinary(NeverCookDependenciesField, Result.NeverCookDependencies))
+			{
+				return false;
+			}
+		}
 	}
 	return true;
 }
diff --git a/Engine/Source/Editor/UnrealEd/Private/Cooker/PackageResultsMessage.h b/Engine/Source/Editor/UnrealEd/Private/Cooker/PackageResultsMessage.h
--- a/Engine/Source/Editor/UnrealEd/Private/Cooker/PackageResultsMessage.h
+++ b/Engine/Source/Editor/UnrealEd/Private/Cooker/PackageResultsMessage.h
@@ -35,13 +35,13 @@ public:
 		TConstArrayView<UE::CompactBinaryTCP::FMarshalledMessage> GetMessages() const { return Messages; }
 		TArray<UE::CompactBinaryTCP::FMarshalledMessage> ReleaseMessages();
 
-		bool IsSuccessful() const { return bSuccessful; }
-		void SetSuccessful(bool bInSuccessful) { bSuccessful = bInSuccessful; }
+		ECookResult GetCookResults() const { return CookResults; }
+		void SetCookResults(ECookResult Value) { CookResults = Value; }
 
 	private:
 		TArray<UE::CompactBinaryTCP::FMarshalledMessage> Messages;
 		const ITargetPlatform* Platform = nullptr;
-		bool bSuccessful = false;
+		ECookResult CookResults = ECookResult::NotAttempted;
 
 		friend FPackageRemoteResult;
 		friend FPackageResultsMessage;
@@ -77,6 +77,9 @@ public:
 	TArray<FPlatformResult, TInlineAllocator<1>>& GetPlatforms() { return Platforms; }
 	void SetPlatforms(TConstArrayView<ITargetPlatform*> OrderedSessionPlatforms);
 
+	void SetNeverCookDependencies(TArray<FName>&& InNeverCookDependencies) { NeverCookDependencies = MoveTemp(InNeverCookDependencies); }
+	TConstArrayView<FName> GetNeverCookDependencies() const { return NeverCookDependencies; }
+
 	/**
 	 * A non-atomic RefCount that can be used for storage of a refcount by the user (e.g. CookWorkerClient)
 	 * If used from multiple threads, the user must access it only within the user's external critical section.
@@ -116,6 +119,7 @@ private:
 private:
 	// Fields read/writable only from the owner thread.
 	TArray<FAsyncMessage> AsyncMessages;
+	TArray<FName> NeverCookDependencies;
 	FName PackageName;
 	/** If failure reason is InvalidSuppressCookReason, it was saved. Otherwise, holds the suppression reason */
 	ESuppressCookReason SuppressCookReason;
diff --git a/Engine/Source/Editor/UnrealEd/Public/Commandlets/AssetRegistryGenerator.h b/Engine/Source/Editor/UnrealEd/Public/Commandlets/AssetRegistryGenerator.h
--- a/Engine/Source/Editor/UnrealEd/Public/Commandlets/AssetRegistryGenerator.h
+++ b/Engine/Source/Editor/UnrealEd/Public/Commandlets/AssetRegistryGenerator.h
@@ -63,27 +63,47 @@ public:
 	/** Info about a GeneratorPackage (see ICookPackageSplitter) loaded from previous iterative cooks. */
 	struct FGeneratorPackageInfo
 	{
-		TArray<FName> Generated;
-		bool bGeneratedModified = false;
+		TMap<FName, FGuid> Generated;
 	};
 
+	enum EDifference
+	{
+		// Cooked packages have files that can be loaded at runtime.
+		IdenticalCooked,
+		ModifiedCooked,
+		RemovedCooked,
+
+		// Uncooked packages either were skipped by the platform (including editor-only packages) or had an error during Save
+		IdenticalUncooked,
+		ModifiedUncooked,
+		RemovedUncooked,
+
+		// NeverCookPlaceholders are packages that were marked NeverCook, so we did not attempt to save them, but they
+		// were also marked as necessary for dependency testing, so we included them in the AssetRegistry
+		IdenticalNeverCookPlaceholder,
+		ModifiedNeverCookPlaceholder,
+		RemovedNeverCookPlaceholder,
+
+		// Scripts have entries in the AssetRegistry for dependency testing, but are embedded in the binary and do not
+		// have their own files
+		IdenticalScript,
+		ModifiedScript,
+		RemovedScript,
+	};
 	/**
 	 * Differences between the current and the previous state.
 	 */
 	struct FAssetRegistryDifference
 	{
-		/** ModifiedPackages list of packages which existed beforeand now, but need to be recooked */
-		TSet<FName> ModifiedPackages;
-		/** NewPackages list of packages that did not exist before, but exist now */
-		TSet<FName> NewPackages;
-		/** RemovedPackages list of packages that existed before, but do not any more */
-		TSet<FName> RemovedPackages;
-		/** IdenticalCookedPackages list of cooked packages that have not changed */
-		TSet<FName> IdenticalCookedPackages;
-		/** IdenticalUncookedPackages list of uncooked packages that have not changed.These were filtered out by platform or editor only */
-		TSet<FName> IdenticalUncookedPackages;
-		/** Set of Generator packages and the packages they generated. */
-		TMap<FName, FGeneratorPackageInfo> PreviousGeneratorPackages;
+		/** Collection of all non-generated packages contained in the previous cook, and their difference category. */
+		TMap<FName, EDifference> Packages;
+		/**
+		 * Collection of Generator packages and the packages they generated. The keys of the map are Generator
+		 * packages, and these packages exist in this->Packages. But the values include a list of Generated packages,
+		 * which are not present in Packages. Generated packages cannot be given a difference category until the
+		 * Generator runs, or if the Generator is unmodified.
+		 */
+		TMap<FName, FGeneratorPackageInfo> GeneratorPackages;
 	};
 
 	/** Sets asset registry from a previous run that is used for iterative cooking. */
@@ -99,8 +119,8 @@ public:
 	void ComputePackageDifferences(const FComputeDifferenceOptions& Options, const FAssetRegistryState& PreviousState, FAssetRegistryDifference& OutDifference);
 
 	/** Computes just the list of packages in the PreviousState that no longer exist in the current state. */
-	void ComputePackageRemovals(const FAssetRegistryState& PreviousState, TArray<FName>& RemovedPackages,
-		TMap<FName, FGeneratorPackageInfo>& PreviousGeneratorPackages);
+	void ComputePackageRemovals(const FAssetRegistryState& PreviousState, TArray<FName>& OutRemovedPackages,
+		TMap<FName, FGeneratorPackageInfo>& OutGeneratorPackages, int32& OutNumNeverCookPlaceHolderPackages);
 
 	/**
 	 * FinalizeChunkIDs 
@@ -211,7 +231,8 @@ public:
 	 * @param SavePackageResult The metadata to associate with the given package name
 	 * @param bIncludeOnlyDiskAssets Include only disk assets or else also enumerate memory assets
 	 */
-	void UpdateAssetRegistryData(const UPackage& Package, FSavePackageResultStruct& SavePackageResult,
+	void UpdateAssetRegistryData(FName PackageName, const UPackage* Package,
+		UE::Cook::ECookResult CookResult, FSavePackageResultStruct* SavePackageResult,
 		FCookTagList&& InArchiveCookTagList, bool bIncludeOnlyDiskAssets,
 		TOptional<FAssetPackageData>&& OverrideAssetPackageData,
 		TOptional<TArray<FAssetDependency>>&& OverridePackageDependencies);
@@ -455,8 +476,8 @@ class IAssetRegistryReporter
 public:
 	virtual ~IAssetRegistryReporter() {}
 
-	virtual void UpdateAssetRegistryData(FPackageData& PackageData, const UPackage& Package,
-		FSavePackageResultStruct& SavePackageResult, FCookTagList&& InArchiveCookTagList,
+	virtual void UpdateAssetRegistryData(FName PackageName, const UPackage* Package, UE::Cook::ECookResult CookResult,
+		FSavePackageResultStruct* SavePackageResult, FCookTagList&& InArchiveCookTagList,
 		bool bIncludeOnlyDiskAssets, TOptional<FAssetPackageData>&& OverrideAssetPackageData, 
 		TOptional<TArray<FAssetDependency>>&& OverridePackageDependencies) = 0;
 
@@ -470,13 +491,14 @@ public:
 	{
 	}
 
-	virtual void UpdateAssetRegistryData(FPackageData& PackageData, const UPackage& Package,
-		FSavePackageResultStruct& SavePackageResult, FCookTagList&& InArchiveCookTagList,
+	virtual void UpdateAssetRegistryData(FName PackageName, const UPackage* Package, UE::Cook::ECookResult CookResult,
+		FSavePackageResultStruct* SavePackageResult, FCookTagList&& InArchiveCookTagList,
 		bool bIncludeOnlyDiskAssets, TOptional<FAssetPackageData>&& OverrideAssetPackageData,
 		TOptional<TArray<FAssetDependency>>&& OverridePackageDependencies) override
 	{
-		Generator.UpdateAssetRegistryData(Package, SavePackageResult, MoveTemp(InArchiveCookTagList),
-			bIncludeOnlyDiskAssets, MoveTemp(OverrideAssetPackageData), MoveTemp(OverridePackageDependencies));
+		Generator.UpdateAssetRegistryData(PackageName, Package, CookResult, SavePackageResult,
+			MoveTemp(InArchiveCookTagList), bIncludeOnlyDiskAssets, MoveTemp(OverrideAssetPackageData),
+			MoveTemp(OverridePackageDependencies));
 	}
 
 private:
@@ -488,8 +510,8 @@ class FAssetRegistryReporterRemote : public IAssetRegistryReporter
 public:
 	FAssetRegistryReporterRemote(FCookWorkerClient& InClient, const ITargetPlatform* InTargetPlatform);
 
-	virtual void UpdateAssetRegistryData(FPackageData& PackageData, const UPackage& Package,
-		FSavePackageResultStruct& SavePackageResult, FCookTagList&& InArchiveCookTagList,
+	virtual void UpdateAssetRegistryData(FName PackageName, const UPackage* Package, UE::Cook::ECookResult CookResult,
+		FSavePackageResultStruct* SavePackageResult, FCookTagList&& InArchiveCookTagList,
 		bool bIncludeOnlyDiskAssets, TOptional<FAssetPackageData>&& OverrideAssetPackageData, 
 		TOptional<TArray<FAssetDependency>>&& OverridePackageDependencies) override;
 
diff --git a/Engine/Source/Editor/UnrealEd/Public/CookPackageSplitter.h b/Engine/Source/Editor/UnrealEd/Public/CookPackageSplitter.h
--- a/Engine/Source/Editor/UnrealEd/Public/CookPackageSplitter.h
+++ b/Engine/Source/Editor/UnrealEd/Public/CookPackageSplitter.h
@@ -8,6 +8,7 @@
 #include "Containers/List.h"
 #include "Containers/UnrealString.h"
 #include "CoreMinimal.h"
+#include "Hash/Blake3.h"
 #include "HAL/PlatformCrt.h"
 #include "HAL/PreprocessorHelpers.h"
 #include "Misc/Optional.h"
@@ -18,6 +19,7 @@
 class UClass;
 class UObject;
 class UPackage;
+struct FAssetDependency;
 template <typename FuncType> class TFunctionRef;
 
 /**
@@ -55,9 +57,22 @@ public:
 	/** Data sent to the cooker to describe each desired generated package */
 	struct FGeneratedPackage
 	{
+		/** Parent path for the generated package. If empty, uses the generator's package path. */
 		FString GeneratedRootPath;
-		FString RelativePath;		// Generated package relative to GeneratedRootPath/Generated_ if GeneratedRootPath member is specified, relative to  Parent/Generated otherwise.
-		TArray<FName> Dependencies; // LongPackageNames that the generated package references
+		/** Generated package relative to <GeneratedRootPath>/_Generated_. */
+		FString RelativePath;
+		UE_DEPRECATED(5.3, "Write to PackageDependencies instead")
+		TArray<FName> Dependencies;
+		/**
+		 * AssetRegistry dependencies for the generated package. AR dependencies cause packages to be added to the
+		 * current cook and cause invalidation of this package in iterative cooks if any dependencies change.
+		 */
+		TArray<FAssetDependency> PackageDependencies;
+		/**
+		 * Hash of the data used to construct the generated package that is not covered by the dependencies.
+		 * Changes to this hash will cause invalidation of the package during iterative cooks.
+		 */
+		FBlake3Hash GenerationHash;
 		/* GetGenerateList must specify true if the package will be a map (.umap, contains a UWorld or ULevel), else false */
 		void SetCreateAsMap(bool bInCreateAsMap) { bCreateAsMap = bInCreateAsMap; }
 		const TOptional<bool>& GetCreateAsMap() const { return bCreateAsMap; }
@@ -68,6 +83,12 @@ public:
 	/** Return the list of packages to generate. */
 	virtual TArray<FGeneratedPackage> GetGenerateList(const UPackage* OwnerPackage, const UObject* OwnerObject) = 0;
 
+	/**
+	 * Return a list of NeverCook packages that should be added to the cook results as uncooked packages for
+	 * use in detecting changes during iterative cooks. Used for support of ExternalActors; most splitters should not use this.
+	 */
+	virtual TArray<FName> GetNeverCookDependencies() const { return TArray<FName>(); }
+
 	/** Representation of a generated package that is provided when populating the generator package. */
 	struct FGeneratedPackageForPreSave
 	{
diff --git a/Engine/Source/Runtime/AssetRegistry/Public/AssetRegistry/IAssetRegistry.h b/Engine/Source/Runtime/AssetRegistry/Public/AssetRegistry/IAssetRegistry.h
--- a/Engine/Source/Runtime/AssetRegistry/Public/AssetRegistry/IAssetRegistry.h
+++ b/Engine/Source/Runtime/AssetRegistry/Public/AssetRegistry/IAssetRegistry.h
@@ -93,12 +93,31 @@ struct FAssetRegistryDependencyOptions
 struct FAssetDependency
 {
 	FAssetIdentifier AssetId;
-	UE::AssetRegistry::EDependencyCategory Category;
-	UE::AssetRegistry::EDependencyProperty Properties;
+	UE::AssetRegistry::EDependencyCategory Category = UE::AssetRegistry::EDependencyCategory::None;
+	UE::AssetRegistry::EDependencyProperty Properties = UE::AssetRegistry::EDependencyProperty::None;
+
 	bool operator==(const FAssetDependency& Other) const
 	{
 		return AssetId == Other.AssetId && Category == Other.Category && Properties == Other.Properties;
 	}
+	bool LexicalLess(const FAssetDependency& Other) const
+	{
+		if (!(AssetId == Other.AssetId))
+		{
+			return AssetId.LexicalLess(Other.AssetId);
+		}
+		if (Category != Other.Category)
+		{
+			return Category < Other.Category;
+		}
+		return Properties < Other.Properties;
+	}
+
+	static FAssetDependency PackageDependency(FName PackageName, UE::AssetRegistry::EDependencyProperty Properties)
+	{
+		return FAssetDependency{ FAssetIdentifier(PackageName), UE::AssetRegistry::EDependencyCategory::Package,
+			Properties };
+	}
 
 	ASSETREGISTRY_API void WriteCompactBinary(FCbWriter& Writer) const;
 private:
diff --git a/Engine/Source/Runtime/Engine/Private/WorldPartition/WorldPartitionCookPackageSplitter.cpp b/Engine/Source/Runtime/Engine/Private/WorldPartition/WorldPartitionCookPackageSplitter.cpp
--- a/Engine/Source/Runtime/Engine/Private/WorldPartition/WorldPartitionCookPackageSplitter.cpp
+++ b/Engine/Source/Runtime/Engine/Private/WorldPartition/WorldPartitionCookPackageSplitter.cpp
@@ -4,9 +4,11 @@
 
 #if WITH_EDITOR
 
+#include "AssetRegistry/IAssetRegistry.h"
 #include "Engine/Level.h"
 #include "WorldPartition/WorldPartition.h"
 #include "WorldPartition/WorldPartitionLog.h"
+#include "WorldPartition/WorldPartitionRuntimeCell.h"
 #include "Editor.h"
 
 // Register FWorldPartitionCookPackageSplitter for UWorld class
@@ -112,6 +114,25 @@ TArray<ICookPackageSplitter::FGeneratedPackage> FWorldPartitionCookPackageSplitt
 	return PackagesToGenerate;
 }
 
+TArray<FName> FWorldPartitionCookPackageSplitter::GetNeverCookDependencies() const
+{
+	TSet<FName> NeverCookDependencies;
+	for (const IWorldPartitionCookPackageGenerator* CookPackageGenerator : CookContext.GetCookPackageGenerators())
+	{
+		if (const TArray<FWorldPartitionCookPackage*>* CookPackages = CookContext.GetCookPackages(CookPackageGenerator))
+		{
+			for (const FWorldPartitionCookPackage* CookPackage : *CookPackages)
+			{
+				if (UWorldPartitionRuntimeCell* Cell = CookPackageGenerator->GetCellForPackage(*CookPackage))
+				{
+					NeverCookDependencies.Append(Cell->GetActorPackageNames());
+				}
+			}
+		}
+	}
+	return NeverCookDependencies.Array();
+}
+
 bool FWorldPartitionCookPackageSplitter::PopulateGeneratedPackage(UPackage* OwnerPackage, UObject* OwnerObject,
 	const FGeneratedPackageForPopulate& GeneratedPackage, TArray<UObject*>& OutObjectsToMove,
 	TArray<UPackage*>& OutModifiedPackages)
@@ -187,10 +208,21 @@ void FWorldPartitionCookPackageSplitter::BuildPackagesToGenerateList(TArray<ICoo
 				ICookPackageSplitter::FGeneratedPackage& GeneratedPackage = PackagesToGenerate.Emplace_GetRef();
 				GeneratedPackage.GeneratedRootPath = CookPackage->Root;
 				GeneratedPackage.RelativePath = CookPackage->RelativePath;
+				// GenerationHash is left as empty. All dependencies for the package's bytes are specified by PackageDependencies
 
 				CookPackage->Type == FWorldPartitionCookPackage::EType::Level ? GeneratedPackage.SetCreateAsMap(true) : GeneratedPackage.SetCreateAsMap(false);
 
-				// @todo_ow: Set dependencies once we get iterative cooking working
+				// Fill generated package dependencies for iterative cooking
+				if (UWorldPartitionRuntimeCell* Cell = CookPackageGenerator->GetCellForPackage(*CookPackage))
+				{
+					TSet<FName> ActorPackageNames = Cell->GetActorPackageNames();
+					GeneratedPackage.PackageDependencies.Reset(ActorPackageNames.Num());
+					for (FName ActorPackageName : ActorPackageNames)
+					{
+						GeneratedPackage.PackageDependencies.Add(FAssetDependency::PackageDependency(ActorPackageName,
+							UE::AssetRegistry::EDependencyProperty::Hard | UE::AssetRegistry::EDependencyProperty::Game));
+					}
+				}
 			}
 		}
 	}
diff --git a/Engine/Source/Runtime/Engine/Private/WorldPartition/WorldPartitionCookPackageSplitter.h b/Engine/Source/Runtime/Engine/Private/WorldPartition/WorldPartitionCookPackageSplitter.h
--- a/Engine/Source/Runtime/Engine/Private/WorldPartition/WorldPartitionCookPackageSplitter.h
+++ b/Engine/Source/Runtime/Engine/Private/WorldPartition/WorldPartitionCookPackageSplitter.h
@@ -19,6 +19,7 @@ public:
 	FWorldPartitionCookPackageSplitter();
 	virtual ~FWorldPartitionCookPackageSplitter();
 
+	TArray<FName> GetNeverCookDependencies() const override;
 	virtual void Teardown(ETeardown Status) override;
 	virtual bool UseInternalReferenceToAvoidGarbageCollect() override { return true; }
 	virtual TArray<ICookPackageSplitter::FGeneratedPackage> GetGenerateList(const UPackage* OwnerPackage, const UObject* OwnerObject) override;
