AssetEditorSubsystem : Show Warning Toast if attempting to edit cooked asset.

diff --git a/Engine/Source/Editor/UnrealEd/Private/Subsystems/AssetEditorSubsystem.cpp b/Engine/Source/Editor/UnrealEd/Private/Subsystems/AssetEditorSubsystem.cpp
--- a/Engine/Source/Editor/UnrealEd/Private/Subsystems/AssetEditorSubsystem.cpp
+++ b/Engine/Source/Editor/UnrealEd/Private/Subsystems/AssetEditorSubsystem.cpp
@@ -339,6 +339,10 @@ bool UAssetEditorSubsystem::OpenEditorForAsset(UObject* Asset, const EToolkitMod
 			// Don't open asset editors for cooked packages
 			if (Package->bIsCookedForEditor)
 			{
+				if (TSharedPtr<SNotificationItem> InfoItem = FSlateNotificationManager::Get().AddNotification(FNotificationInfo(LOCTEXT("NotifyBlockedByCookedAsset", "Unable to Edit Cooked asset"))))
+				{
+					InfoItem->SetCompletionState(SNotificationItem::CS_Fail);
+				}
 				return false;
 			}
 
