Fixed issue when overwriting a texture

diff --git a/Engine/Plugins/Interchange/Runtime/Source/Import/Private/Texture/InterchangeTextureFactory.cpp b/Engine/Plugins/Interchange/Runtime/Source/Import/Private/Texture/InterchangeTextureFactory.cpp
--- a/Engine/Plugins/Interchange/Runtime/Source/Import/Private/Texture/InterchangeTextureFactory.cpp
+++ b/Engine/Plugins/Interchange/Runtime/Source/Import/Private/Texture/InterchangeTextureFactory.cpp
@@ -35,6 +35,7 @@
 #include "InterchangeVolumeTextureNode.h"
 #include "Misc/ConfigCacheIni.h"
 #include "Misc/CoreStats.h"
+#include "Misc/Paths.h"
 #include "Nodes/InterchangeBaseNode.h"
 #include "ProfilingDebugging/CpuProfilerTrace.h"
 #include "Serialization/EditorBulkData.h"
@@ -1523,6 +1524,21 @@ UObject* UInterchangeTextureFactory::ImportAssetObject_GameThread(const FImportA
 		//This is a reimport, we are just re-updating the source data
 		Texture = ExistingAsset;
 	}
+	else
+	{
+		UInterchangeResultError_Generic* Message = AddMessage<UInterchangeResultError_Generic>();
+		Message->SourceAssetName = Arguments.SourceData->GetFilename();
+		Message->DestinationAssetName = Arguments.AssetName;
+		Message->AssetType = TextureClass;
+
+		const FText TargetClassName = FText::FromString(TextureClass->GetName());
+		const FText ExistingClassName = FText::FromString(ExistingAsset->GetClass()->GetName());
+		const FText AssetName = FText::FromString(Arguments.AssetName);
+		const FText FolderName = FText::FromString(FPaths::GetPath(Arguments.Parent->GetPathName()));
+		Message->Text = FText::Format(NSLOCTEXT("InterchangeTextureFactory", "ClassMismatch", "You cannot create a '{0}' asset named '{1}' in '{2}', as there is already a '{3}' asset with the same name in this folder."), TargetClassName, AssetName, FolderName, ExistingClassName);
+		
+		return nullptr;
+	}
 
 	if (!Texture)
 	{
diff --git a/Engine/Source/Editor/UnrealEd/Private/Editor.cpp b/Engine/Source/Editor/UnrealEd/Private/Editor.cpp
--- a/Engine/Source/Editor/UnrealEd/Private/Editor.cpp
+++ b/Engine/Source/Editor/UnrealEd/Private/Editor.cpp
@@ -411,6 +411,10 @@ UE::Interchange::FAssetImportResultRef FReimportManager::ReimportAsync(UObject*
 			{
 				if (bUseInterchangeFramework && CanReimportHandler->IsInterchangeFactory())
 				{
+					// Make sure SourceFilenames reflects the source filenames in Obj
+					SourceFilenames.Empty();
+					CanReimportHandler->CanReimport(Obj, SourceFilenames);
+
 					int32 RealSourceFileIndex = SourceFileIndex == INDEX_NONE ? 0 : SourceFileIndex;
 					int32 RealValidSourceFileIndex = SourceFilenames.IsValidIndex(RealSourceFileIndex) ? RealSourceFileIndex : 0;
 					UE::Interchange::FScopedSourceData ScopedSourceData(SourceFilenames[RealValidSourceFileIndex]);
@@ -422,6 +426,17 @@ UE::Interchange::FAssetImportResultRef FReimportManager::ReimportAsync(UObject*
 						ImportAssetParameters.ReimportAsset = Obj;
 						ImportAssetParameters.ReimportSourceIndex = SourceFileIndex;
 						UE::Interchange::FAssetImportResultRef ImportResult = InterchangeManager.ImportAssetAsync(FString(), ScopedSourceData.GetSourceData(), ImportAssetParameters);
+
+						TFunction<void(UE::Interchange::FImportResult&)> AppendAndBroadcastImportResultIfNeeded =
+							[](UE::Interchange::FImportResult& Result)
+						{
+							UInterchangeManager& InterchangeManager = UInterchangeManager::GetInterchangeManager();
+							TStrongObjectPtr<UInterchangeResultsContainer> ResultsContainer(Result.GetResults());
+							InterchangeManager.OnBatchImportComplete.Broadcast(ResultsContainer);
+						};
+
+						ImportResult->OnDone(AppendAndBroadcastImportResultIfNeeded);
+
 						return ImportResult;
 					}
 				}
