Fixed localization meta-data gather issues

* Inherited properties are no longer gathered for each derived type.
* The OutputKeys from the dashboard no longer contains errant quotes.
* The OutputKeys formatting no longer converts to FText internally.

diff --git a/Engine/Source/Editor/UnrealEd/Classes/Commandlets/GatherTextFromMetadataCommandlet.h b/Engine/Source/Editor/UnrealEd/Classes/Commandlets/GatherTextFromMetadataCommandlet.h
--- a/Engine/Source/Editor/UnrealEd/Classes/Commandlets/GatherTextFromMetadataCommandlet.h
+++ b/Engine/Source/Editor/UnrealEd/Classes/Commandlets/GatherTextFromMetadataCommandlet.h
@@ -28,7 +28,7 @@ public:
 	{
 		TArray<FString> InputKeys;
 		TArray<FString> OutputNamespaces;
-		TArray<FText> OutputKeys;
+		TArray<FString> OutputKeys;
 	};
 
 private:
diff --git a/Engine/Source/Editor/UnrealEd/Private/Commandlets/GatherTextFromMetadataCommandlet.cpp b/Engine/Source/Editor/UnrealEd/Private/Commandlets/GatherTextFromMetadataCommandlet.cpp
--- a/Engine/Source/Editor/UnrealEd/Private/Commandlets/GatherTextFromMetadataCommandlet.cpp
+++ b/Engine/Source/Editor/UnrealEd/Private/Commandlets/GatherTextFromMetadataCommandlet.cpp
@@ -253,12 +253,7 @@ int32 UGatherTextFromMetaDataCommandlet::Main( const FString& Params )
 	FGatherParameters Arguments;
 	GetStringArrayFromConfig(*SectionName, TEXT("InputKeys"), Arguments.InputKeys, GatherTextConfigPath);
 	GetStringArrayFromConfig(*SectionName, TEXT("OutputNamespaces"), Arguments.OutputNamespaces, GatherTextConfigPath);
-	TArray<FString> OutputKeys;
-	GetStringArrayFromConfig(*SectionName, TEXT("OutputKeys"), OutputKeys, GatherTextConfigPath);
-	for(const auto& OutputKey : OutputKeys)
-	{
-		Arguments.OutputKeys.Add(FText::FromString(OutputKey));
-	}
+	GetStringArrayFromConfig(*SectionName, TEXT("OutputKeys"), Arguments.OutputKeys, GatherTextConfigPath);
 
 	// Execute gather.
 	GatherTextFromUObjects(IncludePathFilters, ExcludePathFilters, Arguments);
@@ -319,7 +314,7 @@ void UGatherTextFromMetaDataCommandlet::GatherTextFromField(UField* Field, const
 	// For structs, also gather the new non-object field values.
 	if (UStruct* Struct = Cast<UStruct>(Field))
 	{
-		for (TFieldIterator<FField> FieldIt(Struct); FieldIt; ++FieldIt)
+		for (TFieldIterator<FField> FieldIt(Struct, EFieldIterationFlags::None); FieldIt; ++FieldIt)
 		{
 			GatherTextFromField(*FieldIt, Arguments, InPlatformName);
 		}
@@ -346,20 +341,20 @@ void UGatherTextFromMetaDataCommandlet::GatherTextFromField(UField* Field, const
 
 				for (int32 j = 0; j < Arguments.InputKeys.Num(); ++j)
 				{
-					FFormatNamedArguments PatternArguments;
-					PatternArguments.Add(TEXT("FieldPath"), FText::FromString(Enum->GetFullGroupName(false) + TEXT(".") + Enum->GetNameStringByIndex(i)));
+					FStringFormatNamedArguments PatternArguments;
+					PatternArguments.Add(TEXT("FieldPath"), Enum->GetFullGroupName(false) + TEXT(".") + Enum->GetNameStringByIndex(i));
 
 					if (Enum->HasMetaData(*Arguments.InputKeys[j], i))
 					{
 						const FString& MetaDataValue = Enum->GetMetaData(*Arguments.InputKeys[j], i);
 						if (!MetaDataValue.IsEmpty())
 						{
-							PatternArguments.Add(TEXT("MetaDataValue"), FText::FromString(MetaDataValue));
+							PatternArguments.Add(TEXT("MetaDataValue"), MetaDataValue);
 
 							const FString Namespace = Arguments.OutputNamespaces[j];
 							FLocItem LocItem(MetaDataValue);
 							FManifestContext Context;
-							Context.Key = FText::Format(Arguments.OutputKeys[j], PatternArguments).ToString();
+							Context.Key = FString::Format(*Arguments.OutputKeys[j], PatternArguments);
 							Context.SourceLocation = FString::Printf(TEXT("Meta-data for key %s of enum value %s of enum %s in %s"), *Arguments.InputKeys[j], *Enum->GetNameStringByIndex(i), *Enum->GetName(), *Enum->GetFullGroupName(true));
 							Context.PlatformName = InPlatformName;
 							GatherManifestHelper->AddSourceText(Namespace, LocItem, Context);
@@ -434,21 +429,21 @@ void UGatherTextFromMetaDataCommandlet::GatherTextFromFieldImpl(FieldType* Field
 {
 	for (int32 i = 0; i < Arguments.InputKeys.Num(); ++i)
 	{
-		FFormatNamedArguments PatternArguments;
-		PatternArguments.Add(TEXT("FieldPath"), FText::FromString(Field->GetFullGroupName(false)));
+		FStringFormatNamedArguments PatternArguments;
+		PatternArguments.Add(TEXT("FieldPath"), Field->GetFullGroupName(false));
 
 		if (Field->HasMetaData(*Arguments.InputKeys[i]))
 		{
 			const FString& MetaDataValue = Field->GetMetaData(*Arguments.InputKeys[i]);
 			if (!MetaDataValue.IsEmpty())
 			{
-				PatternArguments.Add(TEXT("MetaDataValue"), FText::FromString(MetaDataValue));
+				PatternArguments.Add(TEXT("MetaDataValue"), MetaDataValue);
 
 				const UStruct* FieldOwnerType = Field->GetOwnerStruct();
 				const FString Namespace = Arguments.OutputNamespaces[i];
 				FLocItem LocItem(MetaDataValue);
 				FManifestContext Context;
-				Context.Key = FText::Format(Arguments.OutputKeys[i], PatternArguments).ToString();
+				Context.Key = FString::Format(*Arguments.OutputKeys[i], PatternArguments);
 				Context.SourceLocation = FString::Printf(TEXT("Meta-data for key %s of member %s in %s (type: %s, owner: %s)"), *Arguments.InputKeys[i], *Field->GetName(), *Field->GetFullGroupName(true), *Field->GetClass()->GetName(), FieldOwnerType ? *FieldOwnerType->GetName() : TEXT("<null>"));
 				Context.PlatformName = InPlatformName;
 				GatherManifestHelper->AddSourceText(Namespace, LocItem, Context);
