Allow using ScriptStruct as a field owner to mean "all struct types"

This is comparable with using Object to mean "all class types"

#rnx

diff --git a/Engine/Source/Editor/UnrealEd/Private/Commandlets/GatherTextFromMetadataCommandlet.cpp b/Engine/Source/Editor/UnrealEd/Private/Commandlets/GatherTextFromMetadataCommandlet.cpp
--- a/Engine/Source/Editor/UnrealEd/Private/Commandlets/GatherTextFromMetadataCommandlet.cpp
+++ b/Engine/Source/Editor/UnrealEd/Private/Commandlets/GatherTextFromMetadataCommandlet.cpp
@@ -211,22 +211,28 @@ int32 UGatherTextFromMetaDataCommandlet::Main( const FString& Params )
 			TArray<FString> FieldOwnerTypeStrs;
 			GetStringArrayFromConfig(*SectionName, InConfigKey, FieldOwnerTypeStrs, GatherTextConfigPath);
 
-			TArray<const UStruct*> AllFieldOwnerTypes;
-			GetObjectsOfClass(UClass::StaticClass(), (TArray<UObject*>&)AllFieldOwnerTypes, false);
-			GetObjectsOfClass(UScriptStruct::StaticClass(), (TArray<UObject*>&)AllFieldOwnerTypes, false);
+			TArray<const UStruct*> AllFieldOwnerClassTypes;
+			TArray<const UStruct*> AllFieldOwnerScriptStructTypes;
+			GetObjectsOfClass(UClass::StaticClass(), (TArray<UObject*>&)AllFieldOwnerClassTypes, false);
+			GetObjectsOfClass(UScriptStruct::StaticClass(), (TArray<UObject*>&)AllFieldOwnerScriptStructTypes, false);
 
 			for (const FString& FieldOwnerTypeStr : FieldOwnerTypeStrs)
 			{
 				const bool bIsWildcard = FieldOwnerTypeStr.GetCharArray().Contains(TEXT('*')) || FieldOwnerTypeStr.GetCharArray().Contains(TEXT('?'));
 				if (bIsWildcard)
 				{
-					for (const UStruct* FieldOwnerType : AllFieldOwnerTypes)
+					auto AddFieldOwnersMatchingWildcard = [&FieldOwnerTypeStr, &OutFieldOwnerTypes](const TArray<const UStruct*>& AllFieldOwnerTypes)
 					{
-						if (FieldOwnerType->GetName().MatchesWildcard(FieldOwnerTypeStr))
+						for (const UStruct* FieldOwnerType : AllFieldOwnerTypes)
 						{
-							OutFieldOwnerTypes.Add(FieldOwnerType);
+							if (FieldOwnerType->GetName().MatchesWildcard(FieldOwnerTypeStr))
+							{
+								OutFieldOwnerTypes.Add(FieldOwnerType);
+							}
 						}
-					}
+					};
+					AddFieldOwnersMatchingWildcard(AllFieldOwnerClassTypes);
+					AddFieldOwnersMatchingWildcard(AllFieldOwnerScriptStructTypes);
 				}
 				else
 				{
@@ -242,6 +248,11 @@ int32 UGatherTextFromMetaDataCommandlet::Main( const FString& Params )
 					{
 						GetDerivedClasses(FieldOwnerClass, (TArray<UClass*>&)OutFieldOwnerTypes);
 					}
+					if (FieldOwnerType == UScriptStruct::StaticClass())
+					{
+						// Structs don't have a catch-all base, so we allow ScriptStruct to mean "all struct types"
+						OutFieldOwnerTypes.Append(AllFieldOwnerScriptStructTypes);
+					}
 				}
 			}
 		};
