WorldPartition - Builder commandlet: Added new option -RunningFromUnrealEd
* Skips the error(s) / warning(s) summary at the end of the commandlet execution, with the side effect of binding the process exit code to the actual result of the commandlet
* Enables FastExit to exit the process as soon as the commandlet finishes, which avoids possible crashes during shutdown. The engine isn't shutdown cleanly.

diff --git a/Engine/Source/Editor/UnrealEd/Private/Commandlets/WorldPartitionBuilderCommandlet.cpp b/Engine/Source/Editor/UnrealEd/Private/Commandlets/WorldPartitionBuilderCommandlet.cpp
--- a/Engine/Source/Editor/UnrealEd/Private/Commandlets/WorldPartitionBuilderCommandlet.cpp
+++ b/Engine/Source/Editor/UnrealEd/Private/Commandlets/WorldPartitionBuilderCommandlet.cpp
@@ -42,6 +42,12 @@ int32 UWorldPartitionBuilderCommandlet::Main(const FString& Params)
 		LogWorldPartitionBuilderCommandlet.SetVerbosity(ELogVerbosity::Verbose);
 	}
 
+	if (Switches.Contains(TEXT("RunningFromUnrealEd")))
+	{
+		ShowErrorCount = false;	// This has the side effect of making the process return code match the return code of the commandlet
+		FastExit = true;		// Faster exit which avoids crash during shutdown. The engine isn't shutdown cleanly.
+	}
+
 	// This will convert incomplete package name to a fully qualified path
 	FString WorldFilename;
 	if (!FPackageName::SearchForPackageOnDisk(Tokens[0], &Tokens[0], &WorldFilename))
diff --git a/Engine/Source/Editor/WorldPartitionEditor/Private/WorldPartitionEditorModule.cpp b/Engine/Source/Editor/WorldPartitionEditor/Private/WorldPartitionEditorModule.cpp
--- a/Engine/Source/Editor/WorldPartitionEditor/Private/WorldPartitionEditorModule.cpp
+++ b/Engine/Source/Editor/WorldPartitionEditor/Private/WorldPartitionEditorModule.cpp
@@ -402,6 +402,7 @@ void FWorldPartitionEditorModule::RunCommandletAsExternalProcess(const FString&
 	CommandletArgsArray.Add(TEXT("\"") + ProjectPath + TEXT('"'));
 	CommandletArgsArray.Add(TEXT("-BaseDir=\"") + FString(FPlatformProcess::BaseDir()) + TEXT('"'));
 	CommandletArgsArray.Add(TEXT("-Unattended"));
+	CommandletArgsArray.Add(TEXT("-RunningFromUnrealEd"));
 	CommandletArgsArray.Add(InCommandletArgs);
 
 	OnExecuteCommandletEvent.Broadcast(CommandletArgsArray);
diff --git a/Engine/Source/Runtime/Engine/Classes/Commandlets/Commandlet.h b/Engine/Source/Runtime/Engine/Classes/Commandlets/Commandlet.h
--- a/Engine/Source/Runtime/Engine/Classes/Commandlets/Commandlet.h
+++ b/Engine/Source/Runtime/Engine/Classes/Commandlets/Commandlet.h
@@ -87,6 +87,10 @@ class UCommandlet : public UObject
 	UPROPERTY()
 	uint32 ShowProgress:1;
 
+	/** Whether to exit the process as soon as the commandlet completes, skipping the engine shutdown */
+	UPROPERTY()
+	uint32 FastExit:1;
+
 	/**
 	 * Entry point for your commandlet
 	 *
diff --git a/Engine/Source/Runtime/Engine/Private/Commandlets/Commandlet.cpp b/Engine/Source/Runtime/Engine/Private/Commandlets/Commandlet.cpp
--- a/Engine/Source/Runtime/Engine/Private/Commandlets/Commandlet.cpp
+++ b/Engine/Source/Runtime/Engine/Private/Commandlets/Commandlet.cpp
@@ -23,6 +23,7 @@ UCommandlet::UCommandlet(const FObjectInitializer& ObjectInitializer)
 	IsEditor = true;
 	ShowErrorCount = true;
 	ShowProgress = true;
+	FastExit = false;
 }
 
 
diff --git a/Engine/Source/Runtime/Launch/Private/LaunchEngineLoop.cpp b/Engine/Source/Runtime/Launch/Private/LaunchEngineLoop.cpp
--- a/Engine/Source/Runtime/Launch/Private/LaunchEngineLoop.cpp
+++ b/Engine/Source/Runtime/Launch/Private/LaunchEngineLoop.cpp
@@ -4245,7 +4245,8 @@ int32 FEngineLoop::PreInitPostStartupScreen(const TCHAR* CmdLine)
 				}
 			}
 
-			if (FParse::Param(FCommandLine::Get(), TEXT("fastexit")))
+			const bool bShouldPerformFastExit = Commandlet->FastExit || FParse::Param(FCommandLine::Get(), TEXT("fastexit"));
+			if (bShouldPerformFastExit)
 			{
 				FPlatformMisc::RequestExitWithStatus(true, ErrorLevel);
 			}
