Cooker: Iterative DLC Cooks and WorldPartition.
* Fix bugs preventing iterative DLC cooks from cooking iteratively.
* Fix progress display in DLC cooks to only report the new cooked packages.
* Removed GetNeverCookDependencies() from CookPackageSplitter since  we can get the information from AssetRegistry and GeneratedPackage dependencies.
[RN] Minor, Cooking

diff --git a/Engine/Plugins/Experimental/GameFeatures/Source/GameFeatures/Private/GameFeatureDataAssetDependencyGatherer.cpp b/Engine/Plugins/Experimental/GameFeatures/Source/GameFeatures/Private/GameFeatureDataAssetDependencyGatherer.cpp
--- a/Engine/Plugins/Experimental/GameFeatures/Source/GameFeatures/Private/GameFeatureDataAssetDependencyGatherer.cpp
+++ b/Engine/Plugins/Experimental/GameFeatures/Source/GameFeatures/Private/GameFeatureDataAssetDependencyGatherer.cpp
@@ -43,7 +43,8 @@ void FGameFeatureDataAssetDependencyGatherer::GatherDependencies(const FAssetDat
 
 			for (const FAssetData& FilteredAsset : FilteredAssets)
 			{
-				OutDependencies.Emplace(IAssetDependencyGatherer::FGathereredDependency{ FilteredAsset.PackageName, UE::AssetRegistry::EDependencyProperty::Game });
+				OutDependencies.Emplace(IAssetDependencyGatherer::FGathereredDependency{ FilteredAsset.PackageName,
+					UE::AssetRegistry::EDependencyProperty::Game | UE::AssetRegistry::EDependencyProperty::Build });
 			}
 		}
 	}
diff --git a/Engine/Source/Editor/UnrealEd/Classes/CookOnTheSide/CookOnTheFlyServer.h b/Engine/Source/Editor/UnrealEd/Classes/CookOnTheSide/CookOnTheFlyServer.h
--- a/Engine/Source/Editor/UnrealEd/Classes/CookOnTheSide/CookOnTheFlyServer.h
+++ b/Engine/Source/Editor/UnrealEd/Classes/CookOnTheSide/CookOnTheFlyServer.h
@@ -1271,8 +1271,8 @@ private:
 
 	/** Cook (save) a package and process the results */
 	void SaveCookedPackage(UE::Cook::FSaveCookedPackageContext& Context);
-	/** Helper for save of generator packages using ExternalActors: Record special-case ExternalActor packages for iterative builds. */
-	void RecordNeverCookDependencies(TConstArrayView<FName> NeverCookDependencies);
+	/** Helper for package saves using ExternalActors: record ExternalActors for iterative builds. */
+	void RecordExternalActorDependencies(TConstArrayView<FName> ExternalActorDependencies);
 	friend class UE::Cook::FSaveCookedPackageContext;
 
 	/**
@@ -1349,6 +1349,7 @@ private:
 	static UCookOnTheFlyServer* ActiveCOTFS;
 	uint32		StatLoadedPackageCount = 0;
 	uint32		StatSavedPackageCount = 0;
+	int32		PackageDataFromBaseGameNum = 0;
 	UE::Cook::FStatHistoryInt NumObjectsHistory;
 	UE::Cook::FStatHistoryInt VirtualMemoryHistory;
 
diff --git a/Engine/Source/Editor/UnrealEd/Private/Commandlets/AssetRegistryGenerator.cpp b/Engine/Source/Editor/UnrealEd/Private/Commandlets/AssetRegistryGenerator.cpp
--- a/Engine/Source/Editor/UnrealEd/Private/Commandlets/AssetRegistryGenerator.cpp
+++ b/Engine/Source/Editor/UnrealEd/Private/Commandlets/AssetRegistryGenerator.cpp
@@ -189,6 +189,7 @@ FAssetRegistryGenerator::FAssetRegistryGenerator(const ITargetPlatform* InPlatfo
 	: AssetRegistry(FModuleManager::LoadModuleChecked<FAssetRegistryModule>(TEXT("AssetRegistry")).Get())
 	, TargetPlatform(InPlatform)
 	, bGenerateChunks(false)
+	, bClonedGlobalAssetRegistry(false)
 	, HighestChunkId(0)
 	, DependencyInfo(*GetMutableDefault<UChunkDependencyInfo>())
 {
@@ -1108,11 +1109,11 @@ void FAssetRegistryGenerator::UpdateKeptPackages()
 {
 	for (TPair<FName, FIterativelySkippedPackageUpdateData>& Pair : PreviousPackagesToUpdate)
 	{
-		for (const FAssetData& PreviousAssetData : Pair.Value.AssetDatas)
+		for (FAssetData& PreviousAssetData : Pair.Value.AssetDatas)
 		{
-			State.UpdateAssetData(PreviousAssetData);
+			State.UpdateAssetData(MoveTemp(PreviousAssetData), true /* bCreateIfNotExists */);
 		}
-		*State.CreateOrGetAssetPackageData(Pair.Key) = Pair.Value.PackageData;
+		*State.CreateOrGetAssetPackageData(Pair.Key) = MoveTemp(Pair.Value.PackageData);
 
 		if (!Pair.Value.PackageDependencies.IsEmpty())
 		{
@@ -1123,6 +1124,7 @@ void FAssetRegistryGenerator::UpdateKeptPackages()
 			State.AddReferencers(FAssetIdentifier(Pair.Key), Pair.Value.PackageReferencers);
 		}
 	}
+	PreviousPackagesToUpdate.Empty();
 }
 
 static void AppendCookTagsToTagMap(TArray<TPair<FName, FString>>&& InTags, FAssetDataTagMap& OutTags)
@@ -1258,12 +1260,28 @@ void FAssetRegistryGenerator::Initialize(const TArray<FName> &InStartupPackages,
 
 	if (bInitializeFromExisting)
 	{
+		bClonedGlobalAssetRegistry = true;
 		AssetRegistry.InitializeTemporaryAssetRegistryState(State, SaveOptions);
 	}
 
 	FGameDelegates::Get().GetAssignLayerChunkDelegate() = FAssignLayerChunkDelegate::CreateStatic(AssignLayerChunkDelegate);
 }
 
+void FAssetRegistryGenerator::CloneGlobalAssetRegistryFilteredByPreviousState(const FAssetRegistryState& PreviousState)
+{
+	FAssetRegistrySerializationOptions SaveOptions;
+	AssetRegistry.InitializeSerializationOptions(SaveOptions, TargetPlatform->IniPlatformName());
+
+	TSet<FName> KeepPackages;
+	const TMap<FName, const FAssetPackageData*> &PreviousPackageDatas = PreviousState.GetAssetPackageDataMap();
+	KeepPackages.Reserve(PreviousPackageDatas.Num());
+	for (const TPair<FName, const FAssetPackageData*>& Pair : PreviousPackageDatas)
+	{
+		KeepPackages.Add(Pair.Key);
+	}
+	AssetRegistry.InitializeTemporaryAssetRegistryState(State, SaveOptions, false /* bRefreshExisting */, KeepPackages);
+}
+
 static UE::Cook::ECookResult DiskSizeToCookResult(int64 DiskSize)
 {
 	if (DiskSize >= 0)
@@ -1458,7 +1476,7 @@ void FAssetRegistryGenerator::ComputePackageDifferences(const FComputeDifference
 			// current state, then they must have changed themselves and so are already in the modified set.
 			TArray<FAssetIdentifier> Referencers;
 			State.GetReferencers(ModifiedPackage, Referencers, UE::AssetRegistry::EDependencyCategory::Package,
-				UE::AssetRegistry::EDependencyQuery::Hard | UE::AssetRegistry::EDependencyQuery::Build);
+				UE::AssetRegistry::EDependencyQuery::Build);
 
 			for (const FAssetIdentifier& Referencer : Referencers)
 			{
@@ -1625,11 +1643,11 @@ void FAssetRegistryGenerator::FinalizeChunkIDs(const TSet<FName>& InCookedPackag
 	// Development only packages should have been marked via UpdateAssetRegistryData with a negative size indicating why they were not cooked 
 	for (FName DevelopmentOnlyPackage : DevelopmentOnlyPackages)
 	{
-		FAssetPackageData* PackageData = State.CreateOrGetAssetPackageData(DevelopmentOnlyPackage);
-		if (PackageData->DiskSize >= 0)
+		FAssetPackageData* PackageData = State.GetAssetPackageData(DevelopmentOnlyPackage);
+		if (PackageData && PackageData->DiskSize >= 0)
 		{
 			UE_LOG(LogAssetRegistryGenerator, Warning,
-				TEXT("Package %s is a development-only package but was not marked with DiskSize explaining why it was not cooked. Marking it as a generic cook failure."),
+				TEXT("Package %s is a development-only package but was not marked with DiskSize explaining why it was not cooked. Marking it as a generic skip."),
 				*DevelopmentOnlyPackage.ToString());
 			PackageData->DiskSize = -1;
 		}
@@ -2788,7 +2806,6 @@ void FAssetRegistryGenerator::UpdateAssetRegistryData(FName PackageName, const U
 	PreviousPackagesToUpdate.Remove(PackageName);
 
 	uint32 NewPackageFlags = 0;
-	FAssetPackageData* AssetPackageData = GetAssetPackageData(PackageName);
 	bool bSaveSucceeded = CookResult == UE::Cook::ECookResult::Succeeded;
 
 	// Copy latest data for all Assets in the package into the cooked registry. This should be done even
@@ -2806,13 +2823,27 @@ void FAssetRegistryGenerator::UpdateAssetRegistryData(FName PackageName, const U
 		CreateOrFindAssetDatas(*Package);
 	}
 
-	// Set Overrides for generated packages
+	FAssetPackageData* AssetPackageData = GetAssetPackageData(PackageName);
+	if (!bClonedGlobalAssetRegistry)
+	{
+		// Copy the Dependencies and PackageData from the global AssetRegistry
+		TOptional<FAssetPackageData> GlobalPackageData = AssetRegistry.GetAssetPackageDataCopy(PackageName);
+		if (GlobalPackageData)
+		{
+			*AssetPackageData = MoveTemp(*GlobalPackageData);
+		}
+		TArray<FAssetDependency> GlobalDependencies;
+		AssetRegistry.GetDependencies(PackageName, GlobalDependencies, UE::AssetRegistry::EDependencyCategory::All);
+		State.SetDependencies(FAssetIdentifier(PackageName), GlobalDependencies, UE::AssetRegistry::EDependencyCategory::All);
+	}
 	if (OverrideAssetPackageData)
 	{
+		// Generated packages pass in their own AssetPackageData
 		*AssetPackageData = MoveTemp(*OverrideAssetPackageData);
 	}
 	if (OverridePackageDependencies)
 	{
+		// Generated and Generator packages pass in their own dependencies
 		SetOverridePackageDependencies(PackageName, *OverridePackageDependencies);
 	}
 
diff --git a/Engine/Source/Editor/UnrealEd/Private/CookOnTheFlyServer.cpp b/Engine/Source/Editor/UnrealEd/Private/CookOnTheFlyServer.cpp
--- a/Engine/Source/Editor/UnrealEd/Private/CookOnTheFlyServer.cpp
+++ b/Engine/Source/Editor/UnrealEd/Private/CookOnTheFlyServer.cpp
@@ -1621,7 +1621,7 @@ void UCookOnTheFlyServer::UpdateDisplay(UE::Cook::FTickStackData& StackData, boo
 		return;
 	}
 
-	const int32 CookedPackagesCount = PackageDatas->GetNumCooked() - PackageDatas->GetNumCooked(ECookResult::NeverCookPlaceholder);
+	const int32 CookedPackagesCount = PackageDatas->GetNumCooked() - PackageDatas->GetNumCooked(ECookResult::NeverCookPlaceholder) - PackageDataFromBaseGameNum;
 	const int32 CookPendingCount = WorkerRequests->GetNumExternalRequests() + PackageDatas->GetMonitor().GetNumInProgress();
 	if (bForceDisplay ||
 		(DeltaProgressDisplayTime >= GCookProgressUpdateTime && CookPendingCount != 0 &&
@@ -2876,9 +2876,10 @@ UE::Cook::EPollStatus UCookOnTheFlyServer::ConditionalCreateGeneratorPackage(UE:
 	Private::FRegisteredCookPackageSplitter* Splitter = nullptr;
 	UObject* SplitDataObject = nullptr;
 	bool bGeneratorExists = false;
+	bool bIncomplete = false;
 	ON_SCOPE_EXIT
 	{
-		if (!bGeneratorExists)
+		if (!bGeneratorExists && !bIncomplete)
 		{
 			// Destroy any old GeneratorPackage if we no longer find we need one
 			PackageData.DestroyGeneratorPackage();
@@ -2926,6 +2927,7 @@ UE::Cook::EPollStatus UCookOnTheFlyServer::ConditionalCreateGeneratorPackage(UE:
 
 	if (bPrecaching)
 	{
+		bIncomplete = true;
 		return EPollStatus::Incomplete;
 	}
 
@@ -5665,8 +5667,8 @@ void FSaveCookedPackageContext::FinishPlatform()
 
 			if (!bHasFirstPlatformResults)
 			{
-				GeneratorPackage->FetchNeverCookDependencies();
-				COTFS.RecordNeverCookDependencies(GeneratorPackage->GetNeverCookDependencies());
+				GeneratorPackage->FetchExternalActorDependencies();
+				COTFS.RecordExternalActorDependencies(GeneratorPackage->GetExternalActorDependencies());
 			}
 		}
 		else if (GeneratorPackage = PackageData.GetGeneratedOwner(); GeneratorPackage)
@@ -5833,7 +5835,7 @@ void FSaveCookedPackageContext::FinishPackage()
 
 } // namespace UE::Cook
 
-void UCookOnTheFlyServer::RecordNeverCookDependencies(TConstArrayView<FName> NeverCookDependencies)
+void UCookOnTheFlyServer::RecordExternalActorDependencies(TConstArrayView<FName> ExternalActorDependencies)
 {
 	using namespace UE::Cook;
 
@@ -5847,7 +5849,7 @@ void UCookOnTheFlyServer::RecordNeverCookDependencies(TConstArrayView<FName> Nev
 	// WorldPartitionCookPackageSplitter. They are marked as NeverCook, but we need to add them
 	// to the cook results so we can detect whether they change in iterative cooks. Call a function
 	// on the splitter to get a list of them and add them.
-	for (FName DependencyName : NeverCookDependencies)
+	for (FName DependencyName : ExternalActorDependencies)
 	{
 		FPackageData* DependencyData = PackageDatas->TryAddPackageDataByPackageName(DependencyName);
 		if (DependencyData)
@@ -7202,6 +7204,11 @@ void UCookOnTheFlyServer::PopulateCookedPackages(TArrayView<const ITargetPlatfor
 			continue;
 		}
 
+		if (!PlatformAssetRegistry.HasClonedGlobalAssetRegistry())
+		{
+			PlatformAssetRegistry.CloneGlobalAssetRegistryFilteredByPreviousState(*PreviousAssetRegistry);
+		}
+
 		TMap<FName, FAssetRegistryGenerator::FGeneratorPackageInfo> PreviousGeneratorPackages;
 		if (bHybridIterativeEnabled)
 		{
@@ -8816,10 +8823,10 @@ void UCookOnTheFlyServer::PrintFinishStats()
 	UE_LOG(LogCook, Display, TEXT("Peak Used virtual %u MiB Peak Used physical %u MiB"), MemStats.PeakUsedVirtual / 1024 / 1024, MemStats.PeakUsedPhysical / 1024 / 1024);
 
 	COOK_STAT(UE_LOG(LogCook, Display, TEXT("Packages Cooked: %d, Packages Iteratively Skipped: %d, Packages Skipped by Platform: %d, Total Packages: %d"),
-		PackageDatas->GetNumCooked(ECookResult::Succeeded) - DetailedCookStats::NumPackagesIterativelySkipped,
+		PackageDatas->GetNumCooked(ECookResult::Succeeded) - DetailedCookStats::NumPackagesIterativelySkipped - PackageDataFromBaseGameNum,
 		DetailedCookStats::NumPackagesIterativelySkipped,
 		PackageDatas->GetNumCooked(ECookResult::Failed),
-		PackageDatas->GetNumCooked() - PackageDatas->GetNumCooked(ECookResult::NeverCookPlaceholder)));
+		PackageDatas->GetNumCooked() - PackageDatas->GetNumCooked(ECookResult::NeverCookPlaceholder) - PackageDataFromBaseGameNum));
 }
 
 void UCookOnTheFlyServer::PrintDetailedCookStats()
@@ -10245,7 +10252,8 @@ void UCookOnTheFlyServer::RecordDLCPackagesFromBaseGame(FBeginCookContext& Begin
 		if (ActivePackageList.Num() > 0)
 		{
 			SCOPED_BOOT_TIMING("AddPackageDataByFileNamesForPlatform");
-			PackageDatas->AddExistingPackageDatasForPlatform(ActivePackageList, TargetPlatform, bFirstAddExistingPackageDatas);
+			PackageDatas->AddExistingPackageDatasForPlatform(ActivePackageList, TargetPlatform,
+				bFirstAddExistingPackageDatas, PackageDataFromBaseGameNum);
 		}
 
 		TArray<FName>& PlatformBasedPackages = CookByTheBookOptions->BasedOnReleaseCookedPackages.FindOrAdd(PlatformName);
@@ -10269,6 +10277,7 @@ void UCookOnTheFlyServer::RecordDLCPackagesFromBaseGame(FBeginCookContext& Begin
 			if (UE::Cook::FPackageData* PackageData = PackageDatas->TryAddPackageDataByFileName(FName(*AssetPath)))
 			{
 				PackageData->SetPlatformsCooked(BeginContext.TargetPlatforms, UE::Cook::ECookResult::Succeeded);
+				++PackageDataFromBaseGameNum;
 			}
 			else
 			{
diff --git a/Engine/Source/Editor/UnrealEd/Private/Cooker/CookPackageData.cpp b/Engine/Source/Editor/UnrealEd/Private/Cooker/CookPackageData.cpp
--- a/Engine/Source/Editor/UnrealEd/Private/Cooker/CookPackageData.cpp
+++ b/Engine/Source/Editor/UnrealEd/Private/Cooker/CookPackageData.cpp
@@ -14,6 +14,7 @@
 #include "Cooker/CookPlatformManager.h"
 #include "Cooker/CookRequestCluster.h"
 #include "Cooker/CookWorkerClient.h"
+#include "Cooker/PackageTracker.h"
 #include "CookOnTheSide/CookOnTheFlyServer.h"
 #include "Containers/StringView.h"
 #include "EditorDomain/EditorDomain.h"
@@ -1592,10 +1593,43 @@ bool FGeneratorPackage::TryGenerateList(UObject* OwnerObject, FPackageDatas& Pac
 	return true;
 }
 
-void FGeneratorPackage::FetchNeverCookDependencies()
+void FGeneratorPackage::FetchExternalActorDependencies()
 {
 	check(IsInitialized());
-	NeverCookDependencies = CookPackageSplitterInstance->GetNeverCookDependencies();
+
+	// The Generator package declares all its ExternalActor dependencies in its AssetRegistry dependencies
+	// The Generator's generated packages can also include ExternalActors from other maps due to level instancing,
+	// these are included in the dependencies reported by the Generator for each GeneratedPackage in the data
+	// returned from GetGenerateList. These sets will overlap; take the union.
+	ExternalActorDependencies.Reset();
+	IAssetRegistry::GetChecked().GetDependencies(GetOwner().GetPackageName(), ExternalActorDependencies,
+		UE::AssetRegistry::EDependencyCategory::Package);
+	for (const FCookGenerationInfo& Info : PackagesToGenerate)
+	{
+		ExternalActorDependencies.Reserve(Info.GetDependencies().Num() + ExternalActorDependencies.Num());
+		for (const FAssetDependency& Dependency : Info.GetDependencies())
+		{
+			ExternalActorDependencies.Add(Dependency.AssetId.PackageName);
+		}
+	}
+	Algo::Sort(ExternalActorDependencies, FNameFastLess());
+	ExternalActorDependencies.SetNum(Algo::Unique(ExternalActorDependencies));
+	FPackageDatas& PackageDatas = this->GetOwner().GetPackageDatas();
+	FThreadSafeSet<FName>& NeverCookPackageList =
+		GetOwner().GetPackageDatas().GetCookOnTheFlyServer().PackageTracker->NeverCookPackageList;
+
+	// Keep only on-disk PackageDatas that are marked as NeverCook
+	ExternalActorDependencies.RemoveAll([&PackageDatas, &NeverCookPackageList](FName PackageName)
+		{
+			FPackageData* PackageData = PackageDatas.TryAddPackageDataByPackageName(PackageName);
+			if (!PackageData)
+			{
+				return true;
+			}
+			bool bIsNeverCook = NeverCookPackageList.Contains(PackageData->GetFileName());
+			return !bIsNeverCook;
+		});
+	ExternalActorDependencies.Shrink();
 }
 
 FCookGenerationInfo* FGeneratorPackage::FindInfo(const FPackageData& PackageData)
@@ -2740,7 +2774,7 @@ FName FPackageDatas::GetStandardFileName(FStringView InFileName)
 }
 
 void FPackageDatas::AddExistingPackageDatasForPlatform(TConstArrayView<FConstructPackageData> ExistingPackages,
-	const ITargetPlatform* TargetPlatform, bool bExpectPackageDatasAreNew)
+	const ITargetPlatform* TargetPlatform, bool bExpectPackageDatasAreNew, int32& OutPackageDataFromBaseGameNum)
 {
 	int32 NumPackages = ExistingPackages.Num();
 	if (NumPackages == 0)
@@ -2790,6 +2824,7 @@ void FPackageDatas::AddExistingPackageDatasForPlatform(TConstArrayView<FConstruc
 		}
 		PackageData->SetPlatformCooked(TargetPlatform, ECookResult::Succeeded);
 	}
+	OutPackageDataFromBaseGameNum += ExistingPackages.Num();
 }
 
 FPackageData* FPackageDatas::UpdateFileName(FName PackageName)
diff --git a/Engine/Source/Editor/UnrealEd/Private/Cooker/CookPackageData.h b/Engine/Source/Editor/UnrealEd/Private/Cooker/CookPackageData.h
--- a/Engine/Source/Editor/UnrealEd/Private/Cooker/CookPackageData.h
+++ b/Engine/Source/Editor/UnrealEd/Private/Cooker/CookPackageData.h
@@ -778,6 +778,8 @@ public:
 	void IterativeCookValidateOrClear(FGeneratorPackage& Generator,
 		TConstArrayView<const ITargetPlatform*> RequestedPlatforms, const FGuid& PreviousGuid, bool& bOutIdentical);
 
+	TConstArrayView<FAssetDependency> GetDependencies() const { return PackageDependencies; }
+
 public:
 	FBeginCacheObjects BeginCacheObjects;
 	FGuid Guid;
@@ -815,8 +817,8 @@ public:
 
 	/** Call the Splitter's GetGenerateList and create the PackageDatas */
 	bool TryGenerateList(UObject* OwnerObject, FPackageDatas& PackageDatas);
-	/** Call the Splitter's GetNeverCookDependencies and store them on this. */
-	void FetchNeverCookDependencies();
+	/** Record all dependencies from the Generator that are ExternalActor dependencies and store them on this. */
+	void FetchExternalActorDependencies();
 
 	/** Accessor for the packages to generate */
 	TArrayView<UE::Cook::FCookGenerationInfo> GetPackagesToGenerate() { check(IsInitialized()); return PackagesToGenerate; }
@@ -861,8 +863,8 @@ public:
 	void SetOwnerPackage(UPackage* InPackage) { OwnerPackage = InPackage; }
 	void SetPreviousGeneratedPackages(TMap<FName, FGuid>&& Packages) { PreviousGeneratedPackages = MoveTemp(Packages); }
 
-	TConstArrayView<FName> GetNeverCookDependencies() { check(IsInitialized()); return NeverCookDependencies; }
-	TArray<FName> ReleaseNeverCookDependencies() { TArray<FName> Result = MoveTemp(NeverCookDependencies); return Result; }
+	TConstArrayView<FName> GetExternalActorDependencies() { check(IsInitialized()); return ExternalActorDependencies; }
+	TArray<FName> ReleaseExternalActorDependencies() { TArray<FName> Result = MoveTemp(ExternalActorDependencies); return Result; }
 
 private:
 	void ConditionalNotifyCompletion(ICookPackageSplitter::ETeardown Status);
@@ -877,7 +879,7 @@ private:
 	TArray<FCookGenerationInfo> PackagesToGenerate;
 	TWeakObjectPtr<UPackage> OwnerPackage;
 	TMap<FName, FGuid> PreviousGeneratedPackages;
-	TArray<FName> NeverCookDependencies;
+	TArray<FName> ExternalActorDependencies;
 
 	int32 NextPopulateIndex = 0;
 	int32 RemainingToPopulate = 0;
@@ -1232,7 +1234,7 @@ public:
 
 	/** Create and mark-cooked a batch of PackageDatas, used by DLC for cooked-in-earlier-release packages. */
 	void AddExistingPackageDatasForPlatform(TConstArrayView<FConstructPackageData> ExistingPackages,
-		const ITargetPlatform* TargetPlatform, bool bExpectPackageDatasAreNew);
+		const ITargetPlatform* TargetPlatform, bool bExpectPackageDatasAreNew, int32& OutPackageDataFromBaseGameNum);
 
 	/**
 	 * Try to find the PackageData for the given PackageName.
diff --git a/Engine/Source/Editor/UnrealEd/Private/Cooker/CookWorkerClient.cpp b/Engine/Source/Editor/UnrealEd/Private/Cooker/CookWorkerClient.cpp
--- a/Engine/Source/Editor/UnrealEd/Private/Cooker/CookWorkerClient.cpp
+++ b/Engine/Source/Editor/UnrealEd/Private/Cooker/CookWorkerClient.cpp
@@ -193,7 +193,7 @@ void FCookWorkerClient::ReportPromoteToSaveComplete(FPackageData& PackageData)
 	Result->SetPlatforms(OrderedSessionPlatforms);
 	if (FGeneratorPackage* Generator = PackageData.GetGeneratorPackage(); Generator)
 	{
-		Result->SetNeverCookDependencies(Generator->ReleaseNeverCookDependencies());
+		Result->SetExternalActorDependencies(Generator->ReleaseExternalActorDependencies());
 	}
 
 	int32 NumPlatforms = OrderedSessionPlatforms.Num();
diff --git a/Engine/Source/Editor/UnrealEd/Private/Cooker/CookWorkerServer.cpp b/Engine/Source/Editor/UnrealEd/Private/Cooker/CookWorkerServer.cpp
--- a/Engine/Source/Editor/UnrealEd/Private/Cooker/CookWorkerServer.cpp
+++ b/Engine/Source/Editor/UnrealEd/Private/Cooker/CookWorkerServer.cpp
@@ -657,7 +657,7 @@ void FCookWorkerServer::RecordResults(FPackageResultsMessage& Message)
 				PackageData->SetPlatformCooked(TargetPlatform, PlatformResult.GetCookResults());
 				HandleReceivedPackagePlatformMessages(*PackageData, TargetPlatform, PlatformResult.ReleaseMessages());
 			}
-			COTFS.RecordNeverCookDependencies(Result.GetNeverCookDependencies());
+			COTFS.RecordExternalActorDependencies(Result.GetExternalActorDependencies());
 			if (Result.IsReferencedOnlyByEditorOnlyData())
 			{
 				COTFS.PackageTracker->UncookedEditorOnlyPackages.AddUnique(Result.GetPackageName());
diff --git a/Engine/Source/Editor/UnrealEd/Private/Cooker/PackageResultsMessage.cpp b/Engine/Source/Editor/UnrealEd/Private/Cooker/PackageResultsMessage.cpp
--- a/Engine/Source/Editor/UnrealEd/Private/Cooker/PackageResultsMessage.cpp
+++ b/Engine/Source/Editor/UnrealEd/Private/Cooker/PackageResultsMessage.cpp
@@ -147,9 +147,9 @@ void FPackageResultsMessage::Write(FCbWriter& Writer) const
 				Writer.EndObject();
 			}
 			Writer.EndArray();
-			if (!Result.NeverCookDependencies.IsEmpty())
+			if (!Result.ExternalActorDependencies.IsEmpty())
 			{
-				Writer << "NCD" << Result.NeverCookDependencies;
+				Writer << "EAD" << Result.ExternalActorDependencies;
 			}
 		}
 		Writer.EndObject();
@@ -214,10 +214,10 @@ bool FPackageResultsMessage::TryRead(FCbObjectView Object)
 				return false;
 			}
 		}
-		FCbFieldView NeverCookDependenciesField = ResultObject["NCD"];
-		if (NeverCookDependenciesField.IsArray())
+		FCbFieldView ExternalActorDependenciesField = ResultObject["EAD"];
+		if (ExternalActorDependenciesField.IsArray())
 		{
-			if (!LoadFromCompactBinary(NeverCookDependenciesField, Result.NeverCookDependencies))
+			if (!LoadFromCompactBinary(ExternalActorDependenciesField, Result.ExternalActorDependencies))
 			{
 				return false;
 			}
diff --git a/Engine/Source/Editor/UnrealEd/Private/Cooker/PackageResultsMessage.h b/Engine/Source/Editor/UnrealEd/Private/Cooker/PackageResultsMessage.h
--- a/Engine/Source/Editor/UnrealEd/Private/Cooker/PackageResultsMessage.h
+++ b/Engine/Source/Editor/UnrealEd/Private/Cooker/PackageResultsMessage.h
@@ -77,8 +77,8 @@ public:
 	TArray<FPlatformResult, TInlineAllocator<1>>& GetPlatforms() { return Platforms; }
 	void SetPlatforms(TConstArrayView<ITargetPlatform*> OrderedSessionPlatforms);
 
-	void SetNeverCookDependencies(TArray<FName>&& InNeverCookDependencies) { NeverCookDependencies = MoveTemp(InNeverCookDependencies); }
-	TConstArrayView<FName> GetNeverCookDependencies() const { return NeverCookDependencies; }
+	void SetExternalActorDependencies(TArray<FName>&& InExternalActorDependencies) { ExternalActorDependencies = MoveTemp(InExternalActorDependencies); }
+	TConstArrayView<FName> GetExternalActorDependencies() const { return ExternalActorDependencies; }
 
 	/**
 	 * A non-atomic RefCount that can be used for storage of a refcount by the user (e.g. CookWorkerClient)
@@ -119,7 +119,7 @@ private:
 private:
 	// Fields read/writable only from the owner thread.
 	TArray<FAsyncMessage> AsyncMessages;
-	TArray<FName> NeverCookDependencies;
+	TArray<FName> ExternalActorDependencies;
 	FName PackageName;
 	/** If failure reason is InvalidSuppressCookReason, it was saved. Otherwise, holds the suppression reason */
 	ESuppressCookReason SuppressCookReason;
diff --git a/Engine/Source/Editor/UnrealEd/Public/Commandlets/AssetRegistryGenerator.h b/Engine/Source/Editor/UnrealEd/Public/Commandlets/AssetRegistryGenerator.h
--- a/Engine/Source/Editor/UnrealEd/Public/Commandlets/AssetRegistryGenerator.h
+++ b/Engine/Source/Editor/UnrealEd/Public/Commandlets/AssetRegistryGenerator.h
@@ -46,9 +46,20 @@ public:
 	 * Initializes manifest generator - creates manifest lists, hooks up delegates.
 	 */
 	void Initialize(const TArray<FName> &StartupPackages, bool bInitializeFromExisting);
+	/**
+	 * Used by DLC which does not Clone the entire global AR. Copies data from the global assetregistry
+	 * for all packages present in the previous AssetRegistry
+	 */
+	void CloneGlobalAssetRegistryFilteredByPreviousState(const FAssetRegistryState& PreviousState);
 
 	const ITargetPlatform* GetTargetPlatform() const { return TargetPlatform; }
 
+	/**
+	 * Whether *this cloned the global the AssetRegistry during its construction.
+	 * If false, package data must be copied during UpdataAssetRegistryData.
+	 */
+	bool HasClonedGlobalAssetRegistry() const { return bClonedGlobalAssetRegistry; }
+
 	/**
 	 * Options when computing the differences between current and previous state.
 	 */
@@ -312,6 +323,7 @@ private:
 	TSet<FName> PackagesContainingMaps;
 	/** Should the chunks be generated or only asset registry */
 	bool bGenerateChunks;
+	bool bClonedGlobalAssetRegistry;
 	/** Highest chunk id, being used for geneating dependency tree */
 	int32 HighestChunkId;
 	/** Array of Maps with chunks<->packages assignments */
diff --git a/Engine/Source/Editor/UnrealEd/Public/CookPackageSplitter.h b/Engine/Source/Editor/UnrealEd/Public/CookPackageSplitter.h
--- a/Engine/Source/Editor/UnrealEd/Public/CookPackageSplitter.h
+++ b/Engine/Source/Editor/UnrealEd/Public/CookPackageSplitter.h
@@ -83,12 +83,6 @@ public:
 	/** Return the list of packages to generate. */
 	virtual TArray<FGeneratedPackage> GetGenerateList(const UPackage* OwnerPackage, const UObject* OwnerObject) = 0;
 
-	/**
-	 * Return a list of NeverCook packages that should be added to the cook results as uncooked packages for
-	 * use in detecting changes during iterative cooks. Used for support of ExternalActors; most splitters should not use this.
-	 */
-	virtual TArray<FName> GetNeverCookDependencies() const { return TArray<FName>(); }
-
 	/** Representation of a generated package that is provided when populating the generator package. */
 	struct FGeneratedPackageForPreSave
 	{
diff --git a/Engine/Source/Runtime/AssetRegistry/Private/AssetRegistry.cpp b/Engine/Source/Runtime/AssetRegistry/Private/AssetRegistry.cpp
--- a/Engine/Source/Runtime/AssetRegistry/Private/AssetRegistry.cpp
+++ b/Engine/Source/Runtime/AssetRegistry/Private/AssetRegistry.cpp
@@ -4474,13 +4474,34 @@ bool ReadAssetFile(FPackageReader& PackageReader, IAssetRegistry::FLoadPackageRe
 
 }
 
-void UAssetRegistryImpl::InitializeTemporaryAssetRegistryState(FAssetRegistryState& OutState, const FAssetRegistrySerializationOptions& Options, bool bRefreshExisting) const
+void UAssetRegistryImpl::InitializeTemporaryAssetRegistryState(FAssetRegistryState& OutState,
+	const FAssetRegistrySerializationOptions& Options, bool bRefreshExisting,
+	const TSet<FName>& RequiredPackages, const TSet<FName>& RemovePackages) const
 {
 	using FAssetDataMap = UE::AssetRegistry::Private::FAssetDataMap;
 
 	FReadScopeLock InterfaceScopeLock(InterfaceLock);
 	const FAssetRegistryState& State = GuardedData.GetState();
-	OutState.InitializeFromExisting(State.CachedAssets, State.CachedDependsNodes, State.CachedPackageData, Options, bRefreshExisting ? FAssetRegistryState::EInitializationMode::OnlyUpdateExisting : FAssetRegistryState::EInitializationMode::Rebuild);
+	if (!RequiredPackages.IsEmpty() || !RemovePackages.IsEmpty())
+	{
+		if (bRefreshExisting)
+		{
+			// InitializeFromExistingAndPrune does not support EInitializationMode so we have to Initialize and then Prune
+			OutState.InitializeFromExisting(State.CachedAssets, State.CachedDependsNodes, State.CachedPackageData, Options,
+				FAssetRegistryState::EInitializationMode::OnlyUpdateExisting);
+			OutState.PruneAssetData(RequiredPackages, RemovePackages, Options);
+		}
+		else
+		{
+			TSet<int32> UnusedChunksToKeep;
+			OutState.InitializeFromExistingAndPrune(State, RequiredPackages, RemovePackages, UnusedChunksToKeep, Options);
+		}
+	}
+	else
+	{
+		OutState.InitializeFromExisting(State.CachedAssets, State.CachedDependsNodes, State.CachedPackageData, Options,
+			bRefreshExisting ? FAssetRegistryState::EInitializationMode::OnlyUpdateExisting : FAssetRegistryState::EInitializationMode::Rebuild);
+	}
 }
 
 #if ASSET_REGISTRY_STATE_DUMPING_ENABLED
diff --git a/Engine/Source/Runtime/AssetRegistry/Private/AssetRegistry.h b/Engine/Source/Runtime/AssetRegistry/Private/AssetRegistry.h
--- a/Engine/Source/Runtime/AssetRegistry/Private/AssetRegistry.h
+++ b/Engine/Source/Runtime/AssetRegistry/Private/AssetRegistry.h
@@ -119,7 +119,9 @@ public:
 	virtual SIZE_T GetAllocatedSize(bool bLogDetailed = false) const override;
 	virtual void LoadPackageRegistryData(FArchive& Ar, FLoadPackageRegistryData& InOutData) const override;
 	virtual void LoadPackageRegistryData(const FString& PackageFilename, FLoadPackageRegistryData& InOutData) const override;
-	virtual void InitializeTemporaryAssetRegistryState(FAssetRegistryState& OutState, const FAssetRegistrySerializationOptions& Options, bool bRefreshExisting = false) const override;
+	virtual void InitializeTemporaryAssetRegistryState(FAssetRegistryState& OutState, const FAssetRegistrySerializationOptions& Options,
+		bool bRefreshExisting = false, const TSet<FName>& RequiredPackages = TSet<FName>(),
+		const TSet<FName>& RemovePackages = TSet<FName>()) const override;
 #if ASSET_REGISTRY_STATE_DUMPING_ENABLED
 	virtual void DumpState(const TArray<FString>& Arguments, TArray<FString>& OutPages, int32 LinesPerPage = 1) const override;
 #endif
diff --git a/Engine/Source/Runtime/AssetRegistry/Private/AssetRegistryState.cpp b/Engine/Source/Runtime/AssetRegistry/Private/AssetRegistryState.cpp
--- a/Engine/Source/Runtime/AssetRegistry/Private/AssetRegistryState.cpp
+++ b/Engine/Source/Runtime/AssetRegistry/Private/AssetRegistryState.cpp
@@ -2201,14 +2201,13 @@ void FAssetRegistryState::GetPrimaryAssetsIds(TSet<FPrimaryAssetId>& OutPrimaryA
 const FAssetPackageData* FAssetRegistryState::GetAssetPackageData(FName PackageName) const
 {
 	FAssetPackageData* const* FoundData = CachedPackageData.Find(PackageName);
-	if (FoundData)
-	{
-		return *FoundData;
-	}
-	else
-	{
-		return nullptr;
-	}
+	return FoundData ? *FoundData : nullptr;
+}
+
+ FAssetPackageData* FAssetRegistryState::GetAssetPackageData(FName PackageName)
+{
+	FAssetPackageData** FoundData = CachedPackageData.Find(PackageName);
+	return FoundData ? *FoundData : nullptr;
 }
 
 FAssetPackageData* FAssetRegistryState::CreateOrGetAssetPackageData(FName PackageName)
diff --git a/Engine/Source/Runtime/AssetRegistry/Public/AssetRegistry/AssetRegistryState.h b/Engine/Source/Runtime/AssetRegistry/Public/AssetRegistry/AssetRegistryState.h
--- a/Engine/Source/Runtime/AssetRegistry/Public/AssetRegistry/AssetRegistryState.h
+++ b/Engine/Source/Runtime/AssetRegistry/Public/AssetRegistry/AssetRegistryState.h
@@ -415,8 +415,9 @@ public:
 	/** Get the set of primary assets contained in this state */
 	void GetPrimaryAssetsIds(TSet<FPrimaryAssetId>& OutPrimaryAssets) const;
 
-	/** Returns non-editable pointer to the asset package data */
+	/** Returns pointer to the asset package data */
 	const FAssetPackageData* GetAssetPackageData(FName PackageName) const;
+	FAssetPackageData* GetAssetPackageData(FName PackageName);
 
 	/** Returns all package names */
 	void GetPackageNames(TArray<FName>& OutPackageNames) const
diff --git a/Engine/Source/Runtime/AssetRegistry/Public/AssetRegistry/IAssetRegistry.h b/Engine/Source/Runtime/AssetRegistry/Public/AssetRegistry/IAssetRegistry.h
--- a/Engine/Source/Runtime/AssetRegistry/Public/AssetRegistry/IAssetRegistry.h
+++ b/Engine/Source/Runtime/AssetRegistry/Public/AssetRegistry/IAssetRegistry.h
@@ -705,8 +705,12 @@ public:
 	 * @param OutState			This will be filled in with a copy of the asset data, platform data, and dependency data
 	 * @param Options			Serialization options that will be used to write this later
 	 * @param bRefreshExisting	If true, will not delete or add packages in OutState and will just update things that already exist
+	 * @param RequiredPackages  If non-empty, only packages in this set will be kept.
+	 * @param RemovePackages	If non-empty, packages in this set will not be kept.
 	 */
-	virtual void InitializeTemporaryAssetRegistryState(FAssetRegistryState& OutState, const FAssetRegistrySerializationOptions& Options, bool bRefreshExisting = false) const = 0;
+	virtual void InitializeTemporaryAssetRegistryState(FAssetRegistryState& OutState, const FAssetRegistrySerializationOptions& Options,
+		bool bRefreshExisting = false, const TSet<FName>& RequiredPackages = TSet<FName>(),
+		const TSet<FName>& RemovePackages = TSet<FName>()) const = 0;
 
 	UE_DEPRECATED(5.0, "Receiving a pointer is not threadsafe. Use other functions on IAssetRegistry to access the same data, or contact Epic Core team to add the threadsafe functions you require..")
 	virtual const FAssetRegistryState* GetAssetRegistryState() const = 0;
diff --git a/Engine/Source/Runtime/Engine/Private/ExternalActorAssetDependencyGatherer.cpp b/Engine/Source/Runtime/Engine/Private/ExternalActorAssetDependencyGatherer.cpp
--- a/Engine/Source/Runtime/Engine/Private/ExternalActorAssetDependencyGatherer.cpp
+++ b/Engine/Source/Runtime/Engine/Private/ExternalActorAssetDependencyGatherer.cpp
@@ -28,7 +28,8 @@ void FExternalActorAssetDependencyGatherer::GatherDependencies(const FAssetData&
 
 		for (const FAssetData& FilteredAsset : FilteredAssets)
 		{
-			OutDependencies.Emplace(IAssetDependencyGatherer::FGathereredDependency{ FilteredAsset.PackageName, UE::AssetRegistry::EDependencyProperty::Game });
+			OutDependencies.Emplace(IAssetDependencyGatherer::FGathereredDependency{ FilteredAsset.PackageName,
+				UE::AssetRegistry::EDependencyProperty::Game  | UE::AssetRegistry::EDependencyProperty::Build });
 		}
 	}
 }
diff --git a/Engine/Source/Runtime/Engine/Private/WorldPartition/WorldPartitionCookPackageSplitter.cpp b/Engine/Source/Runtime/Engine/Private/WorldPartition/WorldPartitionCookPackageSplitter.cpp
--- a/Engine/Source/Runtime/Engine/Private/WorldPartition/WorldPartitionCookPackageSplitter.cpp
+++ b/Engine/Source/Runtime/Engine/Private/WorldPartition/WorldPartitionCookPackageSplitter.cpp
@@ -114,25 +114,6 @@ TArray<ICookPackageSplitter::FGeneratedPackage> FWorldPartitionCookPackageSplitt
 	return PackagesToGenerate;
 }
 
-TArray<FName> FWorldPartitionCookPackageSplitter::GetNeverCookDependencies() const
-{
-	TSet<FName> NeverCookDependencies;
-	for (const IWorldPartitionCookPackageGenerator* CookPackageGenerator : CookContext.GetCookPackageGenerators())
-	{
-		if (const TArray<FWorldPartitionCookPackage*>* CookPackages = CookContext.GetCookPackages(CookPackageGenerator))
-		{
-			for (const FWorldPartitionCookPackage* CookPackage : *CookPackages)
-			{
-				if (UWorldPartitionRuntimeCell* Cell = CookPackageGenerator->GetCellForPackage(*CookPackage))
-				{
-					NeverCookDependencies.Append(Cell->GetActorPackageNames());
-				}
-			}
-		}
-	}
-	return NeverCookDependencies.Array();
-}
-
 bool FWorldPartitionCookPackageSplitter::PopulateGeneratedPackage(UPackage* OwnerPackage, UObject* OwnerObject,
 	const FGeneratedPackageForPopulate& GeneratedPackage, TArray<UObject*>& OutObjectsToMove,
 	TArray<UPackage*>& OutModifiedPackages)
diff --git a/Engine/Source/Runtime/Engine/Private/WorldPartition/WorldPartitionCookPackageSplitter.h b/Engine/Source/Runtime/Engine/Private/WorldPartition/WorldPartitionCookPackageSplitter.h
--- a/Engine/Source/Runtime/Engine/Private/WorldPartition/WorldPartitionCookPackageSplitter.h
+++ b/Engine/Source/Runtime/Engine/Private/WorldPartition/WorldPartitionCookPackageSplitter.h
@@ -19,7 +19,6 @@ public:
 	FWorldPartitionCookPackageSplitter();
 	virtual ~FWorldPartitionCookPackageSplitter();
 
-	TArray<FName> GetNeverCookDependencies() const override;
 	virtual void Teardown(ETeardown Status) override;
 	virtual bool UseInternalReferenceToAvoidGarbageCollect() override { return true; }
 	virtual TArray<ICookPackageSplitter::FGeneratedPackage> GetGenerateList(const UPackage* OwnerPackage, const UObject* OwnerObject) override;
