Fixes the isssue of having the vertical toolbar of a mode dragged out into a separate window in any of your mode layouts, then navigate to modeling mode ~ you would see an empty floating window (that you could close) along with the modeling toolkit.

diff --git a/Engine/Source/Editor/UnrealEd/Private/Toolkits/BaseToolkit.cpp b/Engine/Source/Editor/UnrealEd/Private/Toolkits/BaseToolkit.cpp
--- a/Engine/Source/Editor/UnrealEd/Private/Toolkits/BaseToolkit.cpp
+++ b/Engine/Source/Editor/UnrealEd/Private/Toolkits/BaseToolkit.cpp
@@ -414,8 +414,12 @@ void FModeToolkit::InvokeUI()
 					ActiveToolBarRows.Emplace(ScriptableMode->GetID(), Palette, GetToolPaletteDisplayName(Palette), PaletteWidget);
 				}
 			}
-			TSharedPtr<SDockTab> CreatedToolbarTab = ModeUILayerPtr->GetTabManager()->TryInvokeTab(UAssetEditorUISubsystem::VerticalToolbarID);
-			ModeToolbarTab = CreatedToolbarTab;
+			if (!HasToolkitBuilder())
+			{
+				const TSharedPtr<SDockTab> CreatedToolbarTab = ModeUILayerPtr->GetTabManager()->TryInvokeTab(UAssetEditorUISubsystem::VerticalToolbarID);
+				ModeToolbarTab = CreatedToolbarTab;				
+			}
+
 
 		}
 	}
