SourceControl / Skein: Automatically attempt to check out files when a save occurs and provide the user a relevant error in the case of a conflict.

diff --git a/Engine/Source/Editor/PackagesDialog/Private/PackagesDialog.cpp b/Engine/Source/Editor/PackagesDialog/Private/PackagesDialog.cpp
--- a/Engine/Source/Editor/PackagesDialog/Private/PackagesDialog.cpp
+++ b/Engine/Source/Editor/PackagesDialog/Private/PackagesDialog.cpp
@@ -263,7 +263,22 @@ void FPackagesDialogModule::AddPackageItem(UPackage* InPackage, ECheckBoxState I
  */
 void FPackagesDialogModule::AddButton(EDialogReturnType Type, const FText& Name, const FText& ToolTip/*=FText()*/,  TAttribute<bool> Disabled /*= false*/)
 {
-	PackagesDialogWidget.Get()->AddButton(MakeShareable(new FPackageButton(this, Type, Name, ToolTip, Disabled)));
+	EDialogButtonStyle Style = (Type == DRT_Save) || (Type == DRT_CheckOut) ? DBS_Primary : DBS_Normal;
+	PackagesDialogWidget.Get()->AddButton(MakeShareable(new FPackageButton(this, Type, Style, Name, ToolTip, Disabled)));
+}
+
+/**
+ * Adds a new button to the package dialog window
+ *
+ * @param	Type		The type of this button
+ * @param	Style		The style of this button
+ * @param	Name		The name to display
+ * @param	ToolTip		The tooltip to display
+ * @param	Disabled	If the button should be disabled
+ */
+void FPackagesDialogModule::AddButton(EDialogReturnType Type, EDialogButtonStyle Style, const FText& Name, const FText& ToolTip/*=FText()*/, TAttribute<bool> Disabled /*= false*/)
+{
+	PackagesDialogWidget.Get()->AddButton(MakeShareable(new FPackageButton(this, Type, Style, Name, ToolTip, Disabled)));
 }
 
 /**
diff --git a/Engine/Source/Editor/PackagesDialog/Private/SPackagesDialog.cpp b/Engine/Source/Editor/PackagesDialog/Private/SPackagesDialog.cpp
--- a/Engine/Source/Editor/PackagesDialog/Private/SPackagesDialog.cpp
+++ b/Engine/Source/Editor/PackagesDialog/Private/SPackagesDialog.cpp
@@ -338,8 +338,8 @@ void SPackagesDialog::AddButton(TSharedPtr<FPackageButton> Button)
 	.Padding(5, 0)
 	[
 		SNew(SButton) 
-			.ButtonStyle(&FAppStyle::Get(), (Button->GetType() == DRT_Save || Button->GetType() == DRT_CheckOut) ? "PrimaryButton" : "Button")
-			.TextStyle(&FAppStyle::Get(), (Button->GetType() == DRT_Save || Button->GetType() == DRT_CheckOut) ? "PrimaryButtonText" : "ButtonText")
+			.ButtonStyle(&FAppStyle::Get(), (Button->GetStyle() == DBS_Primary) ? "PrimaryButton" : "Button")
+			.TextStyle(&FAppStyle::Get(), (Button->GetStyle() == DBS_Primary) ? "PrimaryButtonText" : "ButtonText")
 			.Text(Button->GetName())
 			.ToolTipText(Button->GetToolTip())
 			.IsEnabled(Button.Get(), &FPackageButton::IsEnabled)
diff --git a/Engine/Source/Editor/PackagesDialog/Private/SPackagesDialog.h b/Engine/Source/Editor/PackagesDialog/Private/SPackagesDialog.h
--- a/Engine/Source/Editor/PackagesDialog/Private/SPackagesDialog.h
+++ b/Engine/Source/Editor/PackagesDialog/Private/SPackagesDialog.h
@@ -49,11 +49,12 @@ struct FKeyEvent;
 class FPackageButton : public TSharedFromThis<FPackageButton>
 {
 public:
-	FPackageButton(FPackagesDialogModule* InModule, EDialogReturnType InType, const FText& InName, const FText& InToolTip, TAttribute<bool> InDisabled = false)
+	FPackageButton(FPackagesDialogModule* InModule, EDialogReturnType InType, EDialogButtonStyle InStyle, const FText& InName, const FText& InToolTip, TAttribute<bool> InDisabled = false)
 		: Module(InModule)
 		, Name(InName)
 		, ToolTip(InToolTip)
 		, Type(InType)
+		, Style(InStyle)
 		, Clicked(false)
 		, Disabled(InDisabled)
 	{ }
@@ -99,10 +100,17 @@ public:
 	/**
 	 * Gets the type of the button
 	 *
-	 * @return the tupe of the button
+	 * @return the type of the button
 	 */
 	EDialogReturnType GetType() const { return Type; }
 
+	/**
+	 * Gets the style of the button
+	 *
+	 * @return the style of the button
+	 */
+	EDialogButtonStyle GetStyle() const { return Style; }
+
 	/**
 	 * Sets if the button should be disabled
 	 *
@@ -129,6 +137,7 @@ private:
 	FText Name;						// Name of the button
 	FText ToolTip;					// Tool tip for this button
 	EDialogReturnType Type;			// Button type
+	EDialogButtonStyle Style;		// Button style
 	bool Clicked;					// Stores if the button was clicked to close the dialog
 	TAttribute<bool> Disabled;		// Stores if the button is disabled or not
 };
diff --git a/Engine/Source/Editor/PackagesDialog/Public/PackagesDialog.h b/Engine/Source/Editor/PackagesDialog/Public/PackagesDialog.h
--- a/Engine/Source/Editor/PackagesDialog/Public/PackagesDialog.h
+++ b/Engine/Source/Editor/PackagesDialog/Public/PackagesDialog.h
@@ -18,9 +18,16 @@ enum EDialogReturnType
 	DRT_MakeWritable,
 	DRT_Cancel,
 	DRT_Skip,
+	DRT_Revert,
 	DRT_None
 };
 
+enum EDialogButtonStyle
+{
+	DBS_Normal			= 0,
+	DBS_Primary,
+};
+
 class FPackagesDialogModule : public IModuleInterface
 {
 
@@ -107,11 +114,13 @@ public:
 	 * Adds a new button to the package dialog window
 	 *
 	 * @param	Type		The type of this button
+	 * @param	ButtonStyle	The style of this button
 	 * @param	Name		The name to display
 	 * @param	ToolTip		The tooltip to display
 	 * @param	Disabled	If the button should be disabled
 	 */
 	virtual void AddButton(EDialogReturnType Type, const FText& Name, const FText& ToolTip = FText(), TAttribute<bool> Disabled = false);
+	virtual void AddButton(EDialogReturnType Type, EDialogButtonStyle ButtonStyle, const FText& Name, const FText& ToolTip = FText(), TAttribute<bool> Disabled = false);
 
 	/**
 	 * Override this to set whether your module is allowed to be unloaded on the fly
diff --git a/Engine/Source/Editor/UnrealEd/Private/FileHelpers.cpp b/Engine/Source/Editor/UnrealEd/Private/FileHelpers.cpp
--- a/Engine/Source/Editor/UnrealEd/Private/FileHelpers.cpp
+++ b/Engine/Source/Editor/UnrealEd/Private/FileHelpers.cpp
@@ -229,6 +229,18 @@ namespace FileDialogHelpers
 	}
 }
 
+/**
+* Checks if the alternate checkout workflow should be used
+* In this workflow:
+* - the checkout dialog is not used but checkout is automatic.
+* - if any checkout fails, a dialog is shown with the option to revert or save.
+* - if SourceControl is unavailable, a warning is shown about making changes while offline, followed by a save.
+*/
+static bool UseAlternateCheckoutWorkflow()
+{
+	return ISourceControlModule::Get().GetProvider().GetName() == TEXT("Unreal Revision Control");
+}
+
 /**
 * Prompts user with a confirmation dialog if there are checkouts or modifications in other branches
 *
@@ -1898,62 +1910,7 @@ bool FEditorFileUtils::PromptToCheckoutPackagesInternal(bool bCheckDirty, const
 				CheckoutPackagesDialogModule.GetResults( PkgsToMakeWritable, ECheckBoxState::Undetermined );
 				CheckoutPackagesDialogModule.GetResults( PkgsToMakeWritable, ECheckBoxState::Checked);
 
-				bool bPackageFailedWritable = false;
-				FString PkgsWhichFailedWritable;
-
-				// Attempt to make writable each package the user checked
-				for( TArray<UPackage*>::TIterator PkgsToMakeWritableIter( PkgsToMakeWritable ); PkgsToMakeWritableIter; ++PkgsToMakeWritableIter )
-				{
-					UPackage* PackageToMakeWritable = *PkgsToMakeWritableIter;
-					FString Filename;
-
-					bool bFoundFile = FPackageName::DoesPackageExist( PackageToMakeWritable->GetName(), &Filename );
-					if( bFoundFile )
-					{
-						// If we're ignoring the package due to the user ignoring it for saving, remove it from the ignore list
-						// as getting here means we've explicitly decided to save the asset.
-						PackagesNotSavedDuringSaveAll.Remove( PackageToMakeWritable->GetName() );
-
-						// Get the fully qualified filename.
-						const FString FullFilename = FPaths::ConvertRelativePathToFull(Filename);
-
-						// Knock off the read only flag from the current file attributes
-						if (FPlatformFileManager::Get().GetPlatformFile().SetReadOnly(*Filename, false))
-						{
-							// Add to PackagesNotToPromptAnyMore only if not added to Uncontrolled Changelist.
-							// If added to Uncontrolled Changelist, we want the checkout prompt to be displayed again if the file is reverted
-							if (!UncontrolledChangelistModule.OnMakeWritable(Filename))
-							{
-								PackagesNotToPromptAnyMore.Add(PackageToMakeWritable->GetName());
-							}
-
-							if (OutPackagesCheckedOutOrMadeWritable)
-							{
-								OutPackagesCheckedOutOrMadeWritable->Add(PackageToMakeWritable);
-							}
-						}
-						else
-						{
-							bPackageFailedWritable = true;
-							PkgsWhichFailedWritable += FString::Printf( TEXT("\n%s"), *PackageToMakeWritable->GetName() );
-						}
-					}
-					else if (OutPackagesCheckedOutOrMadeWritable)
-					{
-						OutPackagesCheckedOutOrMadeWritable->Append(PackagesToCheckOut);
-					}
-				}
-
-				if ( bPackageFailedWritable ) 
-				{
-					FFormatNamedArguments Arguments;
-					Arguments.Add(TEXT("Packages"), FText::FromString( PkgsWhichFailedWritable ));
-					FText MessageFormatting = NSLOCTEXT("FileHelper", "FailedMakingWritableDlgMessageFormatting", "The following assets could not be made writable:{Packages}");
-					FText Message = FText::Format( MessageFormatting, Arguments );
-
-					FText Title = NSLOCTEXT("FileHelper", "FailedMakingWritableDlg_Title", "Unable to make assets writable");
-					FMessageDialog::Open(EAppMsgType::Ok, Message, &Title);
-				}
+				MakePackagesWritable(PkgsToMakeWritable, OutPackagesCheckedOutOrMadeWritable, nullptr);
 
 				bPerformedOperation = true;
 			}
@@ -1991,6 +1948,79 @@ bool FEditorFileUtils::PromptToCheckoutPackagesInternal(bool bCheckDirty, const
 	return bResult;
 }
 
+void FEditorFileUtils::MakePackagesWritable(const TArray<UPackage*>& PackagesToMakeWritable, TArray<UPackage*>* OutPackagesMadeWritable, TArray<UPackage*>* OutPackagesMadeWritableFailed)
+{
+	TArray<UPackage*> PackagesMadeWritableSuccess;
+	TArray<UPackage*> PackagesMadeWritableFailure;
+	PackagesMadeWritableSuccess.Reserve(PackagesToMakeWritable.Num());
+	PackagesMadeWritableFailure.Reserve(PackagesToMakeWritable.Num());
+
+	// Attempt to make writable each package the user checked
+	FUncontrolledChangelistsModule& UncontrolledChangelistModule = FUncontrolledChangelistsModule::Get();
+	for (UPackage* PackageToMakeWritable : PackagesToMakeWritable)
+	{
+		FString Filename;
+
+		bool bFoundFile = FPackageName::DoesPackageExist(PackageToMakeWritable->GetName(), &Filename);
+		if (bFoundFile)
+		{
+			// If we're ignoring the package due to the user ignoring it for saving, remove it from the ignore list
+			// as getting here means we've explicitly decided to save the asset.
+			PackagesNotSavedDuringSaveAll.Remove(PackageToMakeWritable->GetName());
+
+			// Get the fully qualified filename.
+			const FString FullFilename = FPaths::ConvertRelativePathToFull(Filename);
+
+			// Knock off the read only flag from the current file attributes
+			if (FPlatformFileManager::Get().GetPlatformFile().SetReadOnly(*Filename, false))
+			{
+				// Add to PackagesNotToPromptAnyMore only if not added to Uncontrolled Changelist.
+				// If added to Uncontrolled Changelist, we want the checkout prompt to be displayed again if the file is reverted
+				if (!UncontrolledChangelistModule.OnMakeWritable(Filename))
+				{
+					PackagesNotToPromptAnyMore.Add(PackageToMakeWritable->GetName());
+				}
+
+				PackagesMadeWritableSuccess.Add(PackageToMakeWritable);
+			}
+			else
+			{
+				PackagesMadeWritableFailure.Add(PackageToMakeWritable);
+			}
+		}
+		else
+		{
+			PackagesMadeWritableSuccess.Add(PackageToMakeWritable);
+		}
+	}
+
+	if (PackagesMadeWritableFailure.Num() > 0)
+	{
+		FString PkgsWhichFailedWritable;
+		for (UPackage* Package : PackagesMadeWritableFailure)
+		{
+			PkgsWhichFailedWritable += FString::Printf(TEXT("\n%s"), *Package->GetName());
+		}
+
+		FFormatNamedArguments Arguments;
+		Arguments.Add(TEXT("Packages"), FText::FromString(PkgsWhichFailedWritable));
+		FText MessageFormatting = NSLOCTEXT("FileHelper", "FailedMakingWritableDlgMessageFormatting", "The following assets could not be made writable:{Packages}");
+		FText Message = FText::Format(MessageFormatting, Arguments);
+		FText Title = NSLOCTEXT("FileHelper", "FailedMakingWritableDlg_Title", "Unable to make assets writable");
+
+		FMessageDialog::Open(EAppMsgType::Ok, Message, &Title);
+	}
+
+	if (OutPackagesMadeWritable)
+	{
+		OutPackagesMadeWritable->Append(PackagesMadeWritableSuccess);
+	}
+	if (OutPackagesMadeWritableFailed)
+	{
+		OutPackagesMadeWritableFailed->Append(PackagesMadeWritableFailure);
+	}
+}
+
 ECommandResult::Type FEditorFileUtils::CheckoutPackages(const TArray<UPackage*>& PkgsToCheckOut, TArray<UPackage*>* OutPackagesCheckedOut, const bool bErrorIfAlreadyCheckedOut, const bool bConfirmPackageBranchCheckOutStatus)
 {
 	TRACE_CPUPROFILER_EVENT_SCOPE(FEditorFileUtils_CheckoutPackages);
@@ -4446,15 +4476,24 @@ FEditorFileUtils::EPromptReturnCode FEditorFileUtils::PromptForCheckoutAndSave(
 
 		TArray<UPackage*> PackagesCheckedOutOrMadeWritable;
 		TArray<UPackage*> PackagesNotNeedingCheckout;
+		TArray<UPackage*> PackagesToRevert;
 
 		// Prompt to check-out any packages under source control
 		bool bUserResponse = true;
+		bool bAutomaticCheckout = UseAlternateCheckoutWorkflow();
 
 		if (!bAlreadyCheckedOut)
 		{
-			const bool bPromptingAfterModify = false;
-			const bool bAllowSkip = true;
-			bUserResponse = FEditorFileUtils::PromptToCheckoutPackagesInternal(false, PackagesToSave, &PackagesCheckedOutOrMadeWritable, &PackagesNotNeedingCheckout, bPromptingAfterModify, bAllowSkip);
+			if (bAutomaticCheckout)
+			{
+				bUserResponse = FEditorFileUtils::AutomaticCheckoutOrPromptToRevertPackages(PackagesToSave, &PackagesCheckedOutOrMadeWritable, &PackagesNotNeedingCheckout, &PackagesToRevert);
+			}
+			else
+			{
+				const bool bPromptingAfterModify = false;
+				const bool bAllowSkip = true;
+				bUserResponse = FEditorFileUtils::PromptToCheckoutPackagesInternal(false, PackagesToSave, &PackagesCheckedOutOrMadeWritable, &PackagesNotNeedingCheckout, bPromptingAfterModify, bAllowSkip);
+			}
 		}
 
 		if( bAlreadyCheckedOut || (bUserResponse && (PackagesCheckedOutOrMadeWritable.Num() > 0 || PackagesNotNeedingCheckout.Num() > 0)) )
@@ -4488,6 +4527,42 @@ FEditorFileUtils::EPromptReturnCode FEditorFileUtils::PromptForCheckoutAndSave(
 			// The user cancelled the checkout dialog, so set the return response accordingly  
 			ReturnResponse = PR_Cancelled;
 		}
+
+		if (PackagesToRevert.Num() > 0)
+		{
+			// Check if the world should be reloaded after the revert.
+			bool bReloadWorld = false;
+			if (UWorld* EditorWorld = GEditor->GetEditorWorldContext().World())
+			{
+				UPackage* EditorWorldPackage = EditorWorld->GetPackage();
+				if (PackagesToRevert.Contains(EditorWorldPackage))
+				{
+					// If the world file is reverted, the world should be reloaded.
+					bReloadWorld = true;
+				}
+				else
+				{
+					// If one of the external files is reverted, the world should be reloaded.
+					for (UPackage* Package : PackagesToRevert)
+					{
+						FString PackageName = Package->GetName();
+						if (PackageName.Contains(FPackagePath::GetExternalActorsFolderName()) ||
+							PackageName.Contains(FPackagePath::GetExternalObjectsFolderName()))
+						{
+							bReloadWorld = true;
+							break;
+						}
+					}
+				}
+			}
+
+			// Save the packages that need to be reverted, so the SourceControl can act on them.
+			TArray<UPackage*> PackagesSaveFailed;
+			InternalPromptForCheckoutAndSave(PackagesToRevert, /*bUseDialog=*/false, PackagesSaveFailed);
+
+			// Revert packages that could not be checked out.
+			USourceControlHelpers::RevertAndReloadPackages(USourceControlHelpers::PackageFilenames(PackagesToRevert), /*bRevertAll=*/false, /*bReloadWorld=*/bReloadWorld);
+		}
 	}
 
 	return ReturnResponse;
@@ -4506,6 +4581,263 @@ FEditorFileUtils::EPromptReturnCode FEditorFileUtils::PromptForCheckoutAndSave(c
 		bCanBeDeclined);
 }
 
+/* Return 'true' to indicate that the packages (in OutPackagesCheckedOutOrMadeWritable) should be saved, or 'false' to cancel saving. */
+bool FEditorFileUtils::AutomaticCheckoutOrPromptToRevertPackages(const TArray<UPackage*>& PackagesToCheckOut, TArray<UPackage*>* OutPackagesCheckedOutOrMadeWritable, TArray<UPackage*>* OutPackagesNotNeedingCheckout, TArray<UPackage*>* OutPackagesToRevert)
+{
+	ISourceControlProvider& SourceControlProvider = ISourceControlModule::Get().GetProvider();
+
+	// Is there anything to work with?
+	int32 NumPackages = PackagesToCheckOut.Num();
+	if (NumPackages == 0)
+	{
+		return true;
+	}
+
+	// Build map from PackageFileName -> Package.
+	TMap<FString, UPackage*> PackageMap;
+	for (UPackage* Package : PackagesToCheckOut)
+	{
+		PackageMap.Add(USourceControlHelpers::PackageFilename(Package), Package);
+	}
+
+	// Determine initial states.
+	SourceControlProvider.Execute(ISourceControlOperation::Create<FUpdateStatus>(), PackagesToCheckOut);
+
+	TArray<FSourceControlStateRef> InitialStates;
+	SourceControlProvider.GetState(PackagesToCheckOut, InitialStates, EStateCacheUsage::Use);
+
+	TArray<UPackage*> PackagesCheckOutNeeded;
+	TArray<UPackage*> PackagesCheckOutImpossible;
+	TArray<UPackage*> PackagesCheckedOutAlready;
+	PackagesCheckOutNeeded.Reserve(PackagesToCheckOut.Num());
+	PackagesCheckOutImpossible.Reserve(PackagesToCheckOut.Num());
+	PackagesCheckedOutAlready.Reserve(PackagesToCheckOut.Num());
+
+	for (const FSourceControlStateRef& State : InitialStates)
+	{
+		const FString& PackageFilename = State->GetFilename();
+
+		if (State->IsCheckedOut())
+		{
+			// No need to check it out.
+			PackagesCheckedOutAlready.Add(PackageMap[PackageFilename]);
+		}
+		else if (State->IsCheckedOutOther() || !State->IsCurrent())
+		{
+			// Can't check it out.
+			PackagesCheckOutImpossible.Add(PackageMap[PackageFilename]);
+		}
+		else
+		{
+			// Try to check it out.
+			PackagesCheckOutNeeded.Add(PackageMap[PackageFilename]);
+		}
+	}
+
+	// Result value indicates whether to continue with saving.
+	bool bResult = false;
+
+	// Is SourceControl online?
+	bool bSourceControlEnabled = SourceControlProvider.IsEnabled();
+	bool bSourceControlAvailable = SourceControlProvider.IsAvailable();
+	if (bSourceControlAvailable)
+	{
+		// Yes, SourceControl is online.
+		// Try to check out the packages.
+
+		TArray<UPackage*> PackagesCheckOutSuccess;
+		TArray<UPackage*> PackagesCheckOutFailure;
+		PackagesCheckOutSuccess.Reserve(PackagesCheckOutNeeded.Num());
+		PackagesCheckOutFailure.Reserve(PackagesCheckOutNeeded.Num());
+
+		TArray<UPackage*> PackagesWritableSuccess;
+		TArray<UPackage*> PackagesWritableFailure;
+		TArray<UPackage*> PackagesToRevert;
+
+		if (PackagesCheckOutNeeded.Num() > 0)
+		{
+			FScopedSlowTask SlowTask(static_cast<float>(PackagesCheckOutNeeded.Num()), LOCTEXT("CheckingOutPackages", "Checking out packages..."));
+			SlowTask.MakeDialog();
+
+			int32 NumCheckOutImpossible = 0;
+
+			// Loop while attempting to check-out the packages.
+			// If succeeds, then break out of the loop.
+			// If failed, inspect the state of the files again and see if somebody just checked out or submitted the file. Then retry with the remaining files.
+			do
+			{
+				NumCheckOutImpossible = PackagesCheckOutImpossible.Num();
+
+				ECommandResult::Type Result = SourceControlProvider.Execute(ISourceControlOperation::Create<FCheckOut>(), PackagesCheckOutNeeded);
+				if (Result == ECommandResult::Succeeded)
+				{
+					PackagesCheckOutSuccess.Append(PackagesCheckOutNeeded);
+					PackagesCheckOutNeeded.Empty();
+				}
+				else
+				{
+					TArray<FSourceControlStateRef> States;
+					if (SourceControlProvider.GetState(PackagesCheckOutNeeded, States, EStateCacheUsage::ForceUpdate) == ECommandResult::Succeeded)
+					{
+						for (const FSourceControlStateRef& State : States)
+						{
+							const FString& PackageFilename = State->GetFilename();
+
+							if (State->IsCheckedOut())
+							{
+								// If provider supports partial checkout this could happen.
+								UPackage* Package = PackageMap[PackageFilename];
+								PackagesCheckOutNeeded.Remove(Package);
+								PackagesCheckOutSuccess.Add(Package);
+							}
+							else if (State->IsCheckedOutOther() || !State->IsCurrent())
+							{
+								// Somebody just beat us to it.
+								UPackage* Package = PackageMap[PackageFilename];
+								PackagesCheckOutNeeded.Remove(Package);
+								PackagesCheckOutImpossible.Add(Package);
+							}
+						}
+					}
+				}
+			} while (NumCheckOutImpossible != PackagesCheckOutImpossible.Num() && PackagesCheckOutNeeded.Num() > 0);
+
+			SlowTask.EnterProgressFrame(static_cast<float>(PackagesCheckOutNeeded.Num()));
+		}
+
+		// Any remaining packages have failed for unknown reasons.
+		PackagesCheckOutFailure = PackagesCheckOutNeeded;
+
+		// Were we able to check out all packages?
+		if (PackagesCheckOutImpossible.Num() > 0 || PackagesCheckOutFailure.Num() > 0)
+		{
+			// No.
+			// Show dialog with assets that weren't checked out.
+
+			const FText DialogTitle = NSLOCTEXT("PackagesDialogModule", "CheckoutPackagesFailedDialogTitle", "Check-out failed");
+			const FText DialogHeading = NSLOCTEXT("PackagesDialogModule", "CheckoutPackagesFailedDialogMessage",
+				"While saving, check-out failed for the following assets. Revert your changes to these assets and sync to the latest snapshot to avoid conflicts with your teammates.\r\n\r\n"
+				"If necessary, you may also proceed by saving locally only, but you will likely run into conflicts later when trying to check in these changes.\r\n\r\n"
+				"Tip: Turn on automatic checkout and automatic undo in your Unreal Revision Control settings to avoid future conflicts with your teammates and conflict warnings."
+			);
+
+			FPackagesDialogModule& CheckoutPackagesDialogModule = FModuleManager::LoadModuleChecked<FPackagesDialogModule>(TEXT("PackagesDialog"));
+			CheckoutPackagesDialogModule.CreatePackagesDialog(
+				DialogTitle,
+				DialogHeading,
+				/*InReadOnly=*/true,
+				/*InAllowSourceControlConnection*/true
+			);
+
+			TArray<UPackage*> PackagesNotCheckedOut;
+			PackagesNotCheckedOut.Append(PackagesCheckOutFailure);
+			PackagesNotCheckedOut.Append(PackagesCheckOutImpossible);
+
+			for (UPackage* Package : PackagesNotCheckedOut)
+			{
+				FSourceControlStatePtr State = SourceControlProvider.GetState(Package, EStateCacheUsage::Use);
+				if (!State->IsCurrent())
+				{
+					CheckoutPackagesDialogModule.AddPackageItem(Package, ECheckBoxState::Unchecked, true, TEXT("SavePackages.SCC_DlgNotCurrent"), State->GetDisplayTooltip().ToString());
+				}
+				else if (State->IsCheckedOutOther())
+				{
+					CheckoutPackagesDialogModule.AddPackageItem(Package, ECheckBoxState::Unchecked, true, TEXT("SavePackages.SCC_DlgCheckedOutOther"), State->GetDisplayTooltip().ToString());
+				}
+				else
+				{
+					CheckoutPackagesDialogModule.AddPackageItem(Package, ECheckBoxState::Unchecked, true, TEXT("SavePackages.SCC_DlgNoIcon"), State->GetDisplayTooltip().ToString());
+				}
+			}
+
+			// The Revert button will allow the user to undo changes to those assets.
+			CheckoutPackagesDialogModule.AddButton(DRT_Revert, DBS_Primary, NSLOCTEXT("PackagesDialogModule", "Dlg_RevertButton", "Revert My Changes"), NSLOCTEXT("PackagesDialogModule", "Dlg_RevertButtonTooltip", "Revert changes to files that could not be checked out (recommended)."));
+
+			// The Save button will allow the user to proceed with saving those assets anyway, thereby risking conflicts.
+			CheckoutPackagesDialogModule.AddButton(DRT_Save, DBS_Normal, NSLOCTEXT("PackagesDialogModule", "Dlg_SaveButton", "Save Locally Only"), NSLOCTEXT("PackagesDialogModule", "Dlg_RevertButtonTooltip", "Save changes to files that could not be checked out anyway. You will likely be unable to check-in these changes."));
+
+			EDialogReturnType UserResponse = CheckoutPackagesDialogModule.ShowPackagesDialog();
+			if (UserResponse == DRT_Revert)
+			{
+				PackagesToRevert = PackagesNotCheckedOut;
+			}
+			if (UserResponse == DRT_Save)
+			{
+				// Make the packages writable and proceed to save.
+				MakePackagesWritable(PackagesNotCheckedOut, &PackagesWritableSuccess, &PackagesWritableFailure);
+			}
+		}
+		else
+		{
+			// Yes.
+			// All packages were checked out.
+		}
+
+		// Populate output values.
+		if (OutPackagesCheckedOutOrMadeWritable)
+		{
+			OutPackagesCheckedOutOrMadeWritable->Append(PackagesCheckOutSuccess);
+			OutPackagesCheckedOutOrMadeWritable->Append(PackagesWritableSuccess);
+		}
+		if (OutPackagesNotNeedingCheckout)
+		{
+			OutPackagesNotNeedingCheckout->Append(PackagesCheckedOutAlready);
+		}
+		if (OutPackagesToRevert)
+		{
+			OutPackagesToRevert->Append(PackagesToRevert);
+		}
+
+		// Save if anything was checked out or made writable.
+		bResult = (PackagesCheckOutSuccess.Num() > 0) || (PackagesWritableSuccess.Num() > 0) || (PackagesCheckedOutAlready.Num() > 0);
+	}
+	else if (bSourceControlEnabled)
+	{
+		// No, SourceControl is offline.
+		// Warn the user that they're working in offline mode.
+
+		FText OfflineTitle = NSLOCTEXT("FileHelper", "OfflineDialogTitle", "Warning - Offline, Conflicts may occur");
+		FText OfflineMessage = NSLOCTEXT("FileHelper", "OfflineDialogMessage",
+			"You've made changes while offline. These and any further changes made offline could conflict with your teammates' work when you reconnect.\r\n\r\n"
+			"Reconnect as soon as possible to minimize conflicts and continue making changes.");
+
+		FSuppressableWarningDialog::FSetupInfo Info(OfflineMessage, OfflineTitle, TEXT("ShowOfflineModeWarning"), GEditorPerProjectIni);
+		Info.ConfirmText = NSLOCTEXT("FileHelper", "OfflineDialog_ConfirmText", "Ok");
+		Info.CheckBoxText = NSLOCTEXT("FileHelper", "OfflineDialog_CheckBoxText", "Don't show this again");
+		Info.bDefaultToSuppressInTheFuture = true;
+		FSuppressableWarningDialog OfflineModeWarningDialog(Info);
+
+		TArray<UPackage*> PackagesWritableSuccess;
+		TArray<UPackage*> PackagesWritableFailure;
+
+		// Show the warning and if the user doesn't cancel, make the files writable and proceed with the save.
+		bool bMakeWritable = (OfflineModeWarningDialog.ShowModal() != FSuppressableWarningDialog::EResult::Cancel);
+		if (bMakeWritable)
+		{
+			TArray<UPackage*> PackagesToMakeWritable;
+			PackagesToMakeWritable.Append(PackagesCheckOutNeeded);
+			PackagesToMakeWritable.Append(PackagesCheckOutImpossible);
+
+			MakePackagesWritable(PackagesToMakeWritable, &PackagesWritableSuccess, &PackagesWritableFailure);
+		}
+
+		// Populate output values.
+		if (OutPackagesCheckedOutOrMadeWritable)
+		{
+			OutPackagesCheckedOutOrMadeWritable->Append(PackagesWritableSuccess);
+		}
+		if (OutPackagesNotNeedingCheckout)
+		{
+			OutPackagesNotNeedingCheckout->Append(PackagesCheckedOutAlready);
+		}
+
+		// Save if anything was made writable.
+		bResult = (PackagesWritableSuccess.Num() > 0) || (PackagesCheckedOutAlready.Num() > 0);
+	}
+
+	return bResult;
+}
+
 bool FEditorFileUtils::SaveWorlds(UWorld* InWorld, const FString& RootPath, const TCHAR* Prefix, TArray<FString>& OutFilenames)
 {
 	const FScopedBusyCursor BusyCursor;
diff --git a/Engine/Source/Editor/UnrealEd/Public/FileHelpers.h b/Engine/Source/Editor/UnrealEd/Public/FileHelpers.h
--- a/Engine/Source/Editor/UnrealEd/Public/FileHelpers.h
+++ b/Engine/Source/Editor/UnrealEd/Public/FileHelpers.h
@@ -656,6 +656,12 @@ private:
 	/** Prompts for package checkout without reentrance check */
 	static bool PromptToCheckoutPackagesInternal(bool bCheckDirty, const TArray<UPackage*>& PackagesToCheckOut, TArray<UPackage*>* OutPackagesCheckedOutOrMadeWritable, TArray<UPackage*>* OutPackagesNotNeedingCheckout, const bool bPromptingAfterModify, const bool bAllowSkip);
 
+	/** Tries to checkout packages without prompting and only prompts (for revert) if that fails */
+	static bool AutomaticCheckoutOrPromptToRevertPackages(const TArray<UPackage*>& PackagesToCheckOut, TArray<UPackage*>* OutPackagesCheckedOutOrMadeWritable, TArray<UPackage*>* OutPackagesNotNeedingCheckout, TArray<UPackage*>* OutPackagesToRevert);
+
+	/** Tries to make the given packages writable */
+	static void MakePackagesWritable(const TArray<UPackage*>& PackagesToMakeWritable, TArray<UPackage*>* OutPackagesMadeWritable, TArray<UPackage*>* OutPackagesMadeWritableFailed);
+
 	static bool bIsLoadingDefaultStartupMap;
 
 	/** Flag used to determine if the checkout and save prompt is already open to prevent re-entrance */
