SourceControl: Fixed text conflict from NSLOCTEXT macro.

diff --git a/Engine/Source/Editor/UnrealEd/Private/FileHelpers.cpp b/Engine/Source/Editor/UnrealEd/Private/FileHelpers.cpp
--- a/Engine/Source/Editor/UnrealEd/Private/FileHelpers.cpp
+++ b/Engine/Source/Editor/UnrealEd/Private/FileHelpers.cpp
@@ -4754,7 +4754,7 @@ bool FEditorFileUtils::AutomaticCheckoutOrPromptToRevertPackages(const TArray<UP
 			CheckoutPackagesDialogModule.AddButton(DRT_Revert, DBS_Primary, NSLOCTEXT("PackagesDialogModule", "Dlg_RevertButton", "Revert My Changes"), NSLOCTEXT("PackagesDialogModule", "Dlg_RevertButtonTooltip", "Revert changes to files that could not be checked out (recommended)."));
 
 			// The Save button will allow the user to proceed with saving those assets anyway, thereby risking conflicts.
-			CheckoutPackagesDialogModule.AddButton(DRT_Save, DBS_Normal, NSLOCTEXT("PackagesDialogModule", "Dlg_SaveButton", "Save Locally Only"), NSLOCTEXT("PackagesDialogModule", "Dlg_RevertButtonTooltip", "Save changes to files that could not be checked out anyway. You will likely be unable to check-in these changes."));
+			CheckoutPackagesDialogModule.AddButton(DRT_Save, DBS_Normal, NSLOCTEXT("PackagesDialogModule", "Dlg_SaveButton", "Save Locally Only"), NSLOCTEXT("PackagesDialogModule", "Dlg_SaveButtonTooltip", "Save changes to files that could not be checked out anyway. You will likely be unable to check-in these changes."));
 
 			EDialogReturnType UserResponse = CheckoutPackagesDialogModule.ShowPackagesDialog();
 			if (UserResponse == DRT_Revert)
