json name fixup

diff --git a/Engine/Source/Editor/UnrealEd/Private/Cooker/CookProfiling.cpp b/Engine/Source/Editor/UnrealEd/Private/Cooker/CookProfiling.cpp
--- a/Engine/Source/Editor/UnrealEd/Private/Cooker/CookProfiling.cpp
+++ b/Engine/Source/Editor/UnrealEd/Private/Cooker/CookProfiling.cpp
@@ -1111,7 +1111,7 @@ void LogCookStats(ECookMode::Type CookMode)
 			{
 				JsonWriter->WriteObjectStart(ProfileEntry.Key);
 				JsonWriter->WriteValue(TEXT("Path"), ProfileEntry.Path);
-				JsonWriter->WriteValue(TEXT("TimInSec"), ProfileEntry.Value);
+				JsonWriter->WriteValue(TEXT("Value"), ProfileEntry.Value);
 				JsonWriter->WriteObjectEnd();
 			}
 			JsonWriter->WriteObjectEnd();
