Fix a couple of issues where packages were aggressively dirtied
 - Call cancel on the transaction if it is actually a camera movement input
 - Reset package dirty state when a widget drag is cancelled (for example with escape key)

diff --git a/Engine/Source/Editor/UnrealEd/Classes/Settings/LevelEditorViewportSettings.h b/Engine/Source/Editor/UnrealEd/Classes/Settings/LevelEditorViewportSettings.h
--- a/Engine/Source/Editor/UnrealEd/Classes/Settings/LevelEditorViewportSettings.h
+++ b/Engine/Source/Editor/UnrealEd/Classes/Settings/LevelEditorViewportSettings.h
@@ -331,6 +331,10 @@ class UNREALED_API ULevelEditorViewportSettings
 	/** The sensitivity of mouse movement when rotating the camera. */
 	UPROPERTY(EditAnywhere, config, Category=Controls, meta=(DisplayName="Mouse Sensitivity", ClampMin="0.01",ClampMax="1.0") )
 	float MouseSensitivty;
+
+	/** Camera movement notification toggle to switch back to the behavior that caused camera notifications to be sent during gizmo movement. */
+	UPROPERTY(EditAnywhere, config, Category = Controls)
+	uint32 bUseLegacyCameraMovementNotifications : 1;
 	
 	/** Whether or not to invert mouse on the y axis in free look mode */
 	UPROPERTY(EditAnywhere, config, Category = Controls, meta = (DisplayName = "Invert Mouse Look Y Axis"))
diff --git a/Engine/Source/Editor/UnrealEd/Private/EditorViewportClient.cpp b/Engine/Source/Editor/UnrealEd/Private/EditorViewportClient.cpp
--- a/Engine/Source/Editor/UnrealEd/Private/EditorViewportClient.cpp
+++ b/Engine/Source/Editor/UnrealEd/Private/EditorViewportClient.cpp
@@ -1892,7 +1892,7 @@ void FEditorViewportClient::UpdateCameraMovement( float DeltaTime )
 		EditorMovementDeltaUpperBound = .15f;
 #endif
 		// Check whether the camera is being moved by the mouse or keyboard
-		bool bHasMovement = bIsTracking;
+		bool bHasMovement = GetDefault<ULevelEditorViewportSettings>()->bUseLegacyCameraMovementNotifications;
 
 		if ((*CameraUserImpulseData).RotateYawVelocityModifier != 0.0f ||
 			(*CameraUserImpulseData).RotatePitchVelocityModifier != 0.0f ||
@@ -1908,6 +1908,7 @@ void FEditorViewportClient::UpdateCameraMovement( float DeltaTime )
 		{
 			bHasMovement = true;
 		}
+		bHasMovement = bHasMovement && bIsTracking;
 
 		BeginCameraMovement(bHasMovement);
 
diff --git a/Engine/Source/Editor/UnrealEd/Private/LevelEditorViewport.cpp b/Engine/Source/Editor/UnrealEd/Private/LevelEditorViewport.cpp
--- a/Engine/Source/Editor/UnrealEd/Private/LevelEditorViewport.cpp
+++ b/Engine/Source/Editor/UnrealEd/Private/LevelEditorViewport.cpp
@@ -1801,6 +1801,18 @@ void FTrackingTransaction::Cancel()
 	{
 		ScopedTransaction->Cancel();
 	}
+
+	// If the transaction is cancelled, reset the package dirty states to reflect their original state
+	for (const TTuple<UPackage*, bool>& InitialPackageDirtyState : InitialPackageDirtyStates)
+	{
+		UPackage* Package = InitialPackageDirtyState.Key;
+		const bool bInitialDirtyState = InitialPackageDirtyState.Value;
+		if (!bInitialDirtyState && Package->IsDirty())
+		{
+			Package->SetDirtyFlag(false);
+		}
+	}
+
 	End();
 }
 
@@ -2146,6 +2158,7 @@ bool FLevelEditorViewportClient::ShouldLockPitch() const
 
 void FLevelEditorViewportClient::BeginCameraMovement(bool bHasMovement)
 {
+	const bool bIsUsingLegacyMovementNotify = GetDefault<ULevelEditorViewportSettings>()->bUseLegacyCameraMovementNotifications;
 	// If there's new movement broadcast it
 	if (bHasMovement)
 	{
@@ -2167,10 +2180,15 @@ void FLevelEditorViewportClient::BeginCameraMovement(bool bHasMovement)
 					FCoreUObjectDelegates::OnPreObjectPropertyChanged.Broadcast(ActorLock, PropertyChain);
 				}
 			}
+
+			if (!bIsUsingLegacyMovementNotify)
+			{
+				TrackingTransaction.Cancel();
+			}
 			bIsCameraMoving = true;
 		}
 	}
-	else
+	else if (bIsUsingLegacyMovementNotify || !bIsTracking)
 	{
 		bIsCameraMoving = false;
 	}
diff --git a/Engine/Source/Editor/UnrealEd/Private/Settings/SettingsClasses.cpp b/Engine/Source/Editor/UnrealEd/Private/Settings/SettingsClasses.cpp
--- a/Engine/Source/Editor/UnrealEd/Private/Settings/SettingsClasses.cpp
+++ b/Engine/Source/Editor/UnrealEd/Private/Settings/SettingsClasses.cpp
@@ -938,6 +938,7 @@ ULevelEditorViewportSettings::ULevelEditorViewportSettings( const FObjectInitial
 	bAllowScreenRotate = false;
 	bShowActorEditorContext = true;
 	MouseSensitivty = .2f;
+	bUseLegacyCameraMovementNotifications = false;
 	// Set a default preview mesh
 	PreviewMeshes.Add(FSoftObjectPath("/Engine/EditorMeshes/ColorCalibrator/SM_ColorCalibrator.SM_ColorCalibrator"));
 	LastInViewportMenuLocation = FVector2D(EForceInit::ForceInitToZero);
