SourceControl: Writable (modified) file status is not automatically updated in the Revision Control column (Outliner).

If the packages were made writable, instead of FCheckOut, the SourceControlStateChanged event is not broadcasted because the SourceControlProvider was never involved. After saving to disk though, the files have become modified which will be picked up by SourceControl providers like Git or Skein if an FUpdateStatus is performed on them, even if SourceControl is offline (Skein stores the original hashes in a local database, Git does probably something similar).

diff --git a/Engine/Source/Editor/UnrealEd/Private/FileHelpers.cpp b/Engine/Source/Editor/UnrealEd/Private/FileHelpers.cpp
--- a/Engine/Source/Editor/UnrealEd/Private/FileHelpers.cpp
+++ b/Engine/Source/Editor/UnrealEd/Private/FileHelpers.cpp
@@ -4563,6 +4563,12 @@ FEditorFileUtils::EPromptReturnCode FEditorFileUtils::PromptForCheckoutAndSave(
 			// Revert packages that could not be checked out.
 			USourceControlHelpers::RevertAndReloadPackages(USourceControlHelpers::PackageFilenames(PackagesToRevert), /*bRevertAll=*/false, /*bReloadWorld=*/bReloadWorld);
 		}
+
+		// Update SCC state for packages that were made writable outside of SCC
+		if (PackagesCheckedOutOrMadeWritable.Num())
+		{
+			ISourceControlModule::Get().QueueStatusUpdate(PackagesCheckedOutOrMadeWritable);
+		}
 	}
 
 	return ReturnResponse;
