Added a separate entry in a RuntimeVirtualTexture component for a streaming VT that is specific to mobile rendering
Depending on target platform shading path one or the other SVT will be used
This is required for projects that has a significantly different landscape materials for each shading path

diff --git a/Engine/Source/Editor/UnrealEd/Private/EditorBuildUtils.cpp b/Engine/Source/Editor/UnrealEd/Private/EditorBuildUtils.cpp
--- a/Engine/Source/Editor/UnrealEd/Private/EditorBuildUtils.cpp
+++ b/Engine/Source/Editor/UnrealEd/Private/EditorBuildUtils.cpp
@@ -57,6 +57,8 @@
 #include "WorldPartition/WorldPartitionBuildNavigationOptions.h"
 #include "WorldPartition/WorldPartitionHelpers.h"
 #include "WorldPartition/WorldPartitionRuntimeVirtualTextureBuilder.h"
+#include "AssetCompilingManager.h"
+#include "ComponentRecreateRenderStateContext.h"
 
 DEFINE_LOG_CATEGORY_STATIC(LogEditorBuildUtils, Log, All);
 
@@ -1629,6 +1631,29 @@ bool FEditorBuildUtils::EditorBuildMaterialTextureStreamingData(UPackage* Packag
 	return bAnyPackagesDirtied;
 }
 
+static bool BuildVirtualTextureComponents(IVirtualTexturingEditorModule* Module, EShadingPath ShadingPath, TArray<URuntimeVirtualTextureComponent*>& Components)
+{
+	if (Components.Num() == 0)
+	{
+		return true;
+	}
+
+	FScopedSlowTask BuildTask(static_cast<float>(Components.Num()), LOCTEXT("VirtualTextureBuild", "Building Virtual Textures"));
+	BuildTask.MakeDialog(true);
+
+	for (URuntimeVirtualTextureComponent* Component : Components)
+	{
+		BuildTask.EnterProgressFrame();
+
+		// Note that Build*() functions return true if the associated Has*() functions return false
+		if (BuildTask.ShouldCancel() || !Module->BuildStreamedMips(ShadingPath, Component))
+		{
+			return false;
+		}
+	}
+	return true;
+}
+
 bool FEditorBuildUtils::EditorBuildVirtualTexture(UWorld* InWorld)
 {
 	if (InWorld == nullptr)
@@ -1657,33 +1682,50 @@ bool FEditorBuildUtils::EditorBuildVirtualTexture(UWorld* InWorld)
 			UWorldPartitionRuntimeVirtualTextureBuilder::LoadRuntimeVirtualTextureActors(WorldPartition, ForEachActorWithLoadingResult);
 		}
 
-		TArray<URuntimeVirtualTextureComponent*> Components;
+		// We will need to build VTs for both shading paths
+		const ERHIFeatureLevel::Type CurFeatureLevel = InWorld->FeatureLevel;
+		const ERHIFeatureLevel::Type AltFeatureLevel = (CurFeatureLevel == ERHIFeatureLevel::ES3_1 ? GMaxRHIFeatureLevel : ERHIFeatureLevel::ES3_1);
+		const EShadingPath CurShadingPath = FSceneInterface::GetShadingPath(CurFeatureLevel);
+		const EShadingPath AltShadingPath = FSceneInterface::GetShadingPath(AltFeatureLevel);
+
+		TArray<URuntimeVirtualTextureComponent*> Components[2];
 		for (TObjectIterator<URuntimeVirtualTextureComponent> It; It; ++It)
 		{
-			if (Module->HasStreamedMips(*It))
+			if (Module->HasStreamedMips(CurShadingPath, *It))
 			{
-				Components.Add(*It);
+				Components[0].Add(*It);
 			}
-		}
 
-		if (Components.Num() == 0)
+			if (Module->HasStreamedMips(AltShadingPath, *It))
+			{
+				Components[1].Add(*It);
+			}
+		}
+		
+		// Build for a current feature level first
+		if (!BuildVirtualTextureComponents(Module, CurShadingPath, Components[0]))
 		{
-			return true;
+			return false;
 		}
-
-		FScopedSlowTask BuildTask(static_cast<float>(Components.Num()), LOCTEXT("VirtualTextureBuild", "Building Virtual Textures"));
-		BuildTask.MakeDialog(true);
-
-		for (URuntimeVirtualTextureComponent* Component : Components)
+		
+		// Build for others if any
+		bool bResult = true;
+		if (Components[1].Num() != 0)
 		{
-			BuildTask.EnterProgressFrame();
-
-			// Note that Build*() functions return true if the associated Has*() functions return false
-			if (BuildTask.ShouldCancel() || !Module->BuildStreamedMips(Component))
+			InWorld->ChangeFeatureLevel(AltFeatureLevel);
+			// Make sure all assets are finished compiling. Recreate render state after shader compilation complete
 			{
-				return false;
+				UMaterialInterface::SubmitRemainingJobsForWorld(InWorld);
+				FAssetCompilingManager::Get().FinishAllCompilation();
+				FAssetCompilingManager::Get().ProcessAsyncTasks();
+				FGlobalComponentRecreateRenderStateContext Context;
 			}
+			bResult = BuildVirtualTextureComponents(Module, AltShadingPath, Components[1]);
 		}
+		
+		// Restore world feature level
+		InWorld->ChangeFeatureLevel(CurFeatureLevel);
+		return bResult;
 	}
 	
 	return true;
diff --git a/Engine/Source/Editor/UnrealEd/Private/WorldPartition/WorldPartitionRuntimeVirtualTextureBuilder.cpp b/Engine/Source/Editor/UnrealEd/Private/WorldPartition/WorldPartitionRuntimeVirtualTextureBuilder.cpp
--- a/Engine/Source/Editor/UnrealEd/Private/WorldPartition/WorldPartitionRuntimeVirtualTextureBuilder.cpp
+++ b/Engine/Source/Editor/UnrealEd/Private/WorldPartition/WorldPartitionRuntimeVirtualTextureBuilder.cpp
@@ -13,6 +13,8 @@
 #include "Modules/ModuleManager.h"
 #include "UObject/UObjectIterator.h"
 #include "ComponentRecreateRenderStateContext.h"
+#include "Materials/MaterialInstance.h"
+#include "ShaderCompiler.h"
 
 DEFINE_LOG_CATEGORY_STATIC(LogWorldPartitionRuntimeVirtualTextureBuilder, All, All);
 
@@ -36,6 +38,23 @@ void UWorldPartitionRuntimeVirtualTextureBuilder::LoadRuntimeVirtualTextureActor
 	FAssetCompilingManager::Get().FinishAllCompilation();
 }
 
+static bool BuildVirtualTextureComponents(IVirtualTexturingEditorModule* Module, EShadingPath ShadingPath, TArray<URuntimeVirtualTextureComponent*>& Components)
+{
+	if (Components.Num() == 0)
+	{
+		return true;
+	}
+
+	for (URuntimeVirtualTextureComponent* Component : Components)
+	{
+		if (!Module->BuildStreamedMips(ShadingPath, Component))
+		{
+			return false;
+		}
+	}
+	return true;
+}
+
 bool UWorldPartitionRuntimeVirtualTextureBuilder::RunInternal(UWorld* World, const FCellInfo& InCellInfo, FPackageSourceControlHelper& PackageHelper)
 {
 	IVirtualTexturingEditorModule& VTModule = FModuleManager::Get().LoadModuleChecked<IVirtualTexturingEditorModule>("VirtualTexturingEditor");
@@ -56,14 +75,38 @@ bool UWorldPartitionRuntimeVirtualTextureBuilder::RunInternal(UWorld* World, con
 		FGlobalComponentRecreateRenderStateContext Context;
 	}
 
+	// We will need to build VTs for both shading paths
+	const ERHIFeatureLevel::Type CurFeatureLevel = World->FeatureLevel;
+	const ERHIFeatureLevel::Type AltFeatureLevel = (CurFeatureLevel == ERHIFeatureLevel::ES3_1 ? GMaxRHIFeatureLevel : ERHIFeatureLevel::ES3_1);
+	const EShadingPath CurShadingPath = FSceneInterface::GetShadingPath(CurFeatureLevel);
+	const EShadingPath AltShadingPath = FSceneInterface::GetShadingPath(AltFeatureLevel);
+
 	TSet<UPackage*> ModifiedPackages;
+	TArray<URuntimeVirtualTextureComponent*> Components[2];
 	for (TObjectIterator<URuntimeVirtualTextureComponent> It(RF_ClassDefaultObject, false, EInternalObjectFlags::Garbage); It; ++It)
 	{
 		if (It->GetWorld() == World)
 		{
-			if (VTModule.HasStreamedMips(*It) && VTModule.BuildStreamedMips(*It))
+			if (VTModule.HasStreamedMips(CurShadingPath, *It))
+			{
+				Components[0].Add(*It);
+			}
+
+			if (VTModule.HasStreamedMips(AltShadingPath, *It))
+			{
+				Components[1].Add(*It);
+			}
+		}
+	}
+
+	// Build for a current feature level first
+	if (Components[0].Num() != 0)
+	{
+		for (URuntimeVirtualTextureComponent* Component : Components[0])
+		{
+			if (VTModule.BuildStreamedMips(CurShadingPath, Component))
 			{
-				if (UVirtualTextureBuilder* VTBuilder = It->GetStreamingTexture(); VTBuilder->GetPackage()->IsDirty())
+				if (UVirtualTextureBuilder* VTBuilder = Component->GetStreamingTexture(); VTBuilder->GetPackage()->IsDirty())
 				{
 					ModifiedPackages.Add(VTBuilder->GetPackage());
 				}
@@ -71,6 +114,44 @@ bool UWorldPartitionRuntimeVirtualTextureBuilder::RunInternal(UWorld* World, con
 		}
 	}
 
+	// Build for others if any
+	if (Components[1].Num() != 0)
+	{
+		// Commandlets do not initialize shader resources for alternate feature levels, do it now
+		{
+			bool bUpdateProgressDialog = false;
+			bool bCacheAllRemainingShaders = true;
+			UMaterialInterface::SetGlobalRequiredFeatureLevel(AltFeatureLevel, true);
+			UMaterial::AllMaterialsCacheResourceShadersForRendering(bUpdateProgressDialog, bCacheAllRemainingShaders);
+			UMaterialInstance::AllMaterialsCacheResourceShadersForRendering(bUpdateProgressDialog, bCacheAllRemainingShaders);
+			CompileGlobalShaderMap(AltFeatureLevel);
+		}
+	
+		World->ChangeFeatureLevel(AltFeatureLevel);
+		
+		// Make sure all assets are finished compiling. Recreate render state after shader compilation complete
+		{
+			UMaterialInterface::SubmitRemainingJobsForWorld(World);
+			FAssetCompilingManager::Get().FinishAllCompilation();
+			FAssetCompilingManager::Get().ProcessAsyncTasks();
+			FGlobalComponentRecreateRenderStateContext Context;
+		}
+				
+		for (URuntimeVirtualTextureComponent* Component : Components[1])
+		{
+			if (VTModule.BuildStreamedMips(AltShadingPath, Component))
+			{
+				if (UVirtualTextureBuilder* VTBuilder = Component->GetStreamingTexture(); VTBuilder->GetPackage()->IsDirty())
+				{
+					ModifiedPackages.Add(VTBuilder->GetPackage());
+				}
+			}
+		}
+	}
+
+	// Restore world feature level
+	World->ChangeFeatureLevel(CurFeatureLevel);
+
 	// Wait for VT Textures to be ready before saving
 	FAssetCompilingManager::Get().FinishAllCompilation();
 
diff --git a/Engine/Source/Editor/VirtualTexturingEditor/Private/RuntimeVirtualTextureBuildStreamingMips.cpp b/Engine/Source/Editor/VirtualTexturingEditor/Private/RuntimeVirtualTextureBuildStreamingMips.cpp
--- a/Engine/Source/Editor/VirtualTexturingEditor/Private/RuntimeVirtualTextureBuildStreamingMips.cpp
+++ b/Engine/Source/Editor/VirtualTexturingEditor/Private/RuntimeVirtualTextureBuildStreamingMips.cpp
@@ -122,6 +122,12 @@ namespace
 namespace RuntimeVirtualTexture
 {
 	bool HasStreamedMips(URuntimeVirtualTextureComponent* InComponent)
+	{
+		EShadingPath ShadingPath = (InComponent && InComponent->GetScene()) ? InComponent->GetScene()->GetShadingPath() : EShadingPath::Deferred;
+		return HasStreamedMips(ShadingPath, InComponent);
+	}
+
+	bool HasStreamedMips(EShadingPath ShadingPath, URuntimeVirtualTextureComponent* InComponent)
 	{
 		if (InComponent == nullptr)
 		{
@@ -138,12 +144,23 @@ namespace RuntimeVirtualTexture
 			return false;
 		}
 
+		if (ShadingPath == EShadingPath::Mobile && !InComponent->GetStreamingTexture()->bSeparateTextureForMobile)
+		{
+			return false;
+		}
+
 		return true;
 	}
-
+	
 	bool BuildStreamedMips(URuntimeVirtualTextureComponent* InComponent, ERuntimeVirtualTextureDebugType DebugType)
 	{
-		if (!HasStreamedMips(InComponent))
+		EShadingPath ShadingPath = (InComponent && InComponent->GetScene()) ? InComponent->GetScene()->GetShadingPath() : EShadingPath::Deferred;
+		return BuildStreamedMips(ShadingPath, InComponent, DebugType);
+	}
+
+	bool BuildStreamedMips(EShadingPath ShadingPath, URuntimeVirtualTextureComponent* InComponent, ERuntimeVirtualTextureDebugType DebugType)
+	{
+		if (!HasStreamedMips(ShadingPath, InComponent))
 		{
 			return true;
 		}
@@ -304,7 +321,7 @@ namespace RuntimeVirtualTexture
 		// Place final pixel data into the runtime virtual texture
 		Task.EnterProgressFrame(TaskWorkBuildBulkData);
 
-		InComponent->InitializeStreamingTexture(ImageSizeX, ImageSizeY, (uint8*)FinalPixels.GetData());
+		InComponent->InitializeStreamingTexture(ShadingPath, ImageSizeX, ImageSizeY, (uint8*)FinalPixels.GetData());
 
 		return true;
 	}
diff --git a/Engine/Source/Editor/VirtualTexturingEditor/Private/RuntimeVirtualTextureBuildStreamingMips.h b/Engine/Source/Editor/VirtualTexturingEditor/Private/RuntimeVirtualTextureBuildStreamingMips.h
--- a/Engine/Source/Editor/VirtualTexturingEditor/Private/RuntimeVirtualTextureBuildStreamingMips.h
+++ b/Engine/Source/Editor/VirtualTexturingEditor/Private/RuntimeVirtualTextureBuildStreamingMips.h
@@ -11,7 +11,9 @@ namespace RuntimeVirtualTexture
 {
 	/** Returns true if the component describes a runtime virtual texture that has streaming mips. */
 	bool HasStreamedMips(URuntimeVirtualTextureComponent* InComponent);
+	bool HasStreamedMips(EShadingPath ShadingPath, URuntimeVirtualTextureComponent* InComponent);
 
 	/** Build the streaming mips texture. */
 	bool BuildStreamedMips(URuntimeVirtualTextureComponent* InComponent, ERuntimeVirtualTextureDebugType DebugType);
+	bool BuildStreamedMips(EShadingPath ShadingPath, URuntimeVirtualTextureComponent* InComponent, ERuntimeVirtualTextureDebugType DebugType);
 };
diff --git a/Engine/Source/Editor/VirtualTexturingEditor/Private/VirtualTexturingEditorModule.cpp b/Engine/Source/Editor/VirtualTexturingEditor/Private/VirtualTexturingEditorModule.cpp
--- a/Engine/Source/Editor/VirtualTexturingEditor/Private/VirtualTexturingEditorModule.cpp
+++ b/Engine/Source/Editor/VirtualTexturingEditor/Private/VirtualTexturingEditorModule.cpp
@@ -30,8 +30,8 @@ public:
 	//~ End IModuleInterface Interface.
 
 	//~ Begin IVirtualTexturingEditorModule Interface.
-	virtual bool HasStreamedMips(URuntimeVirtualTextureComponent* InComponent) const override;
-	virtual bool BuildStreamedMips(URuntimeVirtualTextureComponent* InComponent) const override;
+	virtual bool HasStreamedMips(EShadingPath ShadingPath, URuntimeVirtualTextureComponent* InComponent) const override;
+	virtual bool BuildStreamedMips(EShadingPath ShadingPath, URuntimeVirtualTextureComponent* InComponent) const override;
 	virtual void ConvertVirtualTextures(const TArray<UTexture2D *>& Textures, bool bConvertBackToNonVirtual, const TArray<UMaterial *>* RelatedMaterials /* = nullptr */) const override;
 	virtual void ConvertVirtualTexturesWithDialog(const TArray<UTexture2D *>& Textures, bool bConvertBackToNonVirtual) const override;
 	//~ End IVirtualTexturingEditorModule Interface.
@@ -82,14 +82,14 @@ void FVirtualTexturingEditorModule::OnPlacementModeRefresh(FName CategoryName)
 	}
 }
 
-bool FVirtualTexturingEditorModule::HasStreamedMips(URuntimeVirtualTextureComponent* InComponent) const
+bool FVirtualTexturingEditorModule::HasStreamedMips(EShadingPath ShadingPath, URuntimeVirtualTextureComponent* InComponent) const
 {
-	return RuntimeVirtualTexture::HasStreamedMips(InComponent);
+	return RuntimeVirtualTexture::HasStreamedMips(ShadingPath, InComponent);
 }
 
-bool FVirtualTexturingEditorModule::BuildStreamedMips(URuntimeVirtualTextureComponent* InComponent) const
+bool FVirtualTexturingEditorModule::BuildStreamedMips(EShadingPath ShadingPath, URuntimeVirtualTextureComponent* InComponent) const
 {
-	return RuntimeVirtualTexture::BuildStreamedMips(InComponent, ERuntimeVirtualTextureDebugType::None);
+	return RuntimeVirtualTexture::BuildStreamedMips(ShadingPath, InComponent, ERuntimeVirtualTextureDebugType::None);
 }
 
 void FVirtualTexturingEditorModule::ConvertVirtualTextures(const TArray<UTexture2D *>& Textures, bool bConvertBackToNonVirtual, const TArray<UMaterial *>* RelatedMaterials /* = nullptr */) const
diff --git a/Engine/Source/Editor/VirtualTexturingEditor/Public/VirtualTexturingEditorModule.h b/Engine/Source/Editor/VirtualTexturingEditor/Public/VirtualTexturingEditorModule.h
--- a/Engine/Source/Editor/VirtualTexturingEditor/Public/VirtualTexturingEditorModule.h
+++ b/Engine/Source/Editor/VirtualTexturingEditor/Public/VirtualTexturingEditorModule.h
@@ -7,15 +7,16 @@
 
 class UTexture2D;
 class UMaterial;
+enum class EShadingPath;
 
 /** Module for virtual texturing editor extensions. */
 class IVirtualTexturingEditorModule : public IModuleInterface
 {
 public:
 	/** Returns true if the component describes a runtime virtual texture that has streaming low mips. */
-	virtual bool HasStreamedMips(class URuntimeVirtualTextureComponent* InComponent) const = 0;
+	virtual bool HasStreamedMips(EShadingPath ShadingPath, class URuntimeVirtualTextureComponent* InComponent) const = 0;
 	/** Build the contents of the streaming low mips. */
-	virtual bool BuildStreamedMips(class URuntimeVirtualTextureComponent* InComponent) const = 0;
+	virtual bool BuildStreamedMips(EShadingPath ShadingPath, class URuntimeVirtualTextureComponent* InComponent) const = 0;
 	
 	virtual void ConvertVirtualTextures(const TArray<UTexture2D *>& Textures, bool bConvertBackToNonVirtual, const TArray<UMaterial *>* RelatedMaterials /* = nullptr */) const = 0;
 	virtual void ConvertVirtualTexturesWithDialog(const TArray<UTexture2D*>& Textures, bool bConvertBackToNonVirtual) const = 0;
diff --git a/Engine/Source/Runtime/Engine/Classes/Components/RuntimeVirtualTextureComponent.h b/Engine/Source/Runtime/Engine/Classes/Components/RuntimeVirtualTextureComponent.h
--- a/Engine/Source/Runtime/Engine/Classes/Components/RuntimeVirtualTextureComponent.h
+++ b/Engine/Source/Runtime/Engine/Classes/Components/RuntimeVirtualTextureComponent.h
@@ -120,7 +120,7 @@ public:
 	int32 NumStreamingMips() const { return FMath::Clamp(StreamLowMips, 0, 12); }
 
 	/** Get if we want to use any streaming low mips on this component. */
-	bool IsStreamingLowMips() const;
+	bool IsStreamingLowMips(EShadingPath ShadingPath) const;
 
 	/** Public getter for debug streaming mips flag. */
 	bool IsBuildDebugStreamingMips() { return bBuildDebugStreamingMips; }
@@ -129,13 +129,15 @@ public:
 	TEnumAsByte<ETextureLossyCompressionAmount> GetLossyCompressionAmount() const { return LossyCompressionAmount; }
 
 	/** Returns true if there are StreamingTexure contents but they are not valid for use. */
-	bool IsStreamingTextureInvalid() const;
+	bool IsStreamingTextureInvalid(EShadingPath ShadingPath) const;
 
 #if WITH_EDITOR
+	/** Returns true if there are StreamingTexure contents but they are not valid for use, taking into account all rendering modes */
+	bool IsStreamingTextureInvalid() const;
 	/** Set a new asset to hold the low mip streaming texture. This should only be called directly before setting data to the new asset. */
 	void SetStreamingTexture(UVirtualTextureBuilder* InTexture) { StreamingTexture = InTexture; }
 	/** Initialize the low mip streaming texture with the passed in size and data. */
-	void InitializeStreamingTexture(uint32 InSizeX, uint32 InSizeY, uint8* InData);
+	void InitializeStreamingTexture(EShadingPath ShadingPath, uint32 InSizeX, uint32 InSizeY, uint8* InData);
 #endif
 
 #if WITH_EDITOR
diff --git a/Engine/Source/Runtime/Engine/Classes/VT/VirtualTextureBuilder.h b/Engine/Source/Runtime/Engine/Classes/VT/VirtualTextureBuilder.h
--- a/Engine/Source/Runtime/Engine/Classes/VT/VirtualTextureBuilder.h
+++ b/Engine/Source/Runtime/Engine/Classes/VT/VirtualTextureBuilder.h
@@ -51,17 +51,29 @@ public:
 	UPROPERTY(VisibleAnywhere, BlueprintReadOnly, Category = Texture)
 	TObjectPtr<class UVirtualTexture2D> Texture;
 
+	/** The UTexture object for Mobile rendering. */
+	UPROPERTY(VisibleAnywhere, BlueprintReadOnly, Category = Texture)
+	TObjectPtr<class UVirtualTexture2D> TextureMobile;
+
 	/** Some client defined hash of that defines how the Texture was built. */
 	UPROPERTY()
 	uint64 BuildHash;
 
+	/** Virtual texture for a specific shading path */
+	UVirtualTexture2D* GetVirtualTexture(EShadingPath ShadingPath) const;
+
+	/** Whether to use a separate texture for Mobile rendering. A separate texure will be built using mobile preview editor mode */
+	UPROPERTY(EditAnywhere, Category = Texture)
+	bool bSeparateTextureForMobile = false;
+
 #if WITH_EDITOR
 	/** Creates a new UVirtualTexture2D and stores it in the contained Texture. */
-	void BuildTexture(FVirtualTextureBuildDesc const& BuildDesc);
+	void BuildTexture(EShadingPath ShadingPath, FVirtualTextureBuildDesc const& BuildDesc);
 #endif
 
 protected:
 	//~ Begin UObject Interface
 	virtual void Serialize(FArchive& Ar) override;
+	virtual void PostLoad() override;
 	//~ End UObject Interface
 };
diff --git a/Engine/Source/Runtime/Engine/Private/Components/RuntimeVirtualTextureComponent.cpp b/Engine/Source/Runtime/Engine/Private/Components/RuntimeVirtualTextureComponent.cpp
--- a/Engine/Source/Runtime/Engine/Private/Components/RuntimeVirtualTextureComponent.cpp
+++ b/Engine/Source/Runtime/Engine/Private/Components/RuntimeVirtualTextureComponent.cpp
@@ -191,7 +191,7 @@ uint64 URuntimeVirtualTextureComponent::CalculateStreamingTextureSettingsHash()
 	return Settings.PackedValue;
 }
 
-bool URuntimeVirtualTextureComponent::IsStreamingLowMips() const
+bool URuntimeVirtualTextureComponent::IsStreamingLowMips(EShadingPath ShadingPath) const
 {
 #if WITH_EDITOR
 	if (!bUseStreamingLowMipsInEditor && GIsEditor)
@@ -199,16 +199,30 @@ bool URuntimeVirtualTextureComponent::IsStreamingLowMips() const
 		return false;
 	}
 #endif
-	return VirtualTexture != nullptr && StreamingTexture != nullptr && StreamingTexture->Texture != nullptr;
+	return VirtualTexture != nullptr && StreamingTexture != nullptr && StreamingTexture->GetVirtualTexture(ShadingPath) != nullptr;
 }
 
-bool URuntimeVirtualTextureComponent::IsStreamingTextureInvalid() const
+bool URuntimeVirtualTextureComponent::IsStreamingTextureInvalid(EShadingPath ShadingPath) const
 {
-	return VirtualTexture != nullptr && StreamingTexture != nullptr && StreamingTexture->Texture != nullptr && StreamingTexture->BuildHash != CalculateStreamingTextureSettingsHash();
+	return 
+		VirtualTexture != nullptr && 
+		StreamingTexture != nullptr && 
+		StreamingTexture->GetVirtualTexture(ShadingPath) != nullptr && 
+		StreamingTexture->BuildHash != CalculateStreamingTextureSettingsHash();
 }
 
 #if WITH_EDITOR
 
+bool URuntimeVirtualTextureComponent::IsStreamingTextureInvalid() const
+{
+	return 
+		VirtualTexture != nullptr && 
+		StreamingTexture != nullptr && 
+		StreamingTexture->GetVirtualTexture(EShadingPath::Deferred) != nullptr && 
+		StreamingTexture->GetVirtualTexture(EShadingPath::Mobile) != nullptr && 
+		StreamingTexture->BuildHash != CalculateStreamingTextureSettingsHash();
+}
+
 // RAII class to release and recreate runtime virtual texture producers associated with a UVirtualTextureBuilder.
 // Required around modifications of a UVirtualTextureBuilder because virtual producers hold pointers to the internal data.
 class FScopedRuntimeVirtualTextureRecreate
@@ -252,7 +266,7 @@ static void GetLayerFormatSettings(FTextureFormatSettings& OutFormatSettings, EP
 	OutFormatSettings.CompressionYCoCg = IsLayerYCoCg;
 	OutFormatSettings.SRGB = IsLayerSRGB;
 }
-void URuntimeVirtualTextureComponent::InitializeStreamingTexture(uint32 InSizeX, uint32 InSizeY, uint8* InData)
+void URuntimeVirtualTextureComponent::InitializeStreamingTexture(EShadingPath ShadingPath, uint32 InSizeX, uint32 InSizeY, uint8* InData)
 {
 	// We need an existing StreamingTexture object to update.
 	if (VirtualTexture != nullptr && StreamingTexture != nullptr)
@@ -288,7 +302,7 @@ void URuntimeVirtualTextureComponent::InitializeStreamingTexture(uint32 InSizeX,
 		BuildDesc.InData = InData;
 
 		StreamingTexture->Modify();
-		StreamingTexture->BuildTexture(BuildDesc);
+		StreamingTexture->BuildTexture(ShadingPath, BuildDesc);
 	}
 }
 
diff --git a/Engine/Source/Runtime/Engine/Private/VT/VirtualTextureBuilder.cpp b/Engine/Source/Runtime/Engine/Private/VT/VirtualTextureBuilder.cpp
--- a/Engine/Source/Runtime/Engine/Private/VT/VirtualTextureBuilder.cpp
+++ b/Engine/Source/Runtime/Engine/Private/VT/VirtualTextureBuilder.cpp
@@ -5,6 +5,10 @@
 #include "RenderUtils.h"
 #include "VT/VirtualTexture.h"
 
+#if WITH_EDITOR
+#include "Interfaces/ITargetPlatform.h"
+#endif
+
 #include UE_INLINE_GENERATED_CPP_BY_NAME(VirtualTextureBuilder)
 
 UVirtualTextureBuilder::UVirtualTextureBuilder(const FObjectInitializer& ObjectInitializer)
@@ -18,27 +22,74 @@ UVirtualTextureBuilder::~UVirtualTextureBuilder()
 
 void UVirtualTextureBuilder::Serialize(FArchive& Ar)
 {
-	if (Ar.IsCooking() && Ar.IsSaving() && !UseVirtualTexturing(GMaxRHIFeatureLevel, Ar.CookingTarget()))
+#if WITH_EDITOR	
+	if (Ar.IsCooking() && Ar.IsSaving())
 	{
-		// Clear Texture during cook for platforms that don't support virtual texturing
 		UVirtualTexture2D* TextureBackup = Texture;
-		Texture = nullptr;
+		UVirtualTexture2D* TextureMobileBackup = TextureMobile;
+		
+		// Clear Texture during cook for platforms that don't support virtual texturing
+		if (!UseVirtualTexturing(GMaxRHIFeatureLevel, Ar.CookingTarget()))
+		{
+			Texture = nullptr;
+			TextureMobile = nullptr;
+		}
+
+		// Selectivly serialize VirtualTexture for platforms that support Deferred or/and Mobile rendering
+		if (!Ar.CookingTarget()->SupportsFeature(ETargetPlatformFeatures::DeferredRendering) && bSeparateTextureForMobile)
+		{
+			Texture = nullptr;
+		}
+				
+		if (!Ar.CookingTarget()->SupportsFeature(ETargetPlatformFeatures::MobileRendering))
+		{
+			TextureMobile = nullptr;
+		}
+
 		Super::Serialize(Ar);
+
 		Texture = TextureBackup;
+		TextureMobile = TextureMobileBackup;
 	}
 	else
+#endif
 	{
 		Super::Serialize(Ar);
 	}
 }
 
-#if WITH_EDITOR
+void UVirtualTextureBuilder::PostLoad()
+{
+	Super::PostLoad();
+
+	// Discard one of the VTs on a cooked platform that support both rendering modes
+	if (FPlatformProperties::RequiresCookedData())
+	{
+		if (FSceneInterface::GetShadingPath(GMaxRHIFeatureLevel) == EShadingPath::Mobile && bSeparateTextureForMobile)
+		{
+			Texture = nullptr;
+		}
+		else
+		{
+			TextureMobile = nullptr;
+		}
+	}
+}
 
-void UVirtualTextureBuilder::BuildTexture(FVirtualTextureBuildDesc const& BuildDesc)
+UVirtualTexture2D* UVirtualTextureBuilder::GetVirtualTexture(EShadingPath ShadingPath) const
 {
-	BuildHash = BuildDesc.BuildHash;
+	if (ShadingPath == EShadingPath::Mobile && bSeparateTextureForMobile)
+	{
+		return TextureMobile;
+	}
 
-	Texture = NewObject<UVirtualTexture2D>(this, TEXT("Texture"));
+	return Texture;
+}
+
+#if WITH_EDITOR
+
+static void BuildVirtualTexture2D(UVirtualTexture2D* Texture, FVirtualTextureBuildDesc const& BuildDesc)
+{
 	Texture->VirtualTextureStreaming = true;
 	Texture->LODGroup = BuildDesc.LODGroup;
 
@@ -59,5 +110,33 @@ void UVirtualTextureBuilder::BuildTexture(FVirtualTextureBuildDesc const& BuildD
 	Texture->PostEditChange();
 }
 
+
+void UVirtualTextureBuilder::BuildTexture(EShadingPath ShadingPath, FVirtualTextureBuildDesc const& BuildDesc)
+{
+	if (!bSeparateTextureForMobile)
+	{
+		// Always clear mobile specific VT if property is switched off
+		TextureMobile = nullptr;
+	}
+	
+	if (ShadingPath == EShadingPath::Mobile)
+	{
+		if (!bSeparateTextureForMobile)
+		{
+			return;
+		}
+		
+		BuildHash = BuildDesc.BuildHash;
+		TextureMobile = NewObject<UVirtualTexture2D>(this, TEXT("TextureMobile"));
+		BuildVirtualTexture2D(TextureMobile, BuildDesc);
+	}
+	else
+	{
+		BuildHash = BuildDesc.BuildHash;
+		Texture = NewObject<UVirtualTexture2D>(this, TEXT("Texture"));
+		BuildVirtualTexture2D(Texture, BuildDesc);
+	}
+}
+
 #endif
 
diff --git a/Engine/Source/Runtime/Landscape/Private/LandscapeRender.cpp b/Engine/Source/Runtime/Landscape/Private/LandscapeRender.cpp
--- a/Engine/Source/Runtime/Landscape/Private/LandscapeRender.cpp
+++ b/Engine/Source/Runtime/Landscape/Private/LandscapeRender.cpp
@@ -243,7 +243,13 @@ void ULandscapeComponent::GetUsedMaterials(TArray<UMaterialInterface*>& OutMater
 	{
 		OutMaterials.Append(MobileMaterialInterfaces.FilterByPredicate([](UMaterialInterface* MaterialInstance) { return MaterialInstance != nullptr; }));
 	}
-
+#if	WITH_EDITORONLY_DATA	
+	if (MobileCombinationMaterialInstances.Num())
+	{
+		OutMaterials.Append(MobileCombinationMaterialInstances.FilterByPredicate([](UMaterialInstanceConstant* MaterialInstance) { return MaterialInstance != nullptr; }));
+	}
+#endif
+	
 	if (OverrideMaterial)
 	{
 		OutMaterials.Add(OverrideMaterial);
diff --git a/Engine/Source/Runtime/Renderer/Private/VT/RuntimeVirtualTextureSceneProxy.cpp b/Engine/Source/Runtime/Renderer/Private/VT/RuntimeVirtualTextureSceneProxy.cpp
--- a/Engine/Source/Runtime/Renderer/Private/VT/RuntimeVirtualTextureSceneProxy.cpp
+++ b/Engine/Source/Runtime/Renderer/Private/VT/RuntimeVirtualTextureSceneProxy.cpp
@@ -59,13 +59,15 @@ FRuntimeVirtualTextureSceneProxy::FRuntimeVirtualTextureSceneProxy(URuntimeVirtu
 
 		const ERuntimeVirtualTextureMaterialType MaterialType = VirtualTexture->GetMaterialType();
 		const bool bClearTextures = VirtualTexture->GetClearTextures();
+		FSceneInterface* SceneInterface = InComponent->GetScene();
+		const EShadingPath ShadingPath = SceneInterface ? SceneInterface->GetShadingPath() : EShadingPath::Deferred;
 
 		// The producer object created here will be passed into the virtual texture system which will take ownership.
-		IVirtualTexture* Producer = new FRuntimeVirtualTextureProducer(ProducerDesc, ProducerId, MaterialType, bClearTextures, InComponent->GetScene(), Transform, Bounds);
-
+		IVirtualTexture* Producer = new FRuntimeVirtualTextureProducer(ProducerDesc, ProducerId, MaterialType, bClearTextures, SceneInterface, Transform, Bounds);
+		
 		// Create a producer for the streaming low mips. 
 		// This is bound with the main producer so that one allocated VT can use both runtime or streaming producers dependent on mip level.
-		if (InComponent->IsStreamingLowMips())
+		if (InComponent->IsStreamingLowMips(ShadingPath))
 		{
 			if (CVarVTStreamingMips.GetValueOnAnyThread() == 0)
 			{
@@ -79,7 +81,7 @@ FRuntimeVirtualTextureSceneProxy::FRuntimeVirtualTextureSceneProxy(URuntimeVirtu
 				});
 #endif
 			}
-			else if (InComponent->IsStreamingTextureInvalid())
+			else if (InComponent->IsStreamingTextureInvalid(ShadingPath))
 			{
 #if !UE_BUILD_SHIPPING
 				// Notify that streaming texture is invalid since this can cause performance regression.
@@ -94,7 +96,7 @@ FRuntimeVirtualTextureSceneProxy::FRuntimeVirtualTextureSceneProxy(URuntimeVirtu
 			}
 			else
 			{
-				UVirtualTexture2D* StreamingTexture = InComponent->GetStreamingTexture()->Texture;
+				UVirtualTexture2D* StreamingTexture = InComponent->GetStreamingTexture()->GetVirtualTexture(ShadingPath);
 
 				FVTProducerDescription StreamingProducerDesc;
 				IVirtualTexture* StreamingProducer = RuntimeVirtualTexture::CreateStreamingTextureProducer(StreamingTexture, ProducerDesc, StreamingProducerDesc);
