UPackageTools::UnloadPackages clears undo history because transaction records can hold onto assets it needs to unload

diff --git a/Engine/Source/Editor/UnrealEd/Private/PackageTools.cpp b/Engine/Source/Editor/UnrealEd/Private/PackageTools.cpp
--- a/Engine/Source/Editor/UnrealEd/Private/PackageTools.cpp
+++ b/Engine/Source/Editor/UnrealEd/Private/PackageTools.cpp
@@ -20,6 +20,7 @@
 #include "UObject/GCObjectScopeGuard.h"
 #include "Serialization/ArchiveFindCulprit.h"
 #include "Misc/PackageName.h"
+#include "Editor/Transactor.h"
 #include "Editor/EditorPerProjectUserSettings.h"
 #include "ISourceControlOperation.h"
 #include "SourceControlOperations.h"
@@ -29,7 +30,6 @@
 #include "Editor.h"
 #include "Dialogs/Dialogs.h"
 
-
 #include "ObjectTools.h"
 #include "Kismet2/BlueprintEditorUtils.h"
 #include "Kismet2/KismetEditorUtilities.h"
@@ -387,12 +387,16 @@ UPackageTools::UPackageTools(const FObjectInitializer& ObjectInitializer)
 			FlushAsyncLoading();
 			(*GFlushStreamingFunc)();
 
+			GWarn->BeginSlowTask(NSLOCTEXT("UnrealEd", "Unloading", "Unloading"), true);
+
 			// Remove potential references to to-be deleted objects from the GB selection set.
 			GEditor->GetSelectedObjects()->GetElementSelectionSet()->ClearSelection(FTypedElementSelectionOptions());
 
-			bool bScriptPackageWasUnloaded = false;
-
-			GWarn->BeginSlowTask( NSLOCTEXT("UnrealEd", "Unloading", "Unloading"), true );
+			// Clear undo history because transaction records can hold onto assets we want to unload
+			if (GEditor->Trans)
+			{
+				GEditor->Trans->Reset(NSLOCTEXT("UnrealEd", "UnloadPackagesResetUndo", "Unload Assets"));
+			}
 
 			// First add all packages to unload to the root set so they don't get garbage collected while we are operating on them
 			TArray<UPackage*> PackagesAddedToRoot;
@@ -410,6 +414,7 @@ UPackageTools::UPackageTools(const FObjectInitializer& ObjectInitializer)
 			FlushAsyncCompilation(PackagesToUnload.Array());
 
 			// Now try to clean up assets in all packages to unload.
+			bool bScriptPackageWasUnloaded = false;
 			int32 PackageIndex = 0;
 			for (UPackage* PackageBeingUnloaded : PackagesToUnload)
 			{
