WorldPartition - Builder commandlet: Added new option -RunningFromUnrealEd
* Skips the error(s) / warning(s) summary at the end of the commandlet execution, with the side effect of binding the process exit code to the actual result of the commandlet
* Enables FastExit to exit the process as soon as the commandlet finishes, which avoids possible crashes during shutdown. The engine isn't shutdown cleanly.

diff --git a/Engine/Source/Editor/UnrealEd/Private/Commandlets/WorldPartitionConvertCommandlet.cpp b/Engine/Source/Editor/UnrealEd/Private/Commandlets/WorldPartitionConvertCommandlet.cpp
--- a/Engine/Source/Editor/UnrealEd/Private/Commandlets/WorldPartitionConvertCommandlet.cpp
+++ b/Engine/Source/Editor/UnrealEd/Private/Commandlets/WorldPartitionConvertCommandlet.cpp
@@ -166,7 +166,13 @@ UWorldPartitionConvertCommandlet::UWorldPartitionConvertCommandlet(const FObject
 	, WorldExtent(HALF_WORLD_MAX)
 	, LandscapeGridSize(4)
 	, DataLayerFactory(NewObject<UDataLayerFactory>())
-{}
+{
+	if (FParse::Param(FCommandLine::Get(), TEXT("RunningFromUnrealEd")))
+	{
+		ShowErrorCount = false;	// This has the side effect of making the process return code match the return code of the commandlet
+		FastExit = true;		// Faster exit which avoids crash during shutdown. The engine isn't shutdown cleanly.
+	}
+}
 
 UWorld* UWorldPartitionConvertCommandlet::LoadWorld(const FString& LevelToLoad)
 {
@@ -234,14 +240,15 @@ UWorldPartition* UWorldPartitionConvertCommandlet::CreateWorldPartition(AWorldSe
 	{
 		WorldPartition->EditorHash->LoadConfig(*EditorHashClass, *LevelConfigFilename);
 		WorldPartition->RuntimeHash->LoadConfig(*RuntimeHashClass, *LevelConfigFilename);
+
 		// Use specified existing default HLOD layer if valid 
-		if (UHLODLayer* ExistingHLODLayer = LoadObject<UHLODLayer>(NULL, *DefaultHLODLayerAsset))
+		if (UHLODLayer* ExistingHLODLayer = LoadObject<UHLODLayer>(NULL, *DefaultHLODLayerAsset, nullptr, LOAD_NoWarn))
 		{
 			WorldPartition->DefaultHLODLayer = ExistingHLODLayer;
 		}
-		else
+		else if(UHLODLayer* FoundLayer = HLODLayers.FindRef(DefaultHLODLayerName))
 		{
-			WorldPartition->DefaultHLODLayer = HLODLayers.FindRef(DefaultHLODLayerName);
+			WorldPartition->DefaultHLODLayer = FoundLayer;
 		}
 	}
 
