Fix the fbx scene importer option dialog not copying correctly the options.

#rnx

diff --git a/Engine/Source/Editor/UnrealEd/Private/Fbx/SFbxSceneOptionWindow.cpp b/Engine/Source/Editor/UnrealEd/Private/Fbx/SFbxSceneOptionWindow.cpp
--- a/Engine/Source/Editor/UnrealEd/Private/Fbx/SFbxSceneOptionWindow.cpp
+++ b/Engine/Source/Editor/UnrealEd/Private/Fbx/SFbxSceneOptionWindow.cpp
@@ -1393,7 +1393,10 @@ bool SFbxSceneOptionWindow::CanImport()  const
 
 void SFbxSceneOptionWindow::CopyFbxOptionsToFbxOptions(UnFbx::FBXImportOptions *SourceOptions, UnFbx::FBXImportOptions *DestinationOptions)
 {
-	DestinationOptions = SourceOptions;
+	if (ensure(SourceOptions && DestinationOptions))
+	{
+		*DestinationOptions = *SourceOptions;
+	}
 }
 
 void SFbxSceneOptionWindow::CopyStaticMeshOptionsToFbxOptions(UnFbx::FBXImportOptions *ImportSettings, UFbxSceneImportOptionsStaticMesh* StaticMeshOptions)
