Add ability to disable AssetGuidelines from prompting the user to add new plugins/project settings

#localization none
#tests Tested that AssetGuidelines no longer prompted user when new API was used

diff --git a/Engine/Source/Editor/UnrealEd/Classes/Editor/AssetGuideline.h b/Engine/Source/Editor/UnrealEd/Classes/Editor/AssetGuideline.h
--- a/Engine/Source/Editor/UnrealEd/Classes/Editor/AssetGuideline.h
+++ b/Engine/Source/Editor/UnrealEd/Classes/Editor/AssetGuideline.h
@@ -57,6 +57,8 @@ public:
 private:
 
 	TWeakPtr<class SNotificationItem> NotificationPtr;
+
+	static bool bAssetGuidelinesEnabled;
 #endif // WITH_EDITORONLY_DATA
 
 #if WITH_EDITOR
@@ -67,6 +69,9 @@ public:
 	virtual void BeginDestroy() override;
 	/** End UObject interface */
 
+	static void EnableAssetGuidelines(bool isEnabled) { bAssetGuidelinesEnabled = isEnabled; }
+	static bool AreAssetGuidelinesEnabled() { return bAssetGuidelinesEnabled; }
+
 private:
 	void EnableMissingGuidelines(TArray<FString> IncorrectPlugins, TArray<FIniStringValue> IncorrectProjectSettings);
 	void DismissNotifications();
diff --git a/Engine/Source/Editor/UnrealEd/Private/AssetGuideline.cpp b/Engine/Source/Editor/UnrealEd/Private/AssetGuideline.cpp
--- a/Engine/Source/Editor/UnrealEd/Private/AssetGuideline.cpp
+++ b/Engine/Source/Editor/UnrealEd/Private/AssetGuideline.cpp
@@ -35,6 +35,10 @@
 
 #define LOCTEXT_NAMESPACE "AssetGuideine"
 
+DEFINE_LOG_CATEGORY_STATIC(LogAssetGuideline, Warning, All);
+
+bool UAssetGuideline::bAssetGuidelinesEnabled = true;
+
 bool UAssetGuideline::IsPostLoadThreadSafe() const
 {
 	return true;
@@ -126,6 +130,12 @@ void UAssetGuideline::PostLoad()
 			TitleText = !NeededPlugins.IsEmpty() && !NeededProjectSettings.IsEmpty() ? NeedBothGuidelines : !NeededPlugins.IsEmpty() ? NeedPlugins : NeedProjectSettings;
 		}
 
+		if (!bAssetGuidelinesEnabled)
+		{
+			UE_LOG(LogAssetGuideline, Warning, TEXT("%s %s"), *TitleText.ToString(), *SubText.ToString());
+			return;
+		}
+
 		auto WarningHyperLink = [](bool NeedPluginLink, bool NeedProjectSettingLink)
 		{
 			if (NeedProjectSettingLink)
