Removed unused function in \WorldPartitionRuntimeVirtualTextureBuilder.cpp

diff --git a/Engine/Source/Editor/UnrealEd/Private/WorldPartition/WorldPartitionRuntimeVirtualTextureBuilder.cpp b/Engine/Source/Editor/UnrealEd/Private/WorldPartition/WorldPartitionRuntimeVirtualTextureBuilder.cpp
--- a/Engine/Source/Editor/UnrealEd/Private/WorldPartition/WorldPartitionRuntimeVirtualTextureBuilder.cpp
+++ b/Engine/Source/Editor/UnrealEd/Private/WorldPartition/WorldPartitionRuntimeVirtualTextureBuilder.cpp
@@ -39,23 +39,6 @@ void UWorldPartitionRuntimeVirtualTextureBuilder::LoadRuntimeVirtualTextureActor
 	FAssetCompilingManager::Get().FinishAllCompilation();
 }
 
-static bool BuildVirtualTextureComponents(IVirtualTexturingEditorModule* Module, EShadingPath ShadingPath, TArray<URuntimeVirtualTextureComponent*>& Components)
-{
-	if (Components.Num() == 0)
-	{
-		return true;
-	}
-
-	for (URuntimeVirtualTextureComponent* Component : Components)
-	{
-		if (!Module->BuildStreamedMips(ShadingPath, Component))
-		{
-			return false;
-		}
-	}
-	return true;
-}
-
 bool UWorldPartitionRuntimeVirtualTextureBuilder::RunInternal(UWorld* World, const FCellInfo& InCellInfo, FPackageSourceControlHelper& PackageHelper)
 {
 	IVirtualTexturingEditorModule& VTModule = FModuleManager::Get().LoadModuleChecked<IVirtualTexturingEditorModule>("VirtualTexturingEditor");
