AssetRegistryGenerator.h: Remove CoreMinimal.h
#rnx

diff --git a/Engine/Source/Editor/UnrealEd/Private/Commandlets/AssetRegistryGenerator.cpp b/Engine/Source/Editor/UnrealEd/Private/Commandlets/AssetRegistryGenerator.cpp
--- a/Engine/Source/Editor/UnrealEd/Private/Commandlets/AssetRegistryGenerator.cpp
+++ b/Engine/Source/Editor/UnrealEd/Private/Commandlets/AssetRegistryGenerator.cpp
@@ -29,6 +29,7 @@
 #include "Misc/ConfigCacheIni.h"
 #include "Misc/DataDrivenPlatformInfoRegistry.h"
 #include "Misc/FileHelper.h"
+#include "Misc/Paths.h"
 #include "Misc/PathViews.h"
 #include "PakFileUtilities.h"
 #include "Policies/PrettyJsonPrintPolicy.h"
@@ -2521,6 +2522,11 @@ void FAssetRegistryGenerator::AddPackageAndDependenciesToChunk(FChunkPackageSet&
 	}
 }
 
+FString FAssetRegistryGenerator::GetTempPackagingDirectoryForPlatform(const FString& Platform) const
+{
+	return FPaths::ProjectSavedDir() / TEXT("TmpPackaging") / Platform;
+}
+
 void FAssetRegistryGenerator::FixupPackageDependenciesForChunks(FSandboxPlatformFile& InSandboxFile)
 {
 	UE_LOG(LogAssetRegistryGenerator, Log, TEXT("Starting FixupPackageDependenciesForChunks..."));
diff --git a/Engine/Source/Editor/UnrealEd/Public/Commandlets/AssetRegistryGenerator.h b/Engine/Source/Editor/UnrealEd/Public/Commandlets/AssetRegistryGenerator.h
--- a/Engine/Source/Editor/UnrealEd/Public/Commandlets/AssetRegistryGenerator.h
+++ b/Engine/Source/Editor/UnrealEd/Public/Commandlets/AssetRegistryGenerator.h
@@ -2,16 +2,20 @@
 
 #pragma once
 
-#include "CoreMinimal.h"
-
 #include "AssetRegistry/AssetData.h"
 #include "AssetRegistry/AssetRegistryState.h"
+#include "Containers/Array.h"
 #include "Containers/Map.h"
+#include "Containers/Set.h"
+#include "Containers/UnrealString.h"
 #include "Cooker/CookMPCollector.h"
 #include "Misc/AssetRegistryInterface.h"
-#include "Misc/Paths.h"
+#include "Misc/Optional.h"
+#include "Templates/SharedPointer.h"
 #include "Templates/UniquePtr.h"
+#include "UObject/NameTypes.h"
 #include "UObject/Object.h"
+#include "UObject/SoftObjectPath.h"
 #include "UObject/UObjectHash.h"
 
 class FSandboxPlatformFile;
@@ -409,10 +413,7 @@ private:
 	/**
 	 * Returns the path of the temporary packaging directory for the specified platform.
 	 */
-	FString GetTempPackagingDirectoryForPlatform(const FString& Platform) const
-	{
-		return FPaths::ProjectSavedDir() / TEXT("TmpPackaging") / Platform;
-	}
+	FString GetTempPackagingDirectoryForPlatform(const FString& Platform) const;
 
 	/** Returns the config-driven max size of a chunk for the given platform, or -1 for no limit. */
 	int64 GetMaxChunkSizePerPlatform( const ITargetPlatform* Platform ) const;
