Addendum to CL# 24342393

Adding back in the call to SetupParameterCollectionInstances in UWorld::InitWorld() and UMaterialParameterCollection::PostLoad() to fix bugs with parameters not updating on static meshes.

diff --git a/Engine/Source/Editor/UnrealEd/Private/Factories/EditorFactories.cpp b/Engine/Source/Editor/UnrealEd/Private/Factories/EditorFactories.cpp
--- a/Engine/Source/Editor/UnrealEd/Private/Factories/EditorFactories.cpp
+++ b/Engine/Source/Editor/UnrealEd/Private/Factories/EditorFactories.cpp
@@ -666,6 +666,12 @@ UMaterialParameterCollectionFactoryNew::UMaterialParameterCollectionFactoryNew(c
 UObject* UMaterialParameterCollectionFactoryNew::FactoryCreateNew(UClass* Class,UObject* InParent,FName Name,EObjectFlags Flags,UObject* Context,FFeedbackContext* Warn)
 {
 	UMaterialParameterCollection* MaterialParameterCollection = NewObject<UMaterialParameterCollection>(InParent, Class, Name, Flags);
+
+	if (MaterialParameterCollection)
+	{
+		MaterialParameterCollection->SetupWorldParameterCollectionInstances();
+	}
+
 	return MaterialParameterCollection;
 }
 
diff --git a/Engine/Source/Runtime/Engine/Classes/Materials/MaterialParameterCollection.h b/Engine/Source/Runtime/Engine/Classes/Materials/MaterialParameterCollection.h
--- a/Engine/Source/Runtime/Engine/Classes/Materials/MaterialParameterCollection.h
+++ b/Engine/Source/Runtime/Engine/Classes/Materials/MaterialParameterCollection.h
@@ -173,6 +173,9 @@ class UMaterialParameterCollection : public UObject
 		return *UniformBufferStruct;
 	}
 
+	/** Create an instance for this collection in every world. */
+	ENGINE_API void SetupWorldParameterCollectionInstances();
+
 private:
 	virtual ENGINE_API void FinishDestroy() override;
 	virtual ENGINE_API bool IsReadyForFinishDestroy() override;
diff --git a/Engine/Source/Runtime/Engine/Private/Materials/ParameterCollection.cpp b/Engine/Source/Runtime/Engine/Private/Materials/ParameterCollection.cpp
--- a/Engine/Source/Runtime/Engine/Private/Materials/ParameterCollection.cpp
+++ b/Engine/Source/Runtime/Engine/Private/Materials/ParameterCollection.cpp
@@ -47,9 +47,19 @@ void UMaterialParameterCollection::PostLoad()
 	}
 
 	CreateBufferStruct();
+	SetupWorldParameterCollectionInstances();
 	UpdateDefaultResource(true);
 }
 
+void UMaterialParameterCollection::SetupWorldParameterCollectionInstances()
+{
+	for (TObjectIterator<UWorld> It; It; ++It)
+	{
+		UWorld* CurrentWorld = *It;
+		CurrentWorld->AddParameterCollectionInstance(this, true);
+	}
+}
+
 void UMaterialParameterCollection::BeginDestroy()
 {
 	if (DefaultResource)
diff --git a/Engine/Source/Runtime/Engine/Private/World.cpp b/Engine/Source/Runtime/Engine/Private/World.cpp
--- a/Engine/Source/Runtime/Engine/Private/World.cpp
+++ b/Engine/Source/Runtime/Engine/Private/World.cpp
@@ -1895,6 +1895,8 @@ void UWorld::InitWorld(const InitializationValues IVS)
 		AvoidanceManager = NewObject<UAvoidanceManager>(this, GEngine->AvoidanceManagerClass);
 	}
 
+	SetupParameterCollectionInstances();
+
 	if (PersistentLevel->GetOuter() != this)
 	{
 		// Move persistent level into world so the world object won't get garbage collected in the multi- level
