SourceControl: Fixed 'rename' operation (and possibly others) still showing the 'Checkout Dialog' which should no longer be used in 'automatic checkout flow'. See #24299010 for original changelist.

diff --git a/Engine/Source/Editor/UnrealEd/Private/FileHelpers.cpp b/Engine/Source/Editor/UnrealEd/Private/FileHelpers.cpp
--- a/Engine/Source/Editor/UnrealEd/Private/FileHelpers.cpp
+++ b/Engine/Source/Editor/UnrealEd/Private/FileHelpers.cpp
@@ -1803,7 +1803,15 @@ bool FEditorFileUtils::PromptToCheckoutPackages(bool bCheckDirty, const TArray<U
 	// Prevent re-entrance into this function by setting up a guard value (also used by FEditorFileUtils::PromptForCheckoutAndSave)
 	TGuardValue<bool> PromptForCheckoutAndSaveGuard(bIsPromptingForCheckoutAndSave, true);
 
-	return PromptToCheckoutPackagesInternal(bCheckDirty, PackagesToCheckOut, OutPackagesCheckedOutOrMadeWritable, OutPackagesNotNeedingCheckout, bPromptingAfterModify, bAllowSkip);
+	bool bAutomaticCheckout = UseAlternateCheckoutWorkflow();
+	if (bAutomaticCheckout)
+	{
+		return FEditorFileUtils::AutomaticCheckoutOrPromptToRevertPackages(PackagesToCheckOut, OutPackagesCheckedOutOrMadeWritable, OutPackagesNotNeedingCheckout, nullptr);
+	}
+	else
+	{
+		return PromptToCheckoutPackagesInternal(bCheckDirty, PackagesToCheckOut, OutPackagesCheckedOutOrMadeWritable, OutPackagesNotNeedingCheckout, bPromptingAfterModify, bAllowSkip);
+	}
 }
 
 bool FEditorFileUtils::PromptToCheckoutPackagesInternal(bool bCheckDirty, const TArray<UPackage*>& PackagesToCheckOut, TArray<UPackage*>* OutPackagesCheckedOutOrMadeWritable, TArray<UPackage*>* OutPackagesNotNeedingCheckout, const bool bPromptingAfterModify, const bool bAllowSkip )
