Call ResetLoaders in bulk when unloading multiple packages at the same time
  - 28m28s to 4.3s when validating AT4

#rnx

diff --git a/Engine/Source/Editor/UnrealEd/Private/PackageTools.cpp b/Engine/Source/Editor/UnrealEd/Private/PackageTools.cpp
--- a/Engine/Source/Editor/UnrealEd/Private/PackageTools.cpp
+++ b/Engine/Source/Editor/UnrealEd/Private/PackageTools.cpp
@@ -501,24 +501,27 @@ UPackageTools::UPackageTools(const FObjectInitializer& ObjectInitializer)
 					ObjectsInPackage.Reset();
 				}
 
-				// Calling ::ResetLoaders now will force any bulkdata objects still attached to the FLinkerLoad to load
-				// their payloads into memory. If we don't call this now, then the version that will be called during
-				// garbage collection will cause the bulkdata objects to be invalidated rather than loading the payloads 
-				// into memory.
-				// This might seem odd, but if the package we are unloading is being renamed, then the inner UObjects will
-				// be moved to the newly named package rather than being garbage collected and so we need to make sure that
-				// their bulkdata objects remain valid, otherwise renamed packages will not save correctly and cease to function.
-				ResetLoaders(PackageBeingUnloaded);
-
-				if( PackageBeingUnloaded->IsDirty() )
+				// Cleanup.
+				bResult = true;
+			}
+
+			// Calling ::ResetLoaders now will force any bulkdata objects still attached to the FLinkerLoad to load
+			// their payloads into memory. If we don't call this now, then the version that will be called during
+			// garbage collection will cause the bulkdata objects to be invalidated rather than loading the payloads 
+			// into memory.
+			// This might seem odd, but if the package we are unloading is being renamed, then the inner UObjects will
+			// be moved to the newly named package rather than being garbage collected and so we need to make sure that
+			// their bulkdata objects remain valid, otherwise renamed packages will not save correctly and cease to function.
+			ResetLoaders(TArray<UObject*>(PackagesToUnload.Array()));
+
+			for (UPackage* PackageBeingUnloaded : PackagesToUnload)
+			{
+				if (PackageBeingUnloaded->IsDirty())
 				{
 					// The package was marked dirty as a result of something that happened above (e.g callbacks in CollectGarbage).  
 					// Dirty packages we actually care about unloading were filtered above so if the package becomes dirty here it should still be unloaded
 					PackageBeingUnloaded->SetDirtyFlag(false);
 				}
-
-				// Cleanup.
-				bResult = true;
 			}
 
 			// Set the callback for restoring RF_Standalone post reachability analysis.
