Add an option to disable real-time viewports when the Editor process is in the background (off by default)

diff --git a/Engine/Source/Editor/UnrealEd/Classes/Editor/EditorPerformanceSettings.h b/Engine/Source/Editor/UnrealEd/Classes/Editor/EditorPerformanceSettings.h
--- a/Engine/Source/Editor/UnrealEd/Classes/Editor/EditorPerformanceSettings.h
+++ b/Engine/Source/Editor/UnrealEd/Classes/Editor/EditorPerformanceSettings.h
@@ -41,6 +41,10 @@ class UEditorPerformanceSettings : public UDeveloperSettings
 	UPROPERTY(EditAnywhere, config, Category=EditorPerformance, meta=(DisplayName="Use Less CPU when in Background") )
 	uint32 bThrottleCPUWhenNotForeground:1;
 
+	/** Disables realtime viewports when the editor is in the background to reduce GPU performance impact on other applications */
+	UPROPERTY(EditAnywhere, config, Category = EditorPerformance, meta = (DisplayName = "Disable realtime viewports when in Background"))
+	uint32 bDisableRealtimeViewportsWhenNotForeground : 1;
+
 	/** Disables realtime viewports by default when connected via a remote session */
 	UPROPERTY(EditAnywhere, config, Category = EditorPerformance, meta = (DisplayName = "Disable realtime viewports by default in Remote Sessions"))
 	uint32 bDisableRealtimeViewportsInRemoteSessions : 1;
diff --git a/Engine/Source/Editor/UnrealEd/Private/EditorEngine.cpp b/Engine/Source/Editor/UnrealEd/Private/EditorEngine.cpp
--- a/Engine/Source/Editor/UnrealEd/Private/EditorEngine.cpp
+++ b/Engine/Source/Editor/UnrealEd/Private/EditorEngine.cpp
@@ -1620,6 +1620,21 @@ void UEditorEngine::Tick( float DeltaSeconds, bool bIdleMode )
 	// When in PIE if we are in immersive we do not tick the editor world unless there is a visible editor viewport.
 	bool bShouldTickEditorWorld = true;
 
+	// Conditionally disable all viewport rendering when the editor is in the background.
+	// This aims to improve GPU performance of other applications when the editor is not actively used.
+	{
+		const bool bShouldDisableRendering = !FApp::HasFocus() && GetDefault<UEditorPerformanceSettings>()->bDisableRealtimeViewportsWhenNotForeground;
+		const FText SystemDisplayName = LOCTEXT("RealtimeOverrideMessage_BackgroundProcess", "Background Process");
+		for (FEditorViewportClient* const ViewportClient : AllViewportClients)
+		{
+			ViewportClient->RemoveRealtimeOverride(SystemDisplayName, false /*bCheckMissingOverride*/);
+			if (bShouldDisableRendering)
+			{
+				ViewportClient->AddRealtimeOverride(false /*bShouldBeRealtime*/, SystemDisplayName);
+			}
+		}
+	}
+
 	//@todo Multiple Worlds: Do we need to consider what world we are in here?
 
 	// Find which viewport has audio focus, i.e. gets to set the listener location
diff --git a/Engine/Source/Editor/UnrealEd/Private/EditorPerformanceSettings.cpp b/Engine/Source/Editor/UnrealEd/Private/EditorPerformanceSettings.cpp
--- a/Engine/Source/Editor/UnrealEd/Private/EditorPerformanceSettings.cpp
+++ b/Engine/Source/Editor/UnrealEd/Private/EditorPerformanceSettings.cpp
@@ -17,6 +17,7 @@ UEditorPerformanceSettings::UEditorPerformanceSettings(const FObjectInitializer&
 	: Super(ObjectInitializer)
 	, bShowFrameRateAndMemory(false)
 	, bThrottleCPUWhenNotForeground(true)
+	, bDisableRealtimeViewportsWhenNotForeground(false)
 	, bDisableRealtimeViewportsInRemoteSessions(true)
 	, bMonitorEditorPerformance(false)
 	, bEnableScalabilityWarningIndicator(true)
