Don't save empty packages (thanks Michael Andersson).

diff --git a/Engine/Source/Editor/UnrealEd/Private/Commandlets/ConvertLevelsToExternalActorsCommandlet.cpp b/Engine/Source/Editor/UnrealEd/Private/Commandlets/ConvertLevelsToExternalActorsCommandlet.cpp
--- a/Engine/Source/Editor/UnrealEd/Private/Commandlets/ConvertLevelsToExternalActorsCommandlet.cpp
+++ b/Engine/Source/Editor/UnrealEd/Private/Commandlets/ConvertLevelsToExternalActorsCommandlet.cpp
@@ -384,7 +384,14 @@ int32 UConvertLevelsToExternalActorsCommandlet::Main(const FString& Params)
 		Level->ConvertAllActorsToPackaging(bConvertToExternal);
 		UPackage* LevelPackage = Level->GetPackage();
 		PackagesToSave.Add(LevelPackage);
-		PackagesToSave.Append(Level->GetLoadedExternalObjectPackages());
+
+		for (UPackage* ExternalPackage : Level->GetLoadedExternalObjectPackages())
+		{
+			if (!UPackage::IsEmptyPackage(ExternalPackage))
+			{
+				PackagesToSave.Add(ExternalPackage);
+			}
+		}
 	}
 
 	for (UPackage* PackageToSave : PackagesToSave)
