MPCook: Fix replication of FEDLCookChecker data to be sent along with each package result so that we don't lose information for some completed packages when a CookWorker crashes.
#rnx

diff --git a/Engine/Source/Editor/UnrealEd/Private/CookOnTheFlyServer.cpp b/Engine/Source/Editor/UnrealEd/Private/CookOnTheFlyServer.cpp
--- a/Engine/Source/Editor/UnrealEd/Private/CookOnTheFlyServer.cpp
+++ b/Engine/Source/Editor/UnrealEd/Private/CookOnTheFlyServer.cpp
@@ -10367,18 +10367,18 @@ public:
 	virtual FGuid GetMessageType() const override { return MessageType; }
 	virtual const TCHAR* GetDebugName() const override { return TEXT("FEDLMPCollector"); }
 
-	virtual void ClientTick(FMPCollectorClientTickContext& Context) override;
+	virtual void ClientTickPackage(FMPCollectorClientTickPackageContext& Context) override;
 	virtual void ServerReceiveMessage(FMPCollectorServerMessageContext& Context, FCbObjectView Message) override;
 
 	static FGuid MessageType;
 };
 FGuid FEDLMPCollector::MessageType(TEXT("0164FD08F6884F6A82D2D00F8F70B182"));
 
-void FEDLMPCollector::ClientTick(FMPCollectorClientTickContext& Context)
+void FEDLMPCollector::ClientTickPackage(FMPCollectorClientTickPackageContext& Context)
 {
 	FCbWriter Writer;
 	bool bHasData;
-	UE::SavePackageUtilities::EDLCookInfoMoveToCompactBinaryAndClear(Writer, bHasData);
+	UE::SavePackageUtilities::EDLCookInfoMoveToCompactBinaryAndClear(Writer, bHasData, Context.GetPackageName());
 	if (bHasData)
 	{
 		Context.AddMessage(Writer.Save().AsObject());
diff --git a/Engine/Source/Runtime/CoreUObject/Private/UObject/SavePackage/SavePackageUtilities.cpp b/Engine/Source/Runtime/CoreUObject/Private/UObject/SavePackage/SavePackageUtilities.cpp
--- a/Engine/Source/Runtime/CoreUObject/Private/UObject/SavePackage/SavePackageUtilities.cpp
+++ b/Engine/Source/Runtime/CoreUObject/Private/UObject/SavePackage/SavePackageUtilities.cpp
@@ -1156,9 +1156,9 @@ void FEDLCookChecker::Reset()
 {
 	check(!GIsSavingPackage);
 
-	Nodes.Empty();
-	NodeHashToNodeID.Empty();
-	NodePrereqs.Empty();
+	Nodes.Reset();
+	NodeHashToNodeID.Reset();
+	NodePrereqs.Reset();
 	bIsActive = false;
 }
 
@@ -1627,6 +1627,13 @@ void EDLCookInfoMoveToCompactBinaryAndClear(FCbWriter& Writer, bool& bOutHasData
 	FEDLCookChecker::MoveToCompactBinaryAndClear(Writer, bOutHasData);
 }
 
+void EDLCookInfoMoveToCompactBinaryAndClear(FCbWriter& Writer, bool& bOutHasData, FName PackageName)
+{
+	// For simplicity, instead of sending only information related to the given Package, we send all data.
+	LLM_SCOPE_BYTAG(EDLCookChecker);
+	FEDLCookChecker::MoveToCompactBinaryAndClear(Writer, bOutHasData);
+}
+
 bool EDLCookInfoAppendFromCompactBinary(FCbFieldView Field)
 {
 	LLM_SCOPE_BYTAG(EDLCookChecker);
diff --git a/Engine/Source/Runtime/CoreUObject/Public/UObject/SavePackage.h b/Engine/Source/Runtime/CoreUObject/Public/UObject/SavePackage.h
--- a/Engine/Source/Runtime/CoreUObject/Public/UObject/SavePackage.h
+++ b/Engine/Source/Runtime/CoreUObject/Public/UObject/SavePackage.h
@@ -315,6 +315,7 @@ namespace UE::SavePackageUtilities
 	COREUOBJECT_API void VerifyEDLCookInfo(bool bFullReferencesExpected = true);
 	COREUOBJECT_API void EDLCookInfoAddIterativelySkippedPackage(FName LongPackageName);
 	COREUOBJECT_API void EDLCookInfoMoveToCompactBinaryAndClear(FCbWriter& Writer, bool& bOutHasData);
+	COREUOBJECT_API void EDLCookInfoMoveToCompactBinaryAndClear(FCbWriter& Writer, bool& bOutHasData, FName PackageName);
 	COREUOBJECT_API bool EDLCookInfoAppendFromCompactBinary(FCbFieldView Field);
 	COREUOBJECT_API bool CanSkipEditorReferencedPackagesWhenCooking();
 
