Cooker: Diagnostics: Add log statement for skipped packages (e.g. NeverCook) when per-package displaymode is on.
#rnx

diff --git a/Engine/Source/Editor/UnrealEd/Private/CookOnTheFlyServer.cpp b/Engine/Source/Editor/UnrealEd/Private/CookOnTheFlyServer.cpp
--- a/Engine/Source/Editor/UnrealEd/Private/CookOnTheFlyServer.cpp
+++ b/Engine/Source/Editor/UnrealEd/Private/CookOnTheFlyServer.cpp
@@ -2482,9 +2482,26 @@ void UCookOnTheFlyServer::NotifyRemovedFromWorker(UE::Cook::FPackageData& Packag
 
 void UCookOnTheFlyServer::DemoteToIdle(UE::Cook::FPackageData& PackageData, UE::Cook::ESendFlags SendFlags, UE::Cook::ESuppressCookReason Reason)
 {
+	using namespace UE::Cook;
+
 	if (PackageData.IsInProgress())
 	{
 		WorkerRequests->ReportDemoteToIdle(PackageData, Reason);
+		if (GCookProgressDisplay & ((int32)ECookProgressDisplayMode::Instigators | (int32)ECookProgressDisplayMode::PackageNames))
+		{
+			WriteToString<256> PackageNameStr(PackageData.GetPackageName());
+
+			// ExternalActors: Do not send a message for every NeverCook external Actor package; too much spam
+			if (!(Reason == ESuppressCookReason::NeverCook &&
+				PackageNameStr.ToView().Contains(ULevel::GetExternalActorsFolderName())))
+			{
+				UE_CLOG((GCookProgressDisplay & (int32)ECookProgressDisplayMode::Instigators), LogCook, Display,
+					TEXT("Cooking %s, Instigator: { %s } -> Skipped %s"), *PackageNameStr,
+					*(PackageData.GetInstigator().ToString()), LexToString(Reason));
+				UE_CLOG(GCookProgressDisplay & (int32)ECookProgressDisplayMode::PackageNames, LogCook, Display,
+					TEXT("Cooking %s -> Skipped %s"), *PackageNameStr, LexToString(Reason));
+			}
+		}
 	}
 	PackageData.SendToState(UE::Cook::EPackageState::Idle, SendFlags);
 }
diff --git a/Engine/Source/Editor/UnrealEd/Private/Cooker/CookTypes.cpp b/Engine/Source/Editor/UnrealEd/Private/Cooker/CookTypes.cpp
--- a/Engine/Source/Editor/UnrealEd/Private/Cooker/CookTypes.cpp
+++ b/Engine/Source/Editor/UnrealEd/Private/Cooker/CookTypes.cpp
@@ -33,6 +33,29 @@ const TCHAR* LexToString(EReleaseSaveReason Reason)
 	}
 }
 
+const TCHAR* LexToString(ESuppressCookReason Reason)
+{
+	switch (Reason)
+	{
+	case ESuppressCookReason::InvalidSuppressCookReason: return TEXT("InvalidSuppressCookReason");
+	case ESuppressCookReason::AlreadyCooked: return TEXT("AlreadyCooked");
+	case ESuppressCookReason::NeverCook: return TEXT("NeverCook");
+	case ESuppressCookReason::DoesNotExistInWorkspaceDomain: return TEXT("DoesNotExistInWorkspaceDomain");
+	case ESuppressCookReason::ScriptPackage: return TEXT("ScriptPackage");
+	case ESuppressCookReason::NotInCurrentPlugin: return TEXT("NotInCurrentPlugin");
+	case ESuppressCookReason::Redirected: return TEXT("Redirected");
+	case ESuppressCookReason::OrphanedGenerated: return TEXT("OrphanedGenerated");
+	case ESuppressCookReason::LoadError: return TEXT("LoadError");
+	case ESuppressCookReason::SaveError: return TEXT("SaveError");
+	case ESuppressCookReason::OnlyEditorOnly: return TEXT("OnlyEditorOnly");
+	case ESuppressCookReason::CookCanceled: return TEXT("CookCanceled");
+	case ESuppressCookReason::MultiprocessAssignmentError: return TEXT("MultiprocessAssignmentError");
+	case ESuppressCookReason::RetractedByCookDirector: return TEXT("RetractedByCookDirector");
+	default: return TEXT("InvalidSuppressCookReason");
+	}
+}
+
+
 FCookerTimer::FCookerTimer(float InTimeSlice)
 	: StartTime(FPlatformTime::Seconds()), TimeSlice(InTimeSlice)
 {
diff --git a/Engine/Source/Editor/UnrealEd/Private/Cooker/CookTypes.h b/Engine/Source/Editor/UnrealEd/Private/Cooker/CookTypes.h
--- a/Engine/Source/Editor/UnrealEd/Private/Cooker/CookTypes.h
+++ b/Engine/Source/Editor/UnrealEd/Private/Cooker/CookTypes.h
@@ -156,6 +156,7 @@ namespace UE::Cook
 		MultiprocessAssignmentError,
 		RetractedByCookDirector,
 	};
+	const TCHAR* LexToString(UE::Cook::ESuppressCookReason Reason);
 
 	/** The type of callback for External Requests that needs to be executed within the Scheduler's lock. */
 	typedef TUniqueFunction<void()> FSchedulerCallback;
