Cooker: MPCook: Set the log filenames for CookWorkers, and if a CookWorker crashes, write out its log after the last heartbeat. Increase severity of CookWorkerCrash log message from Warning to Error.
#rnx

diff --git a/Engine/Source/Editor/UnrealEd/Classes/CookOnTheSide/CookLog.h b/Engine/Source/Editor/UnrealEd/Classes/CookOnTheSide/CookLog.h
--- a/Engine/Source/Editor/UnrealEd/Classes/CookOnTheSide/CookLog.h
+++ b/Engine/Source/Editor/UnrealEd/Classes/CookOnTheSide/CookLog.h
@@ -5,3 +5,4 @@
 #include "Logging/LogMacros.h"
 
 DECLARE_LOG_CATEGORY_EXTERN(LogCook, Log, All);
+extern FName LogCookName;
diff --git a/Engine/Source/Editor/UnrealEd/Private/CookOnTheFlyServer.cpp b/Engine/Source/Editor/UnrealEd/Private/CookOnTheFlyServer.cpp
--- a/Engine/Source/Editor/UnrealEd/Private/CookOnTheFlyServer.cpp
+++ b/Engine/Source/Editor/UnrealEd/Private/CookOnTheFlyServer.cpp
@@ -144,6 +144,7 @@
 #define LOCTEXT_NAMESPACE "Cooker"
 
 DEFINE_LOG_CATEGORY(LogCook);
+FName LogCookName(TEXT("LogCook"));
 LLM_DEFINE_TAG(Cooker);
 LLM_DEFINE_TAG(Cooker_SavePackage);
 
diff --git a/Engine/Source/Editor/UnrealEd/Private/Cooker/CookDirector.cpp b/Engine/Source/Editor/UnrealEd/Private/Cooker/CookDirector.cpp
--- a/Engine/Source/Editor/UnrealEd/Private/Cooker/CookDirector.cpp
+++ b/Engine/Source/Editor/UnrealEd/Private/Cooker/CookDirector.cpp
@@ -10,6 +10,7 @@
 #include "CookWorkerServer.h"
 #include "CookOnTheSide/CookOnTheFlyServer.h"
 #include "CoreGlobals.h"
+#include "GenericPlatform/GenericPlatformOutputDevices.h"
 #include "LoadBalanceCookBurden.h"
 #include "HAL/PlatformMisc.h"
 #include "HAL/PlatformTime.h"
@@ -17,6 +18,7 @@
 #include "Math/NumericLimits.h"
 #include "Misc/CommandLine.h"
 #include "Misc/ConfigCacheIni.h"
+#include "Misc/PathViews.h"
 #include "PackageTracker.h"
 #include "Serialization/CompactBinary.h"
 #include "Serialization/CompactBinaryWriter.h"
@@ -1198,6 +1200,14 @@ void FCookDirector::TickWorkerShutdowns(ECookDirectorThread TickThread)
 	CompletedWorkers.Empty();
 }
 
+FString FCookDirector::GetWorkerLogFileName(int32 ProfileId)
+{
+	FString DirectorLogFileName = FGenericPlatformOutputDevices::GetAbsoluteLogFilename();
+	FStringView BaseFileName = FPathViews::GetBaseFilenameWithPath(DirectorLogFileName);
+	FStringView Extension = FPathViews::GetExtension(DirectorLogFileName, true /* bIncludeDot */);
+	return FString::Printf(TEXT("%.*s_Worker%d%*s"), BaseFileName.Len(), BaseFileName.GetData(), ProfileId, Extension.Len(), Extension.GetData());
+}
+
 FString FCookDirector::GetWorkerCommandLine(FWorkerId WorkerId, int32 ProfileId)
 {
 	FString CommandLine = FCommandLine::Get();
@@ -1219,7 +1229,8 @@ FString FCookDirector::GetWorkerCommandLine(FWorkerId WorkerId, int32 ProfileId)
 				Token.StartsWith(TEXT("-ShowCookWorker")) ||
 				Token.StartsWith(TEXT("-CoreLimit")) ||
 				Token.StartsWith(TEXT("-PhysicalCoreLimit")) ||
-				Token.StartsWith(TEXT("-CookProcessCount="))
+				Token.StartsWith(TEXT("-CookProcessCount=")) ||
+				Token.StartsWith(TEXT("-abslog="))
 				)
 			{
 				return;
@@ -1244,6 +1255,7 @@ FString FCookDirector::GetWorkerCommandLine(FWorkerId WorkerId, int32 ProfileId)
 	Tokens.Add(FString::Printf(TEXT("-CookDirectorHost=%s"), *WorkerConnectAuthority));
 	Tokens.Add(FString::Printf(TEXT("-MultiprocessId=%d"), WorkerId.GetRemoteIndex() + 1));
 	Tokens.Add(FString::Printf(TEXT("-CookProfileId=%d"), ProfileId));
+	Tokens.Add(FString::Printf(TEXT("-abslog=%s"), *GetWorkerLogFileName(ProfileId)));
 	if (CoreLimit > 0)
 	{
 		Tokens.Add(FString::Printf(TEXT("-PhysicalCoreLimit=%d"), CoreLimit));
diff --git a/Engine/Source/Editor/UnrealEd/Private/Cooker/CookDirector.h b/Engine/Source/Editor/UnrealEd/Private/Cooker/CookDirector.h
--- a/Engine/Source/Editor/UnrealEd/Private/Cooker/CookDirector.h
+++ b/Engine/Source/Editor/UnrealEd/Private/Cooker/CookDirector.h
@@ -172,6 +172,8 @@ private:
 	void TickWorkerConnects(ECookDirectorThread TickThread);
 	/** Tick helper: tick any workers that are shutting down. */
 	void TickWorkerShutdowns(ECookDirectorThread TickThread);
+	/** The LogPath a worker process writes to. */
+	FString GetWorkerLogFileName(int32 ProfileId);
 	/** Get the commandline to launch a worker process with. */
 	FString GetWorkerCommandLine(FWorkerId WorkerId, int32 ProfileId);
 	/** Calls the configured LoadBalanceAlgorithm. Input Requests have been sorted by leaf to root load order. */
diff --git a/Engine/Source/Editor/UnrealEd/Private/Cooker/CookWorkerClient.cpp b/Engine/Source/Editor/UnrealEd/Private/Cooker/CookWorkerClient.cpp
--- a/Engine/Source/Editor/UnrealEd/Private/Cooker/CookWorkerClient.cpp
+++ b/Engine/Source/Editor/UnrealEd/Private/Cooker/CookWorkerClient.cpp
@@ -827,6 +827,7 @@ void FCookWorkerClient::HandleHeartbeatMessage(FMPCollectorClientMessageContext&
 		return;
 	}
 
+	UE_LOG(LogCook, Display, TEXT("CookWorkerHeartbeat: %d"), Message.HeartbeatNumber);
 	SendMessage(FHeartbeatMessage(Message.HeartbeatNumber));
 }
 
diff --git a/Engine/Source/Editor/UnrealEd/Private/Cooker/CookWorkerServer.cpp b/Engine/Source/Editor/UnrealEd/Private/Cooker/CookWorkerServer.cpp
--- a/Engine/Source/Editor/UnrealEd/Private/Cooker/CookWorkerServer.cpp
+++ b/Engine/Source/Editor/UnrealEd/Private/Cooker/CookWorkerServer.cpp
@@ -13,6 +13,8 @@
 #include "Interfaces/ITargetPlatformManagerModule.h"
 #include "Math/NumericLimits.h"
 #include "Misc/AssertionMacros.h"
+#include "Misc/FileHelper.h"
+#include "Misc/Parse.h"
 #include "Misc/ScopeLock.h"
 #include "PackageResultsMessage.h"
 #include "PackageTracker.h"
@@ -52,6 +54,55 @@ void FCookWorkerServer::DetachFromRemoteProcess()
 	bTerminateImmediately = false;
 	SendBuffer.Reset();
 	ReceiveBuffer.Reset();
+
+	if (bNeedCrashDiagnostics)
+	{
+		SendCrashDiagnostics();
+		bNeedCrashDiagnostics = false;
+	}
+}
+
+void FCookWorkerServer::SendCrashDiagnostics()
+{
+	FString LogFileName = Director.GetWorkerLogFileName(ProfileId);
+	UE_LOG(LogCook, Display, TEXT("CookWorker %d log messages written after communication loss:"), ProfileId);
+	FString LogText;
+	if (!FFileHelper::LoadFileToString(LogText, *LogFileName))
+	{
+		UE_LOG(LogCook, Warning, TEXT("No log file found for CookWorker %d."), ProfileId);
+	}
+	else
+	{
+		FString LastSentHeartbeat = FString::Printf(TEXT("CookWorkerHeartbeat: %d"), LastReceivedHeartbeatNumber);
+		int32 StartIndex = INDEX_NONE;
+		for (FStringView MarkerText : { FStringView(LastSentHeartbeat),
+			TEXTVIEW("CookWorkerHeartbeat: "), TEXTVIEW("Connection to CookDirector successful") })
+		{
+			StartIndex = UE::String::FindLast(LogText, MarkerText);
+			if (StartIndex >= 0)
+			{
+				break;
+			}
+		}
+		const TCHAR* StartText = *LogText;
+		FString Line;
+		if (StartIndex != INDEX_NONE)
+		{
+			// Skip the MarkerLine
+			StartText = *LogText + StartIndex;
+			FParse::Line(&StartText, Line);
+			if (*StartText == '\0')
+			{
+				// If there was no line after the MarkerLine, write out the MarkerLine
+				StartText = *LogText + StartIndex;
+			}
+		}
+
+		while (FParse::Line(&StartText, Line))
+		{
+			UE_LOG(LogCook, Display, TEXT("[CookWorker %d]: %s"), ProfileId, *Line);
+		}
+	}
 }
 
 void FCookWorkerServer::ShutdownRemoteProcess()
@@ -362,8 +413,9 @@ void FCookWorkerServer::LaunchProcess()
 	else
 	{
 		// GetLastError information was logged by CreateProc
-		UE_LOG(LogCook, Error, TEXT("CookWorkerServer %d failed to create CookWorker process. Assigned packages will be returned to the director."),
+		UE_LOG(LogCook, Error, TEXT("CookWorkerCrash: %d failed to create CookWorker process. Assigned packages will be returned to the director."),
 			ProfileId);
+		bNeedCrashDiagnostics = true;
 		SendToState(EConnectStatus::LostConnection);
 	}
 }
@@ -380,8 +432,9 @@ void FCookWorkerServer::TickWaitForConnect()
 	{
 		if (!FPlatformProcess::IsProcRunning(CookWorkerHandle))
 		{
-			UE_LOG(LogCook, Error, TEXT("CookWorkerServer %d process terminated before connecting. Assigned packages will be returned to the director."),
+			UE_LOG(LogCook, Error, TEXT("CookWorkerCrash: %d process terminated before connecting. Assigned packages will be returned to the director."),
 				ProfileId);
+			bNeedCrashDiagnostics = true;
 			SendToState(EConnectStatus::LostConnection);
 			return;
 		}
@@ -390,9 +443,10 @@ void FCookWorkerServer::TickWaitForConnect()
 
 	if (CurrentTime - ConnectStartTimeSeconds > WaitForConnectTimeout && !IsCookIgnoreTimeouts())
 	{
-		UE_LOG(LogCook, Error, TEXT("CookWorkerServer %d process failed to connect within %.0f seconds. Assigned packages will be returned to the director."),
+		UE_LOG(LogCook, Error, TEXT("CookWorkerCrash: %d process failed to connect within %.0f seconds. Assigned packages will be returned to the director."),
 			ProfileId, WaitForConnectTimeout);
 		ShutdownRemoteProcess();
+		bNeedCrashDiagnostics = true;
 		SendToState(EConnectStatus::LostConnection);
 		return;
 	}
@@ -434,8 +488,9 @@ void FCookWorkerServer::PumpSendMessages()
 	UE::CompactBinaryTCP::EConnectionStatus Status = UE::CompactBinaryTCP::TryFlushBuffer(Socket, SendBuffer);
 	if (Status == UE::CompactBinaryTCP::Failed)
 	{
-		UE_LOG(LogCook, Error, TEXT("CookWorkerServer %d failed to write to socket, we will shutdown the remote process. Assigned packages will be returned to the director."),
+		UE_LOG(LogCook, Error, TEXT("CookWorker Crash: %d failed to write to socket, we will shutdown the remote process. Assigned packages will be returned to the director."),
 			ProfileId);
+		bNeedCrashDiagnostics = true;
 		SendToState(EConnectStatus::WaitForDisconnect);
 		bTerminateImmediately = true;
 	}
@@ -470,8 +525,9 @@ void FCookWorkerServer::PumpReceiveMessages()
 	EConnectionStatus SocketStatus = TryReadPacket(Socket, ReceiveBuffer, Messages);
 	if (SocketStatus != EConnectionStatus::Okay && SocketStatus != EConnectionStatus::Incomplete)
 	{
-		UE_LOG(LogCook, Warning, TEXT("CookWorkerServer %d failed to read from socket, we will shutdown the remote process. Assigned packages will be returned to the director."),
+		UE_LOG(LogCook, Error, TEXT("CookWorkerCrash: %d failed to read from socket, we will shutdown the remote process. Assigned packages will be returned to the director."),
 			ProfileId);
+		bNeedCrashDiagnostics = true;
 		SendToState(EConnectStatus::WaitForDisconnect);
 		bTerminateImmediately = true;
 		return;
@@ -499,9 +555,10 @@ void FCookWorkerServer::HandleReceiveMessagesInternal()
 		{
 			UE::CompactBinaryTCP::FMarshalledMessage Message = ReceiveMessages.PopFrontValue();
 			UE_CLOG(ConnectStatus != EConnectStatus::PumpingCookComplete && ConnectStatus != EConnectStatus::WaitForDisconnect,
-				LogCook, Error, TEXT("CookWorkerServer %d remote process shut down unexpectedly. Assigned packages will be returned to the director."),
+				LogCook, Error, TEXT("CookWorkerCrash: %d remote process shut down unexpectedly. Assigned packages will be returned to the director."),
 				ProfileId);
 
+			bNeedCrashDiagnostics = true;
 			SendMessageInLock(FAbortWorkerMessage(FAbortWorkerMessage::AbortAcknowledge));
 			SendToState(EConnectStatus::WaitForDisconnect);
 			ReceiveMessages.Reset();
@@ -1025,6 +1082,11 @@ void FLogMessagesMessageHandler::ServerReceiveMessage(FMPCollectorServerMessageC
 
 	for (FReplicatedLogData& LogData : Messages)
 	{
+		if (LogData.Category == LogCookName && LogData.Message.Contains(TEXT("CookWorkerHeatbeat:")))
+		{
+			// Do not spam heartbeat messages into the CookDirector log
+			continue;
+		}
 		GLog->CategorizedLogf(LogData.Category, LogData.Verbosity, TEXT("[CookWorker %d]: %s"),
 			Context.GetProfileId(), *LogData.Message);
 	}
diff --git a/Engine/Source/Editor/UnrealEd/Private/Cooker/CookWorkerServer.h b/Engine/Source/Editor/UnrealEd/Private/Cooker/CookWorkerServer.h
--- a/Engine/Source/Editor/UnrealEd/Private/Cooker/CookWorkerServer.h
+++ b/Engine/Source/Editor/UnrealEd/Private/Cooker/CookWorkerServer.h
@@ -156,6 +156,8 @@ private:
 	void SendToState(EConnectStatus TargetStatus);
 	/** Close the connection and connection resources to the remote process. Does not kill the process. */
 	void DetachFromRemoteProcess();
+	/** Report from the log of the crashed CookWorker. */
+	void SendCrashDiagnostics();
 	/** Kill the Client process (non-graceful termination), and close the connection resources. */
 	void ShutdownRemoteProcess();
 	/** The main implementation of HandleReceiveMessages, only callable from inside the lock. */
@@ -190,6 +192,7 @@ private:
 	FWorkerId WorkerId = FWorkerId::Invalid();
 	EConnectStatus ConnectStatus = EConnectStatus::Uninitialized;
 	bool bTerminateImmediately = false;
+	bool bNeedCrashDiagnostics = false;
 };
 
 /** Information about a PackageData the director sends to cookworkers. */
