Fix it so AdditionalServerGameOptions applies to externally launched dedicated servers as well as listen servers

diff --git a/Engine/Source/Editor/UnrealEd/Private/PlayLevelNewProcess.cpp b/Engine/Source/Editor/UnrealEd/Private/PlayLevelNewProcess.cpp
--- a/Engine/Source/Editor/UnrealEd/Private/PlayLevelNewProcess.cpp
+++ b/Engine/Source/Editor/UnrealEd/Private/PlayLevelNewProcess.cpp
@@ -83,8 +83,11 @@ void UEditorEngine::LaunchNewProcess(const FRequestPlaySessionParams& InParams,
 	else if (NetMode == EPlayNetMode::PIE_ListenServer)
 	{
 		UnrealURLParams += TEXT("?Listen");
+	}
 
-		// Add any additional url parameters the user might have specified.
+	if (NetMode == EPlayNetMode::PIE_ListenServer)
+	{
+		// Add any additional url parameters the user might have specified, for both listen and dedicated servers
 		FString AdditionalServerGameOptions;
 		InParams.EditorPlaySettings->GetAdditionalServerGameOptions(AdditionalServerGameOptions);
 
