Manually reset loaders for packages to be saved.
WorldPartitionBuilder with WorldPartitionResaveActorsBuilder on internal map from 14:30 to 11:20.
This will benefit other builders such as HLODs and minimap generation, etc.

diff --git a/Engine/Source/Editor/UnrealEd/Private/WorldPartition/WorldPartitionBuilder.cpp b/Engine/Source/Editor/UnrealEd/Private/WorldPartition/WorldPartitionBuilder.cpp
--- a/Engine/Source/Editor/UnrealEd/Private/WorldPartition/WorldPartitionBuilder.cpp
+++ b/Engine/Source/Editor/UnrealEd/Private/WorldPartition/WorldPartitionBuilder.cpp
@@ -362,6 +362,8 @@ bool UWorldPartitionBuilder::SavePackages(const TArray<UPackage*>& Packages, FPa
 		}
 	}
 
+	ResetLoaders(TArray<UObject*>(Packages));
+
 	for (int PackageIndex = 0; PackageIndex < Packages.Num(); ++PackageIndex)
 	{
 		const FString& PackageFilename = PackageFilenames[PackageIndex];
