Remove the dedicated option to disable real-time viewports when the editor in the background. Use the existing bThrottleCPUWhenNotForeground option instead, but ignore all the extra special conditions that are handled by UEditorEngine::ShouldThrottleCPUUsage()

diff --git a/Engine/Source/Editor/UnrealEd/Classes/Editor/EditorPerformanceSettings.h b/Engine/Source/Editor/UnrealEd/Classes/Editor/EditorPerformanceSettings.h
--- a/Engine/Source/Editor/UnrealEd/Classes/Editor/EditorPerformanceSettings.h
+++ b/Engine/Source/Editor/UnrealEd/Classes/Editor/EditorPerformanceSettings.h
@@ -41,10 +41,6 @@ class UEditorPerformanceSettings : public UDeveloperSettings
 	UPROPERTY(EditAnywhere, config, Category=EditorPerformance, meta=(DisplayName="Use Less CPU when in Background") )
 	uint32 bThrottleCPUWhenNotForeground:1;
 
-	/** Disables realtime viewports when the editor is in the background to reduce GPU performance impact on other applications */
-	UPROPERTY(EditAnywhere, config, Category = EditorPerformance, meta = (DisplayName = "Disable realtime viewports when in Background"))
-	uint32 bDisableRealtimeViewportsWhenNotForeground : 1;
-
 	/** Disables realtime viewports by default when connected via a remote session */
 	UPROPERTY(EditAnywhere, config, Category = EditorPerformance, meta = (DisplayName = "Disable realtime viewports by default in Remote Sessions"))
 	uint32 bDisableRealtimeViewportsInRemoteSessions : 1;
diff --git a/Engine/Source/Editor/UnrealEd/Private/EditorEngine.cpp b/Engine/Source/Editor/UnrealEd/Private/EditorEngine.cpp
--- a/Engine/Source/Editor/UnrealEd/Private/EditorEngine.cpp
+++ b/Engine/Source/Editor/UnrealEd/Private/EditorEngine.cpp
@@ -1623,7 +1623,8 @@ void UEditorEngine::Tick( float DeltaSeconds, bool bIdleMode )
 	// Conditionally disable all viewport rendering when the editor is in the background.
 	// This aims to improve GPU performance of other applications when the editor is not actively used.
 	{
-		const bool bShouldDisableRendering = !FApp::HasFocus() && GetDefault<UEditorPerformanceSettings>()->bDisableRealtimeViewportsWhenNotForeground;
+		const UEditorPerformanceSettings* PerformanceSettings = GetDefault<UEditorPerformanceSettings>();
+		const bool bShouldDisableRendering = !FApp::HasFocus() && PerformanceSettings->bThrottleCPUWhenNotForeground;
 		const FText SystemDisplayName = LOCTEXT("RealtimeOverrideMessage_BackgroundProcess", "Background Process");
 		for (FEditorViewportClient* const ViewportClient : AllViewportClients)
 		{
diff --git a/Engine/Source/Editor/UnrealEd/Private/EditorPerformanceSettings.cpp b/Engine/Source/Editor/UnrealEd/Private/EditorPerformanceSettings.cpp
--- a/Engine/Source/Editor/UnrealEd/Private/EditorPerformanceSettings.cpp
+++ b/Engine/Source/Editor/UnrealEd/Private/EditorPerformanceSettings.cpp
@@ -17,7 +17,6 @@ UEditorPerformanceSettings::UEditorPerformanceSettings(const FObjectInitializer&
 	: Super(ObjectInitializer)
 	, bShowFrameRateAndMemory(false)
 	, bThrottleCPUWhenNotForeground(true)
-	, bDisableRealtimeViewportsWhenNotForeground(false)
 	, bDisableRealtimeViewportsInRemoteSessions(true)
 	, bMonitorEditorPerformance(false)
 	, bEnableScalabilityWarningIndicator(true)
