Fix non-unity

diff --git a/Engine/Source/Editor/UnrealEd/Private/WorldPartition/WorldPartitionBuilder.cpp b/Engine/Source/Editor/UnrealEd/Private/WorldPartition/WorldPartitionBuilder.cpp
--- a/Engine/Source/Editor/UnrealEd/Private/WorldPartition/WorldPartitionBuilder.cpp
+++ b/Engine/Source/Editor/UnrealEd/Private/WorldPartition/WorldPartitionBuilder.cpp
@@ -6,6 +6,7 @@
 #include "Editor.h"
 #include "EditorWorldUtils.h"
 #include "EngineModule.h"
+#include "UObject/Linker.h"
 #include "HAL/PlatformFileManager.h"
 #include "StaticMeshCompiler.h"
 #include "Engine/World.h"
