Cooker: Suppress per-package Skipped log message for spammy cases like base-game packages from DLC cooks.
#rnx

diff --git a/Engine/Source/Editor/UnrealEd/Private/CookOnTheFlyServer.cpp b/Engine/Source/Editor/UnrealEd/Private/CookOnTheFlyServer.cpp
--- a/Engine/Source/Editor/UnrealEd/Private/CookOnTheFlyServer.cpp
+++ b/Engine/Source/Editor/UnrealEd/Private/CookOnTheFlyServer.cpp
@@ -2488,19 +2488,24 @@ void UCookOnTheFlyServer::DemoteToIdle(UE::Cook::FPackageData& PackageData, UE::
 	if (PackageData.IsInProgress())
 	{
 		WorkerRequests->ReportDemoteToIdle(PackageData, Reason);
+		// If per-package display is on, write a log statement explaining that the package was reachable but skipped.
 		if (GCookProgressDisplay & ((int32)ECookProgressDisplayMode::Instigators | (int32)ECookProgressDisplayMode::PackageNames))
 		{
-			WriteToString<256> PackageNameStr(PackageData.GetPackageName());
-
-			// ExternalActors: Do not send a message for every NeverCook external Actor package; too much spam
-			if (!(Reason == ESuppressCookReason::NeverCook &&
-				PackageNameStr.ToView().Contains(ULevel::GetExternalActorsFolderName())))
+			// Suppress the message in cases that cause large spam like NotInCurrentPlugin for DLC cooks.
+			if (Reason != ESuppressCookReason::NotInCurrentPlugin)
 			{
-				UE_CLOG((GCookProgressDisplay & (int32)ECookProgressDisplayMode::Instigators), LogCook, Display,
-					TEXT("Cooking %s, Instigator: { %s } -> Skipped %s"), *PackageNameStr,
-					*(PackageData.GetInstigator().ToString()), LexToString(Reason));
-				UE_CLOG(GCookProgressDisplay & (int32)ECookProgressDisplayMode::PackageNames, LogCook, Display,
-					TEXT("Cooking %s -> Skipped %s"), *PackageNameStr, LexToString(Reason));
+				WriteToString<256> PackageNameStr(PackageData.GetPackageName());
+
+				// ExternalActors: Do not send a message for every NeverCook external Actor package; too much spam
+				if (!(Reason == ESuppressCookReason::NeverCook &&
+					PackageNameStr.ToView().Contains(ULevel::GetExternalActorsFolderName())))
+				{
+					UE_CLOG((GCookProgressDisplay & (int32)ECookProgressDisplayMode::Instigators), LogCook, Display,
+						TEXT("Cooking %s, Instigator: { %s } -> Skipped %s"), *PackageNameStr,
+						*(PackageData.GetInstigator().ToString()), LexToString(Reason));
+					UE_CLOG(GCookProgressDisplay & (int32)ECookProgressDisplayMode::PackageNames, LogCook, Display,
+						TEXT("Cooking %s -> Skipped %s"), *PackageNameStr, LexToString(Reason));
+				}
 			}
 		}
 	}
