[BugFix] Modified actors are not diff'ing correctly with one file per actor

diff --git a/Engine/Plugins/Editor/ChangelistReview/Source/Private/SSourceControlReviewEntry.cpp b/Engine/Plugins/Editor/ChangelistReview/Source/Private/SSourceControlReviewEntry.cpp
--- a/Engine/Plugins/Editor/ChangelistReview/Source/Private/SSourceControlReviewEntry.cpp
+++ b/Engine/Plugins/Editor/ChangelistReview/Source/Private/SSourceControlReviewEntry.cpp
@@ -5,6 +5,7 @@
 #include "AssetDefinition.h"
 #include "AssetToolsModule.h"
 #include "ClassIconFinder.h"
+#include "DiffUtils.h"
 #include "Engine/Blueprint.h"
 #include "SourceControlHelpers.h"
 #include "AssetRegistry/AssetRegistryModule.h"
@@ -223,7 +224,9 @@ void SSourceControlReviewEntry::TryBindUAssetDiff()
 	UObject* ReviewAsset = nullptr;	
 	if (IsAssetPath(ChangelistFileData.ReviewFileTempPath))
 	{
-		if (UPackage* ReviewFilePkg = LoadPackage(nullptr, FPackagePath::FromLocalPath(ChangelistFileData.ReviewFileTempPath), LOAD_ForDiff | LOAD_DisableCompileOnLoad | LOAD_DisableEngineVersionChecks))
+		const FPackagePath ReviewFileTempPath = FPackagePath::FromLocalPath(ChangelistFileData.ReviewFileTempPath);
+		const FPackagePath AssetPath = FPackagePath::FromLocalPath(ChangelistFileData.AssetFilePath);
+		if (UPackage* ReviewFilePkg = DiffUtils::LoadPackageForDiff(ReviewFileTempPath, AssetPath))
 		{
 			ReviewAsset = FindObject<UObject>(ReviewFilePkg, *ChangelistFileData.AssetName);
 			if(ReviewAsset && ReviewAsset->IsA<UObjectRedirector>())
@@ -236,7 +239,9 @@ void SSourceControlReviewEntry::TryBindUAssetDiff()
 	UObject* PreviousAsset = nullptr;
 	if (IsAssetPath(ChangelistFileData.PreviousFileTempPath))
 	{
-		if (UPackage* PreviousFilePkg = LoadPackage(nullptr, FPackagePath::FromLocalPath(ChangelistFileData.PreviousFileTempPath), LOAD_ForDiff | LOAD_DisableCompileOnLoad | LOAD_DisableEngineVersionChecks))
+		const FPackagePath PreviousFileTempPath = FPackagePath::FromLocalPath(ChangelistFileData.PreviousFileTempPath);
+		const FPackagePath AssetPath = FPackagePath::FromLocalPath(ChangelistFileData.AssetFilePath);
+		if (UPackage* PreviousFilePkg = DiffUtils::LoadPackageForDiff(PreviousFileTempPath, AssetPath))
 		{
 			if (ChangelistFileData.PreviousAssetName.IsEmpty())
 			{
diff --git a/Engine/Source/Developer/AssetTools/Private/AssetTools.cpp b/Engine/Source/Developer/AssetTools/Private/AssetTools.cpp
--- a/Engine/Source/Developer/AssetTools/Private/AssetTools.cpp
+++ b/Engine/Source/Developer/AssetTools/Private/AssetTools.cpp
@@ -102,6 +102,7 @@
 
 #include "AssetDefinition.h"
 #include "AssetDefinitionRegistry.h"
+#include "DiffUtils.h"
 #include "VirtualTexturingEditorModule.h"
 #include "Algo/AnyOf.h"
 #include "Engine/UserDefinedStruct.h"
@@ -2716,33 +2717,28 @@ void UAssetToolsImpl::DiffAgainstDepot( UObject* InObject, const FString& InPack
 				// Get the head revision of this package from source control
 				FString AbsoluteFileName = FPaths::ConvertRelativePathToFull(RelativeFileName);
 				FString TempFileName;
-				if(Revision->Get(TempFileName))
+				if(UPackage* TempPackage = DiffUtils::LoadPackageForDiff(Revision))
 				{
-					// Try and load that package
-					UPackage* TempPackage = LoadPackage(nullptr, *TempFileName, LOAD_ForDiff|LOAD_DisableCompileOnLoad);
-					if(TempPackage != nullptr)
-					{
-						// Grab the old asset from that old package
-						UObject* OldObject = FindObject<UObject>(TempPackage, *InPackageName);
+					// Grab the old asset from that old package
+					UObject* OldObject = FindObject<UObject>(TempPackage, *InPackageName);
 
-						// Recovery for package names that don't match
-						if (OldObject == nullptr)
-						{
-							OldObject = TempPackage->FindAssetInPackage();
-						}
+					// Recovery for package names that don't match
+					if (OldObject == nullptr)
+					{
+						OldObject = TempPackage->FindAssetInPackage();
+					}
 
-						if(OldObject != nullptr)
-						{
-							/* Set the revision information*/
-							FRevisionInfo OldRevision;
-							OldRevision.Changelist = Revision->GetCheckInIdentifier();
-							OldRevision.Date = Revision->GetDate();
-							OldRevision.Revision = Revision->GetRevision();
-
-							FRevisionInfo NewRevision; 
-							NewRevision.Revision = TEXT("");
-							DiffAssets(OldObject, InObject, OldRevision, NewRevision);
-						}
+					if(OldObject != nullptr)
+					{
+						/* Set the revision information*/
+						FRevisionInfo OldRevision;
+						OldRevision.Changelist = Revision->GetCheckInIdentifier();
+						OldRevision.Date = Revision->GetDate();
+						OldRevision.Revision = Revision->GetRevision();
+
+						FRevisionInfo NewRevision; 
+						NewRevision.Revision = TEXT("");
+						DiffAssets(OldObject, InObject, OldRevision, NewRevision);
 					}
 				}
 			}
diff --git a/Engine/Source/Developer/Merge/Private/MergeUtils.cpp b/Engine/Source/Developer/Merge/Private/MergeUtils.cpp
--- a/Engine/Source/Developer/Merge/Private/MergeUtils.cpp
+++ b/Engine/Source/Developer/Merge/Private/MergeUtils.cpp
@@ -1,6 +1,8 @@
 // Copyright Epic Games, Inc. All Rights Reserved.
 
 #include "MergeUtils.h"
+
+#include "DiffUtils.h"
 #include "IAssetTypeActions.h"
 #include "UObject/UObjectHash.h"
 #include "UObject/Package.h"
@@ -112,8 +114,9 @@ UObject const* FMergeToolUtils::LoadAssetFromPackage(const FString& PackageFileN
 	const UObject* Object = nullptr;
 
 	// Try and load that package
-	UPackage* TempPackage = LoadPackage(NULL, *PackageFileName, LOAD_ForDiff | LOAD_DisableCompileOnLoad);
-	if (TempPackage != NULL)
+	const FPackagePath TempPackagePath = FPackagePath::FromLocalPath(PackageFileName);
+	// TODO: @jordan.hoffmann set InOriginalPackagePath parameter if this is used for Actors
+	if (UPackage* TempPackage = DiffUtils::LoadPackageForDiff(TempPackagePath, {}))
 	{
 		// Grab the old asset from that old package
 		UObject* FoundObject = FindObject<UObject>(TempPackage, *AssetName);
diff --git a/Engine/Source/Editor/Kismet/Private/DiffUtils.cpp b/Engine/Source/Editor/Kismet/Private/DiffUtils.cpp
deleted file mode 100644
diff --git a/Engine/Source/Editor/Kismet/Private/SBlueprintEditorToolbar.cpp b/Engine/Source/Editor/Kismet/Private/SBlueprintEditorToolbar.cpp
--- a/Engine/Source/Editor/Kismet/Private/SBlueprintEditorToolbar.cpp
+++ b/Engine/Source/Editor/Kismet/Private/SBlueprintEditorToolbar.cpp
@@ -7,6 +7,7 @@
 #include "BlueprintEditor.h"
 #include "BlueprintEditorCommands.h"
 #include "BlueprintEditorContext.h"
+#include "DiffUtils.h"
 #include "Containers/Array.h"
 #include "Containers/EnumAsByte.h"
 #include "Containers/UnrealString.h"
@@ -571,52 +572,45 @@ static void OnDiffRevisionPicked(FRevisionInfo const& RevisionInfo, TWeakObjectP
 				if (Revision->GetRevision() == RevisionInfo.Revision)
 				{
 					// Get the revision of this package from source control
-					FString PreviousTempPkgName;
-					if (Revision->Get(PreviousTempPkgName))
+					if (UPackage* PreviousTempPkg = DiffUtils::LoadPackageForDiff(Revision))
 					{
-						// Try and load that package
-						UPackage* PreviousTempPkg = LoadPackage(NULL, *PreviousTempPkgName, LOAD_ForDiff | LOAD_DisableCompileOnLoad);
+						UObject* PreviousAsset = nullptr;
 
-						if (PreviousTempPkg != NULL)
+						// If its a levelscript blueprint, find the previous levelscript blueprint in the map
+						if (bIsLevelScriptBlueprint)
 						{
-							UObject* PreviousAsset = NULL;
+							TArray<UObject*> ObjectsInOuter;
+							GetObjectsWithOuter(PreviousTempPkg, ObjectsInOuter);
 
-							// If its a levelscript blueprint, find the previous levelscript blueprint in the map
-							if (bIsLevelScriptBlueprint)
+							// Look for the level script blueprint for this package
+							for (int32 Index = 0; Index < ObjectsInOuter.Num(); Index++)
 							{
-								TArray<UObject*> ObjectsInOuter;
-								GetObjectsWithOuter(PreviousTempPkg, ObjectsInOuter);
-
-								// Look for the level script blueprint for this package
-								for (int32 Index = 0; Index < ObjectsInOuter.Num(); Index++)
+								UObject* Obj = ObjectsInOuter[Index];
+								if (ULevelScriptBlueprint* ObjAsBlueprint = Cast<ULevelScriptBlueprint>(Obj))
 								{
-									UObject* Obj = ObjectsInOuter[Index];
-									if (ULevelScriptBlueprint* ObjAsBlueprint = Cast<ULevelScriptBlueprint>(Obj))
-									{
-										PreviousAsset = ObjAsBlueprint;
-										break;
-									}
+									PreviousAsset = ObjAsBlueprint;
+									break;
 								}
 							}
-							// otherwise its a normal Blueprint
-							else
-							{
-								FString PreviousAssetName = FPaths::GetBaseFilename(Filename, true);
-								PreviousAsset = FindObject<UObject>(PreviousTempPkg, *PreviousAssetName);
-							}
-
-							if (PreviousAsset != NULL)
-							{
-								FAssetToolsModule& AssetToolsModule = FModuleManager::LoadModuleChecked<FAssetToolsModule>(TEXT("AssetTools"));
-								FRevisionInfo OldRevision = { Revision->GetRevision(), Revision->GetCheckInIdentifier(), Revision->GetDate() };
-								FRevisionInfo CurrentRevision = { TEXT(""), Revision->GetCheckInIdentifier(), Revision->GetDate() };
-								AssetToolsModule.Get().DiffAssets(PreviousAsset, BlueprintObj.Get(), OldRevision, CurrentRevision);
-							}
 						}
+						// otherwise its a normal Blueprint
 						else
 						{
-							FMessageDialog::Open(EAppMsgType::Ok, NSLOCTEXT("SourceControl.HistoryWindow", "UnableToLoadAssets", "Unable to load assets to diff. Content may no longer be supported?"));
+							FString PreviousAssetName = FPaths::GetBaseFilename(Filename, true);
+							PreviousAsset = FindObject<UObject>(PreviousTempPkg, *PreviousAssetName);
 						}
+
+						if (PreviousAsset != nullptr)
+						{
+							FAssetToolsModule& AssetToolsModule = FModuleManager::LoadModuleChecked<FAssetToolsModule>(TEXT("AssetTools"));
+							FRevisionInfo OldRevision = { Revision->GetRevision(), Revision->GetCheckInIdentifier(), Revision->GetDate() };
+							FRevisionInfo CurrentRevision = { TEXT(""), Revision->GetCheckInIdentifier(), Revision->GetDate() };
+							AssetToolsModule.Get().DiffAssets(PreviousAsset, BlueprintObj.Get(), OldRevision, CurrentRevision);
+						}
+					}
+					else
+					{
+						FMessageDialog::Open(EAppMsgType::Ok, NSLOCTEXT("SourceControl.HistoryWindow", "UnableToLoadAssets", "Unable to load assets to diff. Content may no longer be supported?"));
 					}
 					break;
 				}
diff --git a/Engine/Source/Editor/Kismet/Public/DiffUtils.h b/Engine/Source/Editor/Kismet/Public/DiffUtils.h
deleted file mode 100644
diff --git a/Engine/Source/Editor/SourceControlWindows/Private/SSourceControlHistory.cpp b/Engine/Source/Editor/SourceControlWindows/Private/SSourceControlHistory.cpp
--- a/Engine/Source/Editor/SourceControlWindows/Private/SSourceControlHistory.cpp
+++ b/Engine/Source/Editor/SourceControlWindows/Private/SSourceControlHistory.cpp
@@ -50,6 +50,7 @@
 #include "IAssetTools.h"
 #include "IAssetTypeActions.h"
 #include "AssetToolsModule.h"
+#include "DiffUtils.h"
 #include "ToolMenu.h"
 #include "ToolMenus.h"
 
@@ -210,10 +211,10 @@ static UObject* GetAssetRevisionObject(TSharedPtr<FHistoryTreeItem> HistoryTreeI
 				TSharedPtr<ISourceControlRevision, ESPMode::ThreadSafe> FileRevision = FileSourceControlState->FindHistoryRevision(RevisionListItem->Revision);
 
 				FString TempPackageName;
-				if (FileRevision.IsValid() && FileRevision->Get(TempPackageName)) // grab the path to a temporary package (where the revision item will be stored)
+				if (FileRevision.IsValid())
 				{
 					// try and load the temporary package
-					AssetPackage = LoadPackage(NULL, *TempPackageName, LOAD_ForDiff|LOAD_DisableCompileOnLoad);
+					AssetPackage = DiffUtils::LoadPackageForDiff(FileRevision);
 				}
 			} // if FileSourceControlState.IsValid()
 		}
@@ -229,6 +230,10 @@ static UObject* GetAssetRevisionObject(TSharedPtr<FHistoryTreeItem> HistoryTreeI
 		// grab the asset from the package - we assume asset name matches file name
 		FString AssetName = FPaths::GetBaseFilename(FileListItem->FileName);
 		AssetObject = FindObject<UObject>(AssetPackage, *AssetName);
+		if (AssetObject == nullptr)
+		{
+			AssetObject = AssetPackage->FindAssetInPackage();
+		}
 
 	} // if HistoryTreeItemIn.IsValid()
 	
@@ -1779,8 +1784,9 @@ bool FSourceControlWindows::DiffAgainstShelvedFile(const FSourceControlStateRef&
 		if (Revision->Get(TempFileName))
 		{
 			// Try and load that package
-			UPackage* TempPackage = LoadPackage(nullptr, *TempFileName, LOAD_ForDiff | LOAD_DisableCompileOnLoad);
-			if (TempPackage != nullptr)
+			const FPackagePath TempPackagePath = FPackagePath::FromLocalPath(TempFileName);
+			const FPackagePath OriginalPackagePath = FPackagePath::FromLocalPath(InFileState->GetFilename());
+			if(UPackage* TempPackage = DiffUtils::LoadPackageForDiff(TempPackagePath, OriginalPackagePath))
 			{
 				// Grab the shelved asset from that package
 				UObject* ShelvedObject = FindObject<UObject>(TempPackage, *SelectedAsset->GetName());
diff --git a/Engine/Source/Editor/UnrealEd/Private/DiffUtils.cpp b/Engine/Source/Editor/UnrealEd/Private/DiffUtils.cpp
new file mode 100644
--- /dev/null
+++ b/Engine/Source/Editor/UnrealEd/Private/DiffUtils.cpp
@@ -0,0 +1,1138 @@
+// Copyright Epic Games, Inc. All Rights Reserved.
+
+#include "DiffUtils.h"
+
+#include "Components/ActorComponent.h"
+#include "Containers/BitArray.h"
+#include "EditorCategoryUtils.h"
+#include "Engine/Blueprint.h"
+#include "IAssetTypeActions.h"
+#include "ISourceControlModule.h"
+#include "ISourceControlProvider.h"
+#include "ISourceControlRevision.h"
+#include "Internationalization/Internationalization.h"
+#include "Math/UnrealMathSSE.h"
+#include "Misc/Attribute.h"
+#include "Misc/CString.h"
+#include "Misc/DateTime.h"
+#include "Misc/Optional.h"
+#include "ObjectEditorUtils.h"
+#include "SlotBase.h"
+#include "Styling/AppStyle.h"
+#include "Styling/SlateColor.h"
+#include "Templates/ChooseClass.h"
+#include "Templates/SubclassOf.h"
+#include "Templates/UnrealTemplate.h"
+#include "Types/SlateConstants.h"
+#include "Types/SlateEnums.h"
+#include "UObject/Class.h"
+#include "UObject/Field.h"
+#include "UObject/Object.h"
+#include "UObject/ObjectMacros.h"
+#include "UObject/PropertyPortFlags.h"
+#include "Widgets/DeclarativeSyntaxSupport.h"
+#include "Widgets/Images/SImage.h"
+#include "Widgets/Text/STextBlock.h"
+#include "Widgets/Views/STableRow.h"
+#include "Widgets/Views/STableViewBase.h"
+#include "Widgets/Views/STreeView.h"
+
+class ITableRow;
+class SWidget;
+
+namespace UEDiffUtils_Private
+{
+	FProperty* Resolve( const UStruct* Class, FName PropertyName )
+	{
+		if(Class == nullptr )
+		{
+			return nullptr;
+		}
+
+		for (FProperty* Prop : TFieldRange<FProperty>(Class))
+		{
+			if( Prop->GetFName() == PropertyName )
+			{
+				return Prop;
+			}
+		}
+
+		return nullptr;
+	}
+
+	FPropertySoftPathSet GetPropertyNameSet(const UStruct* ForStruct)
+	{
+		return FPropertySoftPathSet(DiffUtils::GetVisiblePropertiesInOrderDeclared(ForStruct));
+	}
+}
+
+FResolvedProperty FPropertySoftPath::Resolve(const UObject* Object) const
+{
+	FResolvedProperty ResolvedProperty = Resolve(Object->GetClass(), Object);
+	if (!ResolvedProperty.Property && Object->HasAnyFlags(RF_ClassDefaultObject))
+	{
+		if (const UScriptStruct* SparseClassDataStruct = Object->GetClass()->GetSparseClassDataStruct())
+		{
+			if (const void* SparseClassData = Object->GetClass()->GetSparseClassData(EGetSparseClassDataMethod::ReturnIfNull))
+			{
+				ResolvedProperty = Resolve(SparseClassDataStruct, SparseClassData);
+			}
+		}
+	}
+	return ResolvedProperty;
+}
+
+
+
+int32 FPropertySoftPath::TryReadIndex(const TArray<FChainElement>& LocalPropertyChain, int32& OutIndex)
+{
+	if(OutIndex + 1 < LocalPropertyChain.Num())
+	{
+		FString AsString = LocalPropertyChain[OutIndex + 1].DisplayString;
+		if(AsString.IsNumeric())
+		{
+			++OutIndex;
+			return FCString::Atoi(*AsString);
+		}
+	}
+	return INDEX_NONE;
+};
+
+FResolvedProperty FPropertySoftPath::Resolve(const UStruct* Struct, const void* StructData) const
+{
+	// dig into the object, finding nested objects, etc:
+	const void* CurrentBlock = StructData;
+	const UStruct* NextClass = Struct;
+	const void* NextBlock = CurrentBlock;
+	const FProperty* Property = nullptr;
+
+	for (int32 i = 0; i < PropertyChain.Num(); ++i)
+	{
+		CurrentBlock = NextBlock;
+		const FProperty* NextProperty = UEDiffUtils_Private::Resolve(NextClass, PropertyChain[i].PropertyName);
+
+		// if an index was provided, resolve it
+		const int32 PropertyIndex = TryReadIndex(PropertyChain, i);
+		if (NextProperty && PropertyIndex != INDEX_NONE)
+		{
+			Property = NextProperty;
+			if (const FArrayProperty* ArrayProperty = CastField<FArrayProperty>(Property))
+			{
+				FScriptArrayHelper ArrayHelper(ArrayProperty, Property->ContainerPtrToValuePtr<UObject*>(CurrentBlock));
+				if (ArrayHelper.IsValidIndex(PropertyIndex))
+				{
+					NextProperty = ArrayProperty->Inner;
+					NextBlock = ArrayHelper.GetRawPtr(PropertyIndex);
+				}
+			}
+			else if( const FSetProperty* SetProperty = CastField<FSetProperty>(Property) )
+			{
+				checkf(false, TEXT("Set Indexing not supported yet"));
+				// TODO: @jordan.hoffmann: handle indexing for sets
+			}
+			else if( const FMapProperty* MapProperty = CastField<FMapProperty>(Property) )
+			{
+				checkf(false, TEXT("Map Indexing not supported yet"));
+				// TODO: @jordan.hoffmann: handle indexing for maps
+			}
+		}
+		
+		CurrentBlock = NextBlock;
+		if (NextProperty)
+		{
+			Property = NextProperty;
+			if (const FObjectProperty* ObjectProperty = CastField<FObjectProperty>(Property))
+			{
+				const UObject* NextObject = ObjectProperty->GetObjectPropertyValue(Property->ContainerPtrToValuePtr<UObject*>(CurrentBlock));
+				NextBlock = NextObject;
+				NextClass = NextObject ? NextObject->GetClass() : nullptr;
+			}
+			else if (const FStructProperty* StructProperty = CastField<FStructProperty>(Property))
+			{
+				NextBlock = StructProperty->ContainerPtrToValuePtr<void>(CurrentBlock);
+				NextClass = StructProperty->Struct;
+			}
+			else
+			{
+				break;
+			}
+		}
+		else
+		{
+			break;
+		}
+	}
+
+	return FResolvedProperty(CurrentBlock, Property);
+}
+
+FPropertyPath FPropertySoftPath::ResolvePath(const UObject* Object) const
+{
+	auto UpdateContainerAddress = [](const FProperty* Property, const void* Instance, const void*& OutContainerAddress, const UStruct*& OutContainerStruct)
+	{
+		if( ensure(Instance) )
+		{
+			if(const FObjectProperty* ObjectProperty = CastField<FObjectProperty>(Property))
+			{
+				const UObject* const* InstanceObject = reinterpret_cast<const UObject* const*>(Instance);
+				if( *InstanceObject)
+				{
+					OutContainerAddress = *InstanceObject;
+					OutContainerStruct = (*InstanceObject)->GetClass();
+				}
+			}
+			else if(const FStructProperty* StructProperty = CastField<FStructProperty>(Property))
+			{
+				OutContainerAddress = Instance;
+				OutContainerStruct = StructProperty->Struct;
+			}
+		}
+	};
+
+	const void* ContainerAddress = Object;
+	const UStruct* ContainerStruct = (Object ? Object->GetClass() : nullptr);
+
+	FPropertyPath Ret;
+	for( int32 I = 0; I < PropertyChain.Num(); ++I )
+	{
+		FName PropertyIdentifier = PropertyChain[I].PropertyName;
+		FProperty* ResolvedProperty = UEDiffUtils_Private::Resolve(ContainerStruct, PropertyIdentifier);
+
+		// If the property didn't exist inside the container, return an invalid property
+		if (!ResolvedProperty)
+		{
+			return FPropertyPath();
+		}
+		
+		FPropertyInfo Info(ResolvedProperty, INDEX_NONE);
+		Ret.AddProperty(Info);
+
+		int32 PropertyIndex = TryReadIndex(PropertyChain, I);
+		
+		
+		// calculate offset so we can continue resolving object properties/structproperties:
+		if( const FArrayProperty* ArrayProperty = CastField<FArrayProperty>(ResolvedProperty) )
+		{
+			FScriptArrayHelper ArrayHelper(ArrayProperty, ArrayProperty->ContainerPtrToValuePtr<const void*>( ContainerAddress ));
+			if (ArrayHelper.IsValidIndex(PropertyIndex))
+			{
+				UpdateContainerAddress( ArrayProperty->Inner, ArrayHelper.GetRawPtr(PropertyIndex), ContainerAddress, ContainerStruct );
+
+				FPropertyInfo ArrayInfo(ArrayProperty->Inner, PropertyIndex);
+				Ret.AddProperty(ArrayInfo);
+			}
+		}
+		else if( const FSetProperty* SetProperty = CastField<FSetProperty>(ResolvedProperty) )
+		{
+			if(PropertyIndex != INDEX_NONE)
+			{
+				FScriptSetHelper SetHelper(SetProperty, SetProperty->ContainerPtrToValuePtr<const void*>( ContainerAddress ));
+
+				// Figure out the real index in this instance of the set (sets have gaps in them):
+				int32 RealIndex = -1;
+				for( int32 J = 0; PropertyIndex >= 0; ++J)
+				{
+					++RealIndex;
+					if(SetHelper.IsValidIndex(J))
+					{
+						--PropertyIndex;
+					}
+				}
+
+				UpdateContainerAddress( SetProperty->ElementProp, SetHelper.GetElementPtr(RealIndex), ContainerAddress, ContainerStruct );
+
+				FPropertyInfo SetInfo(SetProperty->ElementProp, RealIndex);
+				Ret.AddProperty(SetInfo);
+			}
+		}
+		else if( const FMapProperty* MapProperty = CastField<FMapProperty>(ResolvedProperty) )
+		{
+			if(PropertyIndex != INDEX_NONE)
+			{
+				FScriptMapHelper MapHelper(MapProperty, MapProperty->ContainerPtrToValuePtr<const void*>( ContainerAddress ));
+				
+				// Figure out the real index in this instance of the map (map have gaps in them):
+				int32 RealIndex = -1;
+				for( int32 J = 0; PropertyIndex >= 0; ++J)
+				{
+					++RealIndex;
+					if(MapHelper.IsValidIndex(J))
+					{
+						--PropertyIndex;
+					}
+				}
+
+				// we have an index, but are we looking into a key or value? Peek ahead to find out:
+				if(ensure(I + 1 < PropertyChain.Num()))
+				{
+					if(PropertyChain[I+1].PropertyName == MapProperty->KeyProp->GetFName())
+					{
+						++I;
+
+						UpdateContainerAddress( MapProperty->KeyProp, MapHelper.GetKeyPtr(RealIndex), ContainerAddress, ContainerStruct );
+
+						FPropertyInfo MakKeyInfo(MapProperty->KeyProp, RealIndex);
+						Ret.AddProperty(MakKeyInfo);
+					}
+					else if(ensure( PropertyChain[I+1].PropertyName == MapProperty->ValueProp->GetFName() ))
+					{	
+						++I;
+
+						UpdateContainerAddress( MapProperty->ValueProp, MapHelper.GetValuePtr(RealIndex), ContainerAddress, ContainerStruct );
+						
+						FPropertyInfo MapValueInfo(MapProperty->ValueProp, RealIndex);
+						Ret.AddProperty(MapValueInfo);
+					}
+				}
+			}
+		}
+		else if (const FObjectProperty* ObjectProperty = CastField<FObjectProperty>(ResolvedProperty))
+		{
+			UpdateContainerAddress( ObjectProperty, ObjectProperty->ContainerPtrToValuePtr<const void*>( ContainerAddress, FMath::Max(PropertyIndex, 0) ), ContainerAddress, ContainerStruct );
+			
+			// handle static arrays:
+			if(PropertyIndex != INDEX_NONE )
+			{
+				FPropertyInfo ObjectInfo(ResolvedProperty, PropertyIndex);
+				Ret.AddProperty(ObjectInfo);
+			}
+		}
+		else if( const FStructProperty* StructProperty = CastField<FStructProperty>(ResolvedProperty) )
+		{
+			UpdateContainerAddress( StructProperty, StructProperty->ContainerPtrToValuePtr<const void*>( ContainerAddress, FMath::Max(PropertyIndex, 0) ), ContainerAddress, ContainerStruct );
+			
+			// handle static arrays:
+			if(PropertyIndex != INDEX_NONE )
+			{
+				FPropertyInfo StructInfo(ResolvedProperty, PropertyIndex);
+				Ret.AddProperty(StructInfo);
+			}
+		}
+		else
+		{
+			// handle static arrays:
+			if(PropertyIndex != INDEX_NONE )
+			{
+				FPropertyInfo StaticArrayInfo(ResolvedProperty, PropertyIndex);
+				Ret.AddProperty(StaticArrayInfo);
+			}
+		}
+	}
+	return Ret;
+}
+
+FString FPropertySoftPath::ToDisplayName() const
+{
+	FString Ret;
+	for( FChainElement Element : PropertyChain )
+	{
+		FString PropertyAsString = Element.DisplayString;
+		if(Ret.IsEmpty())
+		{
+			Ret.Append(PropertyAsString);
+		}
+		else if( PropertyAsString.IsNumeric())
+		{
+			Ret.AppendChar('[');
+			Ret.Append(PropertyAsString);
+			Ret.AppendChar(']');
+		}
+		else
+		{
+			Ret.AppendChar(' ');
+			Ret.Append(PropertyAsString);
+		}
+	}
+	return Ret;
+}
+
+const UObject* DiffUtils::GetCDO(const UBlueprint* ForBlueprint)
+{
+	if (!ForBlueprint
+		|| !ForBlueprint->GeneratedClass)
+	{
+		return NULL;
+	}
+
+	return ForBlueprint->GeneratedClass->ClassDefaultObject;
+}
+
+void DiffUtils::CompareUnrelatedStructs(const UStruct* StructA, const void* A, const UStruct* StructB, const void* B, TArray<FSingleObjectDiffEntry>& OutDifferingProperties)
+{
+	FPropertySoftPathSet PropertiesInA = UEDiffUtils_Private::GetPropertyNameSet(StructA);
+	FPropertySoftPathSet PropertiesInB = UEDiffUtils_Private::GetPropertyNameSet(StructB);
+
+	// any properties in A that aren't in B are differing:
+	auto AddedToA = PropertiesInA.Difference(PropertiesInB).Array();
+	for (const auto& Entry : AddedToA)
+	{
+		OutDifferingProperties.Push(FSingleObjectDiffEntry(Entry, EPropertyDiffType::PropertyAddedToA));
+	}
+
+	// and the converse:
+	auto AddedToB = PropertiesInB.Difference(PropertiesInA).Array();
+	for (const auto& Entry : AddedToB)
+	{
+		OutDifferingProperties.Push(FSingleObjectDiffEntry(Entry, EPropertyDiffType::PropertyAddedToB));
+	}
+
+	// for properties in common, dig out the uproperties and determine if they're identical:
+	if (A && B)
+	{
+		FPropertySoftPathSet Common = PropertiesInA.Intersect(PropertiesInB);
+		for (const auto& PropertyName : Common)
+		{
+			FResolvedProperty AProp = PropertyName.Resolve(StructA, A);
+			FResolvedProperty BProp = PropertyName.Resolve(StructB, B);
+
+			check(AProp != FResolvedProperty() && BProp != FResolvedProperty());
+			TArray<FPropertySoftPath> DifferingSubProperties;
+			if (!DiffUtils::Identical(AProp, BProp, PropertyName, DifferingSubProperties))
+			{
+				for (int DifferingIndex = 0; DifferingIndex < DifferingSubProperties.Num(); DifferingIndex++)
+				{
+					OutDifferingProperties.Push(FSingleObjectDiffEntry(DifferingSubProperties[DifferingIndex], EPropertyDiffType::PropertyValueChanged));
+				}
+			}
+		}
+	}
+}
+
+void DiffUtils::CompareUnrelatedObjects(const UObject* A, const UObject* B, TArray<FSingleObjectDiffEntry>& OutDifferingProperties)
+{
+	if (A && B)
+	{
+		return CompareUnrelatedStructs(A->GetClass(), A, B->GetClass(), B, OutDifferingProperties);
+	}
+}
+
+void DiffUtils::CompareUnrelatedSCS(const UBlueprint* Old, const TArray< FSCSResolvedIdentifier >& OldHierarchy, const UBlueprint* New, const TArray< FSCSResolvedIdentifier >& NewHierarchy, FSCSDiffRoot& OutDifferingEntries )
+{
+	const auto FindEntry = [](TArray< FSCSResolvedIdentifier > const& InArray, const FSCSIdentifier* Value) -> const FSCSResolvedIdentifier*
+	{
+		for (const auto& Node : InArray)
+		{
+			if (Node.Identifier.Name == Value->Name)
+			{
+				return &Node;
+			}
+		}
+		return nullptr;
+	};
+
+	for (const auto& OldNode : OldHierarchy)
+	{
+		const FSCSResolvedIdentifier* NewEntry = FindEntry(NewHierarchy, &OldNode.Identifier);
+
+		if (NewEntry != nullptr)
+		{
+			bool bShouldDiffProperties = true;
+
+			// did it change class?
+			const bool bObjectTypesDiffer = OldNode.Object != nullptr && NewEntry->Object != nullptr && OldNode.Object->GetClass() != NewEntry->Object->GetClass();
+			if (bObjectTypesDiffer)
+			{
+				FSCSDiffEntry Diff = { OldNode.Identifier, ETreeDiffType::NODE_TYPE_CHANGED, FSingleObjectDiffEntry() };
+				OutDifferingEntries.Entries.Push(Diff);
+
+				// Only diff properties if we're still within the same class inheritance hierarchy.
+				bShouldDiffProperties = OldNode.Object->GetClass()->IsChildOf(NewEntry->Object->GetClass()) || NewEntry->Object->GetClass()->IsChildOf(OldNode.Object->GetClass());
+			}
+
+			// did a property change?
+			if(bShouldDiffProperties)
+			{
+				TArray<FSingleObjectDiffEntry> DifferingProperties;
+				DiffUtils::CompareUnrelatedObjects(OldNode.Object, NewEntry->Object, DifferingProperties);
+				for (const auto& Property : DifferingProperties)
+				{
+					// Only include property value change entries if the object types differ.
+					if (!bObjectTypesDiffer || Property.DiffType == EPropertyDiffType::PropertyValueChanged)
+					{
+						FSCSDiffEntry Diff = { OldNode.Identifier, ETreeDiffType::NODE_PROPERTY_CHANGED, Property };
+						OutDifferingEntries.Entries.Push(Diff);
+					}
+				}
+			}
+
+			// did it move?
+			if( NewEntry->Identifier.TreeLocation != OldNode.Identifier.TreeLocation )
+			{
+				FSCSDiffEntry Diff = { OldNode.Identifier, ETreeDiffType::NODE_MOVED, FSingleObjectDiffEntry() };
+				OutDifferingEntries.Entries.Push(Diff);
+			}
+
+			// no change! Do nothing.
+		}
+		else
+		{
+			// not found in the new data, must have been deleted:
+			FSCSDiffEntry Entry = { OldNode.Identifier, ETreeDiffType::NODE_REMOVED, FSingleObjectDiffEntry() };
+			OutDifferingEntries.Entries.Push( Entry );
+		}
+	}
+
+	for (const auto& NewNode : NewHierarchy)
+	{
+		const FSCSResolvedIdentifier* OldEntry = FindEntry(OldHierarchy, &NewNode.Identifier);
+
+		if (OldEntry == nullptr)
+		{
+			FSCSDiffEntry Entry = { NewNode.Identifier, ETreeDiffType::NODE_ADDED, FSingleObjectDiffEntry() };
+			OutDifferingEntries.Entries.Push( Entry );
+		}
+	}
+}
+
+static void AdvanceSetIterator( FScriptSetHelper& SetHelper, int32& Index)
+{
+	do
+	{
+		++Index;
+	}
+	while(Index < SetHelper.GetMaxIndex() && !SetHelper.IsValidIndex(Index));
+}
+
+static void AdvanceMapIterator( FScriptMapHelper& MapHelper, int32& Index)
+{
+	do
+	{
+		++Index;
+	}
+	while(Index < MapHelper.GetMaxIndex() && !MapHelper.IsValidIndex(Index));
+}
+
+static void IdenticalHelper(const FProperty* AProperty, const FProperty* BProperty, const void* AValue, const void* BValue, const FPropertySoftPath& RootPath, TArray<FPropertySoftPath>& DifferingSubProperties, bool bStaticArrayHandled = false)
+{
+	if(AProperty == nullptr || BProperty == nullptr || AProperty->ArrayDim != BProperty->ArrayDim || AProperty->GetClass() != BProperty->GetClass())
+	{
+		DifferingSubProperties.Push(RootPath);
+		return;
+	}
+
+	if(!bStaticArrayHandled && AProperty->ArrayDim != 1)
+	{
+		// Identical does not handle static array case automatically and we have to do the offset calculation ourself because 
+		// our container (e.g. the struct or class or dynamic array) has already done the initial offset calculation:
+		for( int32 I = 0; I < AProperty->ArrayDim; ++I )
+		{
+			int32 Offset = AProperty->ElementSize * I;
+			const void* CurAValue = reinterpret_cast<const void*>(reinterpret_cast<const uint8*>(AValue) + Offset);
+			const void* CurBValue = reinterpret_cast<const void*>(reinterpret_cast<const uint8*>(BValue) + Offset);
+
+			IdenticalHelper(AProperty, BProperty, CurAValue, CurBValue, FPropertySoftPath(RootPath, I), DifferingSubProperties, true);
+		}
+
+		return;
+	}
+	
+	const FStructProperty* APropAsStruct = CastField<FStructProperty>(AProperty);
+	const FArrayProperty* APropAsArray = CastField<FArrayProperty>(AProperty);
+	const FSetProperty* APropAsSet = CastField<FSetProperty>(AProperty);
+	const FMapProperty* APropAsMap = CastField<FMapProperty>(AProperty);
+	const FObjectProperty* APropAsObject = CastField<FObjectProperty>(AProperty);
+	if (APropAsStruct != nullptr)
+	{
+		const FStructProperty* BPropAsStruct = CastFieldChecked<FStructProperty>(const_cast<FProperty*>(BProperty));
+		if (BPropAsStruct->Struct == APropAsStruct->Struct)
+		{
+			if (APropAsStruct->Struct->StructFlags & STRUCT_IdenticalNative)
+			{
+				// If the struct uses CPP identical tests then we need to honor that
+				if (!AProperty->Identical(AValue, BValue, PPF_DeepComparison))
+				{
+					DifferingSubProperties.Push(RootPath);
+				}
+			}
+			else
+			{
+				// Compare sub-properties to detect more granular changes
+				for (TFieldIterator<FProperty> PropertyIt(APropAsStruct->Struct); PropertyIt; ++PropertyIt)
+				{
+					const FProperty* StructProp = *PropertyIt;
+					IdenticalHelper(StructProp, StructProp, StructProp->ContainerPtrToValuePtr<void>(AValue, 0), StructProp->ContainerPtrToValuePtr<void>(BValue, 0), FPropertySoftPath(RootPath, StructProp), DifferingSubProperties);
+				}
+			}
+		}
+		else
+		{
+			DifferingSubProperties.Push(RootPath);
+		}
+	}
+	else if (APropAsArray != nullptr)
+	{
+		const FArrayProperty* BPropAsArray = CastFieldChecked<const FArrayProperty>(BProperty);
+		if(BPropAsArray->Inner->GetClass() == APropAsArray->Inner->GetClass())
+		{
+			FScriptArrayHelper ArrayHelperA(APropAsArray, AValue);
+			FScriptArrayHelper ArrayHelperB(BPropAsArray, BValue);
+		
+			// note any differences in contained types:
+			for (int32 ArrayIndex = 0; ArrayIndex < ArrayHelperA.Num() && ArrayIndex < ArrayHelperB.Num(); ArrayIndex++)
+			{
+				IdenticalHelper(APropAsArray->Inner, APropAsArray->Inner, ArrayHelperA.GetRawPtr(ArrayIndex), ArrayHelperB.GetRawPtr(ArrayIndex), FPropertySoftPath(RootPath, ArrayIndex), DifferingSubProperties);
+			}
+
+			// note any size difference:
+			if (ArrayHelperA.Num() != ArrayHelperB.Num())
+			{
+				DifferingSubProperties.Push(RootPath);
+			}
+		}
+		else
+		{
+			DifferingSubProperties.Push(RootPath);
+		}
+	}
+	else if(APropAsSet != nullptr)
+	{
+		const FSetProperty* BPropAsSet = CastFieldChecked<const FSetProperty>(BProperty);
+		if(BPropAsSet->ElementProp->GetClass() == APropAsSet->ElementProp->GetClass())
+		{
+			FScriptSetHelper SetHelperA(APropAsSet, AValue);
+			FScriptSetHelper SetHelperB(BPropAsSet, BValue);
+
+			if (SetHelperA.Num() != SetHelperB.Num())
+			{
+				// API not robust enough to indicate changes made to # of set elements, would
+				// need to return something more detailed than DifferingSubProperties array:
+				DifferingSubProperties.Push(RootPath);
+			}
+
+			// note any differences in contained elements:
+			const int32 SetSizeA = SetHelperA.Num();
+			const int32 SetSizeB = SetHelperB.Num();
+			
+			int32 SetIndexA = -1;
+			int32 SetIndexB = -1;
+
+			AdvanceSetIterator(SetHelperA, SetIndexA);
+			AdvanceSetIterator(SetHelperB, SetIndexB);
+
+			for (int32 VirtualIndex = 0; VirtualIndex < SetSizeA && VirtualIndex < SetSizeB; ++VirtualIndex)
+			{
+				IdenticalHelper(APropAsSet->ElementProp, APropAsSet->ElementProp, SetHelperA.GetElementPtr(SetIndexA), SetHelperB.GetElementPtr(SetIndexB), FPropertySoftPath(RootPath, VirtualIndex), DifferingSubProperties);
+
+				// advance iterators in step:
+				AdvanceSetIterator(SetHelperA, SetIndexA);
+				AdvanceSetIterator(SetHelperB, SetIndexB);
+			}
+		}
+		else
+		{
+			DifferingSubProperties.Push(RootPath);
+		}
+	}
+	else if(APropAsMap != nullptr)
+	{
+		const FMapProperty* BPropAsMap = CastFieldChecked<const FMapProperty>(BProperty);
+		if(APropAsMap->KeyProp->GetClass() == BPropAsMap->KeyProp->GetClass() && APropAsMap->ValueProp->GetClass() == BPropAsMap->ValueProp->GetClass())
+		{
+			FScriptMapHelper MapHelperA(APropAsMap, AValue);
+			FScriptMapHelper MapHelperB(BPropAsMap, BValue);
+
+			if (MapHelperA.Num() != MapHelperB.Num())
+			{
+				// API not robust enough to indicate changes made to # of set elements, would
+				// need to return something more detailed than DifferingSubProperties array:
+				DifferingSubProperties.Push(RootPath);
+			}
+
+			int32 MapSizeA = MapHelperA.Num();
+			int32 MapSizeB = MapHelperB.Num();
+			
+			int32 MapIndexA = -1;
+			int32 MapIndexB = -1;
+
+			AdvanceMapIterator(MapHelperA, MapIndexA);
+			AdvanceMapIterator(MapHelperB, MapIndexB);
+			
+			for (int32 VirtualIndex = 0; VirtualIndex < MapSizeA && VirtualIndex < MapSizeB; ++VirtualIndex)
+			{
+				IdenticalHelper(APropAsMap->KeyProp, APropAsMap->KeyProp, MapHelperA.GetKeyPtr(MapIndexA), MapHelperB.GetKeyPtr(MapIndexB), FPropertySoftPath(RootPath, VirtualIndex), DifferingSubProperties);
+				IdenticalHelper(APropAsMap->ValueProp, APropAsMap->ValueProp, MapHelperA.GetValuePtr(MapIndexA), MapHelperB.GetValuePtr(MapIndexB), FPropertySoftPath(RootPath, VirtualIndex), DifferingSubProperties);
+
+				AdvanceMapIterator(MapHelperA, MapIndexA);
+				AdvanceMapIterator(MapHelperB, MapIndexB);
+			}
+		}
+		else
+		{
+			DifferingSubProperties.Push(RootPath);
+		}
+	}
+	else if(APropAsObject != nullptr)
+	{
+		// Past container check, do a normal identical check now before going into components
+		if (AProperty->Identical(AValue, BValue, PPF_DeepComparison))
+		{
+			return;
+		}
+
+		const FObjectProperty* BPropAsObject = CastFieldChecked<const FObjectProperty>(BProperty);
+
+		const UObject* A = *((const UObject* const*)AValue);
+		const UObject* B = *((const UObject* const*)BValue);
+
+		// dig into the objects if they are in the same package as our initial object:
+		if(BPropAsObject->HasAnyPropertyFlags(CPF_InstancedReference) && APropAsObject->HasAnyPropertyFlags(CPF_InstancedReference) && A && B && A->GetClass() == B->GetClass())
+		{
+			const UClass* AClass = A->GetClass(); // BClass and AClass are identical!
+
+			// We only want to recurse if this is EditInlineNew and not a component
+			// Other instanced refs are likely to form a type-specific web so recursion doesn't make sense and won't be displayed properly in the details pane
+			if (AClass->HasAnyClassFlags(CLASS_EditInlineNew) && !AClass->IsChildOf(UActorComponent::StaticClass()))
+			{
+				for (TFieldIterator<FProperty> PropertyIt(AClass); PropertyIt; ++PropertyIt)
+				{
+					const FProperty* ClassProp = *PropertyIt;
+					IdenticalHelper(ClassProp, ClassProp, ClassProp->ContainerPtrToValuePtr<void>(A, 0), ClassProp->ContainerPtrToValuePtr<void>(B, 0), FPropertySoftPath(RootPath, ClassProp), DifferingSubProperties);
+				}
+			}
+			else if (A->GetFName() != B->GetFName())
+			{
+				// If the names don't match, report that as a difference as the object was likely changed
+				DifferingSubProperties.Push(RootPath);
+			}
+		}
+		else
+		{
+			DifferingSubProperties.Push(RootPath);
+		}
+	}
+	else
+	{
+		// Passed all container tests that would check for nested properties being wrong
+		if (AProperty->Identical(AValue, BValue, PPF_DeepComparison))
+		{
+			return;
+		}
+
+		DifferingSubProperties.Push(RootPath);
+	}
+}
+
+bool DiffUtils::Identical(const FResolvedProperty& AProp, const FResolvedProperty& BProp, const FPropertySoftPath& RootPath, TArray<FPropertySoftPath>& DifferingProperties)
+{
+	if( AProp.Property == nullptr && BProp.Property == nullptr )
+	{
+		return true;
+	}
+	else if( AProp.Property == nullptr || BProp.Property == nullptr )
+	{
+		return false;
+	}
+
+	const void* AValue = AProp.Property->ContainerPtrToValuePtr<void>(AProp.Object);
+	const void* BValue = BProp.Property->ContainerPtrToValuePtr<void>(BProp.Object);
+
+	// We _could_ just ask the property for comparison but that would make the "identical" functions significantly more complex.
+	// Instead let's write a new function, specific to DiffUtils, that handles the sub properties
+	// NOTE: For Static Arrays, AValue and BValue were, and are, only references to the value at index 0.  So changes to values past index 0 didn't show up before and
+	// won't show up now.  Changes to index 0 will show up as a change to the entire array.
+	IdenticalHelper(AProp.Property, BProp.Property, AValue, BValue, RootPath, DifferingProperties);
+	
+	return DifferingProperties.Num() == 0;
+}
+
+TArray<FPropertySoftPath> DiffUtils::GetVisiblePropertiesInOrderDeclared(const UStruct* ForStruct, const FPropertySoftPath& Scope /*= TArray<FName>()*/)
+{
+	TArray<FPropertySoftPath> Ret;
+	if (ForStruct)
+	{
+		TSet<FString> HiddenCategories = FEditorCategoryUtils::GetHiddenCategories(ForStruct);
+		for (TFieldIterator<FProperty> PropertyIt(ForStruct); PropertyIt; ++PropertyIt)
+		{
+			FName CategoryName = FObjectEditorUtils::GetCategoryFName(*PropertyIt);
+			if (!HiddenCategories.Contains(CategoryName.ToString()))
+			{
+				if (PropertyIt->PropertyFlags&CPF_Edit)
+				{
+					// We don't need to recurse into objects/structs as those will be picked up in the Identical check later
+					FPropertySoftPath NewPath(Scope, *PropertyIt);
+					Ret.Push(NewPath);
+				}
+			}
+		}
+	}
+	return Ret;
+}
+
+TArray<FPropertyPath> DiffUtils::ResolveAll(const UObject* Object, const TArray<FPropertySoftPath>& InSoftProperties)
+{
+	TArray< FPropertyPath > Ret;
+	for (const auto& Path : InSoftProperties)
+	{
+		Ret.Push(Path.ResolvePath(Object));
+	}
+	return Ret;
+}
+
+TArray<FPropertyPath> DiffUtils::ResolveAll(const UObject* Object, const TArray<FSingleObjectDiffEntry>& InDifferences)
+{
+	TArray< FPropertyPath > Ret;
+	for (const auto& Difference : InDifferences)
+	{
+		Ret.Push(Difference.Identifier.ResolvePath(Object));
+	}
+	return Ret;
+}
+
+UPackage* DiffUtils::LoadPackageForDiff(const FPackagePath& InTempPackagePath, const FPackagePath& InOriginalPackagePath)
+{
+	// set up instancing context
+	FLinkerInstancingContext Context;
+	Context.AddPackageMapping(InOriginalPackagePath.GetPackageFName(), InTempPackagePath.GetPackageFName());
+	
+	return LoadPackage(nullptr, *InTempPackagePath.GetPackageName(),
+		LOAD_ForDiff | LOAD_DisableCompileOnLoad | LOAD_DisableEngineVersionChecks, nullptr, &Context);
+}
+
+UPackage* DiffUtils::LoadPackageForDiff(TSharedPtr<ISourceControlRevision> Revision)
+{
+	FString TempFileName;
+	if(Revision->Get(TempFileName))
+	{
+		// Try and load that package
+		const FPackagePath TempPackagePath = FPackagePath::FromLocalPath(TempFileName);
+		const FPackagePath OriginalPackagePath = FPackagePath::FromLocalPath(Revision->GetFilename());
+		return LoadPackageForDiff(TempPackagePath, OriginalPackagePath);
+	}
+	return nullptr;
+}
+
+TSharedPtr<FBlueprintDifferenceTreeEntry> FBlueprintDifferenceTreeEntry::NoDifferencesEntry()
+{
+	// This just generates a widget that tells the user that no differences were detected. Without this
+	// the treeview displaying differences is confusing when no differences are present because it is not obvious
+	// that the control is a treeview (a treeview with no children looks like a listview).
+	const auto GenerateWidget = []() -> TSharedRef<SWidget>
+	{
+		return SNew(STextBlock)
+			.ColorAndOpacity(FLinearColor(.7f, .7f, .7f))
+			.TextStyle(FAppStyle::Get(), TEXT("BlueprintDif.ItalicText"))
+			.Text(NSLOCTEXT("FBlueprintDifferenceTreeEntry", "NoDifferencesLabel", "No differences detected..."));
+	};
+
+	return TSharedPtr<FBlueprintDifferenceTreeEntry>(new FBlueprintDifferenceTreeEntry(
+		FOnDiffEntryFocused()
+		, FGenerateDiffEntryWidget::CreateStatic(GenerateWidget)
+		, TArray< TSharedPtr<FBlueprintDifferenceTreeEntry> >()
+	) );
+}
+
+TSharedPtr<FBlueprintDifferenceTreeEntry> FBlueprintDifferenceTreeEntry::UnknownDifferencesEntry()
+{
+	// Warn about there being unknown differences
+	const auto GenerateWidget = []() -> TSharedRef<SWidget>
+	{
+		return SNew(STextBlock)
+			.ColorAndOpacity(FLinearColor(.7f, .7f, .7f))
+			.TextStyle(FAppStyle::Get(), TEXT("BlueprintDif.ItalicText"))
+			.Text(NSLOCTEXT("FBlueprintDifferenceTreeEntry", "BlueprintTypeNotSupported", "Warning: Detecting differences in this Blueprint type specific data is not yet supported..."));
+	};
+
+	return TSharedPtr<FBlueprintDifferenceTreeEntry>(new FBlueprintDifferenceTreeEntry(
+		FOnDiffEntryFocused()
+		, FGenerateDiffEntryWidget::CreateStatic(GenerateWidget)
+		, TArray< TSharedPtr<FBlueprintDifferenceTreeEntry> >()
+	));
+}
+
+TSharedPtr<FBlueprintDifferenceTreeEntry> FBlueprintDifferenceTreeEntry::CreateCategoryEntry(const FText& LabelText, const FText& ToolTipText, FOnDiffEntryFocused FocusCallback, const TArray< TSharedPtr<FBlueprintDifferenceTreeEntry> >& Children, bool bHasDifferences)
+{
+	const auto CreateDefaultsRootEntry = [](FText LabelText, FText ToolTipText, FLinearColor Color) -> TSharedRef<SWidget>
+	{
+		return SNew(STextBlock)
+			.ToolTipText(ToolTipText)
+			.ColorAndOpacity(Color)
+			.Text(LabelText);
+	};
+
+	return TSharedPtr<FBlueprintDifferenceTreeEntry>(new FBlueprintDifferenceTreeEntry(
+		FocusCallback
+		, FGenerateDiffEntryWidget::CreateStatic(CreateDefaultsRootEntry, LabelText, ToolTipText, DiffViewUtils::LookupColor(bHasDifferences))
+		, Children
+	));
+}
+
+TSharedPtr<FBlueprintDifferenceTreeEntry> FBlueprintDifferenceTreeEntry::CreateCategoryEntryForMerge(const FText& LabelText, const FText& ToolTipText, FOnDiffEntryFocused FocusCallback, const TArray< TSharedPtr<FBlueprintDifferenceTreeEntry> >& Children, bool bHasRemoteDifferences, bool bHasLocalDifferences, bool bHasConflicts)
+{
+	const auto CreateDefaultsRootEntry = [](FText LabelText, FText ToolTipText, bool bInHasRemoteDifferences, bool bInHasLocalDifferences, bool bInHasConflicts) -> TSharedRef<SWidget>
+	{
+		const FLinearColor BaseColor = DiffViewUtils::LookupColor(bInHasRemoteDifferences || bInHasLocalDifferences, bInHasConflicts);
+		return SNew(SHorizontalBox)
+			+ SHorizontalBox::Slot()
+			[
+				SNew(STextBlock)
+				.ToolTipText(ToolTipText)
+				.ColorAndOpacity(BaseColor)
+				.Text(LabelText)
+			]
+			+ DiffViewUtils::Box(true, DiffViewUtils::LookupColor(bInHasRemoteDifferences, bInHasConflicts))
+			+ DiffViewUtils::Box(true, BaseColor)
+			+ DiffViewUtils::Box(true, DiffViewUtils::LookupColor(bInHasLocalDifferences, bInHasConflicts));
+	};
+
+	return TSharedPtr<FBlueprintDifferenceTreeEntry>(new FBlueprintDifferenceTreeEntry(
+		FocusCallback
+		, FGenerateDiffEntryWidget::CreateStatic(CreateDefaultsRootEntry, LabelText, ToolTipText, bHasRemoteDifferences, bHasLocalDifferences, bHasConflicts)
+		, Children
+	));
+}
+
+TSharedRef< STreeView<TSharedPtr< FBlueprintDifferenceTreeEntry > > > DiffTreeView::CreateTreeView(TArray< TSharedPtr<FBlueprintDifferenceTreeEntry> >* DifferencesList)
+{
+	const auto RowGenerator = [](TSharedPtr< FBlueprintDifferenceTreeEntry > Entry, const TSharedRef<STableViewBase>& Owner) -> TSharedRef< ITableRow >
+	{
+		return SNew(STableRow<TSharedPtr<FBlueprintDifferenceTreeEntry> >, Owner)
+			[
+				Entry->GenerateWidget.Execute()
+			];
+	};
+
+	const auto ChildrenAccessor = [](TSharedPtr<FBlueprintDifferenceTreeEntry> InTreeItem, TArray< TSharedPtr< FBlueprintDifferenceTreeEntry > >& OutChildren)
+	{
+		OutChildren = InTreeItem->Children;
+	};
+
+	const auto Selector = [](TSharedPtr<FBlueprintDifferenceTreeEntry> InTreeItem, ESelectInfo::Type Type)
+	{
+		if (InTreeItem.IsValid())
+		{
+			InTreeItem->OnFocus.ExecuteIfBound();
+		}
+	};
+
+	return SNew(STreeView< TSharedPtr< FBlueprintDifferenceTreeEntry > >)
+		.OnGenerateRow(STreeView< TSharedPtr< FBlueprintDifferenceTreeEntry > >::FOnGenerateRow::CreateStatic(RowGenerator))
+		.OnGetChildren(STreeView< TSharedPtr< FBlueprintDifferenceTreeEntry > >::FOnGetChildren::CreateStatic(ChildrenAccessor))
+		.OnSelectionChanged(STreeView< TSharedPtr< FBlueprintDifferenceTreeEntry > >::FOnSelectionChanged::CreateStatic(Selector))
+		.TreeItemsSource(DifferencesList);
+}
+
+int32 DiffTreeView::CurrentDifference(TSharedRef< STreeView<TSharedPtr< FBlueprintDifferenceTreeEntry > > > TreeView, const TArray< TSharedPtr<class FBlueprintDifferenceTreeEntry> >& Differences)
+{
+	auto SelectedItems = TreeView->GetSelectedItems();
+	if (SelectedItems.Num() == 0)
+	{
+		return INDEX_NONE;
+	}
+
+	for (int32 Iter = 0; Iter < SelectedItems.Num(); ++Iter)
+	{
+		int32 Index = Differences.Find(SelectedItems[Iter]);
+		if (Index != INDEX_NONE)
+		{
+			return Index;
+		}
+	}
+
+	return INDEX_NONE;
+}
+
+void DiffTreeView::HighlightNextDifference(TSharedRef< STreeView<TSharedPtr< FBlueprintDifferenceTreeEntry > > > TreeView, const TArray< TSharedPtr<class FBlueprintDifferenceTreeEntry> >& Differences, const TArray< TSharedPtr<class FBlueprintDifferenceTreeEntry> >& RootDifferences)
+{
+	int32 CurrentIndex = CurrentDifference(TreeView, Differences);
+
+	auto Next = Differences[CurrentIndex + 1];
+	// we have to manually expand our parent:
+	for (auto& Test : RootDifferences)
+	{
+		if (Test->Children.Contains(Next))
+		{
+			TreeView->SetItemExpansion(Test, true);
+			break;
+		}
+	}
+
+	TreeView->SetSelection(Next);
+	TreeView->RequestScrollIntoView(Next);
+}
+
+void DiffTreeView::HighlightPrevDifference(TSharedRef< STreeView<TSharedPtr< FBlueprintDifferenceTreeEntry > > > TreeView, const TArray< TSharedPtr<class FBlueprintDifferenceTreeEntry> >& Differences, const TArray< TSharedPtr<class FBlueprintDifferenceTreeEntry> >& RootDifferences)
+{
+	int32 CurrentIndex = CurrentDifference(TreeView, Differences);
+
+	auto Prev = Differences[CurrentIndex - 1];
+	// we have to manually expand our parent:
+	for (auto& Test : RootDifferences)
+	{
+		if (Test->Children.Contains(Prev))
+		{
+			TreeView->SetItemExpansion(Test, true);
+			break;
+		}
+	}
+
+	TreeView->SetSelection(Prev);
+	TreeView->RequestScrollIntoView(Prev);
+}
+
+bool DiffTreeView::HasNextDifference(TSharedRef< STreeView<TSharedPtr< FBlueprintDifferenceTreeEntry > > > TreeView, const TArray< TSharedPtr<class FBlueprintDifferenceTreeEntry> >& Differences)
+{
+	int32 CurrentIndex = CurrentDifference(TreeView, Differences);
+	return Differences.IsValidIndex(CurrentIndex + 1);
+}
+
+bool DiffTreeView::HasPrevDifference(TSharedRef< STreeView<TSharedPtr< FBlueprintDifferenceTreeEntry > > > TreeView, const TArray< TSharedPtr<class FBlueprintDifferenceTreeEntry> >& Differences)
+{
+	int32 CurrentIndex = CurrentDifference(TreeView, Differences);
+	return Differences.IsValidIndex(CurrentIndex - 1);
+}
+
+FLinearColor DiffViewUtils::LookupColor(bool bDiffers, bool bConflicts)
+{
+	if( bConflicts )
+	{
+		return DiffViewUtils::Conflicting();
+	}
+	else if( bDiffers )
+	{
+		return DiffViewUtils::Differs();
+	}
+	else
+	{
+		return DiffViewUtils::Identical();
+	}
+}
+
+FLinearColor DiffViewUtils::Differs()
+{
+	// yellow color
+	return FLinearColor(0.85f,0.71f,0.25f);
+}
+
+FLinearColor DiffViewUtils::Identical()
+{
+	const static FLinearColor ForegroundColor = FAppStyle::GetColor("Graph.ForegroundColor");
+	return ForegroundColor;
+}
+
+FLinearColor DiffViewUtils::Missing()
+{
+	// blue color
+	return FLinearColor(0.3f,0.3f,1.f);
+}
+
+FLinearColor DiffViewUtils::Conflicting()
+{
+	// red color
+	return FLinearColor(1.0f,0.2f,0.3f);
+}
+
+FText DiffViewUtils::PropertyDiffMessage(FSingleObjectDiffEntry Difference, FText ObjectName)
+{
+	FText Message;
+	FString PropertyName = Difference.Identifier.ToDisplayName();
+	switch (Difference.DiffType)
+	{
+	case EPropertyDiffType::PropertyAddedToA:
+		Message = FText::Format(NSLOCTEXT("DiffViewUtils", "PropertyValueChange_Removed", "{0} removed from {1}"), FText::FromString(PropertyName), ObjectName);
+		break;
+	case EPropertyDiffType::PropertyAddedToB:
+		Message = FText::Format(NSLOCTEXT("DiffViewUtils", "PropertyValueChange_Added", "{0} added to {1}"), FText::FromString(PropertyName), ObjectName);
+		break;
+	case EPropertyDiffType::PropertyValueChanged:
+		Message = FText::Format(NSLOCTEXT("DiffViewUtils", "PropertyValueChange", "{0} changed value in {1}"), FText::FromString(PropertyName), ObjectName);
+		break;
+	}
+	return Message;
+}
+
+FText DiffViewUtils::SCSDiffMessage(const FSCSDiffEntry& Difference, FText ObjectName)
+{
+	const FText NodeName = FText::FromName(Difference.TreeIdentifier.Name);
+	FText Text;
+	switch (Difference.DiffType)
+	{
+	case ETreeDiffType::NODE_ADDED:
+		Text = FText::Format(NSLOCTEXT("DiffViewUtils", "NodeAdded", "Added Node {0} to {1}"), NodeName, ObjectName);
+		break;
+	case ETreeDiffType::NODE_REMOVED:
+		Text = FText::Format(NSLOCTEXT("DiffViewUtils", "NodeRemoved", "Removed Node {0} from {1}"), NodeName, ObjectName);
+		break;
+	case ETreeDiffType::NODE_TYPE_CHANGED:
+		Text = FText::Format(NSLOCTEXT("DiffViewUtils", "NodeTypeChanged", "Node {0} changed type in {1}"), NodeName, ObjectName);
+		break;
+	case ETreeDiffType::NODE_PROPERTY_CHANGED:
+		Text = FText::Format(NSLOCTEXT("DiffViewUtils", "NodePropertyChanged", "{0} on {1}"), DiffViewUtils::PropertyDiffMessage(Difference.PropertyDiff, NodeName), ObjectName);
+		break;
+	case ETreeDiffType::NODE_MOVED:
+		Text = FText::Format(NSLOCTEXT("DiffViewUtils", "NodeMoved", "Moved Node {0} in {1}"), NodeName, ObjectName);
+		break;
+	}
+	return Text;
+}
+
+FText DiffViewUtils::GetPanelLabel(const UObject* Asset, const FRevisionInfo& Revision, FText Label)
+{
+	if( !Asset )
+	{
+		return NSLOCTEXT("DiffViewUtils", "NoBlueprint", "None" );
+	}
+	
+	if( !Revision.Revision.IsEmpty() )
+	{
+		FText RevisionData;
+		
+		if(ISourceControlModule::Get().GetProvider().UsesChangelists())
+		{
+			RevisionData = FText::Format(NSLOCTEXT("DiffViewUtils", "RevisionData", "Revision {0} - CL {1} - {2}")
+				, FText::FromString(Revision.Revision)
+				, FText::AsNumber(Revision.Changelist, &FNumberFormattingOptions::DefaultNoGrouping())
+				, FText::FromString(Revision.Date.ToString(TEXT("%m/%d/%Y"))));
+		}
+		else
+		{
+			RevisionData = FText::Format(NSLOCTEXT("DiffViewUtils", "RevisionDataNoChangelist", "Revision {0} - {1}")
+				, FText::FromString(Revision.Revision)
+				, FText::FromString(Revision.Date.ToString(TEXT("%m/%d/%Y"))));		
+		}
+
+		if (Label.IsEmpty())
+		{
+			return FText::Format(NSLOCTEXT("DiffViewUtils", "RevisionLabelTwoLines", "{0}\n{1}")
+				, FText::FromString(Asset->GetName())
+				, RevisionData);
+		}
+		else
+		{
+			return FText::Format(NSLOCTEXT("DiffViewUtils", "RevisionLabel", "{0}\n{1}\n{2}")
+				, Label
+				, FText::FromString(Asset->GetName())
+				, RevisionData);
+		}
+	}
+	else
+	{
+		if (Label.IsEmpty())
+		{
+			return FText::Format(NSLOCTEXT("DiffViewUtils", "RevisionLabelTwoLines", "{0}\n{1}")
+				, FText::FromString(Asset->GetName())
+				, NSLOCTEXT("DiffViewUtils", "LocalRevisionLabel", "Local Revision"));
+		}
+		else
+		{
+			return FText::Format(NSLOCTEXT("DiffViewUtils", "RevisionLabel", "{0}\n{1}\n{2}")
+				, Label
+				, FText::FromString(Asset->GetName())
+				, NSLOCTEXT("DiffViewUtils", "LocalRevisionLabel", "Local Revision"));
+		}
+	}
+}
+
+SHorizontalBox::FSlot::FSlotArguments DiffViewUtils::Box(bool bIsPresent, FLinearColor Color)
+{
+	return MoveTemp(SHorizontalBox::Slot()
+		.AutoWidth()
+		.HAlign(HAlign_Right)
+		.VAlign(VAlign_Center)
+		.Padding(0.5f, 0.f)
+		[
+			SNew(SImage)
+			.ColorAndOpacity(Color)
+			.Image(bIsPresent ? FAppStyle::GetBrush("BlueprintDif.HasGraph") : FAppStyle::GetBrush("BlueprintDif.MissingGraph"))
+		]);
+};
+
diff --git a/Engine/Source/Editor/UnrealEd/Public/DiffUtils.h b/Engine/Source/Editor/UnrealEd/Public/DiffUtils.h
new file mode 100644
--- /dev/null
+++ b/Engine/Source/Editor/UnrealEd/Public/DiffUtils.h
@@ -0,0 +1,357 @@
+// Copyright Epic Games, Inc. All Rights Reserved.
+#pragma once
+
+#include "Containers/Array.h"
+#include "Containers/Set.h"
+#include "Containers/SparseArray.h"
+#include "Containers/UnrealString.h"
+#include "CoreMinimal.h"
+#include "Delegates/Delegate.h"
+#include "HAL/Platform.h"
+#include "HAL/PlatformCrt.h"
+#include "Internationalization/Text.h"
+#include "Math/Color.h"
+#include "Misc/AssertionMacros.h"
+#include "PropertyPath.h"
+#include "Templates/SharedPointer.h"
+#include "Templates/TypeHash.h"
+#include "UObject/NameTypes.h"
+#include "UObject/UnrealType.h"
+#include "UObject/WeakFieldPtr.h"
+#include "Widgets/SBoxPanel.h"
+#include "Widgets/SWidget.h"
+#include "Widgets/Views/STreeView.h"
+
+class ISourceControlRevision;
+class SWidget;
+class UBlueprint;
+class UObject;
+class UStruct;
+struct FRevisionInfo;
+template <typename ItemType> class STreeView;
+
+struct FResolvedProperty
+{
+	explicit FResolvedProperty()
+		: Object(nullptr)
+		, Property(nullptr)
+	{
+	}
+
+	FResolvedProperty(const void* InObject, const FProperty* InProperty)
+		: Object(InObject)
+		, Property(InProperty)
+	{
+	}
+
+	inline bool operator==(const FResolvedProperty& RHS) const
+	{
+		return Object == RHS.Object && Property == RHS.Property;
+	}
+
+	inline bool operator!=(const FResolvedProperty& RHS) const { return !(*this == RHS); }
+
+	const void* Object;
+	const FProperty* Property;
+};
+
+/**
+ * FPropertySoftPath is a string of identifiers used to identify a single member of a UObject. It is primarily
+ * used when comparing unrelated UObjects for Diffing and Merging, but can also be used as a key select
+ * a property in a SDetailsView.
+ */
+struct FPropertySoftPath
+{
+	FPropertySoftPath()
+	{
+	}
+
+	FPropertySoftPath(TArray<FName> InPropertyChain)
+	{
+		for (FName PropertyName : InPropertyChain)
+		{
+			PropertyChain.Push(FChainElement(PropertyName));
+		}
+	}
+
+	FPropertySoftPath( FPropertyPath InPropertyPath )
+	{
+		for (int32 PropertyIndex = 0, end = InPropertyPath.GetNumProperties(); PropertyIndex != end; ++PropertyIndex)
+		{
+			const FPropertyInfo& Info = InPropertyPath.GetPropertyInfo(PropertyIndex);
+			if (Info.ArrayIndex != INDEX_NONE)
+			{
+				PropertyChain.Push(FName(*FString::FromInt(Info.ArrayIndex)));
+			}
+			else
+			{
+				PropertyChain.Push(FChainElement(Info.Property.Get()));
+			}
+		}
+	}
+
+	FPropertySoftPath(const FPropertySoftPath& SubPropertyPath, const FProperty* LeafProperty)
+		: PropertyChain(SubPropertyPath.PropertyChain)
+	{
+		PropertyChain.Push(FChainElement(LeafProperty));
+	}
+
+	FPropertySoftPath(const FPropertySoftPath& SubPropertyPath, int32 ContainerIndex)
+		: PropertyChain(SubPropertyPath.PropertyChain)
+	{
+		PropertyChain.Push(FName(*FString::FromInt(ContainerIndex)));
+	}
+
+	UNREALED_API FResolvedProperty Resolve(const UObject* Object) const;
+	UNREALED_API FResolvedProperty Resolve(const UStruct* Struct, const void* StructData) const;
+	UNREALED_API FPropertyPath ResolvePath(const UObject* Object) const;
+	UNREALED_API FString ToDisplayName() const;
+
+	inline bool IsSubPropertyMatch(const FPropertySoftPath& PotentialBasePropertyPath) const
+	{
+		if (PropertyChain.Num() <= PotentialBasePropertyPath.PropertyChain.Num())
+		{
+			return false;
+		}
+
+		for (int32 CurChainElement = 0; CurChainElement < PotentialBasePropertyPath.PropertyChain.Num(); CurChainElement++)
+		{
+			if (PotentialBasePropertyPath.PropertyChain[CurChainElement] != PropertyChain[CurChainElement])
+			{
+				return false;
+			}
+		}
+
+		return true;
+	}
+
+	inline bool operator==(FPropertySoftPath const& RHS) const
+	{
+		return PropertyChain == RHS.PropertyChain;
+	}
+
+	inline bool operator!=(FPropertySoftPath const& RHS ) const
+	{
+		return !(*this == RHS);
+	}
+private:
+	struct FChainElement
+	{
+		// FName of the property
+		FName PropertyName;
+
+		// Display string of the property
+		FString DisplayString;
+
+		FChainElement(FName InPropertyName, const FString& InDisplayString = FString())
+			: PropertyName(InPropertyName)
+			, DisplayString(InDisplayString)
+		{
+			if (DisplayString.IsEmpty())
+			{
+				DisplayString = PropertyName.ToString();
+			}
+		}
+
+		FChainElement(const FProperty* Property)
+		{
+			if (Property)
+			{
+				PropertyName = Property->GetFName();
+				DisplayString = Property->GetAuthoredName();
+			}
+		}
+
+		inline bool operator==(FChainElement const& RHS) const
+		{
+			return PropertyName == RHS.PropertyName;
+		}
+
+		inline bool operator!=(FChainElement const& RHS) const
+		{
+			return !(*this == RHS);
+		}
+	};
+	static int32 TryReadIndex(const TArray<FChainElement>& LocalPropertyChain, int32& OutIndex);
+
+	friend uint32 GetTypeHash( FPropertySoftPath const& Path );
+	TArray<FChainElement> PropertyChain;
+};
+
+struct FSCSIdentifier
+{
+	FName Name;
+	TArray< int32 > TreeLocation;
+};
+
+struct FSCSResolvedIdentifier
+{
+	FSCSIdentifier Identifier;
+	const UObject* Object;
+};
+
+FORCEINLINE bool operator==( const FSCSIdentifier& A, const FSCSIdentifier& B )
+{
+	return A.Name == B.Name && A.TreeLocation == B.TreeLocation;
+}
+
+FORCEINLINE bool operator!=(const FSCSIdentifier& A, const FSCSIdentifier& B)
+{
+	return !(A == B);
+}
+
+FORCEINLINE uint32 GetTypeHash( FPropertySoftPath const& Path )
+{
+	uint32 Ret = 0;
+	for (const FPropertySoftPath::FChainElement& PropertyElement : Path.PropertyChain)
+	{
+		Ret = Ret ^ GetTypeHash(PropertyElement.PropertyName);
+	}
+	return Ret;
+}
+
+// Trying to restrict us to this typedef because I'm a little skeptical about hashing FPropertySoftPath safely
+typedef TSet< FPropertySoftPath > FPropertySoftPathSet;
+
+namespace EPropertyDiffType
+{
+	enum Type
+	{
+		Invalid,
+
+		PropertyAddedToA,
+		PropertyAddedToB,
+		PropertyValueChanged,
+	};
+}
+
+struct FSingleObjectDiffEntry
+{
+	FSingleObjectDiffEntry()
+		: Identifier()
+		, DiffType(EPropertyDiffType::Invalid)
+	{
+	}
+
+	FSingleObjectDiffEntry(const FPropertySoftPath& InIdentifier, EPropertyDiffType::Type InDiffType )
+		: Identifier(InIdentifier)
+		, DiffType(InDiffType)
+	{
+	}
+
+	FPropertySoftPath Identifier;
+	EPropertyDiffType::Type DiffType;
+};
+
+namespace ETreeDiffType
+{
+	enum Type
+	{
+		NODE_ADDED,
+		NODE_REMOVED,
+		NODE_TYPE_CHANGED,
+		NODE_PROPERTY_CHANGED,
+		NODE_MOVED,
+		/** We could potentially try to identify hierarchy reorders separately from add/remove */
+	};
+}
+
+struct FSCSDiffEntry
+{
+	FSCSDiffEntry( const FSCSIdentifier& InIdentifier, ETreeDiffType::Type InDiffType, const FSingleObjectDiffEntry& InPropertyDiff )
+		: TreeIdentifier(InIdentifier)
+		, DiffType(InDiffType)
+		, PropertyDiff(InPropertyDiff)
+	{
+	}
+
+	FSCSIdentifier TreeIdentifier;
+	ETreeDiffType::Type DiffType;
+	FSingleObjectDiffEntry PropertyDiff;
+};
+
+struct FSCSDiffRoot
+{
+	// use indices in FSCSIdentifier::TreeLocation to find hierarchy..
+	TArray< FSCSDiffEntry > Entries;
+};
+
+namespace DiffUtils
+{
+	UNREALED_API const UObject* GetCDO(const UBlueprint* ForBlueprint);
+	UNREALED_API void CompareUnrelatedStructs(const UStruct* StructA, const void* A, const UStruct* StructB, const void* B, TArray<FSingleObjectDiffEntry>& OutDifferingProperties);
+	UNREALED_API void CompareUnrelatedObjects(const UObject* A, const UObject* B, TArray<FSingleObjectDiffEntry>& OutDifferingProperties);
+	UNREALED_API void CompareUnrelatedSCS(const UBlueprint* Old, const TArray< FSCSResolvedIdentifier >& OldHierarchy, const UBlueprint* New, const TArray< FSCSResolvedIdentifier >& NewHierarchy, FSCSDiffRoot& OutDifferingEntries );
+	UNREALED_API bool Identical(const FResolvedProperty& AProp, const FResolvedProperty& BProp, const FPropertySoftPath& RootPath, TArray<FPropertySoftPath>& DifferingProperties); 
+	UNREALED_API TArray<FPropertySoftPath> GetVisiblePropertiesInOrderDeclared(const UStruct* ForStruct, const FPropertySoftPath& Scope = FPropertySoftPath());
+
+	UNREALED_API TArray<FPropertyPath> ResolveAll(const UObject* Object, const TArray<FPropertySoftPath>& InSoftProperties);
+	UNREALED_API TArray<FPropertyPath> ResolveAll(const UObject* Object, const TArray<FSingleObjectDiffEntry>& InDifferences);
+
+	/**
+	 * @param InTempPackagePath		- filepath of the temporary uasset version (likely in /Saved/Temp/SourceControl)
+	 * @param InOriginalPackagePath	- filepath of the original uasset in the content directory. Strictly speaking most assets don't need this parameter set
+	 *								  but it's needed to instantiate OFPA Actors properly.
+	 */
+	UNREALED_API UPackage* LoadPackageForDiff(const FPackagePath& InTempPackagePath, const FPackagePath& InOriginalPackagePath);
+	UNREALED_API UPackage* LoadPackageForDiff(TSharedPtr<ISourceControlRevision> Revision);
+
+}
+
+DECLARE_DELEGATE(FOnDiffEntryFocused);
+DECLARE_DELEGATE_RetVal(TSharedRef<SWidget>, FGenerateDiffEntryWidget);
+
+class UNREALED_API FBlueprintDifferenceTreeEntry
+{
+public:
+	FBlueprintDifferenceTreeEntry(FOnDiffEntryFocused InOnFocus, FGenerateDiffEntryWidget InGenerateWidget, TArray< TSharedPtr<FBlueprintDifferenceTreeEntry> > InChildren = TArray< TSharedPtr<FBlueprintDifferenceTreeEntry> >())
+		: OnFocus(InOnFocus)
+		, GenerateWidget(InGenerateWidget)
+		, Children(InChildren) 
+	{
+		check( InGenerateWidget.IsBound() );
+	}
+
+	/** Displays message to user saying there are no differences */
+	static TSharedPtr<FBlueprintDifferenceTreeEntry> NoDifferencesEntry();
+
+	/** Displays message to user warning that there may be undetected differences */
+	static TSharedPtr<FBlueprintDifferenceTreeEntry> UnknownDifferencesEntry();
+
+	/** Create category message for the diff UI */
+	static TSharedPtr<FBlueprintDifferenceTreeEntry> CreateCategoryEntry(const FText& LabelText, const FText& ToolTipText, FOnDiffEntryFocused FocusCallback, const TArray< TSharedPtr<FBlueprintDifferenceTreeEntry> >& Children, bool bHasDifferences);
+
+	/** Create category message for the merge UI */
+	static TSharedPtr<FBlueprintDifferenceTreeEntry> CreateCategoryEntryForMerge(const FText& LabelText, const FText& ToolTipText, FOnDiffEntryFocused FocusCallback, const TArray< TSharedPtr<FBlueprintDifferenceTreeEntry> >& Children, bool bHasRemoteDifferences, bool bHasLocalDifferences, bool bHasConflicts);
+	
+	FOnDiffEntryFocused OnFocus;
+	FGenerateDiffEntryWidget GenerateWidget;
+	TArray< TSharedPtr<FBlueprintDifferenceTreeEntry> > Children;
+};
+
+namespace DiffTreeView
+{
+	UNREALED_API TSharedRef< STreeView<TSharedPtr< FBlueprintDifferenceTreeEntry > > > CreateTreeView(TArray< TSharedPtr<FBlueprintDifferenceTreeEntry> >* DifferencesList);
+	UNREALED_API int32 CurrentDifference( TSharedRef< STreeView<TSharedPtr< FBlueprintDifferenceTreeEntry > > > TreeView, const TArray< TSharedPtr<class FBlueprintDifferenceTreeEntry> >& Differences );
+	UNREALED_API void HighlightNextDifference(TSharedRef< STreeView<TSharedPtr< FBlueprintDifferenceTreeEntry > > > TreeView, const TArray< TSharedPtr<class FBlueprintDifferenceTreeEntry> >& Differences, const TArray< TSharedPtr<class FBlueprintDifferenceTreeEntry> >& RootDifferences);
+	UNREALED_API void HighlightPrevDifference(TSharedRef< STreeView<TSharedPtr< FBlueprintDifferenceTreeEntry > > > TreeView, const TArray< TSharedPtr<class FBlueprintDifferenceTreeEntry> >& Differences, const TArray< TSharedPtr<class FBlueprintDifferenceTreeEntry> >& RootDifferences);
+	UNREALED_API bool HasNextDifference(TSharedRef< STreeView<TSharedPtr< FBlueprintDifferenceTreeEntry > > > TreeView, const TArray< TSharedPtr<class FBlueprintDifferenceTreeEntry> >& Differences);
+	UNREALED_API bool HasPrevDifference(TSharedRef< STreeView<TSharedPtr< FBlueprintDifferenceTreeEntry > > > TreeView, const TArray< TSharedPtr<class FBlueprintDifferenceTreeEntry> >& Differences);
+}
+
+struct FRevisionInfo;
+
+namespace DiffViewUtils
+{
+	UNREALED_API FLinearColor LookupColor( bool bDiffers, bool bConflicts = false );
+	UNREALED_API FLinearColor Differs();
+	UNREALED_API FLinearColor Identical();
+	UNREALED_API FLinearColor Missing();
+	UNREALED_API FLinearColor Conflicting();
+
+	UNREALED_API FText PropertyDiffMessage(FSingleObjectDiffEntry Difference, FText ObjectName);
+	UNREALED_API FText SCSDiffMessage(const FSCSDiffEntry& Difference, FText ObjectName);
+	UNREALED_API FText GetPanelLabel(const UObject* Asset, const FRevisionInfo& Revision, FText Label);
+
+	UNREALED_API SHorizontalBox::FSlot::FSlotArguments Box(bool bIsPresent, FLinearColor Color);
+}
diff --git a/Engine/Source/Editor/WorldBrowser/Private/LevelCollectionModel.cpp b/Engine/Source/Editor/WorldBrowser/Private/LevelCollectionModel.cpp
--- a/Engine/Source/Editor/WorldBrowser/Private/LevelCollectionModel.cpp
+++ b/Engine/Source/Editor/WorldBrowser/Private/LevelCollectionModel.cpp
@@ -25,6 +25,7 @@
 #include "IAssetTools.h"
 #include "IAssetTypeActions.h"
 #include "AssetToolsModule.h"
+#include "DiffUtils.h"
 #include "EditorSupportDelegates.h"
 #include "GameFramework/WorldSettings.h"
 
@@ -1316,12 +1317,11 @@ void FLevelCollectionModel::SCCDiffAgainstDepot(const FLevelModelList& InList, U
 					// Get the head revision of this package from source control
 					FString AbsoluteFileName = FPaths::ConvertRelativePathToFull(RelativeFileName);
 					FString TempFileName;
-					if (Revision->Get(TempFileName))
+					if (UPackage* OldPackage = DiffUtils::LoadPackageForDiff(Revision))
 					{
 						// Try and load that package
 						FText NotMapReason;
-						UPackage* OldPackage = LoadPackage(NULL, *TempFileName, LOAD_ForDiff|LOAD_DisableCompileOnLoad);
-						if(OldPackage != NULL && InEditor->PackageIsAMapFile(*TempFileName, NotMapReason))
+						if(InEditor->PackageIsAMapFile(*TempFileName, NotMapReason))
 						{
 							/* Set the revision information*/
 							UPackage* Package = OriginalPackage;
