Do not track the garbage flag state in the object serialization for transactions, instead track in the transaction itself if the transaction deleted/restored an object. This prevents where situations where objects could be restored or deleted incompletely if they are gone for some reason without the user making the active change.

diff --git a/Engine/Source/Editor/UnrealEd/Classes/Editor/Transactor.h b/Engine/Source/Editor/UnrealEd/Classes/Editor/Transactor.h
--- a/Engine/Source/Editor/UnrealEd/Classes/Editor/Transactor.h
+++ b/Engine/Source/Editor/UnrealEd/Classes/Editor/Transactor.h
@@ -44,6 +44,13 @@ protected:
 	public:
 		struct FSerializedObject : public UE::Transaction::FSerializedObject
 		{
+			enum class EPendingKillChange : uint8
+			{
+				None,
+				DeadToAlive,
+				AliveToDead
+			};
+
 			void SetObject(const UObject* InObject)
 			{
 				ObjectId.SetObject(InObject);
@@ -55,6 +62,7 @@ protected:
 				UE::Transaction::FSerializedObject::Reset();
 				ObjectId.Reset();
 				ObjectAnnotation.Reset();
+				PendingKillChange = EPendingKillChange::None;
 			}
 
 			void Swap(FSerializedObject& Other)
@@ -62,6 +70,7 @@ protected:
 				UE::Transaction::FSerializedObject::Swap(Other);
 				ObjectId.Swap(Other.ObjectId);
 				Exchange(ObjectAnnotation, Other.ObjectAnnotation);
+				Exchange(PendingKillChange, Other.PendingKillChange);
 			}
 
 			/** ID of the object when it was serialized */
@@ -69,6 +78,9 @@ protected:
 
 			/** Annotation data for the object stored externally */
 			TSharedPtr<ITransactionObjectAnnotation> ObjectAnnotation;
+
+			/** Whether this transaction marked this actor as garbage, unmarked it, or didn't touch the garbage flag at all. */
+			EPendingKillChange PendingKillChange = EPendingKillChange::None;
 		};
 
 		// Variables.
diff --git a/Engine/Source/Editor/UnrealEd/Private/EditorTransaction.cpp b/Engine/Source/Editor/UnrealEd/Private/EditorTransaction.cpp
--- a/Engine/Source/Editor/UnrealEd/Private/EditorTransaction.cpp
+++ b/Engine/Source/Editor/UnrealEd/Private/EditorTransaction.cpp
@@ -286,6 +286,10 @@ void FTransaction::FObjectRecord::Save(FTransaction* Owner)
 	check(Owner->bFlip);
 	if (!bRestored)
 	{
+		// Since the transaction will be reset, we want to preserve if this transaction affected the garbage flag at all
+		// (since the only thing that should change is the updated state of the object)
+		FSerializedObject::EPendingKillChange PendingKillChange = SerializedObjectFlip.PendingKillChange;
+
 		SerializedObjectFlip.Reset();
 
 		UObject* CurrentObject = Object.Get();
@@ -296,6 +300,8 @@ void FTransaction::FObjectRecord::Save(FTransaction* Owner)
 
 		FWriter Writer(SerializedObjectFlip, bWantsBinarySerialization);
 		SerializeContents(Writer, -Oper);
+
+		SerializedObjectFlip.PendingKillChange = PendingKillChange;
 	}
 }
 
@@ -393,6 +399,12 @@ void FTransaction::FObjectRecord::Finalize( FTransaction* Owner, UE::Transaction
 			// Diff against the object state when the transaction started
 			UE::Transaction::DiffUtil::GenerateObjectDiff(*DiffableObject, CurrentDiffableObject, DeltaChange, UE::Transaction::DiffUtil::FGenerateObjectDiffOptions(), &ArchetypeCache);
 
+			if (DeltaChange.bHasPendingKillChange)
+			{
+				SerializedObject.PendingKillChange = IsValid(CurrentObject) ? FSerializedObject::EPendingKillChange::DeadToAlive: FSerializedObject::EPendingKillChange::AliveToDead;
+				SerializedObjectFlip.PendingKillChange = IsValid(CurrentObject) ? FSerializedObject::EPendingKillChange::AliveToDead : FSerializedObject::EPendingKillChange::DeadToAlive;
+			}
+
 			// If we have a previous snapshot then we need to consider that part of the diff for the finalized object, as systems may 
 			// have been tracking delta-changes between snapshots and this finalization will need to account for those changes too
 			if (DiffableObjectSnapshot)
@@ -854,6 +866,16 @@ void FTransaction::Apply()
 				Object->CheckDefaultSubobjects();
 				Object->PreEditUndo();
 			}
+			
+			// If we get here, we are undoing - in that case, we need to reverse whatever was in the transaction.
+			if (Record.SerializedObject.PendingKillChange == FObjectRecord::FSerializedObject::EPendingKillChange::AliveToDead)
+			{
+				Object->ClearGarbage();
+			}
+			else if (Record.SerializedObject.PendingKillChange == FObjectRecord::FSerializedObject::EPendingKillChange::DeadToAlive)
+			{
+				Object->MarkAsGarbage();
+			}
 
 			ChangedObjects.Add(Object, FChangedObjectValue(i, Record.SerializedObject.ObjectAnnotation));
 		}
diff --git a/Engine/Source/Runtime/CoreUObject/Private/UObject/Obj.cpp b/Engine/Source/Runtime/CoreUObject/Private/UObject/Obj.cpp
--- a/Engine/Source/Runtime/CoreUObject/Private/UObject/Obj.cpp
+++ b/Engine/Source/Runtime/CoreUObject/Private/UObject/Obj.cpp
@@ -1459,28 +1459,6 @@ void UObject::Serialize(FStructuredArchive::FRecord Record)
 			SerializeScriptProperties(Record.EnterField(TEXT("Properties")));
 		}
 
-		// Keep track of pending kill
-		if (UnderlyingArchive.IsTransacting())
-		{
-			bool WasKill = !IsValidChecked(this);
-			if (UnderlyingArchive.IsLoading())
-			{
-				Record << SA_VALUE(TEXT("WasKill"), WasKill);
-				if (WasKill)
-				{
-					MarkAsGarbage();
-				}
-				else
-				{
-					ClearGarbage();
-				}
-			}
-			else if (UnderlyingArchive.IsSaving())
-			{
-				Record << SA_VALUE(TEXT("WasKill"), WasKill);
-			}
-		}
-
 		// Keep track of transient
 		if (UnderlyingArchive.IsTransacting())
 		{
