Cooker HiddenDependencies analytics.
1) Add -HiddenDependenciesDebug parameter to give warnings for any package references from objects that were not declared in the AssetRegistry dependencies. Removing all hidden dependencies is necessary for robust iterative cooks.
2) Add or extend the PackageScope around all cook operations: Set UE_TRACK_REFERENCING_PACKAGE_SCOPED scope, set packagename for CrashTracker and HiddenDependencies analytics.
3) Remove HybridIterativeDebug parameter; it was the same behavior as the better-named HiddenDependenciesDebug.

#rnx

diff --git a/Engine/Source/Editor/UnrealEd/Classes/CookOnTheSide/CookOnTheFlyServer.h b/Engine/Source/Editor/UnrealEd/Classes/CookOnTheSide/CookOnTheFlyServer.h
--- a/Engine/Source/Editor/UnrealEd/Classes/CookOnTheSide/CookOnTheFlyServer.h
+++ b/Engine/Source/Editor/UnrealEd/Classes/CookOnTheSide/CookOnTheFlyServer.h
@@ -12,11 +12,13 @@
 #include "IPlatformFileSandboxWrapper.h"
 #include "Misc/EnumClassFlags.h"
 #include "Misc/Optional.h"
+#include "Misc/PackageAccessTracking.h"
 #include "Templates/Function.h"
 #include "Templates/RefCounting.h"
 #include "Templates/UniquePtr.h"
 #include "TickableEditorObject.h"
 #include "UObject/Object.h"
+#include "UObject/ObjectHandleTracking.h"
 #include "UObject/ObjectMacros.h"
 #include "UObject/Package.h"
 #include "UObject/UObjectGlobals.h"
@@ -35,6 +37,7 @@ class IPlugin;
 class ITargetPlatform;
 enum class ODSCRecompileCommand;
 struct FBeginCookContext;
+struct FCrashContextExtendedWriter;
 struct FGenericMemoryStats;
 struct FPropertyChangedEvent;
 struct FResourceSizeEx;
@@ -146,6 +149,7 @@ namespace UE::Cook
 {
 	class FAssetRegistryMPCollector;
 	class FBuildDefinitions;
+	class FCachedDependencies;
 	class FCookDirector;
 	class FCookWorkerClient;
 	class FCookWorkerServer;
@@ -1269,6 +1273,9 @@ private:
 		return (CookFlags & InCookFlags) != ECookInitializationFlags::None;
 	}
 
+	void RouteBeginCacheForCookedPlatformData(FName PackageName, UObject* Obj, const ITargetPlatform* TargetPlatform);
+	bool RouteIsCachedCookedPlatformDataLoaded(FName PackageName, UObject* Obj, const ITargetPlatform* TargetPlatform);
+
 	/** Cook (save) a package and process the results */
 	void SaveCookedPackage(UE::Cook::FSaveCookedPackageContext& Context);
 	/** Helper for package saves using ExternalActors: record ExternalActors for iterative builds. */
@@ -1346,6 +1353,37 @@ private:
 	 */
 	void DeleteOutputForPackage(FName PackageName, const ITargetPlatform* TargetPlatform);
 
+	/**
+	 * Set the given package as the active package for diagnostics (e.g. hidden dependencies).
+ 	 * Must be balanced with call to ClearActivePackage. Consider using FScopedActivePackage instead.
+	 * @param PackageTrackingOpsName Optional op for TObjectPtr tracking. Set to NAME_None to not set TObjectPtr context.
+	 */
+	void SetActivePackage(FName PackageName, FName PackageTrackingOpsName);
+	/** Clear the package for diagnostics. Must be balanced with call to SetActivePackage */
+	void ClearActivePackage();
+	struct FActivePackageData
+	{
+		FName PackageName;
+		bool bActive = false;
+		UE_TRACK_REFERENCING_PACKAGE_DECLARE_SCOPE_VARIABLE(ReferenceTrackingScope);
+	};
+	/** Scoped type to call Set/ClearActivePackage. */
+	struct FScopedActivePackage
+	{
+		FScopedActivePackage(UCookOnTheFlyServer& InCOTFS, FName InPackageName, FName InPackageTrackingOpsName);
+		~FScopedActivePackage();
+
+		UCookOnTheFlyServer& COTFS;
+	};
+	/** Callback for FGenericCrashContext; provides the current ActivePackage as context. */
+	void DumpCrashContext(FCrashContextExtendedWriter& Writer);
+	/** Callback for analytics when a new UPackage is loaded. */
+	void OnDiscoveredPackageDebug(FName PackageName, const UE::Cook::FInstigator& Instigator);
+	/** Callback for analytics when a TObjectPtr is read. */
+	void OnObjectHandleReadDebug(TArrayView<const UObject*const> ReadObjects);
+	/** Send warnings/telemetry when a discovered or read package is found to be a hidden dependency. */
+	void ReportHiddenDependency(FName Referencer, FName Dependency);
+
 	static UCookOnTheFlyServer* ActiveCOTFS;
 	uint32		StatLoadedPackageCount = 0;
 	uint32		StatSavedPackageCount = 0;
@@ -1353,6 +1391,23 @@ private:
 	UE::Cook::FStatHistoryInt NumObjectsHistory;
 	UE::Cook::FStatHistoryInt VirtualMemoryHistory;
 
+	FCriticalSection HiddenDependenciesLock;
+	/**
+	 * AllowList or BlockList for reporting hidden dependencies, parsed from ini and commandline, only used if bHiddenDependenciesDebug.
+	 * Read/Write only within HiddenDependenciesLock.
+	 */
+	TSet<FName> HiddenDependenciesClassPathFilterList;
+	/**
+	 * A cache of TSet<FName> of dependencies for each Package, guarded by HiddenDependenciesLock.
+	 * Read/Write only within HiddenDependenciesLock.
+	 */
+	TUniquePtr<UE::Cook::FCachedDependencies> CachedDependencies;
+
+	/** Registration handle for TObjectPtr's OnRead delegate. */
+	UE::CoreUObject::FObjectHandleTrackingCallbackId ObjectHandleReadHandle;
+	/** Data tracking the package currently having calls made (Load/Save/Other) from COTFS. Used for diagnostics. */
+	FActivePackageData ActivePackageData;
+
 	/** True when the which packages we need to cook changes because e.g. a platform was added to the sessionplatforms. */
 	bool bPackageFilterDirty = false;
 	/** True when PumpLoads has detected it is blocked on async work and CookOnTheFlyServer should do work elsewhere. */
@@ -1369,8 +1424,8 @@ private:
 	 * cooked in hybrid-iterative builds.
 	 */
 	bool bHybridIterativeEnabled = true;
-	/** Test mode for the debug of hybrid iterative dependencies. */
-	bool bHybridIterativeDebug = false;
+	bool bHiddenDependenciesDebug = false;
+	bool bHiddenDependenciesClassPathFilterListIsAllowList = true;
 	bool bFirstCookInThisProcessInitialized = false;
 	bool bFirstCookInThisProcess = true;
 	bool bImportBehaviorCallbackInstalled = false;
@@ -1448,6 +1503,8 @@ private:
 	friend UE::Cook::FGeneratorPackage;
 	friend UE::Cook::FInitializeConfigSettings;
 	friend UE::Cook::FPackageData;
+	friend UE::Cook::FPackageDatas;
+	friend UE::Cook::FPackageTracker;
 	friend UE::Cook::FPackageWriterMPCollector;
 	friend UE::Cook::FPendingCookedPlatformData;
 	friend UE::Cook::FPlatformManager;
diff --git a/Engine/Source/Editor/UnrealEd/Private/CookOnTheFlyServer.cpp b/Engine/Source/Editor/UnrealEd/Private/CookOnTheFlyServer.cpp
--- a/Engine/Source/Editor/UnrealEd/Private/CookOnTheFlyServer.cpp
+++ b/Engine/Source/Editor/UnrealEd/Private/CookOnTheFlyServer.cpp
@@ -55,6 +55,7 @@
 #include "EngineGlobals.h"
 #include "FileServerMessages.h"
 #include "GameDelegates.h"
+#include "GenericPlatform/GenericPlatformCrashContext.h"
 #include "GlobalShader.h"
 #include "HAL/FileManager.h"
 #include "HAL/IConsoleManager.h"
@@ -468,6 +469,9 @@ UCookOnTheFlyServer::~UCookOnTheFlyServer()
 	FCoreUObjectDelegates::GetPreGarbageCollectDelegate().RemoveAll(this);
 	FCoreUObjectDelegates::GetPostGarbageCollect().RemoveAll(this);
 	GetTargetPlatformManager()->GetOnTargetPlatformsInvalidatedDelegate().RemoveAll(this);
+#if WITH_ADDITIONAL_CRASH_CONTEXTS
+	FGenericCrashContext::OnAdditionalCrashContextDelegate().RemoveAll(this);
+#endif
 
 	if (HasAnyFlags(RF_ClassDefaultObject))
 	{
@@ -2700,6 +2704,8 @@ void UCookOnTheFlyServer::LoadPackageInQueue(UE::Cook::FPackageData& PackageData
 				RejectPackageToLoad(PackageData, TEXT("is a generated package which could not load its generator"), ESuppressCookReason::LoadError);
 				return;
 			}
+			FScopedActivePackage ScopedActivePackage(*this, OwnerPackageData.GetPackageName(),
+				PackageAccessTrackingOps::NAME_CookerBuildObject);
 			Generator->GetCookPackageSplitterInstance()->OnOwnerReloaded(OwnerPackage, SplitterDataObject);
 			Generator->SetOwnerPackage(OwnerPackage);
 		}
@@ -2793,6 +2799,11 @@ void UCookOnTheFlyServer::QueueDiscoveredPackageData(UE::Cook::FPackageData& Pac
 		return;
 	}
 
+	if (bHiddenDependenciesDebug)
+	{
+		OnDiscoveredPackageDebug(PackageData.GetPackageName(), Instigator);
+	}
+
 	if (CookOnTheFlyRequestManager)
 	{
 		if (PackageData.IsGenerated())
@@ -3193,12 +3204,16 @@ UE::Cook::EPollStatus UCookOnTheFlyServer::BeginCacheObjectsToMove(UE::Cook::FGe
 		if (Info.IsGenerator())
 		{
 			ConstructGeneratedPackagesForPresave(PackageData, Generator, GeneratedPackagesForPresave);
+			FScopedActivePackage ScopedActivePackage(*this, Generator.GetOwner().GetPackageName(),
+				PackageAccessTrackingOps::NAME_CookerBuildObject);
 			bPopulateSucceeded = Splitter->PopulateGeneratorPackage(Package, SplitDataObject, GeneratedPackagesForPresave,
 				ObjectsToMove, KeepReferencedPackages);
 		}
 		else
 		{
 			ICookPackageSplitter::FGeneratedPackageForPopulate SplitterInfo{ Info.RelativePath, Info.GeneratedRootPath, Package, Info.IsCreateAsMap() };
+			FScopedActivePackage ScopedActivePackage(*this, Generator.GetOwner().GetPackageName(),
+				PackageAccessTrackingOps::NAME_CookerBuildObject);
 			bPopulateSucceeded = Splitter->PopulateGeneratedPackage(Package, SplitDataObject, SplitterInfo,
 				ObjectsToMove, KeepReferencedPackages);
 		}
@@ -3246,11 +3261,15 @@ UE::Cook::EPollStatus UCookOnTheFlyServer::PreSaveGeneratorPackage(UE::Cook::FPa
 
 	TArray<UPackage*> KeepReferencedPackages;
 	ConstructGeneratedPackagesForPresave(PackageData, Generator, GeneratedPackagesForPresave);
-	if (!Splitter->PreSaveGeneratorPackage(Package, SplitDataObject, GeneratedPackagesForPresave, KeepReferencedPackages))
 	{
-		UE_LOG(LogCook, Error, TEXT("PackageSplitter returned false from PreSaveGeneratorPackage. Splitter=%s"),
-			*Generator.GetSplitDataObjectName().ToString());
-		return EPollStatus::Error;
+		FScopedActivePackage ScopedActivePackage(*this, Generator.GetOwner().GetPackageName(),
+			PackageAccessTrackingOps::NAME_CookerBuildObject);
+		if (!Splitter->PreSaveGeneratorPackage(Package, SplitDataObject, GeneratedPackagesForPresave, KeepReferencedPackages))
+		{
+			UE_LOG(LogCook, Error, TEXT("PackageSplitter returned false from PreSaveGeneratorPackage. Splitter=%s"),
+				*Generator.GetSplitDataObjectName().ToString());
+			return EPollStatus::Error;
+		}
 	}
 	Info.AddKeepReferencedPackages(KeepReferencedPackages);
 
@@ -3430,11 +3449,15 @@ UE::Cook::EPollStatus UCookOnTheFlyServer::TryPopulateGeneratedPackage(UE::Cook:
 	PopulateData.Package = GeneratedPackage;
 	PopulateData.bCreatedAsMap = GeneratedInfo.IsCreateAsMap();
 	TArray<UPackage*> KeepReferencedPackages;
-	if (!Splitter->PreSaveGeneratedPackage(OwnerPackage, OwnerObject, PopulateData, KeepReferencedPackages))
 	{
-		UE_LOG(LogCook, Error, TEXT("PackageSplitter returned false from PreSaveGeneratedPackage. Splitter=%s, Generated=%s."),
-			*Generator.GetSplitDataObjectName().ToString(), *GeneratedPackageName);
-		return EPollStatus::Error;
+		FScopedActivePackage ScopedActivePackage(*this, Generator.GetOwner().GetPackageName(),
+			PackageAccessTrackingOps::NAME_CookerBuildObject);
+		if (!Splitter->PreSaveGeneratedPackage(OwnerPackage, OwnerObject, PopulateData, KeepReferencedPackages))
+		{
+			UE_LOG(LogCook, Error, TEXT("PackageSplitter returned false from PreSaveGeneratedPackage. Splitter=%s, Generated=%s."),
+				*Generator.GetSplitDataObjectName().ToString(), *GeneratedPackageName);
+			return EPollStatus::Error;
+		}
 	}
 	GeneratedInfo.AddKeepReferencedPackages(KeepReferencedPackages);
 	bool bPackageIsMap = GeneratedPackage->ContainsMap();
@@ -3611,7 +3634,6 @@ UE::Cook::EPollStatus UCookOnTheFlyServer::CallBeginCacheOnObjects(UE::Cook::FPa
 	FWeakObjectPtr* ObjectsData = Objects.GetData();
 	int NumPlatforms = TargetPlatforms.Num();
 	int NumIndexes = Objects.Num() * NumPlatforms;
-	UE_TRACK_REFERENCING_PACKAGE_SCOPED(Package, PackageAccessTrackingOps::NAME_CookerBuildObject);
 	while (NextIndex < NumIndexes)
 	{
 		int ObjectIndex = NextIndex / NumPlatforms;
@@ -3649,9 +3671,9 @@ UE::Cook::EPollStatus UCookOnTheFlyServer::CallBeginCacheOnObjects(UE::Cook::FPa
 			*CurrentAsyncCache -= 1;
 		}
 
-		RouteBeginCacheForCookedPlatformData(Obj, TargetPlatform);
+		RouteBeginCacheForCookedPlatformData(PackageData.GetPackageName(), Obj, TargetPlatform);
 		++NextIndex;
-		if (RouteIsCachedCookedPlatformDataLoaded(Obj, TargetPlatform))
+		if (RouteIsCachedCookedPlatformDataLoaded(PackageData.GetPackageName(), Obj, TargetPlatform))
 		{
 			if (CurrentAsyncCache)
 			{
@@ -3887,13 +3909,7 @@ bool UCookOnTheFlyServer::LoadPackageForCooking(UE::Cook::FPackageData& PackageD
 	UE::Cook::FPackageData* ReportingPackageData)
 {
 	UE_SCOPED_HIERARCHICAL_COOKTIMER_AND_DURATION(LoadPackageForCooking, DetailedCookStats::TickCookOnTheSideLoadPackagesTimeSec);
-
-	check(PackageTracker->LoadingPackageData == nullptr);
-	PackageTracker->LoadingPackageData = &PackageData;
-	ON_SCOPE_EXIT
-	{
-		PackageTracker->LoadingPackageData = nullptr;
-	};
+	FScopedActivePackage ScopedActivePackage(*this, PackageData.GetPackageName(), NAME_None);
 
 	FString PackageName = PackageData.GetPackageName().ToString();
 	OutPackage = FindObject<UPackage>(nullptr, *PackageName);
@@ -3919,6 +3935,9 @@ bool UCookOnTheFlyServer::LoadPackageForCooking(UE::Cook::FPackageData& PackageD
 		UPackage* LoadedPackage;
 		{
 			LLM_SCOPE(ELLMTag::Untagged); // Reset the scope so that untagged memory in the package shows up as Untagged rather than Cooker
+#if ENABLE_COOK_STATS
+			++DetailedCookStats::NumRequestedLoads;
+#endif
 			LoadedPackage = LoadPackage(nullptr, *FileName, LOAD_None);
 		}
 		if (IsValid(LoadedPackage) && LoadedPackage->IsFullyLoaded())
@@ -3959,6 +3978,46 @@ bool UCookOnTheFlyServer::LoadPackageForCooking(UE::Cook::FPackageData& PackageD
 	return bSuccess;
 }
 
+void UCookOnTheFlyServer::SetActivePackage(FName PackageName, FName PackageTrackingOpsName)
+{
+	check(!ActivePackageData.bActive);
+	ActivePackageData.bActive = true;
+	ActivePackageData.PackageName = PackageName;
+	if (!PackageTrackingOpsName.IsNone())
+	{
+		UE_TRACK_REFERENCING_PACKAGE_ACTIVATE_SCOPE_VARIABLE(ActivePackageData.ReferenceTrackingScope,
+			PackageName, PackageTrackingOpsName);
+	}
+}
+
+void UCookOnTheFlyServer::ClearActivePackage()
+{
+	check(ActivePackageData.bActive);
+	UE_TRACK_REFERENCING_PACKAGE_DEACTIVATE_SCOPE_VARIABLE(ActivePackageData.ReferenceTrackingScope);
+	ActivePackageData.PackageName = NAME_None;
+	ActivePackageData.bActive = false;
+}
+
+UCookOnTheFlyServer::FScopedActivePackage::FScopedActivePackage(UCookOnTheFlyServer& InCOTFS, FName PackageName, FName PackageTrackingOpsName)
+	: COTFS(InCOTFS)
+{
+	COTFS.SetActivePackage(PackageName, PackageTrackingOpsName);
+}
+UCookOnTheFlyServer::FScopedActivePackage::~FScopedActivePackage()
+{
+	COTFS.ClearActivePackage();
+}
+
+void UCookOnTheFlyServer::DumpCrashContext(FCrashContextExtendedWriter& Writer)
+{
+#if WITH_ADDITIONAL_CRASH_CONTEXTS
+	if (ActivePackageData.bActive)
+	{
+		Writer.AddString(TEXT("ActivePackage"), *WriteToString<256>(ActivePackageData.PackageName));
+	}
+#endif
+}
+
 void UCookOnTheFlyServer::ProcessUnsolicitedPackages(TArray<FName>* OutDiscoveredPackageNames,
 	TMap<FName, UE::Cook::FInstigator>* OutInstigators)
 {
@@ -4383,16 +4442,16 @@ void UCookOnTheFlyServer::TickPrecacheObjectsForPlatforms(const float TimeSlice,
 			continue;
 		}
 
-		UE_TRACK_REFERENCING_PACKAGE_SCOPED(Material, PackageAccessTrackingOps::NAME_CookerBuildObject);
+		FName PackageName = Material->GetPackage()->GetFName();
 		for (const ITargetPlatform* TargetPlatform : TargetPlatforms)
 		{
 			if (!TargetPlatform)
 			{
 				continue;
 			}
-			if (!RouteIsCachedCookedPlatformDataLoaded(Material, TargetPlatform))
+			if (!RouteIsCachedCookedPlatformDataLoaded(PackageName, Material, TargetPlatform))
 			{
-				RouteBeginCacheForCookedPlatformData(Material, TargetPlatform);
+				RouteBeginCacheForCookedPlatformData(PackageName, Material, TargetPlatform);
 				AllMaterialsCompiled = false;
 			}
 		}
@@ -4424,16 +4483,16 @@ void UCookOnTheFlyServer::TickPrecacheObjectsForPlatforms(const float TimeSlice,
 			continue;
 		}
 
-		UE_TRACK_REFERENCING_PACKAGE_SCOPED(Texture, PackageAccessTrackingOps::NAME_CookerBuildObject);
+		FName PackageName = Texture->GetPackage()->GetFName();
 		for (const ITargetPlatform* TargetPlatform : TargetPlatforms)
 		{
 			if (!TargetPlatform)
 			{
 				continue;
 			}
-			if (!RouteIsCachedCookedPlatformDataLoaded(Texture, TargetPlatform))
+			if (!RouteIsCachedCookedPlatformDataLoaded(PackageName, Texture, TargetPlatform))
 			{
-				RouteBeginCacheForCookedPlatformData(Texture, TargetPlatform);
+				RouteBeginCacheForCookedPlatformData(PackageName, Texture, TargetPlatform);
 			}
 		}
 		if (Timer.IsTimeUp())
@@ -5425,6 +5484,7 @@ void UCookOnTheFlyServer::SaveCookedPackage(UE::Cook::FSaveCookedPackageContext&
 				try
 				{
 					LLM_SCOPE_BYTAG(Cooker_SavePackage);
+					FScopedActivePackage ScopedActivePackage(*this, Context.PackageData.GetPackageName(), NAME_None);
 					Context.SavePackageResult = GEditor->Save(Package, Context.World, *Context.PlatFilename, SaveArgs);
 				}
 				catch (std::exception&)
@@ -5793,25 +5853,18 @@ void FSaveCookedPackageContext::FinishPackage()
 	if (!COTFS.CookByTheBookOptions->bSkipSoftReferences)
 	{
 		// Also request any localized variants of this package
-		if (!FPackageName::IsLocalizedPackage(Package->GetName()))
+		for (FName LocalizedPackageName : FRequestCluster::GetLocalizationReferences(PackageFName, COTFS))
 		{
-			const TArray<FName>* LocalizedVariants = COTFS.CookByTheBookOptions->SourceToLocalizedPackageVariants.Find(PackageFName);
-			if (LocalizedVariants)
+			UE::Cook::FPackageData* LocalizedPackageData = COTFS.PackageDatas->TryAddPackageDataByPackageName(LocalizedPackageName);
+			if (LocalizedPackageData)
 			{
-				for (const FName& LocalizedPackageName : *LocalizedVariants)
-				{
-					UE::Cook::FPackageData* LocalizedPackageData = COTFS.PackageDatas->TryAddPackageDataByPackageName(LocalizedPackageName);
-					if (LocalizedPackageData)
-					{
-						COTFS.QueueDiscoveredPackageData(*LocalizedPackageData, UE::Cook::FInstigator(UE::Cook::EInstigator::SoftDependency, PackageFName));
-					}
-				}
+				COTFS.QueueDiscoveredPackageData(*LocalizedPackageData, UE::Cook::FInstigator(UE::Cook::EInstigator::SoftDependency, PackageFName));
 			}
 		}
 
 		// Add SoftObjectPaths to the cook. This has to be done after the package save to catch any SoftObjectPaths that were added during save.
 		TSet<FName> SoftObjectPackages;
-		GRedirectCollector.ProcessSoftObjectPathPackageList(PackageFName, false, SoftObjectPackages);
+		GRedirectCollector.ProcessSoftObjectPathPackageList(PackageFName, false /* bGetEditorOnly */, SoftObjectPackages);
 		for (FName SoftObjectPackage : SoftObjectPackages)
 		{
 			TMap<FSoftObjectPath, FSoftObjectPath> RedirectedPaths;
@@ -5901,12 +5954,14 @@ void UCookOnTheFlyServer::Initialize( ECookMode::Type DesiredCookMode, ECookInit
 
 	PackageDatas = MakeUnique<UE::Cook::FPackageDatas>(*this);
 	PlatformManager = MakeUnique<UE::Cook::FPlatformManager>();
-	PackageTracker = MakeUnique<UE::Cook::FPackageTracker>(*PackageDatas.Get());
+	PackageTracker = MakeUnique<UE::Cook::FPackageTracker>(*this);
 	DiffModeHelper = MakeUnique<FDiffModeCookServerUtils>();
 	BuildDefinitions = MakeUnique<UE::Cook::FBuildDefinitions>();
 	CookByTheBookOptions = MakeUnique<UE::Cook::FCookByTheBookOptions>();
 	CookOnTheFlyOptions = MakeUnique<UE::Cook::FCookOnTheFlyOptions>();
 	AssetRegistry = IAssetRegistry::Get();
+	CachedDependencies = MakeUnique<UE::Cook::FCachedDependencies>();
+
 	if (!IsCookWorkerMode())
 	{
 		WorkerRequests.Reset(new UE::Cook::FWorkerRequestsLocal());
@@ -6009,7 +6064,9 @@ void UCookOnTheFlyServer::Initialize( ECookMode::Type DesiredCookMode, ECookInit
 	FCoreDelegates::TSOnFConfigDeleted().AddUObject(this, &UCookOnTheFlyServer::OnFConfigDeleted);
 
 	GetTargetPlatformManager()->GetOnTargetPlatformsInvalidatedDelegate().AddUObject(this, &UCookOnTheFlyServer::OnTargetPlatformsInvalidated);
-
+#if WITH_ADDITIONAL_CRASH_CONTEXTS
+	FGenericCrashContext::OnAdditionalCrashContextDelegate().AddUObject(this, &UCookOnTheFlyServer::DumpCrashContext);
+#endif
 }
 
 
@@ -6117,8 +6174,6 @@ void FInitializeConfigSettings::LoadLocal(const FString& InOutputDirectoryOverri
 			MaxAsyncCacheForType.Add(CacheSetting.Key, Count);
 		}
 	}
-
-	bHybridIterativeDebug = FParse::Param(FCommandLine::Get(), TEXT("hybriditerativedebug"));
 }
 
 }
@@ -6167,6 +6222,60 @@ void UCookOnTheFlyServer::SetInitializeConfigSettings(UE::Cook::FInitializeConfi
 	// The rest of this function parses config settings that are reparsed on every CookDirector and CookWorker rather than
 	// being replicated from CookDirector to CookWorker
 	bCanSkipEditorReferencedPackagesWhenCooking = UE::SavePackageUtilities::CanSkipEditorReferencedPackagesWhenCooking();
+
+	// Debugging hidden dependencies
+	bHiddenDependenciesDebug = FParse::Param(FCommandLine::Get(), TEXT("HiddenDependenciesDebug"));
+	if (bHiddenDependenciesDebug)
+	{
+		FScopeLock HiddenDependenciesScopeLock(&HiddenDependenciesLock);
+
+		FString ClassPathListStr;
+		TOptional<bool> bAllowList;
+		if (FParse::Value(FCommandLine::Get(), TEXT("-HiddenDependenciesIgnore="), ClassPathListStr))
+		{
+			bAllowList = false;
+		}
+		if (FParse::Value(FCommandLine::Get(), TEXT("-HiddenDependenciesReport="), ClassPathListStr))
+		{
+			if (bAllowList.IsSet() && !*bAllowList)
+			{
+				UE_LOG(LogCook, Error, TEXT("-HiddenDependenciesIgnore and HiddenDepenciesReport are mutually exclusive. HiddenDepenciesIgnore setting will be discarded."));
+			}
+			bAllowList = true;
+		}
+		bHiddenDependenciesClassPathFilterListIsAllowList = bAllowList.IsSet() ? *bAllowList : false;
+		if (!bHiddenDependenciesClassPathFilterListIsAllowList)
+		{
+			TArray<FString> ClassPaths;
+			GConfig->GetArray(TEXT("TargetDomain"), TEXT("IterativeClassBlockList"), ClassPaths, GEditorIni);
+			for (const FString& ClassPath : ClassPaths)
+			{
+				FTopLevelAssetPath Path(ClassPath);
+				if (!Path.IsValid())
+				{
+					UE_LOG(LogCook, Error, TEXT("Invalid Editor:[TargetDomain]:IterativeClassBlockList entry %s. Expected an array of fullpaths such as /Script/Engine.Material"),
+						*ClassPath);
+					continue;
+				}
+				HiddenDependenciesClassPathFilterList.Add(FName(Path.ToString()));
+			}
+		}
+		if (!ClassPathListStr.IsEmpty())
+		{
+			UE::String::ParseTokensMultiple(ClassPathListStr, { '+', ',', ';' },
+				[this](FStringView Token)
+				{
+					FTopLevelAssetPath Path(Token);
+					if (!Path.IsValid())
+					{
+						UE_LOG(LogCook, Error, TEXT("Invalid %s=<ClassPath> setting. Expected a comma-delimited list of fullpaths such as /Script/Engine.Material"),
+							bHiddenDependenciesClassPathFilterListIsAllowList ? TEXT("-HiddenDependenciesReport") : TEXT("-HiddenDependenciesIgnore"));
+						return;
+					}
+					HiddenDependenciesClassPathFilterList.Add(FName(Path.ToString()));
+				});
+		}
+	}
 }
 
 bool UCookOnTheFlyServer::TryInitializeCookWorker()
@@ -7173,22 +7282,6 @@ bool UCookOnTheFlyServer::SaveCurrentIniSettings(const ITargetPlatform* TargetPl
 
 void UCookOnTheFlyServer::OnRequestClusterCompleted(const UE::Cook::FRequestCluster& RequestCluster)
 {
-	using namespace UE::Cook;
-	if (IsDirectorCookByTheBook())
-	{
-		if (bHybridIterativeDebug)
-		{
-			PackageDatas->LockAndEnumeratePackageDatas([](FPackageData* PackageData)
-			{
-				if (!PackageData->AreAllRequestedPlatformsExplored())
-				{
-					UE_LOG(LogCook, Warning, TEXT("Missing dependency: existing requested Package %s was not explored in the first cluster."),
-						*WriteToString<256>(PackageData->GetPackageName()));
-				}
-			});
-			PackageDatas->SetLogDiscoveredPackages(true); // Turn this on for the rest of the cook after the initial cluster
-		}
-	}
 }
 
 FAsyncIODelete& UCookOnTheFlyServer::GetAsyncIODelete()
@@ -10438,6 +10531,11 @@ void UCookOnTheFlyServer::RegisterCookByTheBookDelegates()
 	{
 		FCoreUObjectDelegates::PackageCreatedForLoad.AddUObject(this, &UCookOnTheFlyServer::MaybeMarkPackageAsAlreadyLoaded);
 	}
+	if (bHiddenDependenciesDebug)
+	{
+		ObjectHandleReadHandle = UE::CoreUObject::AddObjectHandleReadCallback(
+			[this](TArrayView<const UObject*const> ReadObjects) { UCookOnTheFlyServer::OnObjectHandleReadDebug(ReadObjects); });
+	}
 }
 
 void UCookOnTheFlyServer::UnregisterCookByTheBookDelegates()
@@ -10446,6 +10544,11 @@ void UCookOnTheFlyServer::UnregisterCookByTheBookDelegates()
 	{
 		FCoreUObjectDelegates::PackageCreatedForLoad.RemoveAll(this);
 	}
+	if (ObjectHandleReadHandle.IsValid())
+	{
+		UE::CoreUObject::RemoveObjectHandleReadCallback(ObjectHandleReadHandle);
+		ObjectHandleReadHandle = UE::CoreUObject::FObjectHandleTrackingCallbackId();
+	}
 }
 
 
@@ -11019,9 +11122,8 @@ uint32 UCookOnTheFlyServer::CookFullLoadAndSave()
 								if (!bIsTexture || bSaveConcurrent)
 								{
 									UE_SCOPED_HIERARCHICAL_COOKTIMER(FullLoadAndSave_BeginCache);
-									UE_TRACK_REFERENCING_PACKAGE_SCOPED(Package, PackageAccessTrackingOps::NAME_CookerBuildObject);
-									RouteBeginCacheForCookedPlatformData(Obj, TargetPlatform);
-									if (!RouteIsCachedCookedPlatformDataLoaded(Obj, TargetPlatform))
+									RouteBeginCacheForCookedPlatformData(Obj->GetPackage()->GetFName(), Obj, TargetPlatform);
+									if (!RouteIsCachedCookedPlatformDataLoaded(Obj->GetPackage()->GetFName(), Obj, TargetPlatform))
 									{
 										bAllPlatformDataLoaded = false;
 									}
@@ -11138,10 +11240,9 @@ uint32 UCookOnTheFlyServer::CookFullLoadAndSave()
 			{
 				UObject* Obj = ObjectsToWaitForCookedPlatformData[ObjIdx];
 				bool bAllPlatformDataLoaded = true;
-				UE_TRACK_REFERENCING_PACKAGE_SCOPED(Obj, PackageAccessTrackingOps::NAME_CookerBuildObject);
 				for (const ITargetPlatform* TargetPlatform : TargetPlatforms)
 				{
-					if (!RouteIsCachedCookedPlatformDataLoaded(Obj, TargetPlatform))
+					if (!RouteIsCachedCookedPlatformDataLoaded(Obj->GetPackage()->GetFName(), Obj, TargetPlatform))
 					{
 						bAllPlatformDataLoaded = false;
 						break;
@@ -11209,7 +11310,6 @@ uint32 UCookOnTheFlyServer::CookFullLoadAndSave()
 						GetObjectsWithOuter(PrecachePackage, ObjsInPackage, false);
 					}
 
-					UE_TRACK_REFERENCING_PACKAGE_SCOPED(PrecachePackage, PackageAccessTrackingOps::NAME_CookerBuildObject);
 					for (UObject* Obj : ObjsInPackage)
 					{
 						if (Obj->HasAnyFlags(RF_Transient) || !Obj->IsA(UTexture::StaticClass()))
@@ -11219,7 +11319,7 @@ uint32 UCookOnTheFlyServer::CookFullLoadAndSave()
 
 						for (const ITargetPlatform* TargetPlatform : TargetPlatforms)
 						{
-							RouteBeginCacheForCookedPlatformData(Obj, TargetPlatform);
+							RouteBeginCacheForCookedPlatformData(Obj->GetPackage()->GetFName(), Obj, TargetPlatform);
 						}
 					}
 				}
@@ -11570,6 +11670,332 @@ void UCookOnTheFlyServer::RegisterLocalizationChunkDataGenerator()
 	}
 }
 
+void UCookOnTheFlyServer::RouteBeginCacheForCookedPlatformData(FName PackageName, UObject* Obj, const ITargetPlatform* TargetPlatform)
+{
+	LLM_SCOPE_BYTAG(Cooker_CachedPlatformData);
+	UE_SCOPED_TEXT_COOKTIMER(*WriteToString<128>(GetClassTraceScope(Obj), TEXT("_BeginCacheForCookedPlatformData")));
+	UE_SCOPED_COOK_STAT(PackageName, EPackageEventStatType::BeginCacheForCookedPlatformData);
+	FScopedActivePackage ScopedActivePackage(*this, PackageName, PackageAccessTrackingOps::NAME_CookerBuildObject);
+	Obj->BeginCacheForCookedPlatformData(TargetPlatform);
+}
+
+bool UCookOnTheFlyServer::RouteIsCachedCookedPlatformDataLoaded(FName PackageName, UObject* Obj, const ITargetPlatform* TargetPlatform)
+{
+	LLM_SCOPE_BYTAG(Cooker_CachedPlatformData);
+	UE_SCOPED_TEXT_COOKTIMER(*WriteToString<128>(GetClassTraceScope(Obj), TEXT("_IsCachedCookedPlatformDataLoaded")));
+	UE_SCOPED_COOK_STAT(Obj->GetPackage()->GetFName(), EPackageEventStatType::IsCachedCookedPlatformDataLoaded);
+	FScopedActivePackage ScopedActivePackage(*this, PackageName, PackageAccessTrackingOps::NAME_CookerBuildObject);
+	return Obj->IsCachedCookedPlatformDataLoaded(TargetPlatform);
+}
+
+namespace UE::Cook
+{
+
+/**
+ * Keep a map from PackageName to TSet<FName> package dependencies.
+ * Prune old entries in the map when memory limit is reached.
+ */
+class FCachedDependencies
+{
+public:
+	const TSet<FName>& FindOrAdd(FName PackageName);
+
+private:
+	struct FEntry
+	{
+		TSet<FName> Dependencies;
+		FName PackageName;
+	};
+
+	void CalculateDependencies(FEntry& Entry);
+
+	TRingBuffer<TUniquePtr<FEntry>> FIFOEntries;
+	TMap<FName, FEntry*> ExistingPackages;
+	TArray<FName> ScratchDependenciesArray;
+	int32 MaxNum = 1024;
+	int32 MaxEntryDependenciesKeepSize = 1024;
+};
+
+const TSet<FName>& FCachedDependencies::FindOrAdd(FName PackageName)
+{
+	FEntry* Entry;
+	FEntry*& EntryInExistingPackages = ExistingPackages.FindOrAdd(PackageName);
+	if (!EntryInExistingPackages)
+	{
+		TUniquePtr<FEntry> AllocatedEntry;
+		if (FIFOEntries.Num() >= MaxNum)
+		{
+			AllocatedEntry = FIFOEntries.PopFrontValue();
+			// If the PackageNames were equal, then setting the new name followed by removal of the old name would not work.
+			// They can't be equal because PackageName did not exist in ExistingPackages so it cannot exist in FIFOEntries.
+			check(AllocatedEntry->PackageName != PackageName);
+			EntryInExistingPackages = AllocatedEntry.Get();
+			ExistingPackages.Remove(AllocatedEntry->PackageName);
+			// After calling Remove, EntryInExistingPackages is possibly invalidated so we cannot access it
+		}
+		else
+		{
+			AllocatedEntry.Reset(new FEntry());
+			EntryInExistingPackages = AllocatedEntry.Get();
+		}
+		Entry = AllocatedEntry.Get();
+		FIFOEntries.Add(MoveTemp(AllocatedEntry));
+
+		Entry->PackageName = PackageName;
+		CalculateDependencies(*Entry);
+	}
+	else
+	{
+		Entry = EntryInExistingPackages;
+	}
+	return Entry->Dependencies;
+}
+
+void FCachedDependencies::CalculateDependencies(FEntry& Entry)
+{
+	// Get dependencies into the array scratch
+	ScratchDependenciesArray.Reset();
+	IAssetRegistry::GetChecked().GetDependencies(Entry.PackageName, ScratchDependenciesArray,
+		UE::AssetRegistry::EDependencyCategory::Package);
+
+	// Shrink the old Dependencies if it is over the keep size and we don't need the space, and reserve the new size
+	if (Entry.Dependencies.GetMaxIndex() > MaxEntryDependenciesKeepSize &&
+		ScratchDependenciesArray.Num() < MaxEntryDependenciesKeepSize / 2)
+	{
+		Entry.Dependencies.Empty();
+	}
+
+	// Copy the array into the Dependencies set
+	Entry.Dependencies.Append(ScratchDependenciesArray);
+
+	// Shrink the no-longer-needed array scratch if is over the keep size
+	if (ScratchDependenciesArray.Num() > MaxEntryDependenciesKeepSize)
+	{
+		ScratchDependenciesArray.Empty();
+	}
+}
+
+} // namespace UE::Cook
+
+void UCookOnTheFlyServer::OnDiscoveredPackageDebug(FName PackageName, const UE::Cook::FInstigator& Instigator)
+{
+	using namespace UE::Cook;
+
+	if (!bHiddenDependenciesDebug)
+	{
+		return;
+	}
+	switch (Instigator.Category)
+	{
+	case EInstigator::StartupPackage:
+		// Not a Hidden dependency
+		return;
+	default:
+		break;
+	}
+
+	bool bShouldReport = true;
+	PackageDatas->UpdateThreadsafePackageData(PackageName,
+		[&bShouldReport](FThreadsafePackageData& Value, bool bNew)
+		{
+			if (!bNew)
+			{
+				switch (Value.Instigator.Category)
+				{
+				case EInstigator::NotYetRequested:
+				case EInstigator::InvalidCategory:
+					break;
+				default:
+					// Discovered earlier; nothing to report now
+					bShouldReport = false;
+					return;
+				}
+
+				if (!Value.bCookable)
+				{
+					// Skipped during exploration because it is not cookable; we will skip it again in QueueDiscoveredPackageData.
+					bShouldReport = false;
+					return;
+				}
+
+				if (Value.bHasLoggedDiscoveryWarning)
+				{
+					// Discovered and warned earlier; has not yet completed the request phase so Instigator is not yet set
+					// Do not log it again
+					bShouldReport = false;
+					return;
+				}
+			}
+
+			bShouldReport = true;
+			Value.bHasLoggedDiscoveryWarning = true;
+		});
+
+	if (!bShouldReport)
+	{
+		return;
+	}
+	ReportHiddenDependency(Instigator.Referencer, PackageName);
+}
+
+void UCookOnTheFlyServer::OnObjectHandleReadDebug(TArrayView<const UObject*const> ReadObjects)
+{
+	using namespace UE::Cook;
+#if UE_WITH_PACKAGE_ACCESS_TRACKING
+	if (ReadObjects.IsEmpty() || (ReadObjects.Num() == 1 && (!ReadObjects[0] || !ReadObjects[0]->HasAnyFlags(RF_Public))))
+	{
+		return;
+	}
+
+	PackageAccessTracking_Private::FTrackedData* AccumulatedScopeData = PackageAccessTracking_Private::FPackageAccessRefScope::GetCurrentThreadAccumulatedData();
+	if (!AccumulatedScopeData || AccumulatedScopeData->BuildOpName.IsNone())
+	{
+		return;
+	}
+
+	FName ReferencerPackageName = AccumulatedScopeData->PackageName;
+
+	TArray<FName, TInlineAllocator<16>> DependencyPackageNames;
+	for (const UObject* ReadObject : ReadObjects)
+	{
+		if (!ReadObject || !ReadObject->HasAnyFlags(RF_Public))
+		{
+			continue;
+		}
+		UPackage* DependencyPackage = ReadObject->GetOutermost();
+		if (DependencyPackage->HasAnyFlags(RF_Transient))
+		{
+			continue;
+		}
+		FName DependencyPackageName = DependencyPackage->GetFName();
+		if (ReferencerPackageName == DependencyPackageName)
+		{
+			continue;
+		}
+
+		if (FPackageName::IsScriptPackage(WriteToString<256>(DependencyPackageName)))
+		{
+			continue;
+		}
+		DependencyPackageNames.Add(DependencyPackageName);
+	}
+	if (DependencyPackageNames.IsEmpty())
+	{
+		return;
+	}
+
+	{
+		LLM_SCOPE_BYTAG(Cooker);
+		FScopeLock HiddenDependenciesScopeLock(&HiddenDependenciesLock);
+		const TSet<FName>& ExpectedDependencyPackageNames = CachedDependencies->FindOrAdd(ReferencerPackageName);
+		DependencyPackageNames.RemoveAllSwap([&ExpectedDependencyPackageNames](FName DependencyPackageName)
+			{
+				return ExpectedDependencyPackageNames.Contains(DependencyPackageName);
+			});
+	}
+	if (DependencyPackageNames.IsEmpty())
+	{
+		return;
+	}
+
+	// Only report the first hidden dependency from a referencerpackage, to reduce spam
+	bool bShouldReport = true;
+	PackageDatas->UpdateThreadsafePackageData(ReferencerPackageName,
+		[&bShouldReport](FThreadsafePackageData& Value, bool bNew)
+		{
+			if (Value.bHasLoggedDependencyWarning)
+			{
+				bShouldReport = false;
+				return;
+			}
+			Value.bHasLoggedDependencyWarning = true;
+		});
+	if (!bShouldReport)
+	{
+		return;
+	}
+	ReportHiddenDependency(ReferencerPackageName, DependencyPackageNames[0]);
+#endif
+}
+
+void UCookOnTheFlyServer::ReportHiddenDependency(FName Referencer, FName Dependency)
+{
+	using namespace UE::Cook;
+
+	FScopeLock HiddenDependenciesScopeLock(&HiddenDependenciesLock);
+
+	if (!HiddenDependenciesClassPathFilterList.IsEmpty())
+	{
+		TOptional<FAssetPackageData> AssetPackageData;
+
+		bool bImportedClassInFilterList = false;
+		if (!Referencer.IsNone())
+		{
+			AssetPackageData = AssetRegistry->GetAssetPackageDataCopy(Referencer);
+			if (AssetPackageData.IsSet())
+			{
+				for (FName ImportedClass : AssetPackageData->ImportedClasses)
+				{
+					if (HiddenDependenciesClassPathFilterList.Contains(ImportedClass))
+					{
+						bImportedClassInFilterList = true;
+						break;
+					}
+				}
+			}
+			if (!bImportedClassInFilterList)
+			{
+				TOptional<FThreadsafePackageData> Data = PackageDatas->FindThreadsafePackageData(Referencer);
+				FName Generator = Data ? Data->Generator : NAME_None;
+				if (!Generator.IsNone())
+				{
+					AssetPackageData = AssetRegistry->GetAssetPackageDataCopy(Generator);
+					if (AssetPackageData.IsSet())
+					{
+						for (FName ImportedClass : AssetPackageData->ImportedClasses)
+						{
+							if (HiddenDependenciesClassPathFilterList.Contains(ImportedClass))
+							{
+								bImportedClassInFilterList = true;
+								break;
+							}
+						}
+					}
+				}
+			}
+		}
+
+		bool bShouldReport = (bHiddenDependenciesClassPathFilterListIsAllowList ?
+			bImportedClassInFilterList : !bImportedClassInFilterList);
+		if (!bShouldReport)
+		{
+			return;
+		}
+	}
+
+	FString PrimaryAssetClassPath;
+	if (!Referencer.IsNone())
+	{
+		TArray<FAssetData> AssetDatas;
+		AssetRegistry->GetAssetsByPackageName(Referencer, AssetDatas, true /* bIncludeOnlyDiskAssets */);
+		TArray<const FAssetData*, TInlineAllocator<1>> AssetDataPtrs;
+		AssetDataPtrs.Reserve(AssetDatas.Num());
+		for (FAssetData& AssetData : AssetDatas)
+		{
+			AssetDataPtrs.Add(&AssetData);
+		}
+		const FAssetData* PrimaryAssetData = UE::AssetRegistry::GetMostImportantAsset(AssetDataPtrs,
+			false /* bRequireOneTopLevelAsset */);
+		if (PrimaryAssetData && PrimaryAssetData->AssetClassPath.IsValid())
+		{
+			PrimaryAssetClassPath = FString::Printf(TEXT(", PrimaryAssetClass=%s"),
+				*PrimaryAssetData->AssetClassPath.ToString());
+		}
+	}
+	UE_LOG(LogCook, Warning, TEXT("Hidden dependency discovered: %s depends on %s%s."),
+		*Referencer.ToString(), *Dependency.ToString(), *PrimaryAssetClassPath);
+}
+
 static
 void ConditionalWaitOnCommandFile(FStringView GateName, TFunctionRef<void (FStringView)> CommandHandler)
 {
diff --git a/Engine/Source/Editor/UnrealEd/Private/Cooker/CookPackageData.cpp b/Engine/Source/Editor/UnrealEd/Private/Cooker/CookPackageData.cpp
--- a/Engine/Source/Editor/UnrealEd/Private/Cooker/CookPackageData.cpp
+++ b/Engine/Source/Editor/UnrealEd/Private/Cooker/CookPackageData.cpp
@@ -740,10 +740,6 @@ void FPackageData::OnEnterIdle()
 
 void FPackageData::OnExitIdle()
 {
-	if (PackageDatas.GetLogDiscoveredPackages())
-	{
-		UE_LOG(LogCook, Warning, TEXT("Missing dependency: Package %s discovered after initial dependency search."), *WriteToString<256>(PackageName));
-	}
 }
 
 void FPackageData::OnEnterRequest()
@@ -1476,12 +1472,17 @@ bool FGeneratorPackage::TryGenerateList(UObject* OwnerObject, FPackageDatas& Pac
 	FPackageData& OwnerPackageData = GetOwner();
 	UPackage* LocalOwnerPackage = OwnerPackageData.GetPackage();
 	check(LocalOwnerPackage);
-	TArray<ICookPackageSplitter::FGeneratedPackage> GeneratorDatas =
-		CookPackageSplitterInstance->GetGenerateList(LocalOwnerPackage, OwnerObject);
+	UCookOnTheFlyServer& COTFS = GetOwner().GetPackageDatas().GetCookOnTheFlyServer();
+
+	TArray<ICookPackageSplitter::FGeneratedPackage> GeneratorDatas;
+	{
+		UCookOnTheFlyServer::FScopedActivePackage ScopedActivePackage(COTFS, OwnerPackageData.GetPackageName(),
+			PackageAccessTrackingOps::NAME_CookerBuildObject);
+		GeneratorDatas = CookPackageSplitterInstance->GetGenerateList(LocalOwnerPackage, OwnerObject);
+	}
 	PackagesToGenerate.Reset(GeneratorDatas.Num());
 	TArray<const ITargetPlatform*, TInlineAllocator<1>> RequestedPlatforms;
 	OwnerPackageData.GetRequestedPlatforms(RequestedPlatforms);
-	UCookOnTheFlyServer& COTFS = GetOwner().GetPackageDatas().GetCookOnTheFlyServer();
 	bool bHybridIterativeEnabled = COTFS.bHybridIterativeEnabled;
 
 	int32 NumIterativeModified = 0; 
@@ -1786,6 +1787,9 @@ UPackage* FGeneratorPackage::CreateGeneratedUPackage(FCookGenerationInfo& Genera
 	const UPackage* InOwnerPackage, const TCHAR* GeneratedPackageName)
 {
 	check(IsInitialized());
+#if ENABLE_COOK_STATS
+	++DetailedCookStats::NumRequestedLoads;
+#endif
 	UPackage* GeneratedPackage = CreatePackage(GeneratedPackageName);
 	PRAGMA_DISABLE_DEPRECATION_WARNINGS;
 	GeneratedPackage->SetGuid(GeneratedInfo.Guid);
@@ -1838,6 +1842,9 @@ void FGeneratorPackage::ResetSaveState(FCookGenerationInfo& Info, UPackage* Pack
 		}
 		else
 		{
+			UCookOnTheFlyServer& COTFS = Info.PackageData->GetPackageDatas().GetCookOnTheFlyServer();
+			UCookOnTheFlyServer::FScopedActivePackage ScopedActivePackage(COTFS, GetOwner().GetPackageName(),
+				PackageAccessTrackingOps::NAME_CookerBuildObject);
 			if (Info.IsGenerator())
 			{
 				GetCookPackageSplitterInstance()->PostSaveGeneratorPackage(Package, SplitObject);
@@ -2226,8 +2233,8 @@ bool FPendingCookedPlatformData::PollIsComplete()
 		Release();
 		return true;
 	}
-	UE_TRACK_REFERENCING_PACKAGE_SCOPED(LocalObject, PackageAccessTrackingOps::NAME_CookerBuildObject);
-	if (RouteIsCachedCookedPlatformDataLoaded(LocalObject, TargetPlatform))
+	UCookOnTheFlyServer& COTFS = PackageData.GetPackageDatas().GetCookOnTheFlyServer();
+	if (COTFS.RouteIsCachedCookedPlatformDataLoaded(PackageData.GetPackageName(), LocalObject, TargetPlatform))
 	{
 		Release();
 		return true;
@@ -2435,9 +2442,6 @@ void FPackageDatas::SetBeginCookConfigSettings(FStringView CookShowInstigator)
 			}
 		}
 	}
-
-	// Discoveries during the processing of the initial cluster are expected, so LogDiscoveredPackages must be off.
-	SetLogDiscoveredPackages(false);
 }
 
 FPackageDatas::~FPackageDatas()
@@ -2865,6 +2869,26 @@ FPackageData* FPackageDatas::UpdateFileName(FName PackageName)
 	return PackageData;
 }
 
+FThreadsafePackageData::FThreadsafePackageData()
+	: bInitialized(false)
+	, bCookable(true)
+	, bHasLoggedDiscoveryWarning(false)
+	, bHasLoggedDependencyWarning(false)
+{
+}
+
+void FPackageDatas::UpdateThreadsafePackageData(const FPackageData& PackageData)
+{
+	UpdateThreadsafePackageData(PackageData.GetPackageName(),
+		[&PackageData](FThreadsafePackageData& ThreadsafeData, bool bNew)
+		{
+			ThreadsafeData.Instigator = PackageData.GetInstigator();
+			FGeneratorPackage* Generator = PackageData.GetGeneratedOwner();
+			ThreadsafeData.Generator = Generator ? Generator->GetOwner().GetPackageName() : NAME_None;
+			ThreadsafeData.bCookable = PackageData.CanCookForPlatforms();
+		});
+}
+
 int32 FPackageDatas::GetNumCooked()
 {
 	int32 Count = 0;
@@ -3055,25 +3079,24 @@ void FPackageDatas::RemapTargetPlatforms(const TMap<ITargetPlatform*, ITargetPla
 
 void FPackageDatas::DebugInstigator(FPackageData& PackageData)
 {
-	if (ShowInstigatorPackageData != &PackageData)
-	{
-		return;
-	}
-
-	TArray<FInstigator> Chain = CookOnTheFlyServer.GetInstigatorChain(PackageData.GetPackageName());
-	TStringBuilder<256> ChainText;
-	if (Chain.Num() == 0)
+	if (ShowInstigatorPackageData == &PackageData)
 	{
-		ChainText << TEXT("<NoInstigator>");
-	}
-	bool bFirst = true;
-	for (FInstigator& Instigator : Chain)
-	{
-		if (!bFirst) ChainText << TEXT(" <- ");
-		ChainText << TEXT("{ ") << Instigator.ToString() << TEXT(" }");
-		bFirst = false;
+		TArray<FInstigator> Chain = CookOnTheFlyServer.GetInstigatorChain(PackageData.GetPackageName());
+		TStringBuilder<256> ChainText;
+		if (Chain.Num() == 0)
+		{
+			ChainText << TEXT("<NoInstigator>");
+		}
+		bool bFirst = true;
+		for (FInstigator& Instigator : Chain)
+		{
+			if (!bFirst) ChainText << TEXT(" <- ");
+			ChainText << TEXT("{ ") << Instigator.ToString() << TEXT(" }");
+			bFirst = false;
+		}
+		UE_LOG(LogCook, Display, TEXT("Instigator chain of %s: %s"), *PackageData.GetPackageName().ToString(), ChainText.ToString());
 	}
-	UE_LOG(LogCook, Display, TEXT("Instigator chain of %s: %s"), *PackageData.GetPackageName().ToString(), ChainText.ToString());
+	UpdateThreadsafePackageData(PackageData);
 }
 
 void FRequestQueue::Empty()
diff --git a/Engine/Source/Editor/UnrealEd/Private/Cooker/CookPackageData.h b/Engine/Source/Editor/UnrealEd/Private/Cooker/CookPackageData.h
--- a/Engine/Source/Editor/UnrealEd/Private/Cooker/CookPackageData.h
+++ b/Engine/Source/Editor/UnrealEd/Private/Cooker/CookPackageData.h
@@ -1065,6 +1065,19 @@ public:
 	FPackageDataQueue EntryQueue;
 };
 
+/** Data duplicated from FPackageData that is stored separately for read/write from any thread. */
+struct FThreadsafePackageData
+{
+	FInstigator Instigator;
+	FName Generator;
+	bool bInitialized : 1;
+	bool bCookable : 1;
+	bool bHasLoggedDiscoveryWarning : 1;
+	bool bHasLoggedDependencyWarning : 1;
+
+	FThreadsafePackageData();
+};
+
 typedef TArray<FPendingCookedPlatformData> FPendingCookedPlatformDataContainer;
 
 /*
@@ -1160,6 +1173,11 @@ public:
 	FPackageData& AddPackageDataByPackageNameChecked(const FName& PackageName, bool bRequireExists = true,
 		bool bCreateAsMap = false);
 
+	void UpdateThreadsafePackageData(const FPackageData& PackageData);
+	/** Callback == void (*Callback)(FThreadsafePackageData& Value, bool bNew) */
+	template<typename CallbackType>
+	void UpdateThreadsafePackageData(FName PackageName, CallbackType&& Callback);
+	TOptional<FThreadsafePackageData> FindThreadsafePackageData(FName PackageName);
 	/**
 	 * Return the PackageData with the given FileName if one exists, otherwise return nullptr.
 	 */
@@ -1290,10 +1308,6 @@ public:
 	/** Swap all ITargetPlatform* stored on this instance according to the mapping in @param Remap. */
 	void RemapTargetPlatforms(const TMap<ITargetPlatform*, ITargetPlatform*>& Remap);
 
-	/** Get/Set whether to log newly discovered packages to debug missing transitive references. */
-	void SetLogDiscoveredPackages(bool bLog) { bLogDiscoveredPackages = bLog; }
-	bool GetLogDiscoveredPackages() const { return bLogDiscoveredPackages; }
-
 	/** Called when a PackageData assigns its instigator, for debugging. */
 	void DebugInstigator(FPackageData& PackageData);
 
@@ -1322,6 +1336,8 @@ private:
 	TMap<FName, FPackageData*> PackageNameToPackageData;
 	/** Guarded by ExistenceLock */
 	TMap<FName, FPackageData*> FileNameToPackageData;
+	/* Guarded by ExistenceLock. Duplicates information on FPackageData, but can be read/write from any thread. */
+	TMap<FName, FThreadsafePackageData> ThreadsafePackageDatas;
 	TRingBuffer<FPendingCookedPlatformDataContainer> PendingCookedPlatformDataLists;
 	int32 PendingCookedPlatformDataNum = 0;
 	FRequestQueue RequestQueue;
@@ -1330,10 +1346,9 @@ private:
 	FPackageDataQueue LoadReadyQueue;
 	FPackageDataQueue SaveQueue;
 	UCookOnTheFlyServer& CookOnTheFlyServer;
-	FRWLock ExistenceLock;
+	mutable FRWLock ExistenceLock;
 	FPackageData* ShowInstigatorPackageData = nullptr;
 	double LastPollAsyncTime;
-	bool bLogDiscoveredPackages = false;
 
 	static IAssetRegistry* AssetRegistry;
 };
@@ -1366,4 +1381,25 @@ struct FPoppedPackageDataScope
 #endif
 };
 
+template<typename CallbackType>
+inline void FPackageDatas::UpdateThreadsafePackageData(FName PackageName, CallbackType&& Callback)
+{
+	FWriteScopeLock ExistenceWriteLock(ExistenceLock);
+	FThreadsafePackageData& Value = ThreadsafePackageDatas.FindOrAdd(PackageName);
+	bool bNew = false;
+	if (!Value.bInitialized)
+	{
+		Value.bInitialized = true;
+		bNew = true;
+	}
+	Callback(Value, bNew);
+}
+
+inline TOptional<FThreadsafePackageData> FPackageDatas::FindThreadsafePackageData(FName PackageName)
+{
+	FReadScopeLock ExistenceReadLock(ExistenceLock);
+	FThreadsafePackageData* Value = ThreadsafePackageDatas.Find(PackageName);
+	return Value ? TOptional<FThreadsafePackageData>(*Value) : TOptional<FThreadsafePackageData>();
+}
+
 } // namespace UE::Cook
\ No newline at end of file
diff --git a/Engine/Source/Editor/UnrealEd/Private/Cooker/CookProfiling.cpp b/Engine/Source/Editor/UnrealEd/Private/Cooker/CookProfiling.cpp
--- a/Engine/Source/Editor/UnrealEd/Private/Cooker/CookProfiling.cpp
+++ b/Engine/Source/Editor/UnrealEd/Private/Cooker/CookProfiling.cpp
@@ -244,15 +244,20 @@ namespace DetailedCookStats
 	int32 PeakRequestQueueSize = 0;
 	int32 PeakLoadQueueSize = 0;
 	int32 PeakSaveQueueSize = 0;
+	std::atomic<int32> NumDetectedLoads{ 0 };
+	int32 NumRequestedLoads = 0;
 	uint32 NumPreloadedDependencies = 0;
 	uint32 NumPackagesIterativelySkipped = 0;
 	FCookStatsManager::FAutoRegisterCallback RegisterCookOnTheFlyServerStats([](FCookStatsManager::AddStatFuncRef AddStat)
 		{
-			AddStat(TEXT("Package.Load"), FCookStatsManager::CreateKeyValueArray(TEXT("NumPreloadedDependencies"), NumPreloadedDependencies));
+			AddStat(TEXT("Package.Load"), FCookStatsManager::CreateKeyValueArray(
+				TEXT("NumPreloadedDependencies"), NumPreloadedDependencies,
+				TEXT("NumInlineLoads"), NumDetectedLoads - NumRequestedLoads));
 			AddStat(TEXT("Package.Save"), FCookStatsManager::CreateKeyValueArray(TEXT("NumPackagesIterativelySkipped"), NumPackagesIterativelySkipped));
-			AddStat(TEXT("CookOnTheFlyServer"), FCookStatsManager::CreateKeyValueArray(TEXT("PeakRequestQueueSize"), PeakRequestQueueSize));
-			AddStat(TEXT("CookOnTheFlyServer"), FCookStatsManager::CreateKeyValueArray(TEXT("PeakLoadQueueSize"), PeakLoadQueueSize));
-			AddStat(TEXT("CookOnTheFlyServer"), FCookStatsManager::CreateKeyValueArray(TEXT("PeakSaveQueueSize"), PeakSaveQueueSize));
+			AddStat(TEXT("CookOnTheFlyServer"), FCookStatsManager::CreateKeyValueArray(
+				TEXT("PeakRequestQueueSize"), PeakRequestQueueSize,
+				TEXT("PeakLoadQueueSize"), PeakLoadQueueSize,
+				TEXT("PeakSaveQueueSize"), PeakSaveQueueSize));
 		});
 }
 #endif
diff --git a/Engine/Source/Editor/UnrealEd/Private/Cooker/CookProfiling.h b/Engine/Source/Editor/UnrealEd/Private/Cooker/CookProfiling.h
--- a/Engine/Source/Editor/UnrealEd/Private/Cooker/CookProfiling.h
+++ b/Engine/Source/Editor/UnrealEd/Private/Cooker/CookProfiling.h
@@ -8,6 +8,8 @@
 #include "ProfilingDebugging/CookStats.h"
 #include "Stats/Stats.h"
 
+#include <atomic>
+
 #define OUTPUT_COOKTIMING ENABLE_COOK_STATS
 #define PROFILE_NETWORK 0
 
@@ -155,6 +157,8 @@ extern bool IsUnversioned;
 
 // Stats tracked through FAutoRegisterCallback
 extern uint32 NumPreloadedDependencies;
+extern std::atomic<int32> NumDetectedLoads;
+extern int32 NumRequestedLoads;
 extern uint32 NumPackagesIterativelySkipped;
 extern int32 PeakRequestQueueSize;
 extern int32 PeakLoadQueueSize;
diff --git a/Engine/Source/Editor/UnrealEd/Private/Cooker/CookRequestCluster.cpp b/Engine/Source/Editor/UnrealEd/Private/Cooker/CookRequestCluster.cpp
--- a/Engine/Source/Editor/UnrealEd/Private/Cooker/CookRequestCluster.cpp
+++ b/Engine/Source/Editor/UnrealEd/Private/Cooker/CookRequestCluster.cpp
@@ -853,6 +853,7 @@ void FRequestCluster::FGraphSearch::MarkExplored(FPackageData& PackageData, bool
 		PlatformData.SetExplored(true);
 		PlatformData.SetCookable(bCookable);
 	}
+	Cluster.COTFS.PackageDatas->UpdateThreadsafePackageData(PackageData);
 }
 
 FRequestCluster::FVertexData::FVertexData(EAsyncType, int32 NumPlatforms)
diff --git a/Engine/Source/Editor/UnrealEd/Private/Cooker/CookTypes.cpp b/Engine/Source/Editor/UnrealEd/Private/Cooker/CookTypes.cpp
--- a/Engine/Source/Editor/UnrealEd/Private/Cooker/CookTypes.cpp
+++ b/Engine/Source/Editor/UnrealEd/Private/Cooker/CookTypes.cpp
@@ -217,7 +217,6 @@ FCbWriter& operator<<(FCbWriter& Writer, const UE::Cook::FInitializeConfigSettin
 	Writer << "MaxNumPackagesBeforePartialGC" << Value.MaxNumPackagesBeforePartialGC;
 	Writer << "ConfigSettingDenyList" << Value.ConfigSettingDenyList;
 	Writer << "MaxAsyncCacheForType" << Value.MaxAsyncCacheForType;
-	Writer << "bHybridIterativeDebug" << Value.bHybridIterativeDebug;
 	// Make sure new values are added to LoadFromCompactBinary and MoveOrCopy
 	Writer.EndObject();
 	return Writer;
@@ -253,7 +252,6 @@ bool LoadFromCompactBinary(FCbFieldView Field, UE::Cook::FInitializeConfigSettin
 	bOk = LoadFromCompactBinary(Field["MaxNumPackagesBeforePartialGC"], OutValue.MaxNumPackagesBeforePartialGC) & bOk;
 	bOk = LoadFromCompactBinary(Field["ConfigSettingDenyList"], OutValue.ConfigSettingDenyList) & bOk;
 	bOk = LoadFromCompactBinary(Field["MaxAsyncCacheForType"], OutValue.MaxAsyncCacheForType) & bOk;
-	bOk = LoadFromCompactBinary(Field["bHybridIterativeDebug"], OutValue.bHybridIterativeDebug) & bOk;
 	// Make sure new values are added to MoveOrCopy and operator<<
 	return bOk;
 }
@@ -282,7 +280,6 @@ void FInitializeConfigSettings::MoveOrCopy(SourceType&& Source, TargetType&& Tar
 	Target.MaxNumPackagesBeforePartialGC = Source.MaxNumPackagesBeforePartialGC;
 	Target.ConfigSettingDenyList = MoveTempIfPossible(Source.ConfigSettingDenyList);
 	Target.MaxAsyncCacheForType = MoveTempIfPossible(Source.MaxAsyncCacheForType);
-	Target.bHybridIterativeDebug = Source.bHybridIterativeDebug;
 	// Make sure new values are added to operator<< and LoadFromCompactBinary
 }
 
diff --git a/Engine/Source/Editor/UnrealEd/Private/Cooker/CookTypes.h b/Engine/Source/Editor/UnrealEd/Private/Cooker/CookTypes.h
--- a/Engine/Source/Editor/UnrealEd/Private/Cooker/CookTypes.h
+++ b/Engine/Source/Editor/UnrealEd/Private/Cooker/CookTypes.h
@@ -310,7 +310,6 @@ namespace UE::Cook
 		int32 SoftGCDenominator;
 		TArray<FString> ConfigSettingDenyList;
 		TMap<FName, int32> MaxAsyncCacheForType; // max number of objects of a specific type which are allowed to async cache at once
-		bool bHybridIterativeDebug = false;
 		bool bUseSoftGC = false;
 
 		friend FCbWriter& ::operator<<(FCbWriter& Writer, const UE::Cook::FInitializeConfigSettings& Value);
@@ -374,22 +373,6 @@ namespace UE::Cook
 bool LexTryParseString(FPlatformMemoryStats::EMemoryPressureStatus& OutValue, FStringView Text);
 FString LexToString(FPlatformMemoryStats::EMemoryPressureStatus Value);
 
-inline void RouteBeginCacheForCookedPlatformData(UObject* Obj, const ITargetPlatform* TargetPlatform)
-{
-	LLM_SCOPE_BYTAG(Cooker_CachedPlatformData);
-	UE_SCOPED_TEXT_COOKTIMER(*WriteToString<128>(GetClassTraceScope(Obj), TEXT("_BeginCacheForCookedPlatformData")));
-	UE_SCOPED_COOK_STAT(Obj->GetPackage()->GetFName(), EPackageEventStatType::BeginCacheForCookedPlatformData);
-	Obj->BeginCacheForCookedPlatformData(TargetPlatform);
-}
-
-inline bool RouteIsCachedCookedPlatformDataLoaded(UObject* Obj, const ITargetPlatform* TargetPlatform)
-{
-	LLM_SCOPE_BYTAG(Cooker_CachedPlatformData);
-	UE_SCOPED_TEXT_COOKTIMER(*WriteToString<128>(GetClassTraceScope(Obj), TEXT("_IsCachedCookedPlatformDataLoaded")));
-	UE_SCOPED_COOK_STAT(Obj->GetPackage()->GetFName(), EPackageEventStatType::IsCachedCookedPlatformDataLoaded);
-	return Obj->IsCachedCookedPlatformDataLoaded(TargetPlatform);
-}
-
 //////////////////////////////////////////////////////////////////////////
 // Cook by the book options
 
diff --git a/Engine/Source/Editor/UnrealEd/Private/Cooker/PackageTracker.cpp b/Engine/Source/Editor/UnrealEd/Private/Cooker/PackageTracker.cpp
--- a/Engine/Source/Editor/UnrealEd/Private/Cooker/PackageTracker.cpp
+++ b/Engine/Source/Editor/UnrealEd/Private/Cooker/PackageTracker.cpp
@@ -5,28 +5,14 @@
 #include "CookOnTheFlyServerInterface.h"
 #include "CookPackageData.h"
 #include "CookPlatformManager.h"
+#include "CookProfiling.h"
 #include "Misc/ScopeRWLock.h"
-#include "ProfilingDebugging/CookStats.h"
 #include "UObject/Package.h"
 #include "UObject/UObjectIterator.h"
 
 namespace UE::Cook
 {
 
-#if ENABLE_COOK_STATS
-namespace Stats
-{
-
-// Stats tracked through FAutoRegisterCallback
-static uint32 NumInlineLoads = 0;
-static FCookStatsManager::FAutoRegisterCallback RegisterCookStats([](FCookStatsManager::AddStatFuncRef AddStat)
-	{
-		AddStat(TEXT("Package.Load"), FCookStatsManager::CreateKeyValueArray(TEXT("NumInlineLoads"), NumInlineLoads));
-	});
-
-}
-#endif
-
 void FThreadSafeUnsolicitedPackagesList::AddCookedPackage(const FFilePlatformRequest& PlatformRequest)
 {
 	FScopeLock S(&SyncObject);
@@ -62,9 +48,24 @@ void FThreadSafeUnsolicitedPackagesList::Empty()
 }
 
 
-FPackageTracker::FPackageTracker(FPackageDatas& InPackageDatas)
-	:PackageDatas(InPackageDatas)
+FPackageTracker::FPackageTracker(UCookOnTheFlyServer& InCOTFS)
+	:COTFS(InCOTFS)
+{
+}
+
+FPackageTracker::~FPackageTracker()
+{
+	if (bTrackingInitialized)
+	{
+		GUObjectArray.RemoveUObjectDeleteListener(this);
+		GUObjectArray.RemoveUObjectCreateListener(this);
+	}
+}
+
+void FPackageTracker::InitializeTracking()
 {
+	check(!bTrackingInitialized);
+
 	LLM_SCOPE_BYTAG(Cooker);
 
 	FWriteScopeLock ScopeLock(Lock);
@@ -86,41 +87,53 @@ FPackageTracker::FPackageTracker(FPackageDatas& InPackageDatas)
 
 	GUObjectArray.AddUObjectDeleteListener(this);
 	GUObjectArray.AddUObjectCreateListener(this);
-}
 
-FPackageTracker::~FPackageTracker()
-{
-	GUObjectArray.RemoveUObjectDeleteListener(this);
-	GUObjectArray.RemoveUObjectCreateListener(this);
+	bTrackingInitialized = true;
 }
 
 TMap<UPackage*, FInstigator> FPackageTracker::GetNewPackages()
 {
+	if (!bTrackingInitialized)
+	{
+		InitializeTracking();
+		check(bTrackingInitialized);
+	}
+
 	FWriteScopeLock ScopeLock(Lock);
 	TMap<UPackage*, FInstigator> Result = MoveTemp(NewPackages);
 	NewPackages.Reset();
-	bHasBeenConsumed = true;
 	return Result;
 }
 
 bool FPackageTracker::HasBeenConsumed() const
 {
-	return bHasBeenConsumed;
+	// Tracking is initialized only during GetNewPackages, so we can use bTrackingInitialized to
+	// tell whether GetNewPackages has been called.
+	return bTrackingInitialized;
 }
 
 void FPackageTracker::NotifyUObjectCreated(const class UObjectBase* Object, int32 Index)
 {
 	if (Object->GetClass() == UPackage::StaticClass())
 	{
-		LLM_SCOPE_BYTAG(Cooker);
 		auto Package = const_cast<UPackage*>(static_cast<const UPackage*>(Object));
 
 		if (Package->GetOuter() == nullptr)
 		{
 			LLM_SCOPE_BYTAG(Cooker);
-			if (LoadingPackageData && Package->GetFName() != LoadingPackageData->GetPackageName())
+#if ENABLE_COOK_STATS
+			++DetailedCookStats::NumDetectedLoads;
+#endif
+#if UE_WITH_PACKAGE_ACCESS_TRACKING
+			PackageAccessTracking_Private::FTrackedData* AccumulatedScopeData = PackageAccessTracking_Private::FPackageAccessRefScope::GetCurrentThreadAccumulatedData();
+			FName ReferencerName(AccumulatedScopeData ? AccumulatedScopeData->PackageName : NAME_None);
+#else
+			FName ReferencerName(NAME_None);
+#endif
+			FInstigator Instigator(EInstigator::Unsolicited, ReferencerName);
+			if (COTFS.bHiddenDependenciesDebug)
 			{
-				COOK_STAT(++Stats::NumInlineLoads);
+				COTFS.OnDiscoveredPackageDebug(Package->GetFName(), Instigator);
 			}
 
 			FWriteScopeLock ScopeLock(Lock);
@@ -129,10 +142,8 @@ void FPackageTracker::NotifyUObjectCreated(const class UObjectBase* Object, int3
 				UE_LOG(LogCook, Verbose, TEXT("SoftGC PoorPerformance: Reloaded package %s."), *WriteToString<256>(Package->GetFName()));
 			}
 			LoadedPackages.Add(Package);
-			NewPackages.Add(Package,
-					FInstigator(EInstigator::Unsolicited,
-						LoadingPackageData ? LoadingPackageData->GetPackageName() : NAME_None)
-				);
+			NewPackages.Add(Package, MoveTemp(Instigator));
+
 		}
 	}
 }
diff --git a/Engine/Source/Editor/UnrealEd/Private/Cooker/PackageTracker.h b/Engine/Source/Editor/UnrealEd/Private/Cooker/PackageTracker.h
--- a/Engine/Source/Editor/UnrealEd/Private/Cooker/PackageTracker.h
+++ b/Engine/Source/Editor/UnrealEd/Private/Cooker/PackageTracker.h
@@ -173,7 +173,7 @@ private:
 struct FPackageTracker : public FUObjectArray::FUObjectCreateListener, public FUObjectArray::FUObjectDeleteListener
 {
 public:
-	FPackageTracker(FPackageDatas& InPackageDatas);
+	FPackageTracker(UCookOnTheFlyServer& InCOTFS);
 	~FPackageTracker();
 
 	/** Returns all packages that have been loaded since the last time GetNewPackages was called */
@@ -189,10 +189,7 @@ public:
 	/** Swap all ITargetPlatform* stored on this instance according to the mapping in @param Remap. */
 	void RemapTargetPlatforms(const TMap<ITargetPlatform*, ITargetPlatform*>& Remap);
 
-	/** The package currently being loaded at CookOnTheFlyServer's direct request. Used to determine which load dependencies were not preloaded. */
-	FPackageData* LoadingPackageData = nullptr;
-
-	FPackageDatas& PackageDatas;
+	UCookOnTheFlyServer& COTFS;
 
 	FThreadSafeUnsolicitedPackagesList UnsolicitedCookedPackages;
 	FThreadSafeQueue<FRecompileShaderRequest> RecompileRequests;
@@ -202,8 +199,6 @@ public:
 	FThreadSafeSet<FName> UncookedEditorOnlyPackages; // set of packages that have been rejected due to being referenced by editor-only properties
 	TFastPointerMap<const ITargetPlatform*, TSet<FName>> PlatformSpecificNeverCookPackages;
 
-	bool bHasBeenConsumed = false;
-
 	// Thread-safe enumeration of loaded package. 
 	// A lock is held during enumeration, keep code simple and optimal so the lock is released as fast as possible.
 	template <typename FunctionType>
@@ -226,6 +221,8 @@ public:
 		ExpectedNeverLoadPackages.Empty();
 	}
 private:
+	void InitializeTracking();
+
 	// Protects data for thread-safety
 	FRWLock Lock;
 
@@ -235,6 +232,7 @@ private:
 
 	// This list contains the UPackages loaded since last call to GetNewPackages
 	TMap<UPackage*, FInstigator> NewPackages;
+	bool bTrackingInitialized = false;
 };
 
 } // namespace UE::Cook
diff --git a/Engine/Source/Runtime/CoreUObject/Public/Misc/PackageAccessTracking.h b/Engine/Source/Runtime/CoreUObject/Public/Misc/PackageAccessTracking.h
--- a/Engine/Source/Runtime/CoreUObject/Public/Misc/PackageAccessTracking.h
+++ b/Engine/Source/Runtime/CoreUObject/Public/Misc/PackageAccessTracking.h
@@ -20,12 +20,18 @@ class UPackage;
 #define UE_TRACK_REFERENCING_PACKAGE_DELAYED(TrackerName, Package) if (TrackerName) TrackerName->SetPackageName(Package->GetFName()); else TrackerName.Emplace(Package->GetFName(), TrackerName##_OpName);
 #define UE_TRACK_REFERENCING_PLATFORM_SCOPED(TargetPlatform) PackageAccessTracking_Private::FPackageAccessRefScope ANONYMOUS_VARIABLE(PackageAccessTracker_)(TargetPlatform);
 #define UE_TRACK_REFERENCING_OPNAME_SCOPED(OpName) PackageAccessTracking_Private::FPackageAccessRefScope ANONYMOUS_VARIABLE(PackageAccessTracker_)(OpName);
+#define UE_TRACK_REFERENCING_PACKAGE_DECLARE_SCOPE_VARIABLE(VariableName) TOptional<PackageAccessTracking_Private::FPackageAccessRefScope> VariableName;
+#define UE_TRACK_REFERENCING_PACKAGE_ACTIVATE_SCOPE_VARIABLE(VariableName, Object, OpName) check(!VariableName.IsSet()); VariableName.Emplace(Object, OpName);
+#define UE_TRACK_REFERENCING_PACKAGE_DEACTIVATE_SCOPE_VARIABLE(VariableName) VariableName.Reset();
 #else
 #define UE_TRACK_REFERENCING_PACKAGE_SCOPED(Package, OpName)
 #define UE_TRACK_REFERENCING_PACKAGE_DELAYED_SCOPED(TrackerName, OpName)
 #define UE_TRACK_REFERENCING_PACKAGE_DELAYED(TrackerName, Package)
 #define UE_TRACK_REFERENCING_PLATFORM_SCOPED(TargetPlatform)
 #define UE_TRACK_REFERENCING_OPNAME_SCOPED(OpName)
+#define UE_TRACK_REFERENCING_PACKAGE_DECLARE_SCOPE_VARIABLE(VariableName)
+#define UE_TRACK_REFERENCING_PACKAGE_ACTIVATE_SCOPE_VARIABLE(VariableName, Object, OpName)
+#define UE_TRACK_REFERENCING_PACKAGE_DEACTIVATE_SCOPE_VARIABLE(VariableName)
 #endif //UE_WITH_PACKAGE_ACCESS_TRACKING
 
 ///////////////////////////////////////////////////////////////////////////////////////////////////////////////////////
