Cooker: Remove FullLoadAndSave; it has been replaced in efficiency by MPCook, and further optimizations (concurrent save, async load) will also exceed it even without MPCook or Incremental Cook.
[RN] Minor, Cooking

diff --git a/Engine/Source/Editor/UnrealEd/Classes/CookOnTheSide/CookOnTheFlyServer.h b/Engine/Source/Editor/UnrealEd/Classes/CookOnTheSide/CookOnTheFlyServer.h
--- a/Engine/Source/Editor/UnrealEd/Classes/CookOnTheSide/CookOnTheFlyServer.h
+++ b/Engine/Source/Editor/UnrealEd/Classes/CookOnTheSide/CookOnTheFlyServer.h
@@ -86,7 +86,7 @@ enum class ECookByTheBookOptions
 	NoInputPackages =					0x00000800, // don't include slate content (this cook will probably be missing content unless you know what you are doing)
 	SkipSoftReferences =				0x00001000, // Don't follow soft references when cooking. Usually not viable for a real cook and the results probably wont load properly, but can be useful for debugging
 	SkipHardReferences =				0x00002000, // Don't follow hard references when cooking. Not viable for a real cook, only useful for debugging
-	FullLoadAndSave =					0x00004000, // Load all packages into memory and save them all at once in one tick for speed reasons. This requires a lot of RAM for large games.
+	//Unused=							0x00004000,
 	CookAgainstFixedBase =				0x00010000, // If cooking DLC, assume that the base content can not be modified. 
 	DlcLoadMainAssetRegistry =			0x00020000, // If cooking DLC, populate the main game asset registry
 	ZenStore =							0x00040000, // Store cooked data in Zen Store
@@ -681,26 +681,12 @@ public:
 	/** Tick CookWorker until it finishes or needs to yield. Should only be called when in CookWorker Mode. */
 	uint32 TickCookWorker();
 
-	/**
-	 * Execute CookByTheBook as a load of all packages followed by a save of all packages.
-	 * This mode is sometimes faster than normal CBTB.
-	 * Should be called at the beginning of the cook after StartCookByTheBook, and only if IsFullLoadAndSave.
-	 */
-	uint32 CookFullLoadAndSave();
-
-	/**
-	 * Is the Local CookOnTheFlyServer using FullLoadAndSave?
-	 * FullLoadAndSave is used in SingleProcess cooks only, always false on CookWorkers or MPCook Directors.
-	 */
-	bool IsFullLoadAndSave() const;
-
 	/** Clear all the previously cooked data all cook requests from now on will be considered recook requests */
 	void ClearAllCookedData();
 
 	/** Demote PackageDatas in any queue back to Idle, and eliminate pending requests. Used when canceling a cook. */
 	void CancelAllQueues();
 
-
 	/**
 	 * Clear any cached cooked platform data for a platform and call ClearCachedCookedPlatformData on all UObjects.
 	 * 
diff --git a/Engine/Source/Editor/UnrealEd/Private/Commandlets/CookCommandlet.cpp b/Engine/Source/Editor/UnrealEd/Private/Commandlets/CookCommandlet.cpp
--- a/Engine/Source/Editor/UnrealEd/Private/Commandlets/CookCommandlet.cpp
+++ b/Engine/Source/Editor/UnrealEd/Private/Commandlets/CookCommandlet.cpp
@@ -399,7 +399,12 @@ bool UCookCommandlet::CookByTheBook( const TArray<ITargetPlatform*>& Platforms)
 	CookOptions |= bCookAll ? ECookByTheBookOptions::CookAll : ECookByTheBookOptions::None;
 	CookOptions |= Switches.Contains(TEXT("MAPSONLY")) ? ECookByTheBookOptions::MapsOnly : ECookByTheBookOptions::None;
 	CookOptions |= Switches.Contains(TEXT("NODEV")) ? ECookByTheBookOptions::NoDevContent : ECookByTheBookOptions::None;
-	CookOptions |= Switches.Contains(TEXT("FullLoadAndSave")) ? ECookByTheBookOptions::FullLoadAndSave : ECookByTheBookOptions::None;
+	if (Switches.Contains(TEXT("FullLoadAndSave")))
+	{
+		UE_LOG(LogCook, Warning, TEXT("-FullLoadAndSave has been deprecated; remove the argument to remove this warning.\n")
+			TEXT("For cook optimizations, try using multiprocess cook (-cookprocesscount=<N>, N>1).\n")
+			TEXT("If you still need further optimizations, contact Epic on UDN."));
+	}
 	CookOptions |= bUseZenStore ? ECookByTheBookOptions::ZenStore : ECookByTheBookOptions::None;
 	CookOptions |= Switches.Contains(TEXT("NoGameAlwaysCook")) ? ECookByTheBookOptions::NoGameAlwaysCookPackages : ECookByTheBookOptions::None;
 	CookOptions |= Switches.Contains(TEXT("DisableUnsolicitedPackages")) ? (ECookByTheBookOptions::SkipHardReferences | ECookByTheBookOptions::SkipSoftReferences) : ECookByTheBookOptions::None;
@@ -494,21 +499,14 @@ bool UCookCommandlet::CookByTheBook( const TArray<ITargetPlatform*>& Platforms)
 			CookOnTheFlyServer->StartCookByTheBook(StartupOptions);
 		}
 		DECLARE_SCOPE_CYCLE_COUNTER(TEXT("CookByTheBook.MainLoop"), STAT_CookByTheBook_MainLoop, STATGROUP_LoadTime);
-		if (CookOnTheFlyServer->IsFullLoadAndSave())
-		{
-			CookOnTheFlyServer->CookFullLoadAndSave();
-		}
-		else
+		while (CookOnTheFlyServer->IsInSession())
 		{
-			while (CookOnTheFlyServer->IsInSession())
-			{
-				uint32 TickResults = 0;
-				uint32 UnusedVariable = 0;
+			uint32 TickResults = 0;
+			uint32 UnusedVariable = 0;
 
-				TickResults = CookOnTheFlyServer->TickCookByTheBook(MAX_flt,
-					ShowProgress ? ECookTickFlags::None : ECookTickFlags::HideProgressDisplay);
-				ConditionalCollectGarbage(TickResults, *CookOnTheFlyServer);
-			}
+			TickResults = CookOnTheFlyServer->TickCookByTheBook(MAX_flt,
+				ShowProgress ? ECookTickFlags::None : ECookTickFlags::HideProgressDisplay);
+			ConditionalCollectGarbage(TickResults, *CookOnTheFlyServer);
 		}
 	} while (bTestCook);
 
diff --git a/Engine/Source/Editor/UnrealEd/Private/CookOnTheFlyServer.cpp b/Engine/Source/Editor/UnrealEd/Private/CookOnTheFlyServer.cpp
--- a/Engine/Source/Editor/UnrealEd/Private/CookOnTheFlyServer.cpp
+++ b/Engine/Source/Editor/UnrealEd/Private/CookOnTheFlyServer.cpp
@@ -9852,8 +9852,6 @@ FBeginCookContext UCookOnTheFlyServer::CreateBeginCookByTheBookContext(const FCo
 	CookByTheBookOptions->CreateReleaseVersion = StartupOptions.CreateReleaseVersion;
 	CookByTheBookOptions->bSkipHardReferences = !!(CookOptions & ECookByTheBookOptions::SkipHardReferences);
 	CookByTheBookOptions->bSkipSoftReferences = !!(CookOptions & ECookByTheBookOptions::SkipSoftReferences);
-	CookByTheBookOptions->bFullLoadAndSave = !!(CookOptions & ECookByTheBookOptions::FullLoadAndSave) &&
-		!IsCookingInEditor() && !CookDirector;
 	CookByTheBookOptions->bCookAgainstFixedBase = !!(CookOptions & ECookByTheBookOptions::CookAgainstFixedBase);
 	CookByTheBookOptions->bDlcLoadMainAssetRegistry = !!(CookOptions & ECookByTheBookOptions::DlcLoadMainAssetRegistry);
 	CookByTheBookOptions->bErrorOnEngineContentUse = StartupOptions.bErrorOnEngineContentUse;
@@ -10165,7 +10163,7 @@ void UCookOnTheFlyServer::GenerateInitialRequests(FBeginCookContext& BeginContex
 	CollectFilesToCook(FilesInPath, FilesInPathInstigators, CookMaps, CookDirectories, IniMapSections, CookOptions, TargetPlatforms, GameDefaultObjects);
 
 	// Add soft/hard startup references after collecting requested files and handling empty requests
-	if (!CookByTheBookOptions->bSkipHardReferences && !CookByTheBookOptions->bFullLoadAndSave)
+	if (!CookByTheBookOptions->bSkipHardReferences)
 	{
 		ProcessUnsolicitedPackages(&FilesInPath, &FilesInPathInstigators);
 	}
@@ -10914,607 +10912,6 @@ bool UCookOnTheFlyServer::GetAllPackageFilenamesFromAssetRegistry(const FString&
 	return false;
 }
 
-bool UCookOnTheFlyServer::IsFullLoadAndSave() const
-{
-	return CookByTheBookOptions->bFullLoadAndSave;
-}
-
-uint32 UCookOnTheFlyServer::CookFullLoadAndSave()
-{
-	UE_SCOPED_HIERARCHICAL_COOKTIMER(FullLoadAndSave);
-	check(IsCookByTheBookMode());
-	check(!IsCookingInEditor());
-	check(IsInGameThread());
-
-	uint32 Result = 0;
-
-	const TArray<const ITargetPlatform*>& TargetPlatforms = PlatformManager->GetSessionPlatforms();
-
-	{
-		UE_LOG(LogCook, Display, TEXT("Loading requested packages..."));
-		UE_SCOPED_HIERARCHICAL_COOKTIMER(FullLoadAndSave_RequestedLoads);
-		while (WorkerRequests->HasExternalRequests())
-		{
-			TArray<UE::Cook::FFilePlatformRequest> BuildRequests;
-			TArray<UE::Cook::FSchedulerCallback> SchedulerCallbacks;
-			UE::Cook::EExternalRequestType RequestType = WorkerRequests->DequeueNextCluster(SchedulerCallbacks, BuildRequests);
-			if (RequestType == UE::Cook::EExternalRequestType::Callback)
-			{
-				for (UE::Cook::FSchedulerCallback& SchedulerCallback : SchedulerCallbacks)
-				{
-					SchedulerCallback();
-				}
-				continue;
-			}
-			check(RequestType == UE::Cook::EExternalRequestType::Cook && BuildRequests.Num() > 0);
-			for (UE::Cook::FFilePlatformRequest& ToBuild : BuildRequests)
-			{
-				const FName BuildFilenameFName = ToBuild.GetFilename();
-				if (!PackageTracker->NeverCookPackageList.Contains(BuildFilenameFName))
-				{
-					const FString BuildFilename = BuildFilenameFName.ToString();
-					GIsCookerLoadingPackage = true;
-					UE_SCOPED_HIERARCHICAL_COOKTIMER(LoadPackage);
-					{
-						LLM_SCOPE(ELLMTag::Untagged); // Reset the scope so that untagged memory in the package shows up as Untagged rather than Cooker
-						LoadPackage(nullptr, *BuildFilename, LOAD_None);
-					}
-					if (GShaderCompilingManager)
-					{
-						GShaderCompilingManager->ProcessAsyncResults(true, false);
-					}
-					FAssetCompilingManager::Get().ProcessAsyncTasks(true);
-					GIsCookerLoadingPackage = false;
-				}
-			}
-		}
-	}
-
-	const bool bSaveConcurrent = FParse::Param(FCommandLine::Get(), TEXT("ConcurrentSave"));
-	uint32 SaveFlags = SAVE_KeepGUID | SAVE_Async | (IsCookFlagSet(ECookInitializationFlags::Unversioned) ? SAVE_Unversioned : 0);
-	if (bSaveConcurrent)
-	{
-		SaveFlags |= SAVE_Concurrent;
-	}
-	TArray<UE::Cook::FPackageData*> PackagesToSave;
-	PackagesToSave.Reserve(65536);
-
-	TSet<UPackage*> ProcessedPackages;
-	ProcessedPackages.Reserve(65536);
-
-	TMap<UWorld*, TArray<bool>> WorldsToPostSaveRoot;
-	WorldsToPostSaveRoot.Reserve(1024);
-
-	TArray<UObject*> ObjectsToWaitForCookedPlatformData;
-	ObjectsToWaitForCookedPlatformData.Reserve(65536);
-
-	TArray<FString> PackagesToLoad;
-
-	FObjectSaveContextData ObjectSaveContext(nullptr, nullptr, TEXT(""), SaveFlags);
-
-	do
-	{
-		PackagesToLoad.Reset();
-
-		{
-			UE_LOG(LogCook, Display, TEXT("Caching platform data and discovering string referenced assets..."));
-			UE_SCOPED_HIERARCHICAL_COOKTIMER(FullLoadAndSave_CachePlatformDataAndDiscoverNewAssets);
-			for (TObjectIterator<UPackage> It; It; ++It)
-			{
-				UPackage* Package = *It;
-				check(Package);
-
-				if (ProcessedPackages.Contains(Package))
-				{
-					continue;
-				}
-
-				ProcessedPackages.Add(Package);
-
-				if (Package->HasAnyPackageFlags(PKG_CompiledIn | PKG_ForDiffing | PKG_EditorOnly | PKG_Compiling | PKG_PlayInEditor | PKG_ContainsScript | PKG_ReloadingForCooker))
-				{
-					continue;
-				}
-
-				if (Package == GetTransientPackage())
-				{
-					continue;
-				}
-
-				FName PackageName = Package->GetFName();
-				if (PackageTracker->NeverCookPackageList.Contains(PackageDatas->GetFileNameByPackageName(PackageName)))
-				{
-					// refuse to save this package
-					continue;
-				}
-
-				if (!FPackageName::IsValidLongPackageName(PackageName.ToString()))
-				{
-					continue;
-				}
-
-				if (Package->GetOuter() != nullptr)
-				{
-					UE_LOG(LogCook, Warning, TEXT("Skipping package %s with outermost %s"), *Package->GetName(), *Package->GetOutermost()->GetName());
-					continue;
-				}
-
-				UE::Cook::FPackageData* PackageData = PackageDatas->TryAddPackageDataByPackageName(PackageName);
-				// Legacy behavior: if TryAddPackageDataByPackageName failed, we will still try to load the Package, but we will not try to save it.
-				if (PackageData)
-				{
-					PackageData->SetPackage(Package);
-					PackagesToSave.Add(PackageData);
-				}
-
-
-				{
-					UE_SCOPED_HIERARCHICAL_COOKTIMER(FullLoadAndSave_PerObjectLogic);
-					TSet<UObject*> ProcessedObjects;
-					ProcessedObjects.Reserve(64);
-					bool bObjectsMayHaveBeenCreated = false;
-					do
-					{
-						bObjectsMayHaveBeenCreated = false;
-						TArray<UObject*> ObjsInPackage;
-						{
-							UE_SCOPED_HIERARCHICAL_COOKTIMER(FullLoadAndSave_GetObjectsWithOuter);
-							GetObjectsWithOuter(Package, ObjsInPackage, true);
-						}
-						for (UObject* Obj : ObjsInPackage)
-						{
-							if (Obj->HasAnyFlags(RF_Transient))
-							{
-								continue;
-							}
-
-							if (ProcessedObjects.Contains(Obj))
-							{
-								continue;
-							}
-
-							bObjectsMayHaveBeenCreated = true;
-							ProcessedObjects.Add(Obj);
-
-							UWorld* World = Cast<UWorld>(Obj);
-							bool bInitializedPhysicsSceneForSave = false;
-							bool bForceInitializedWorld = false;
-
-							bool bAllPlatformDataLoaded = true;
-							bool bIsTexture = Obj->IsA(UTexture::StaticClass());
-							bool bFirstPlatform = true;
-							for (const ITargetPlatform* TargetPlatform : TargetPlatforms)
-							{
-								ObjectSaveContext.TargetPlatform = TargetPlatform;
-								ObjectSaveContext.bOuterConcurrentSave = bFirstPlatform;
-								bFirstPlatform = false;
-
-								if (bSaveConcurrent)
-								{
-									GIsCookerLoadingPackage = true;
-									if (World)
-									{
-										UE_SCOPED_HIERARCHICAL_COOKTIMER(FullLoadAndSave_SettingUpWorlds);
-										// We need a physics scene at save time in case code does traces during onsave events.
-										if (ObjectSaveContext.bOuterConcurrentSave)
-										{
-											bInitializedPhysicsSceneForSave = GEditor->InitializePhysicsSceneForSaveIfNecessary(World, bForceInitializedWorld);
-										}
-
-										{
-											UE_SCOPED_HIERARCHICAL_COOKTIMER(FullLoadAndSave_PreSaveWorld);
-											GEditor->OnPreSaveWorld(World, FObjectPreSaveContext(ObjectSaveContext));
-										}
-										{
-											UE_SCOPED_HIERARCHICAL_COOKTIMER(FullLoadAndSave_PreSaveRoot);
-											UE::SavePackageUtilities::CallPreSaveRoot(World, ObjectSaveContext);
-											WorldsToPostSaveRoot.FindOrAdd(World).Add(ObjectSaveContext.bCleanupRequired);
-										}
-									}
-
-									{
-										UE_SCOPED_HIERARCHICAL_COOKTIMER(FullLoadAndSave_PreSave);
-										UE::SavePackageUtilities::CallPreSave(Obj, ObjectSaveContext);
-									}
-									GIsCookerLoadingPackage = false;
-								}
-
-								if (!bIsTexture || bSaveConcurrent)
-								{
-									UE_SCOPED_HIERARCHICAL_COOKTIMER(FullLoadAndSave_BeginCache);
-									RouteBeginCacheForCookedPlatformData(Obj->GetPackage()->GetFName(), Obj, TargetPlatform);
-									if (!RouteIsCachedCookedPlatformDataLoaded(Obj->GetPackage()->GetFName(), Obj, TargetPlatform))
-									{
-										bAllPlatformDataLoaded = false;
-									}
-								}
-							}
-
-							if (!bAllPlatformDataLoaded)
-							{
-								ObjectsToWaitForCookedPlatformData.Add(Obj);
-							}
-
-							if (World && bInitializedPhysicsSceneForSave)
-							{
-								UE_SCOPED_HIERARCHICAL_COOKTIMER(FullLoadAndSave_CleaningUpWorlds);
-								GEditor->CleanupPhysicsSceneThatWasInitializedForSave(World, bForceInitializedWorld);
-							}
-						}
-					} while (bObjectsMayHaveBeenCreated);
-
-					if (bSaveConcurrent)
-					{
-						UE_SCOPED_HIERARCHICAL_COOKTIMER(FullLoadAndSave_MiscPrep);
-						// Precache the metadata so we don't risk rehashing the map in the parallelfor below
-						Package->GetMetaData();
-					}
-				}
-
-				if (!CookByTheBookOptions->bSkipSoftReferences)
-				{
-					UE_SCOPED_HIERARCHICAL_COOKTIMER(ResolveStringReferences);
-					TSet<FName> StringAssetPackages;
-					GRedirectCollector.ProcessSoftObjectPathPackageList(PackageName, false, StringAssetPackages);
-
-					for (FName StringAssetPackage : StringAssetPackages)
-					{
-						TMap<FSoftObjectPath, FSoftObjectPath> RedirectedPaths;
-
-						// If this is a redirector, extract destination from asset registry
-						if (ContainsRedirector(StringAssetPackage, RedirectedPaths))
-						{
-							for (TPair<FSoftObjectPath, FSoftObjectPath>& RedirectedPath : RedirectedPaths)
-							{
-								GRedirectCollector.AddAssetPathRedirection(RedirectedPath.Key, RedirectedPath.Value);
-								PackagesToLoad.Add(FPackageName::ObjectPathToPackageName(RedirectedPath.Value.ToString()));
-							}
-						}
-						else
-						{
-							PackagesToLoad.Add(StringAssetPackage.ToString());
-						}
-					}
-				}
-			}
-		}
-
-		{
-			UE_LOG(LogCook, Display, TEXT("Loading string referenced assets..."));
-			UE_SCOPED_HIERARCHICAL_COOKTIMER(FullLoadAndSave_LoadStringReferencedAssets);
-			GIsCookerLoadingPackage = true;
-			for (const FString& ToLoad : PackagesToLoad)
-			{
-				FName BuildFilenameFName = PackageDatas->GetFileNameByPackageName(FName(*ToLoad));
-				if (!PackageTracker->NeverCookPackageList.Contains(BuildFilenameFName))
-				{
-					{
-						LLM_SCOPE(ELLMTag::Untagged); // Reset the scope so that untagged memory in the package shows up as Untagged rather than Cooker
-						LoadPackage(nullptr, *ToLoad, LOAD_None);
-					}
-					if (GShaderCompilingManager)
-					{
-						GShaderCompilingManager->ProcessAsyncResults(true, false);
-					}
-					FAssetCompilingManager::Get().ProcessAsyncTasks(true);
-				}
-			}
-			GIsCookerLoadingPackage = false;
-		}
-	} while (PackagesToLoad.Num() > 0);
-
-	ProcessedPackages.Empty();
-
-	// When saving concurrently, before starting the concurrent saves, do tasks which are normally done in SavePackage that
-	// cannot be done from other threads or during FScopedSavingFlag in SavePackage.
-	if (bSaveConcurrent)
-	{
-		UE_LOG(LogCook, Display, TEXT("Flushing async loading..."));
-		UE_SCOPED_HIERARCHICAL_COOKTIMER(FullLoadAndSave_FlushAsyncLoading);
-		FlushAsyncLoading();
-
-		FPackageLocalizationManager::Get().ConditionalUpdateCache();
-	}
-
-	if (bSaveConcurrent)
-	{
-		UE_LOG(LogCook, Display, TEXT("Waiting for async tasks..."));
-		UE_SCOPED_HIERARCHICAL_COOKTIMER(FullLoadAndSave_ProcessThreadUntilIdle);
-		FTaskGraphInterface::Get().ProcessThreadUntilIdle(ENamedThreads::GameThread);
-	}
-
-	if (FAssetCompilingManager::Get().GetNumRemainingAssets())
-	{
-		UE_LOG(LogCook, Display, TEXT("Waiting for async compilation..."));
-		// Wait for all assets to finish compiling
-		FAssetCompilingManager::Get().FinishAllCompilation();
-	}
-
-	// Wait for all platform data to be loaded
-	{
-		UE_LOG(LogCook, Display, TEXT("Waiting for cooked platform data..."));
-		UE_SCOPED_HIERARCHICAL_COOKTIMER(FullLoadAndSave_WaitForCookedPlatformData);
-		while (ObjectsToWaitForCookedPlatformData.Num() > 0)
-		{
-			for (int32 ObjIdx = ObjectsToWaitForCookedPlatformData.Num() - 1; ObjIdx >= 0; --ObjIdx)
-			{
-				UObject* Obj = ObjectsToWaitForCookedPlatformData[ObjIdx];
-				bool bAllPlatformDataLoaded = true;
-				for (const ITargetPlatform* TargetPlatform : TargetPlatforms)
-				{
-					if (!RouteIsCachedCookedPlatformDataLoaded(Obj->GetPackage()->GetFName(), Obj, TargetPlatform))
-					{
-						bAllPlatformDataLoaded = false;
-						break;
-					}
-				}
-
-				if (bAllPlatformDataLoaded)
-				{
-					ObjectsToWaitForCookedPlatformData.RemoveAtSwap(ObjIdx, 1, false);
-				}
-			}
-
-			FAssetCompilingManager::Get().ProcessAsyncTasks(true);
-			double CurrentTime = FPlatformTime::Seconds();
-			if (LastCookableObjectTickTime + TickCookableObjectsFrameTime <= CurrentTime)
-			{
-				UE_SCOPED_COOKTIMER(TickCookableObjects);
-				FTickableCookObject::TickObjects(static_cast<float>(CurrentTime - LastCookableObjectTickTime), false /* bTickComplete */);
-				LastCookableObjectTickTime = CurrentTime;
-			}
-
-			FPlatformProcess::Sleep(0.001f);
-		}
-
-		ObjectsToWaitForCookedPlatformData.Empty();
-	}
-
-	{
-		UE_LOG(LogCook, Display, TEXT("Saving packages..."));
-		UE_SCOPED_HIERARCHICAL_COOKTIMER(FullLoadAndSave_Save);
-		check(bIsSavingPackage == false);
-		bIsSavingPackage = true;
-
-		if (bSaveConcurrent)
-		{
-			GIsSavingPackage = true;
-		}
-
-		TArray<FAssetRegistryGenerator*> Generators;
-		Generators.Reserve(TargetPlatforms.Num());
-		for (const ITargetPlatform* Target : TargetPlatforms)
-		{
-			Generators.Add(PlatformManager->GetPlatformData(Target)->RegistryGenerator.Get());
-		}
-		FCriticalSection AssetGeneratorLock;
-
-		int64 ParallelSavedPackages = 0;
-		ParallelFor(PackagesToSave.Num(),
-			[this, &PackagesToSave, &TargetPlatforms, &Generators, &ParallelSavedPackages, SaveFlags, bSaveConcurrent, &AssetGeneratorLock](int32 PackageIdx)
-		{
-			UE::Cook::FPackageData& PackageData = *PackagesToSave[PackageIdx];
-			UPackage* Package = PackageData.GetPackage();
-			check(Package);
-
-			// when concurrent saving is supported, precaching will need to be refactored for concurrency
-			if (!bSaveConcurrent)
-			{
-				// precache texture platform data ahead of save
-				const int32 PrecacheOffset = 512;
-				UPackage* PrecachePackage = PackageIdx + PrecacheOffset < PackagesToSave.Num() ? PackagesToSave[PackageIdx + PrecacheOffset]->GetPackage() : nullptr;
-				if (PrecachePackage)
-				{
-					TArray<UObject*> ObjsInPackage;
-					{
-						GetObjectsWithOuter(PrecachePackage, ObjsInPackage, false);
-					}
-
-					for (UObject* Obj : ObjsInPackage)
-					{
-						if (Obj->HasAnyFlags(RF_Transient) || !Obj->IsA(UTexture::StaticClass()))
-						{
-							continue;
-						}
-
-						for (const ITargetPlatform* TargetPlatform : TargetPlatforms)
-						{
-							RouteBeginCacheForCookedPlatformData(Obj->GetPackage()->GetFName(), Obj, TargetPlatform);
-						}
-					}
-				}
-			}
-
-			const FName& PackageFileName = PackageData.GetFileName();
-			if (!PackageFileName.IsNone())
-			{
-				// Use SandboxFile to do path conversion to properly handle sandbox paths (outside of standard paths in particular).
-				FString Filename = ConvertToFullSandboxPath(PackageFileName.ToString(), true);
-
-				// look for a world object in the package (if there is one, there's a map)
-				EObjectFlags FlagsToCook = RF_Public;
-				TArray<UObject*> ObjsInPackage;
-				UWorld* World = nullptr;
-				{
-					//UE_SCOPED_HIERARCHICAL_COOKTIMER(SaveCookedPackage_FindWorldInPackage);
-					GetObjectsWithOuter(Package, ObjsInPackage, false);
-					for (UObject* Obj : ObjsInPackage)
-					{
-						World = Cast<UWorld>(Obj);
-						if (World)
-						{
-							FlagsToCook = RF_NoFlags;
-							break;
-						}
-					}
-				}
-
-				const FName& PackageName = PackageData.GetPackageName();
-				FString PackageNameStr = PackageName.ToString();
-				bool bExcludeFromNonEditorTargets = IsCookFlagSet(ECookInitializationFlags::SkipEditorContent) && (PackageNameStr.StartsWith(TEXT("/Engine/Editor")) || PackageNameStr.StartsWith(TEXT("/Engine/VREditor")));
-
-				uint32 OriginalPackageFlags = Package->GetPackageFlags();
-
-				TArray<UE::Cook::ECookResult> SavePackageSuccessPerPlatform;
-				SavePackageSuccessPerPlatform.SetNum(TargetPlatforms.Num());
-				for (int32 PlatformIndex = 0; PlatformIndex < TargetPlatforms.Num(); ++PlatformIndex)
-				{
-					const ITargetPlatform* Target = TargetPlatforms[PlatformIndex];
-					FAssetRegistryGenerator& Generator = *Generators[PlatformIndex]; 
-
-
-					// don't save Editor resources from the Engine if the target doesn't have editoronly data
-					bool bCookPackage = (!bExcludeFromNonEditorTargets || Target->HasEditorOnlyData());
-					if (UAssetManager::IsValid() && !UAssetManager::Get().ShouldCookForPlatform(Package, Target))
-					{
-						bCookPackage = false;
-					}
-
-					if (bCookPackage)
-					{
-						FString PlatFilename = Filename.Replace(TEXT("[Platform]"), *Target->PlatformName());
-
-						UE_CLOG(GCookProgressDisplay & (int32)ECookProgressDisplayMode::PackageNames, LogCook, Display, TEXT("Cooking %s -> %s"), *Package->GetName(), *PlatFilename);
-
-						bool bSwap = (!Target->IsLittleEndian()) ^ (!PLATFORM_LITTLE_ENDIAN);
-						if (!Target->HasEditorOnlyData())
-						{
-							Package->SetPackageFlags(PKG_FilterEditorOnly);
-						}
-						else
-						{
-							Package->ClearPackageFlags(PKG_FilterEditorOnly);
-						}
-						
-						FArchiveCookContext CookContext(Package, FArchiveCookContext::ECookByTheBook);
-						FArchiveCookData CookData(*Target, CookContext);
-
-						GIsCookerLoadingPackage = true;
-						check(SavePackageContexts.Num() > PlatformIndex);
-						FSavePackageContext& SavePackageContext = SavePackageContexts[PlatformIndex]->SaveContext;
-						IPackageWriter::FBeginPackageInfo BeginInfo;
-						BeginInfo.PackageName = Package->GetFName();
-						BeginInfo.LooseFilePath = PlatFilename;
-						SavePackageContext.PackageWriter->BeginPackage(BeginInfo);
-						FSavePackageArgs SaveArgs;
-						SaveArgs.TopLevelFlags = FlagsToCook;
-						SaveArgs.bForceByteSwapping = bSwap;
-						SaveArgs.bWarnOfLongFilename = false;
-						SaveArgs.SaveFlags = SaveFlags;
-						SaveArgs.ArchiveCookData = &CookData;
-						SaveArgs.bSlowTask = false;
-						SaveArgs.SavePackageContext = &SavePackageContext;
-						FSavePackageResultStruct SaveResult = GEditor->Save(Package, World, *PlatFilename, SaveArgs);
-						GIsCookerLoadingPackage = false;
-
-						if (SaveResult == ESavePackageResult::Success && UAssetManager::IsValid())
-						{
-							if (!UAssetManager::Get().VerifyCanCookPackage(this, Package->GetFName()))
-							{
-								SaveResult = ESavePackageResult::Error;
-							}
-						}
-
-						// Flush the global AssetRegisty so any AssetData changes from the save are present
-						AssetRegistry->WaitForCompletion();
-						// Update cooked asset registry
-						FAssetPackageData* AssetPackageData = nullptr;
-						{
-							FScopeLock AssetGeneratorScopeLock(&AssetGeneratorLock);
-							UE::Cook::ECookResult CookResult = SaveResult == ESavePackageResult::Success ?
-								UE::Cook::ECookResult::Succeeded : UE::Cook::ECookResult::Failed;
-							Generator.UpdateAssetRegistryData(Package->GetFName(), Package, CookResult, &SaveResult,
-								MoveTemp(*CookContext.GetCookTagList()),
-								false /* bIncludeDiskOnlyAssets */, TOptional<FAssetPackageData>(),
-								TOptional<TArray<FAssetDependency>>());
-							AssetPackageData = Generator.GetAssetPackageData(Package->GetFName());
-						}
-						check(AssetPackageData);
-
-						ICookedPackageWriter::FCommitPackageInfo CommitInfo;
-						if (SaveResult.IsSuccessful())
-						{
-							CommitInfo.Status = IPackageWriter::ECommitStatus::Success;
-						}
-						else
-						{
-							CommitInfo.Status = IPackageWriter::ECommitStatus::Error;
-						}
-						CommitInfo.PackageName = Package->GetFName();
-						PRAGMA_DISABLE_DEPRECATION_WARNINGS;
-						CommitInfo.PackageGuid = AssetPackageData->PackageGuid;
-						PRAGMA_ENABLE_DEPRECATION_WARNINGS;
-						CommitInfo.WriteOptions = IPackageWriter::EWriteOptions::Write | IPackageWriter::EWriteOptions::ComputeHash;
-						SavePackageContext.PackageWriter->CommitPackage(MoveTemp(CommitInfo));
-
-						if (SaveResult.IsSuccessful())
-						{
-							FPlatformAtomics::InterlockedIncrement(&ParallelSavedPackages);
-						}
-
-						if (SaveResult != ESavePackageResult::ReferencedOnlyByEditorOnlyData)
-						{
-							SavePackageSuccessPerPlatform[PlatformIndex] = UE::Cook::ECookResult::Succeeded;
-						}
-						else
-						{
-							SavePackageSuccessPerPlatform[PlatformIndex] = UE::Cook::ECookResult::Failed;
-						}
-					}
-					else
-					{
-						SavePackageSuccessPerPlatform[PlatformIndex] = UE::Cook::ECookResult::Failed;
-					}
-				}
-
-				PackageData.SetPlatformsCooked(TargetPlatforms, SavePackageSuccessPerPlatform);
-
-				if (SavePackageSuccessPerPlatform.Contains(UE::Cook::ECookResult::Failed))
-				{
-					PackageTracker->UncookedEditorOnlyPackages.Add(PackageName);
-				}
-
-				Package->SetPackageFlagsTo(OriginalPackageFlags);
-			}
-		}, !bSaveConcurrent);
-
-		if (bSaveConcurrent)
-		{
-			GIsSavingPackage = false;
-		}
-
-		if (ParallelSavedPackages > 0)
-		{
-			Result |= COSR_CookedPackage;
-		}
-
-		check(bIsSavingPackage == true);
-		bIsSavingPackage = false;
-	}
-
-	if (bSaveConcurrent)
-	{
-		UE_LOG(LogCook, Display, TEXT("Calling PostSaveRoot on worlds..."));
-		UE_SCOPED_HIERARCHICAL_COOKTIMER(FullLoadAndSave_PostSaveRoot);
-		for (const TPair<UWorld*,TArray<bool>>& WorldIt : WorldsToPostSaveRoot)
-		{
-			UWorld* World = WorldIt.Key;
-			const TArray<bool>& PlatformNeedsCleanup = WorldIt.Value;
-			check(World);
-			check(PlatformNeedsCleanup.Num() == TargetPlatforms.Num());
-			for (int PlatformIndex = TargetPlatforms.Num() - 1; PlatformIndex >= 0; --PlatformIndex)
-			{
-				ObjectSaveContext.TargetPlatform = TargetPlatforms[PlatformIndex];
-				ObjectSaveContext.bOuterConcurrentSave = PlatformIndex == 0;
-				UE::SavePackageUtilities::CallPostSaveRoot(World, ObjectSaveContext, PlatformNeedsCleanup[PlatformIndex]);
-			}
-		}
-	}
-
-	CookByTheBookFinished();
-	return Result;
-}
-
 ICookedPackageWriter& UCookOnTheFlyServer::FindOrCreatePackageWriter(const ITargetPlatform* TargetPlatform)
 {
 	return *FindOrCreateSaveContext(TargetPlatform).PackageWriter;
diff --git a/Engine/Source/Editor/UnrealEd/Private/Cooker/CookTypes.cpp b/Engine/Source/Editor/UnrealEd/Private/Cooker/CookTypes.cpp
--- a/Engine/Source/Editor/UnrealEd/Private/Cooker/CookTypes.cpp
+++ b/Engine/Source/Editor/UnrealEd/Private/Cooker/CookTypes.cpp
@@ -410,7 +410,6 @@ FCbWriter& operator<<(FCbWriter& Writer, const UE::Cook::FCookByTheBookOptions&
 	Writer << "AllowUncookedAssetReferences" << Value.bAllowUncookedAssetReferences;
 	Writer << "SkipHardReferences" << Value.bSkipHardReferences;
 	Writer << "SkipSoftReferences" << Value.bSkipSoftReferences;
-	Writer << "FullLoadAndSave" << Value.bFullLoadAndSave;
 	Writer << "CookAgainstFixedBase" << Value.bCookAgainstFixedBase;
 	Writer << "DlcLoadMainAssetRegistry" << Value.bDlcLoadMainAssetRegistry;
 	Writer.EndObject();
@@ -440,7 +439,6 @@ bool LoadFromCompactBinary(FCbFieldView Field, UE::Cook::FCookByTheBookOptions&
 	bOk = LoadFromCompactBinary(Field["AllowUncookedAssetReferences"], OutValue.bAllowUncookedAssetReferences) & bOk;
 	bOk = LoadFromCompactBinary(Field["SkipHardReferences"], OutValue.bSkipHardReferences) & bOk;
 	bOk = LoadFromCompactBinary(Field["SkipSoftReferences"], OutValue.bSkipSoftReferences) & bOk;
-	bOk = LoadFromCompactBinary(Field["FullLoadAndSave"], OutValue.bFullLoadAndSave) & bOk;
 	bOk = LoadFromCompactBinary(Field["CookAgainstFixedBase"], OutValue.bCookAgainstFixedBase) & bOk;
 	bOk = LoadFromCompactBinary(Field["DlcLoadMainAssetRegistry"], OutValue.bDlcLoadMainAssetRegistry) & bOk;
 
diff --git a/Engine/Source/Editor/UnrealEd/Private/Cooker/CookTypes.h b/Engine/Source/Editor/UnrealEd/Private/Cooker/CookTypes.h
--- a/Engine/Source/Editor/UnrealEd/Private/Cooker/CookTypes.h
+++ b/Engine/Source/Editor/UnrealEd/Private/Cooker/CookTypes.h
@@ -423,7 +423,6 @@ public:
 	bool							bAllowUncookedAssetReferences = false; // this is a flag for dlc, will allow DLC to be cook when the fixed base might be missing references.
 	bool							bSkipHardReferences = false;
 	bool							bSkipSoftReferences = false;
-	bool							bFullLoadAndSave = false;
 	bool							bCookAgainstFixedBase = false;
 	bool							bDlcLoadMainAssetRegistry = false;
 
