FEditorFileUtils::IsMapPackageAsset will now return true for external actor packages in plugin mount points.

diff --git a/Engine/Source/Editor/UnrealEd/Private/FileHelpers.cpp b/Engine/Source/Editor/UnrealEd/Private/FileHelpers.cpp
--- a/Engine/Source/Editor/UnrealEd/Private/FileHelpers.cpp
+++ b/Engine/Source/Editor/UnrealEd/Private/FileHelpers.cpp
@@ -5100,11 +5100,15 @@ bool FEditorFileUtils::IsMapPackageAsset(const FString& ObjectPath, FString& Map
 				MapFilePath = PackagePath;
 				return true;
 			}
-			const FString ExternalPackagePath = FString(TEXT("/Game")) / FPackagePath::GetExternalActorsFolderName();
-			if ( FPackageName::FilenameToLongPackageName(PackagePath).StartsWith(ExternalPackagePath) )
+
+			TArray<FString> ObjectPathParts;
+			if (ObjectPath.ParseIntoArray(ObjectPathParts, TEXT("/")) > 1)
 			{
-				MapFilePath = PackagePath;
-				return true;
+				if (ObjectPathParts[1] == FPackagePath::GetExternalActorsFolderName())
+				{
+					MapFilePath = PackagePath;
+					return true;
+				}
 			}
 		}
 	}
