Fixes blank modeling tab upon opening it via
Window->EditorModes

diff --git a/Engine/Source/Editor/UnrealEd/Private/Toolkits/BaseToolkit.cpp b/Engine/Source/Editor/UnrealEd/Private/Toolkits/BaseToolkit.cpp
--- a/Engine/Source/Editor/UnrealEd/Private/Toolkits/BaseToolkit.cpp
+++ b/Engine/Source/Editor/UnrealEd/Private/Toolkits/BaseToolkit.cpp
@@ -549,7 +549,13 @@ void FModeToolkit::UpdatePrimaryModePanel()
 			if (HasToolkitBuilder())
 			{
 				TabPtr->SetParentDockTabStackTabWellHidden(HasToolkitBuilder());
-				return;				
+				const TSharedPtr<SWidget> Content = GetInlineContent() ;
+				
+				if ( Content && InlineContentHolder.IsValid() )
+				{
+					InlineContentHolder->SetContent( Content.ToSharedRef() );
+				}
+				return;
 			}
 		}
 
