Speculative Fix Selection Crash on Garbage Objects

diff --git a/Engine/Source/Editor/UnrealEd/Private/Selection.cpp b/Engine/Source/Editor/UnrealEd/Private/Selection.cpp
--- a/Engine/Source/Editor/UnrealEd/Private/Selection.cpp
+++ b/Engine/Source/Editor/UnrealEd/Private/Selection.cpp
@@ -235,8 +235,8 @@ void USelection::OnElementListSyncEvent(const FTypedElementList& InElementList,
 
 void USelection::Select(UObject* InObject)
 {
-	check(InObject && IsValidObjectToSelect(InObject));
-	if (ElementSelectionSet)
+	check(InObject);
+	if (IsValidObjectToSelect(InObject) && ElementSelectionSet)
 	{
 		FTypedElementHandle ElementHandle = SelectionElementBridge->GetElementHandleForObject(InObject);
 		check(ElementHandle);
