PR #10147: UE-171741: Mark camera bookmark as modified (Contributed by [REDACTED_PERSON])

diff --git a/Engine/Source/Editor/UnrealEd/Public/Bookmarks/BookMarkTypeActions.h b/Engine/Source/Editor/UnrealEd/Public/Bookmarks/BookMarkTypeActions.h
--- a/Engine/Source/Editor/UnrealEd/Public/Bookmarks/BookMarkTypeActions.h
+++ b/Engine/Source/Editor/UnrealEd/Public/Bookmarks/BookMarkTypeActions.h
@@ -26,6 +26,8 @@ public:
 		{
 			if (UWorld * const World = InViewportClient.GetWorld())
 			{
+				Bookmark->Modify();
+
 				// Use the rotation from the first perspective viewport can find.
 				FRotator Rotation(0, 0, 0);
 				if (!InViewportClient.IsOrtho())
