Localization: When 2 paths have the same fuzz rating during fuzzy path matching, give precedence to the path that excludes  the path.
This is primarily used to handle a case where a path is both in the included and excluded filters when gathering source and assets.

diff --git a/Engine/Source/Editor/UnrealEd/Private/Commandlets/GatherTextCommandletBase.cpp b/Engine/Source/Editor/UnrealEd/Private/Commandlets/GatherTextCommandletBase.cpp
--- a/Engine/Source/Editor/UnrealEd/Private/Commandlets/GatherTextCommandletBase.cpp
+++ b/Engine/Source/Editor/UnrealEd/Private/Commandlets/GatherTextCommandletBase.cpp
@@ -210,6 +210,11 @@ FFuzzyPathMatcher::FFuzzyPathMatcher(const TArray<FString>& InIncludePathFilters
 
 		const int32 PathOneFuzzRating = GetFuzzRating(PathOne);
 		const int32 PathTwoFuzzRating = GetFuzzRating(PathTwo);
+		if (PathOneFuzzRating == PathTwoFuzzRating)
+		{
+			// In the case of a tie, allow an exclusion to take priority
+			return (uint8)PathOne.PathType > (uint8)PathTwo.PathType;
+		}
 		return PathOneFuzzRating < PathTwoFuzzRating;
 	});
 }
