SoftGC: Revert the new bSupportCookerSoftGC functionality in case it is the cause of new cook crashes.
#rnx

diff --git a/Engine/Source/Editor/UnrealEd/Private/Commandlets/CookCommandlet.cpp b/Engine/Source/Editor/UnrealEd/Private/Commandlets/CookCommandlet.cpp
--- a/Engine/Source/Editor/UnrealEd/Private/Commandlets/CookCommandlet.cpp
+++ b/Engine/Source/Editor/UnrealEd/Private/Commandlets/CookCommandlet.cpp
@@ -641,7 +641,6 @@ void UCookCommandlet::ConditionalCollectGarbage(uint32 TickResults, UCookOnTheFl
 
 	UE_LOG(LogCookCommandlet, Display, TEXT("GarbageCollection...%s (%s)"), *GCType, *GCReason);
 	{
-		TGuardValue<bool> SoftGCGuard(UPackage::bSupportCookerSoftGC, true);
 		CollectGarbage(RF_NoFlags);
 	}
 
diff --git a/Engine/Source/Editor/UnrealEd/Private/CookOnTheFlyServer.cpp b/Engine/Source/Editor/UnrealEd/Private/CookOnTheFlyServer.cpp
--- a/Engine/Source/Editor/UnrealEd/Private/CookOnTheFlyServer.cpp
+++ b/Engine/Source/Editor/UnrealEd/Private/CookOnTheFlyServer.cpp
@@ -4994,20 +4994,21 @@ void UCookOnTheFlyServer::PreGarbageCollect()
 			});
 	}
 
-	// Add packages to GCKeepObjects. 
+	// Add packages and all RF_Public objects outered to them to GCKeepObjects
 	TArray<UObject*> ObjectsWithOuter;
 	for (UPackage* Package : GCKeepPackages)
 	{
 		GCKeepObjects.Add(Package);
+		ObjectsWithOuter.Reset();
+		GetObjectsWithOuter(Package, ObjectsWithOuter);
+		for (UObject* Obj : ObjectsWithOuter)
+		{
+			if (IsValidChecked(Obj) && Obj->HasAnyFlags(RF_Public))
+			{
+				GCKeepObjects.Add(Obj);
+			}
+		}
 	}
-	for (FPackageData* PackageData : GCKeepPackageDatas)
-	{
-		PackageData->SetKeepReferencedDuringGC(true);
-	}
-
-	// Add all public objects within every package in memory to the UPackage::SoftGCPackageToObjectList container,
-	// so they will be kept in memory if the package is kept in memory.
-	ConstructSoftGCPackageToObjectList(this->SoftGCPackageToObjectListBuffer);
 }
 
 void UCookOnTheFlyServer::CookerAddReferencedObjects(FReferenceCollector& Collector)
@@ -5049,8 +5050,6 @@ void UCookOnTheFlyServer::PostGarbageCollect()
 	check(!SavingPackageData || SavingPackageData->GetPackage() != nullptr);
 
 	GCKeepObjects.Empty();
-	UPackage::SoftGCPackageToObjectList.Empty();
-	SoftGCPackageToObjectListBuffer.Empty();
 
 	PackageDatas->LockAndEnumeratePackageDatas([](FPackageData* PackageData)
 	{
