#UE Readd code that was accidentally removed related to GCKeepPackageDatas

diff --git a/Engine/Source/Editor/UnrealEd/Private/CookOnTheFlyServer.cpp b/Engine/Source/Editor/UnrealEd/Private/CookOnTheFlyServer.cpp
--- a/Engine/Source/Editor/UnrealEd/Private/CookOnTheFlyServer.cpp
+++ b/Engine/Source/Editor/UnrealEd/Private/CookOnTheFlyServer.cpp
@@ -5009,6 +5009,10 @@ void UCookOnTheFlyServer::PreGarbageCollect()
 			}
 		}
 	}
+	for (FPackageData* PackageData : GCKeepPackageDatas)
+	{
+		PackageData->SetKeepReferencedDuringGC(true);
+	}
 }
 
 void UCookOnTheFlyServer::CookerAddReferencedObjects(FReferenceCollector& Collector)
