SoftGC: Move SoftGCPackageToObjectList.Empty to earlier in the function for better maintainability.
#rnx

diff --git a/Engine/Source/Editor/UnrealEd/Private/CookOnTheFlyServer.cpp b/Engine/Source/Editor/UnrealEd/Private/CookOnTheFlyServer.cpp
--- a/Engine/Source/Editor/UnrealEd/Private/CookOnTheFlyServer.cpp
+++ b/Engine/Source/Editor/UnrealEd/Private/CookOnTheFlyServer.cpp
@@ -4728,6 +4728,7 @@ static void ConstructSoftGCPackageToObjectList(TArray<UObject*>& PackageToObject
 	};
 
 	PackageToObjectListBuffer.Empty();
+	UPackage::SoftGCPackageToObjectList.Empty();
 
 	// Iterate over all UObjects in memory (in parallel) and for each valid public object, get its package and add a FPackageObjectPair for it
 	int32 MaxNumberOfObjects = GUObjectArray.GetObjectArrayNum();
@@ -4799,7 +4800,7 @@ static void ConstructSoftGCPackageToObjectList(TArray<UObject*>& PackageToObject
 	PackageToObjectListBuffer.SetNum(PackageObjectPairsNum);
 	UObject** PackageToObjectListBufferPtr = PackageToObjectListBuffer.GetData();
 
-	UPackage::SoftGCPackageToObjectList.Empty(PackagesNum);
+	UPackage::SoftGCPackageToObjectList.Reserve(PackagesNum);
 	int32 PreviousPackageStartIndex = 0;
 	UPackage* PreviousPackage = nullptr;
 	for (int32 Index = 0; Index < PackageObjectPairsNum; ++Index)
