added LLM_SCOPE_BYTAG for shutdown
https://p4-swarm.epicgames.net/reviews/24503865

diff --git a/Engine/Source/Editor/UnrealEd/Private/UnrealEdGlobals.cpp b/Engine/Source/Editor/UnrealEd/Private/UnrealEdGlobals.cpp
--- a/Engine/Source/Editor/UnrealEd/Private/UnrealEdGlobals.cpp
+++ b/Engine/Source/Editor/UnrealEd/Private/UnrealEdGlobals.cpp
@@ -235,6 +235,7 @@ int32 EditorReinit()
 void EditorExit()
 {
 	TRACE_CPUPROFILER_EVENT_SCOPE(EditorExit);
+	LLM_SCOPE(ELLMTag::EngineMisc);
 
 	// Save out any config settings for the editor so they don't get lost
 	GEditor->SaveConfig();
