MPCook logging:
Fix CrashDiagnostics to open the crashed CookWorker logfile with FILEREAD_AllowWrite so that it will succeed even if the crashed process still has it open.
Fix typo in log interpretation CookWorkerHeatbeat -> CookWorkerHeartbeat
Fix spurious sending of crash diagnostics when CookWorker cleanly shuts down.
#rrb Zousar.Shaker
#rnx

diff --git a/Engine/Source/Editor/UnrealEd/Private/Cooker/CookWorkerClient.cpp b/Engine/Source/Editor/UnrealEd/Private/Cooker/CookWorkerClient.cpp
--- a/Engine/Source/Editor/UnrealEd/Private/Cooker/CookWorkerClient.cpp
+++ b/Engine/Source/Editor/UnrealEd/Private/Cooker/CookWorkerClient.cpp
@@ -827,7 +827,7 @@ void FCookWorkerClient::HandleHeartbeatMessage(FMPCollectorClientMessageContext&
 		return;
 	}
 
-	UE_LOG(LogCook, Display, TEXT("CookWorkerHeartbeat: %d"), Message.HeartbeatNumber);
+	UE_LOG(LogCook, Display, TEXT("%.*s %d"), HeartbeatCategoryText.Len(), HeartbeatCategoryText.GetData(), Message.HeartbeatNumber);
 	SendMessage(FHeartbeatMessage(Message.HeartbeatNumber));
 }
 
diff --git a/Engine/Source/Editor/UnrealEd/Private/Cooker/CookWorkerServer.cpp b/Engine/Source/Editor/UnrealEd/Private/Cooker/CookWorkerServer.cpp
--- a/Engine/Source/Editor/UnrealEd/Private/Cooker/CookWorkerServer.cpp
+++ b/Engine/Source/Editor/UnrealEd/Private/Cooker/CookWorkerServer.cpp
@@ -67,16 +67,18 @@ void FCookWorkerServer::SendCrashDiagnostics()
 	FString LogFileName = Director.GetWorkerLogFileName(ProfileId);
 	UE_LOG(LogCook, Display, TEXT("CookWorker %d log messages written after communication loss:"), ProfileId);
 	FString LogText;
-	if (!FFileHelper::LoadFileToString(LogText, *LogFileName))
+	int32 ReadFlags = FILEREAD_AllowWrite; // To be able to open a file for read that might be open for write from another process, we have to specify FILEREAD_AllowWrite
+	if (!FFileHelper::LoadFileToString(LogText, *LogFileName, FFileHelper::EHashOptions::None, ReadFlags))
 	{
 		UE_LOG(LogCook, Warning, TEXT("No log file found for CookWorker %d."), ProfileId);
 	}
 	else
 	{
-		FString LastSentHeartbeat = FString::Printf(TEXT("CookWorkerHeartbeat: %d"), LastReceivedHeartbeatNumber);
+		FString LastSentHeartbeat = FString::Printf(TEXT("%.*s %d"), HeartbeatCategoryText.Len(), HeartbeatCategoryText.GetData(),
+			LastReceivedHeartbeatNumber);
 		int32 StartIndex = INDEX_NONE;
 		for (FStringView MarkerText : { FStringView(LastSentHeartbeat),
-			TEXTVIEW("CookWorkerHeartbeat: "), TEXTVIEW("Connection to CookDirector successful") })
+			HeartbeatCategoryText, TEXTVIEW("Connection to CookDirector successful") })
 		{
 			StartIndex = UE::String::FindLast(LogText, MarkerText);
 			if (StartIndex >= 0)
@@ -554,11 +556,12 @@ void FCookWorkerServer::HandleReceiveMessagesInternal()
 		if (PeekMessage.MessageType == FAbortWorkerMessage::MessageType)
 		{
 			UE::CompactBinaryTCP::FMarshalledMessage Message = ReceiveMessages.PopFrontValue();
-			UE_CLOG(ConnectStatus != EConnectStatus::PumpingCookComplete && ConnectStatus != EConnectStatus::WaitForDisconnect,
-				LogCook, Error, TEXT("CookWorkerCrash: %d remote process shut down unexpectedly. Assigned packages will be returned to the director."),
-				ProfileId);
-
-			bNeedCrashDiagnostics = true;
+			if (ConnectStatus != EConnectStatus::PumpingCookComplete && ConnectStatus != EConnectStatus::WaitForDisconnect)
+			{
+				UE_LOG(LogCook, Error, TEXT("CookWorkerCrash: %d remote process shut down unexpectedly. Assigned packages will be returned to the director."),
+					ProfileId);
+				bNeedCrashDiagnostics = true;
+			}
 			SendMessageInLock(FAbortWorkerMessage(FAbortWorkerMessage::AbortAcknowledge));
 			SendToState(EConnectStatus::WaitForDisconnect);
 			ReceiveMessages.Reset();
@@ -1082,7 +1085,7 @@ void FLogMessagesMessageHandler::ServerReceiveMessage(FMPCollectorServerMessageC
 
 	for (FReplicatedLogData& LogData : Messages)
 	{
-		if (LogData.Category == LogCookName && LogData.Message.Contains(TEXT("CookWorkerHeatbeat:")))
+		if (LogData.Category == LogCookName && LogData.Message.Contains(HeartbeatCategoryText))
 		{
 			// Do not spam heartbeat messages into the CookDirector log
 			continue;
diff --git a/Engine/Source/Editor/UnrealEd/Private/Cooker/CookWorkerServer.h b/Engine/Source/Editor/UnrealEd/Private/Cooker/CookWorkerServer.h
--- a/Engine/Source/Editor/UnrealEd/Private/Cooker/CookWorkerServer.h
+++ b/Engine/Source/Editor/UnrealEd/Private/Cooker/CookWorkerServer.h
@@ -383,6 +383,7 @@ public:
 	int32 HeartbeatNumber;
 	static FGuid MessageType;
 };
+constexpr FStringView HeartbeatCategoryText(TEXTVIEW("CookWorkerHeartbeat:"));
 
 class FPackageWriterMPCollector : public UE::Cook::IMPCollector
 {
