UX Improvements for the "Reopen asset tabs" notification that shows up when you launch the editor

diff --git a/Engine/Source/Editor/MainFrame/Private/Menus/LayoutsMenu.cpp b/Engine/Source/Editor/MainFrame/Private/Menus/LayoutsMenu.cpp
--- a/Engine/Source/Editor/MainFrame/Private/Menus/LayoutsMenu.cpp
+++ b/Engine/Source/Editor/MainFrame/Private/Menus/LayoutsMenu.cpp
@@ -801,7 +801,7 @@ void FLayoutsMenuLoad::ReloadCurrentLayout()
 	{
 		// Save open asset editors + disable manual saving (no need to close them with AssetEditorSubsystem->CloseAllAssetEditors)
 		AssetEditorSubsystem->SaveOpenAssetEditors(true);
-		AssetEditorSubsystem->SetAutoRestoreAndDisableSaving(true);
+		AssetEditorSubsystem->SetAutoRestoreAndDisableSavingOverride(true);
 	}
 	FUnrealEdMisc::Get().AllowSavingLayoutOnClose(false);
 	// Editor is reset on-the-fly
@@ -813,7 +813,7 @@ void FLayoutsMenuLoad::ReloadCurrentLayout()
 		check(AssetEditorSubsystem);
 		// Restore asset editors + disable manual saving and avoid trying to re-open the asset editors twice
 		AssetEditorSubsystem->RestorePreviouslyOpenAssets();
-		AssetEditorSubsystem->SetAutoRestoreAndDisableSaving(false);
+		AssetEditorSubsystem->SetAutoRestoreAndDisableSavingOverride(TOptional<bool>());
 	}
 	// Save layout and create duplicated ini file (DuplicatedEditorLayoutIniFilePath)
 	FPrivateLayoutsMenu::SaveLayoutWithoutRemovingTempLayoutFiles();
diff --git a/Engine/Source/Editor/UnrealEd/Classes/Settings/EditorLoadingSavingSettings.h b/Engine/Source/Editor/UnrealEd/Classes/Settings/EditorLoadingSavingSettings.h
--- a/Engine/Source/Editor/UnrealEd/Classes/Settings/EditorLoadingSavingSettings.h
+++ b/Engine/Source/Editor/UnrealEd/Classes/Settings/EditorLoadingSavingSettings.h
@@ -29,6 +29,18 @@ enum class EAutoSaveMethod : uint8
 	BackupAndOverwrite,
 };
 
+UENUM()
+enum class ERestoreOpenAssetTabsMethod : uint8
+{
+	/** Always prompt the user if they want to restore previously opened asset tabs on launch */
+	AlwaysPrompt = 0,
+	
+	/** Always automatically restore any previously opened asset tabs on launch */
+	AlwaysRestore = 1,
+
+	/** Never restore previously opened asset tabs on launch */
+	NeverRestore = 2
+};
 
 /** A filter used by the auto reimport manager to explicitly include/exclude files matching the specified wildcard */
 USTRUCT()
@@ -95,10 +107,16 @@ public:
 	UPROPERTY(EditAnywhere, config, Category=Startup)
 	uint32 bForceCompilationAtStartup:1;
 
-	/** Whether to restore previously open assets at startup */
+	/** Whether to restore previously open assets at startup after a clean shutdown */
 	UPROPERTY(EditAnywhere, config, Category=Startup)
-	uint32 bRestoreOpenAssetTabsOnRestart:1;
+	ERestoreOpenAssetTabsMethod RestoreOpenAssetTabsOnRestart = ERestoreOpenAssetTabsMethod::AlwaysPrompt;
 
+#if WITH_EDITOR
+	/** Whether to restore previously open assets at startup */
+	UPROPERTY(meta = (DeprecatedProperty, DeprecationMessage = "Use RestoreOpenAssetTabsOnRestart instead"))
+	uint32 bRestoreOpenAssetTabsOnRestart_DEPRECATED:1;
+#endif
+	
 private:
 
 	UPROPERTY(config)
diff --git a/Engine/Source/Editor/UnrealEd/Private/Subsystems/AssetEditorSubsystem.cpp b/Engine/Source/Editor/UnrealEd/Private/Subsystems/AssetEditorSubsystem.cpp
--- a/Engine/Source/Editor/UnrealEd/Private/Subsystems/AssetEditorSubsystem.cpp
+++ b/Engine/Source/Editor/UnrealEd/Private/Subsystems/AssetEditorSubsystem.cpp
@@ -44,7 +44,6 @@ DEFINE_LOG_CATEGORY_STATIC(LogAssetEditorSubsystem, Log, All);
 UAssetEditorSubsystem::UAssetEditorSubsystem()
 	: Super()
 	, bSavingOnShutdown(false)
-	, bAutoRestoreAndDisableSaving(false)
 	, bRequestRestorePreviouslyOpenAssets(false)
 {
 }
@@ -770,92 +769,101 @@ void UAssetEditorSubsystem::RestorePreviouslyOpenAssets()
 	bool bCleanShutdown = true;
 	GConfig->GetBool(TEXT("AssetEditorSubsystem"), TEXT("CleanShutdown"), bCleanShutdown, GEditorPerProjectIni);
 
+	bool bDebuggerAttachedLastSession = false;
+	GConfig->GetBool(TEXT("AssetEditorSubsystem"), TEXT("DebuggerAttached"), bDebuggerAttachedLastSession, GEditorPerProjectIni);
+	
 	SaveOpenAssetEditors(false);
 
+	/** True if the last editor run crashed and did not have a debugger attached
+	 * A "clean" shutdown for our purposes is the logical NOT of this, which includes clean shutdowns without a debugger
+	 * along with any shutdowns when a debugger is attached
+	 */
+	bool bCrashedWithoutDebugger = !bDebuggerAttachedLastSession && !bCleanShutdown;
+
 	if (OpenAssets.Num() > 0)
 	{
-		if (bCleanShutdown)
+		// This option overrides the saved setting
+		if(bAutoRestoreAndDisableSaving.IsSet())
 		{
-			// Do we have permission to automatically re-open the assets, or should we ask?
-			const bool bAutoRestore = GetDefault<UEditorLoadingSavingSettings>()->bRestoreOpenAssetTabsOnRestart || bAutoRestoreAndDisableSaving;
-
-			if (bAutoRestore)
+			// If bAutoRestoreAndDisableSaving is true, we automatically restore the opened assets
+			if(bAutoRestoreAndDisableSaving.GetValue())
 			{
-				// Pretend that we showed the notification and that the user clicked "Restore Now"
 				OpenEditorsForAssets(OpenAssets);
 			}
-			else
-			{
-				// Has this notification previously been suppressed by the user?
-				bool bSuppressNotification = false;
-				GConfig->GetBool(TEXT("AssetEditorSubsystem"), TEXT("SuppressRestorePreviouslyOpenAssetsNotification"), bSuppressNotification, GEditorPerProjectIni);
-
-				if (!bSuppressNotification)
-				{
-					// Ask the user; this doesn't block so will reopen the assets later
-					SpawnRestorePreviouslyOpenAssetsNotification(bCleanShutdown, OpenAssets);
-				}
-			}
+			return;
 		}
-		else
+		
+		/* If we crashed without a debugger attached, always prompt regardless of what the user previously said
+		 * to make sure the user can never get stuck in a crash loop due to corrupted assets etc
+		 */
+		if(bCrashedWithoutDebugger)
 		{
-			// If we crashed, we always ask regardless of what the user previously said
-			SpawnRestorePreviouslyOpenAssetsNotification(bCleanShutdown, OpenAssets);
+			SpawnRestorePreviouslyOpenAssetsNotification(!bCrashedWithoutDebugger, OpenAssets);
+			return;
+		}
+
+		const ERestoreOpenAssetTabsMethod AutoRestoreMethod = GetDefault<UEditorLoadingSavingSettings>()->RestoreOpenAssetTabsOnRestart;
+
+		switch(AutoRestoreMethod)
+		{
+		case ERestoreOpenAssetTabsMethod::AlwaysPrompt:
+			{
+				SpawnRestorePreviouslyOpenAssetsNotification(!bCrashedWithoutDebugger, OpenAssets);
+				break;
+			}
+			
+		case ERestoreOpenAssetTabsMethod::NeverRestore:
+			// Do nothing here since the user does not want to restore anything
+			break;
+		case ERestoreOpenAssetTabsMethod::AlwaysRestore:
+			{
+				// Pretend that we showed the notification and that the user clicked "Restore Now"
+				OpenEditorsForAssets(OpenAssets);
+				break;
+			}
 		}
 	}
 }
 
 void UAssetEditorSubsystem::SetAutoRestoreAndDisableSaving(const bool bInAutoRestoreAndDisableSaving)
 {
-	bAutoRestoreAndDisableSaving = bInAutoRestoreAndDisableSaving;
+	// We preserve legacy behavior, where true is the same but false translates to not having the override set in the new logic
+	SetAutoRestoreAndDisableSavingOverride(bInAutoRestoreAndDisableSaving ? bInAutoRestoreAndDisableSaving : TOptional<bool>());
+}
 
+void UAssetEditorSubsystem::SetAutoRestoreAndDisableSavingOverride(TOptional<bool> bInAutoRestoreAndDisableSaving)
+{
+	bAutoRestoreAndDisableSaving = bInAutoRestoreAndDisableSaving;
+	
 	// Disable any pending request to avoid trying to restore previously opened assets twice
 	bRequestRestorePreviouslyOpenAssets = false;
+
 }
 
 void UAssetEditorSubsystem::SpawnRestorePreviouslyOpenAssetsNotification(const bool bCleanShutdown, const TArray<FString>& AssetsToOpen)
 {
-	/** Utility functions for notifications */
-	struct Local
-	{
-		static ECheckBoxState GetDontAskAgainCheckBoxState()
-		{
-			bool bSuppressNotification = false;
-			GConfig->GetBool(TEXT("AssetEditorSubsystem"), TEXT("SuppressRestorePreviouslyOpenAssetsNotification"), bSuppressNotification, GEditorPerProjectIni);
-			return bSuppressNotification ? ECheckBoxState::Checked : ECheckBoxState::Unchecked;
-		}
-
-		static void OnDontAskAgainCheckBoxStateChanged(ECheckBoxState NewState)
-		{
-			const bool bSuppressNotification = (NewState == ECheckBoxState::Checked);
-			GConfig->SetBool(TEXT("AssetEditorSubsystem"), TEXT("SuppressRestorePreviouslyOpenAssetsNotification"), bSuppressNotification, GEditorPerProjectIni);
-		}
-	};
-
 	FText NotificationMessage = bCleanShutdown
-		? LOCTEXT("ReopenAssetEditorsAfterClose", "{0} asset {0}|plural(one=editor was,other=editors were) open when the editor was last closed. Would you like to re-open them?")
-		: LOCTEXT("ReopenAssetEditorsAfterCrash", "{0} asset {0}|plural(one=editor was,other=editors were) open when the editor quit unexpectedly. Would you like to re-open them?");
+		? LOCTEXT("ReopenAssetEditorsAfterClose", "{0} asset {0}|plural(one=editor was,other=editors were) open when the editor was last closed. Would you like to re-open {0}|plural(one=it,other=them)?")
+		: LOCTEXT("ReopenAssetEditorsAfterCrash", "{0} asset {0}|plural(one=editor was,other=editors were) open when the editor quit unexpectedly. Would you like to re-open {0}|plural(one=it,other=them)?");
 	NotificationMessage = FText::Format(NotificationMessage, AssetsToOpen.Num());
 
 	FNotificationInfo Info = FNotificationInfo(NotificationMessage);
 
 	// Add the buttons
 	Info.ButtonDetails.Add(FNotificationButtonInfo(
-		LOCTEXT("ReopenAssetEditors_Confirm", "Open"),
+		LOCTEXT("ReopenAssetEditors_Confirm", "Yes"),
 		FText(),
 		FSimpleDelegate::CreateUObject(this, &UAssetEditorSubsystem::OnConfirmRestorePreviouslyOpenAssets, AssetsToOpen),
 		SNotificationItem::CS_None
 	));
 	Info.ButtonDetails.Add(FNotificationButtonInfo(
-		LOCTEXT("ReopenAssetEditors_Cancel", "Cancel"),
+		LOCTEXT("ReopenAssetEditors_Cancel", "No"),
 		FText(),
 		FSimpleDelegate::CreateUObject(this, &UAssetEditorSubsystem::OnCancelRestorePreviouslyOpenAssets),
 		SNotificationItem::CS_None
 	));
 
-	// We will let the notification expire automatically after 10 seconds
-	Info.bFireAndForget = true;
-	Info.ExpireDuration = 20.0f;
+	Info.bFireAndForget = false;
 
 	// We want the auto-save to be subtle
 	Info.bUseLargeFont = false;
@@ -865,9 +873,17 @@ void UAssetEditorSubsystem::SpawnRestorePreviouslyOpenAssetsNotification(const b
 	// Only let the user suppress the non-crash version
 	if (bCleanShutdown)
 	{
-		Info.CheckBoxState = TAttribute<ECheckBoxState>::Create(&Local::GetDontAskAgainCheckBoxState);
-		Info.CheckBoxStateChanged = FOnCheckStateChanged::CreateStatic(&Local::OnDontAskAgainCheckBoxStateChanged);
-		Info.CheckBoxText = NSLOCTEXT("ModalDialogs", "DefaultCheckBoxMessage", "Don't show this again");
+		bRememberMyChoiceChecked = false;
+		
+		Info.CheckBoxState = TAttribute<ECheckBoxState>::CreateLambda([this]()
+		{
+			return bRememberMyChoiceChecked ? ECheckBoxState::Checked : ECheckBoxState::Unchecked;
+		});
+		Info.CheckBoxStateChanged = FOnCheckStateChanged::CreateLambda([this](ECheckBoxState NewState)
+		{
+			bRememberMyChoiceChecked = (NewState == ECheckBoxState::Checked) ? true : false;
+		});
+		Info.CheckBoxText = LOCTEXT("RememberCheckBoxMessage", "Remember my choice");
 	}
 
 	// Close any existing notification
@@ -890,14 +906,14 @@ void UAssetEditorSubsystem::OnConfirmRestorePreviouslyOpenAssets(TArray<FString>
 		RestorePreviouslyOpenAssetsNotification->SetFadeOutDuration(0.5f);
 		RestorePreviouslyOpenAssetsNotification->ExpireAndFadeout();
 
-		// If the user suppressed the notification for future sessions, make sure this is reflected in their settings
-		// Note: We do that inside this if statement so that we only do it if we were showing a UI they could interact with
-		bool bSuppressNotification = false;
-		GConfig->GetBool(TEXT("AssetEditorSubsystem"), TEXT("SuppressRestorePreviouslyOpenAssetsNotification"), bSuppressNotification, GEditorPerProjectIni);
-		UEditorLoadingSavingSettings& Settings = *GetMutableDefault<UEditorLoadingSavingSettings>();
-		Settings.bRestoreOpenAssetTabsOnRestart = bSuppressNotification;
-		Settings.PostEditChange();
-
+		// Change the saved setting to AlwaysRestore if the user checked "Remember my choice"
+		if(bRememberMyChoiceChecked)
+		{
+			UEditorLoadingSavingSettings& Settings = *GetMutableDefault<UEditorLoadingSavingSettings>();
+			Settings.RestoreOpenAssetTabsOnRestart = ERestoreOpenAssetTabsMethod::AlwaysRestore;
+			Settings.PostEditChange();
+		}
+		
 		// we do this inside the condition so that it can only be done once. 
 		OpenEditorsForAssets(AssetsToOpen);
 
@@ -910,6 +926,14 @@ void UAssetEditorSubsystem::OnCancelRestorePreviouslyOpenAssets()
 	TSharedPtr<SNotificationItem> RestorePreviouslyOpenAssetsNotification = RestorePreviouslyOpenAssetsNotificationPtr.Pin();
 	if (RestorePreviouslyOpenAssetsNotification.IsValid())
 	{
+		// Change the saved setting to NeverRestore if the user checked "Remember my choice"
+		if(bRememberMyChoiceChecked)
+		{
+			UEditorLoadingSavingSettings& Settings = *GetMutableDefault<UEditorLoadingSavingSettings>();
+			Settings.RestoreOpenAssetTabsOnRestart = ERestoreOpenAssetTabsMethod::NeverRestore;
+			Settings.PostEditChange();
+		}
+		
 		RestorePreviouslyOpenAssetsNotification->SetExpireDuration(0.0f);
 		RestorePreviouslyOpenAssetsNotification->SetFadeOutDuration(0.5f);
 		RestorePreviouslyOpenAssetsNotification->ExpireAndFadeout();
@@ -918,31 +942,33 @@ void UAssetEditorSubsystem::OnCancelRestorePreviouslyOpenAssets()
 
 void UAssetEditorSubsystem::SaveOpenAssetEditors(const bool bOnShutdown)
 {
-	if (!bSavingOnShutdown && !bAutoRestoreAndDisableSaving)
+	// Saving is disabled if bAutoRestoreAndDisableSaving is set
+	if (!bSavingOnShutdown && !bAutoRestoreAndDisableSaving.IsSet())
 	{
 		TArray<FString> OpenAssets;
-		if (bOnShutdown || !FPlatformMisc::IsDebuggerPresent())
+		
+		for (const TPair<IAssetEditorInstance*, FAssetEntry>& EditorPair : OpenedEditors)
 		{
-			for (const TPair<IAssetEditorInstance*, FAssetEntry>& EditorPair : OpenedEditors)
+			IAssetEditorInstance* Editor = EditorPair.Key;
+			if (Editor != nullptr)
 			{
-				IAssetEditorInstance* Editor = EditorPair.Key;
-				if (Editor != nullptr)
+				UObject* EditedObject = EditorPair.Value.ObjectPtr.Get();
+				if (EditedObject != nullptr)
 				{
-					UObject* EditedObject = EditorPair.Value.ObjectPtr.Get();
-					if (EditedObject != nullptr)
+					// only record assets that have a valid saved package
+					UPackage* Package = EditedObject->GetOutermost();
+					if (Package != nullptr && Package->GetFileSize() != 0)
 					{
-						// only record assets that have a valid saved package
-						UPackage* Package = EditedObject->GetOutermost();
-						if (Package != nullptr && Package->GetFileSize() != 0)
-						{
-							OpenAssets.Add(EditedObject->GetPathName());
-						}
+						OpenAssets.Add(EditedObject->GetPathName());
 					}
 				}
 			}
 		}
+
 		GConfig->SetArray(TEXT("AssetEditorSubsystem"), TEXT("OpenAssetsAtExit"), OpenAssets, GEditorPerProjectIni);
 		GConfig->SetBool(TEXT("AssetEditorSubsystem"), TEXT("CleanShutdown"), bOnShutdown, GEditorPerProjectIni);
+		GConfig->SetBool(TEXT("AssetEditorSubsystem"), TEXT("DebuggerAttached"), FPlatformMisc::IsDebuggerPresent(), GEditorPerProjectIni);
+
 		GConfig->Flush(false, GEditorPerProjectIni);
 	}
 }
diff --git a/Engine/Source/Editor/UnrealEd/Public/Subsystems/AssetEditorSubsystem.h b/Engine/Source/Editor/UnrealEd/Public/Subsystems/AssetEditorSubsystem.h
--- a/Engine/Source/Editor/UnrealEd/Public/Subsystems/AssetEditorSubsystem.h
+++ b/Engine/Source/Editor/UnrealEd/Public/Subsystems/AssetEditorSubsystem.h
@@ -251,8 +251,12 @@ public:
 	void RestorePreviouslyOpenAssets();
 
 	/** Sets bAutoRestoreAndDisableSaving and sets bRequestRestorePreviouslyOpenAssets to false to avoid running RestorePreviouslyOpenAssets() twice. */
+	UE_DEPRECATED(5.0, "Please use the SetAutoRestoreAndDisableSavingOverride instead to set/unset the override")
 	void SetAutoRestoreAndDisableSaving(const bool bInAutoRestoreAndDisableSaving);
 
+	/** Set/Unset the bAutoRestoreAndDisableSaving override, along with setting bRequestRestorePreviouslyOpenAssets to false to avoid running RestorePreviouslyOpenAssets() twice */
+	void SetAutoRestoreAndDisableSavingOverride(TOptional<bool> bInAutoRestoreAndDisableSaving);
+
 private:
 
 	/** Handles a package being reloaded */
@@ -352,12 +356,16 @@ private:
 	 * Flag whether to disable SaveOpenAssetEditors() and enable auto-restore on RestorePreviouslyOpenAssets().
 	 * Useful e.g., to allow LayoutsMenu.cpp re-load layouts on-the-fly and reload the previously opened assets.
 	 * If true, SaveOpenAssetEditors() will not save any asset editor and RestorePreviouslyOpenAssets() will automatically open them without asking the user.
-	 * If false, default behavior of both SaveOpenAssetEditors() and RestorePreviouslyOpenAssets().
+	 * If false, SaveOpenAssetEditors() will not save any asset editor and RestorePreviouslyOpenAssets() will never open them.
+	 * If unset, default behavior of both SaveOpenAssetEditors() and RestorePreviouslyOpenAssets().
 	 */
-	bool bAutoRestoreAndDisableSaving;
+	TOptional<bool> bAutoRestoreAndDisableSaving;
 
 	/** Flag whether there has been a request to notify whether to restore previously open assets */
 	bool bRequestRestorePreviouslyOpenAssets;
+	
+	/** True if the "Remember my choice" checkbox on the restore assets notification is checked */
+	bool bRememberMyChoiceChecked = false;
 
 	/** A pointer to the notification used by RestorePreviouslyOpenAssets */
 	TWeakPtr<SNotificationItem> RestorePreviouslyOpenAssetsNotificationPtr;
