Optimized advanced copy by running its consolidate step less frequently

The "AssetTools.AdvancedCopyConsolidationMethod" CVar can control how frequently the consolidate runs, and there are three options available:
* Lazy - Run a consolidate once at the end for all assets and dependencies combined. The fastest method.
* PerAsset - Run a consolidate once for each asset combined with its dependencies. This is very slow, but is the default method for safety.
* PerAssetDependency - Run a consolidate once for each asset dependency. This is incredibly slow, and is only included for regression testing.

When copying ~1,600 assets the different methods produce the following results:
* Lazy: 222 seconds (~4 mins)
* PerAsset: 1818 seconds (~30 mins)
* PerAssetDependency: 5000 seconds (~83 mins)

diff --git a/Engine/Source/Developer/AssetTools/Private/AssetTools.cpp b/Engine/Source/Developer/AssetTools/Private/AssetTools.cpp
--- a/Engine/Source/Developer/AssetTools/Private/AssetTools.cpp
+++ b/Engine/Source/Developer/AssetTools/Private/AssetTools.cpp
@@ -419,6 +419,32 @@ namespace UE::AssetTools::Private
 		TEXT("When set, The package migration will use the new implementation made for 5.1.")
 	);
 
+	enum class EAdvancedCopyConsolidationMethod : uint8
+	{
+		/** Consolidate everything once (new, fast) */
+		Lazy = 0,
+		/** Consolidate once per-asset (default, slow) */
+		PerAsset = 1,
+		/** Consolidate once per-asset dependency (legacy, glacial) */
+		PerAssetDependency = 2,
+	};
+
+	int32 AdvancedCopyConsolidationMethodInt = static_cast<int32>(EAdvancedCopyConsolidationMethod::PerAsset);
+	FAutoConsoleVariableRef CVarAdvancedCopyConsolidationMethod(
+		TEXT("AssetTools.AdvancedCopyConsolidationMethod"),
+		AdvancedCopyConsolidationMethodInt,
+		TEXT("0: Lazy (new, fast), 1: PerAsset (default, slow), 2: PerAssetDependency (legacy, glacial)")
+	);
+
+	EAdvancedCopyConsolidationMethod GetAdvancedCopyConsolidationMethod()
+	{
+		if (AdvancedCopyConsolidationMethodInt >= static_cast<int32>(EAdvancedCopyConsolidationMethod::Lazy) && AdvancedCopyConsolidationMethodInt <= static_cast<int32>(EAdvancedCopyConsolidationMethod::PerAssetDependency))
+		{
+			return static_cast<EAdvancedCopyConsolidationMethod>(AdvancedCopyConsolidationMethodInt);
+		}
+		return EAdvancedCopyConsolidationMethod::PerAsset;
+	}
+
 	// use a struct as a namespace to allow easier friend declarations
 	struct FPackageMigrationImpl
 	{
@@ -2119,8 +2145,8 @@ bool UAssetToolsImpl::AdvancedCopyPackages(
 		ExistingObjectSet.Reserve(SourceAndDestPackages.Num());
 		NewObjectSet.Reserve(SourceAndDestPackages.Num());
 
-		FScopedSlowTask LoopProgress(SourceAndDestPackages.Num(), LOCTEXT("AdvancedCopying", "Copying files and dependencies..."));
-		LoopProgress.MakeDialog();
+		TUniquePtr<FScopedSlowTask> LoopProgress = MakeUnique<FScopedSlowTask>(SourceAndDestPackages.Num(), LOCTEXT("AdvancedCopyPackages.CopyingFilesAndDependencies", "Copying Files and Dependencies..."));
+		LoopProgress->MakeDialog();
 
 		for (const auto& Package : SourceAndDestPackages)
 		{
@@ -2130,7 +2156,7 @@ bool UAssetToolsImpl::AdvancedCopyPackages(
 
 			if (FPackageName::DoesPackageExist(PackageName, &SrcFilename))
 			{
-				LoopProgress.EnterProgressFrame();
+				LoopProgress->EnterProgressFrame();
 				UPackage* Pkg = LoadPackage(nullptr, *PackageName, LOAD_None);
 				if (Pkg)
 				{
@@ -2182,9 +2208,35 @@ bool UAssetToolsImpl::AdvancedCopyPackages(
 		TSet<UObject*> ObjectsAndSubObjectsToReplaceWithin;
 		ObjectTools::GatherSubObjectsForReferenceReplacement(NewObjectSet, ExistingObjectSet, ObjectsAndSubObjectsToReplaceWithin);
 
+		LoopProgress.Reset(); // Note: Reset first as FScopedSlowTask asserts about out-of-order scoping if assigning over an existing FScopedSlowTask
+		LoopProgress = MakeUnique<FScopedSlowTask>(SuccessfullyCopiedSourcePackages.Num(), LOCTEXT("AdvancedCopyPackages.ReplacingAssetReferences", "Replacing Asset References..."));
+		LoopProgress->MakeDialog();
+
+		TMap<UObject*, TArray<UObject*, TInlineAllocator<1>>> Consolidations;
+		TArray<ObjectTools::FReplaceRequest> Requests;
+		auto ConsolidatePendingObjects = [&Consolidations, &Requests, &ObjectsAndSubObjectsToReplaceWithin, &ExistingObjectSet]()
+		{
+			check(Requests.Num() == 0);
+
+			if (Consolidations.Num() > 0)
+			{
+				Requests.Reserve(Consolidations.Num());
+				for (TPair<UObject*, TArray<UObject*, TInlineAllocator<1>>>&Consolidation : Consolidations)
+				{
+					Requests.Add(ObjectTools::FReplaceRequest{ Consolidation.Key, Consolidation.Value });
+				}
+				ObjectTools::ConsolidateObjects(Requests, ObjectsAndSubObjectsToReplaceWithin, ExistingObjectSet, false);
+
+				Consolidations.Reset();
+				Requests.Reset();
+			}
+		};
+
 		TArray<FName> Dependencies;
 		for (FName SuccessfullyCopiedPackage : SuccessfullyCopiedSourcePackages)
 		{
+			LoopProgress->EnterProgressFrame();
+
 			Dependencies.Reset();
 			AssetRegistryModule.Get().GetDependencies(SuccessfullyCopiedPackage, Dependencies);
 
@@ -2207,31 +2259,39 @@ bool UAssetToolsImpl::AdvancedCopyPackages(
 				const int32 DependencyIndex = SuccessfullyCopiedSourcePackages.IndexOfByKey(Dependency);
 				if (DependencyIndex != INDEX_NONE)
 				{
-					TMap<UObject*, TArray<UObject*, TInlineAllocator<1>>> Consolidations;
-					Consolidations.Reserve(DuplicatedObjectsForEachPackage[DependencyIndex].Num());
-
+					Consolidations.Reserve(Consolidations.Num() + DuplicatedObjectsForEachPackage[DependencyIndex].Num());
 					for (const TPair<TSoftObjectPtr<UObject>, TSoftObjectPtr<UObject>>& Duplication : DuplicatedObjectsForEachPackage[DependencyIndex])
 					{
 						UObject* SourceObject = Duplication.Key.Get();
 						UObject* NewObject = Duplication.Value.Get();
 						if (SourceObject && NewObject)
 						{
-							Consolidations.FindOrAdd(NewObject).Add(SourceObject);
+							Consolidations.FindOrAdd(NewObject).AddUnique(SourceObject);
 						}
 					}
+				}
 
-					TArray<ObjectTools::FReplaceRequest> Requests;
-					Requests.Reserve(Consolidations.Num());
-					for (TPair<UObject*, TArray<UObject*, TInlineAllocator<1>>>& Consolidation : Consolidations)
-					{
-						Requests.Add(ObjectTools::FReplaceRequest{Consolidation.Key, Consolidation.Value});
-					}
-
-					ObjectTools::ConsolidateObjects(Requests, ObjectsAndSubObjectsToReplaceWithin, ExistingObjectSet, false);
+				if (UE::AssetTools::Private::GetAdvancedCopyConsolidationMethod() == UE::AssetTools::Private::EAdvancedCopyConsolidationMethod::PerAssetDependency)
+				{
+					ConsolidatePendingObjects();
 				}
 			}
+
+			if (UE::AssetTools::Private::GetAdvancedCopyConsolidationMethod() == UE::AssetTools::Private::EAdvancedCopyConsolidationMethod::PerAsset)
+			{
+				ConsolidatePendingObjects();
+			}
+		}
+
+		LoopProgress.Reset();
+
+		if (UE::AssetTools::Private::GetAdvancedCopyConsolidationMethod() == UE::AssetTools::Private::EAdvancedCopyConsolidationMethod::Lazy)
+		{
+			ConsolidatePendingObjects();
 		}
 
+		check(Consolidations.Num() == 0);
+
 		ObjectTools::CompileBlueprintsAfterRefUpdate(NewObjectSet.Array());
 
 		// The default value for save packages is true if SCC is enabled because the user can use SCC to revert a change
diff --git a/Engine/Source/Editor/UnrealEd/Private/ObjectTools.cpp b/Engine/Source/Editor/UnrealEd/Private/ObjectTools.cpp
--- a/Engine/Source/Editor/UnrealEd/Private/ObjectTools.cpp
+++ b/Engine/Source/Editor/UnrealEd/Private/ObjectTools.cpp
@@ -1134,7 +1134,7 @@ namespace ObjectTools
 			}
 		}
 
-		GWarn->StatusUpdate( 0, 0, NSLOCTEXT("UnrealEd", "ConsolidateAssetsUpdate_FindingReferences", "Finding Asset References...") );
+		GWarn->StatusUpdate( 0, 0, NSLOCTEXT("UnrealEd", "ConsolidateAssetsUpdate_CollectingReferences", "Collecting Asset References...") );
 
 		ReplacementMap.GenerateKeyArray( OutInfo.ReplaceableObjects );
 
@@ -1149,14 +1149,13 @@ namespace ObjectTools
 			TMap<UObject*, int32> CurNumReferencesMap;
 			TMultiMap<UObject*, FProperty*> CurReferencingPropertiesMMap;
 			PropertyArrayType CurReferencedProperties;
-			for ( FThreadSafeObjectIterator ObjIter; ObjIter; ++ObjIter )
+			
+			auto CollectObjectReferencers = [bOnlyNullingOut, &ReplacementMap, &ReferencingPropertiesMapKeys, &ReferencingPropertiesMapValues, &FindRefsArchive, &CurNumReferencesMap, &CurReferencingPropertiesMMap, &CurReferencedProperties](UObject* CurObject)
 			{
-				UObject* CurObject = *ObjIter;
-
 				// Don't bother replacing in objects that are about to be garbage collected
-				if ((ObjectsToReplaceWithin.Num() > 0 && !ObjectsToReplaceWithin.Contains(CurObject)) || !IsValidChecked(CurObject) || CurObject->IsUnreachable())
+				if (!IsValidChecked(CurObject) || CurObject->IsUnreachable())
 				{
-					continue;
+					return;
 				}
 
 				// Unless the "object to replace with" is null, ignore the objects being replaced
@@ -1171,16 +1170,21 @@ namespace ObjectTools
 
 					if ( FindRefsArchive.GetReferenceCounts( CurNumReferencesMap, CurReferencingPropertiesMMap ) > 0  )
 					{
-						CurReferencingPropertiesMMap.GenerateValueArray( CurReferencedProperties );
+						// TODO: FFindReferencersArchive is giving us the leaf property rather than the member property
+						CurReferencedProperties.Reset(CurReferencingPropertiesMMap.Num());
+						for (const TTuple<UObject*, FProperty*>& CurReferencingPropertiesPair : CurReferencingPropertiesMMap)
+						{
+							CurReferencedProperties.AddUnique(CurReferencingPropertiesPair.Value);
+						}
 
 						ReferencingPropertiesMapKeys.Add(CurObject);
 						ReferencingPropertiesMapValues.Add(CurReferencedProperties);
 
 						if ( CurReferencedProperties.Num() > 0)
 						{
-							for ( PropertyArrayType::TConstIterator RefPropIter( CurReferencedProperties ); RefPropIter; ++RefPropIter )
+							for (FProperty* RefProp : CurReferencedProperties)
 							{
-								CurObject->PreEditChange( *RefPropIter );
+								CurObject->PreEditChange(RefProp);
 							}
 						}
 						else
@@ -1189,25 +1193,56 @@ namespace ObjectTools
 						}
 					}
 				}
+			};
+
+			if (ObjectsToReplaceWithin.Num() > 0)
+			{
+				int32 NumObjsCollected = 0;
+				TArray<UObject*> InnerObjects;
+				for (UObject* CurObject : ObjectsToReplaceWithin)
+				{
+					++NumObjsCollected;
+					GWarn->StatusUpdate(NumObjsCollected, ObjectsToReplaceWithin.Num(), NSLOCTEXT("UnrealEd", "ConsolidateAssetsUpdate_CollectingReferences", "Collecting Asset References..."));
+
+					if (CurObject && CurObject->IsValidLowLevel())
+					{
+						CollectObjectReferencers(CurObject);
+
+						// FArchiveReplaceObjectAndStructPropertyRef is recursive into sub-objects, but FFindReferencersArchive 
+						// isn't so we need to handle that ourselves to build the complete set of references
+						InnerObjects.Reset();
+						GetObjectsWithOuter(CurObject, InnerObjects);
+						for (UObject* InnerObject : InnerObjects)
+						{
+							CollectObjectReferencers(InnerObject);
+						}
+					}
+				}
+			}
+			else
+			{
+				for (FThreadSafeObjectIterator ObjIter; ObjIter; ++ObjIter)
+				{
+					CollectObjectReferencers(*ObjIter);
+				}
 			}
 		}
 
+		// Shuffle dependents before the objects that they reference
 		{
 			TBitArray<> TouchedThisItteration(false, ReferencingPropertiesMapKeys.Num());
 			for (int CurrentIndex = 0; CurrentIndex < ReferencingPropertiesMapKeys.Num(); CurrentIndex++)
 			{
+				GWarn->StatusUpdate(CurrentIndex + 1, ReferencingPropertiesMapKeys.Num(), NSLOCTEXT("UnrealEd", "ConsolidateAssetsUpdate_PreparingAssetReferences", "Preparing Asset References..."));
+
 				TouchedThisItteration.Init(false, ReferencingPropertiesMapKeys.Num());
 				FFindReferencersArchive FindDependentArchive(ReferencingPropertiesMapKeys[CurrentIndex], ReferencingPropertiesMapKeys);
 				for (int DependentIndex = CurrentIndex + 1; DependentIndex < ReferencingPropertiesMapKeys.Num(); DependentIndex++)
 				{
 					if (!TouchedThisItteration[DependentIndex] && FindDependentArchive.GetReferenceCount(ReferencingPropertiesMapKeys[DependentIndex]) > 0)
 					{
-						UObject* Key = ReferencingPropertiesMapKeys[CurrentIndex];
-						PropertyArrayType Value = ReferencingPropertiesMapValues[CurrentIndex];
-						ReferencingPropertiesMapKeys[CurrentIndex] = ReferencingPropertiesMapKeys[DependentIndex];
-						ReferencingPropertiesMapValues[CurrentIndex] = ReferencingPropertiesMapValues[DependentIndex];
-						ReferencingPropertiesMapKeys[DependentIndex] = Key;
-						ReferencingPropertiesMapValues[DependentIndex] = Value;
+						Swap(ReferencingPropertiesMapKeys[CurrentIndex], ReferencingPropertiesMapKeys[DependentIndex]);
+						Swap(ReferencingPropertiesMapValues[CurrentIndex], ReferencingPropertiesMapValues[DependentIndex]);
 
 						FindDependentArchive.ResetPotentialReferencer(ReferencingPropertiesMapKeys[CurrentIndex]);
 						TouchedThisItteration[DependentIndex] = true;
@@ -1217,10 +1252,15 @@ namespace ObjectTools
 			}
 		}
 
-		if(ObjectsToReplaceWithin.Num() > 0)
+		// Run the reference replacement
+		if (ObjectsToReplaceWithin.Num() > 0)
 		{
+			int32 NumObjsReplaced = 0;
 			for (UObject* CurObject : ObjectsToReplaceWithin)
 			{
+				++NumObjsReplaced;
+				GWarn->StatusUpdate(NumObjsReplaced, ObjectsToReplaceWithin.Num(), NSLOCTEXT("UnrealEd", "ConsolidateAssetsUpdate_ReplacingReferences", "Replacing Asset References..."));
+
 				if (CurObject && CurObject->IsValidLowLevel())
 				{
 					UBlueprint* BPObjectToUpdate = Cast<UBlueprint>(CurObject);
@@ -1240,28 +1280,29 @@ namespace ObjectTools
 			for (int32 Index = 0; Index < ReferencingPropertiesMapKeys.Num(); Index++)
 			{
 				++NumObjsReplaced;
-				GWarn->StatusUpdate( NumObjsReplaced, ReferencingPropertiesMapKeys.Num(), NSLOCTEXT("UnrealEd", "ConsolidateAssetsUpdate_ReplacingReferences", "Replacing Asset References...") );
+				GWarn->StatusUpdate(NumObjsReplaced, ReferencingPropertiesMapKeys.Num(), NSLOCTEXT("UnrealEd", "ConsolidateAssetsUpdate_ReplacingReferences", "Replacing Asset References..."));
 
 				UObject* CurReplaceObj = ReferencingPropertiesMapKeys[Index];
 				FArchiveReplaceObjectAndStructPropertyRef<UObject> ReplaceAr(CurReplaceObj, ReplacementMap, EArchiveReplaceObjectFlags::IncludeClassGeneratedByRef);
 			}
 		}
+
 		// Now alter the referencing objects the change has completed via PostEditChange,
 		// this is done in a separate loop to prevent reading of data that we want to overwrite
 		int32 NumObjsPostEdited = 0;
 		for (int32 Index = 0; Index < ReferencingPropertiesMapKeys.Num(); Index++)
 		{
 			++NumObjsPostEdited;
-			GWarn->StatusUpdate( NumObjsPostEdited, ReferencingPropertiesMapKeys.Num(), NSLOCTEXT("UnrealEd", "ConsolidateAssetsUpdate_PostEditing", "Performing Post Update Edits...") );
+			GWarn->StatusUpdate( NumObjsPostEdited, ReferencingPropertiesMapKeys.Num(), NSLOCTEXT("UnrealEd", "ConsolidateAssetsUpdate_FinalizingReferences", "Finalizing Asset References...") );
 
 			UObject* CurReplaceObj = ReferencingPropertiesMapKeys[Index];
 			const PropertyArrayType& RefPropArray = ReferencingPropertiesMapValues[Index];
 
 			if (RefPropArray.Num() > 0)
 			{
-				for ( PropertyArrayType::TConstIterator RefPropIter( RefPropArray ); RefPropIter; ++RefPropIter )
+				for (FProperty* RefProp : RefPropArray)
 				{
-					FPropertyChangedEvent PropertyEvent(*RefPropIter, EPropertyChangeType::Redirected);
+					FPropertyChangedEvent PropertyEvent(RefProp, EPropertyChangeType::Redirected);
 					CurReplaceObj->PostEditChangeProperty( PropertyEvent );
 				}
 			}
@@ -1590,12 +1631,12 @@ namespace ObjectTools
 		{
 			// See if this is a blueprint consolidate and replace instances of the generated class
 			UBlueprint* BlueprintToConsolidateTo = Cast<UBlueprint>(Request.New);
-			if (BlueprintToConsolidateTo != NULL && BlueprintToConsolidateTo->GeneratedClass)
+			if (BlueprintToConsolidateTo && BlueprintToConsolidateTo->GeneratedClass)
 			{
-				for ( TArray<UObject*>::TConstIterator ConsolIter( ReplaceInfo.ReplaceableObjects ); ConsolIter; ++ConsolIter )
+				for (UObject* Old : Request.Old)
 				{
-					UBlueprint* BlueprintToConsolidate = Cast<UBlueprint>(*ConsolIter);
-					if (BlueprintToConsolidate != NULL && BlueprintToConsolidate->GeneratedClass)
+					UBlueprint* BlueprintToConsolidate = Cast<UBlueprint>(Old);
+					if (BlueprintToConsolidate && BlueprintToConsolidate->GeneratedClass)
 					{
 						// Replace all instances of objects based on the old blueprint's class with objects based on the new class,
 						// then repair the references on the object being consolidated so those objects can be properly disposed of upon deletion.
