Save the main map before any external actors in SaveMap

This updates SaveMap to use the same order as SaveCurrentLevel, which seems to address some issues where external actors could be left with invalid references after copying a map (and its actors) via advanced copy.

diff --git a/Engine/Source/Editor/UnrealEd/Private/FileHelpers.cpp b/Engine/Source/Editor/UnrealEd/Private/FileHelpers.cpp
--- a/Engine/Source/Editor/UnrealEd/Private/FileHelpers.cpp
+++ b/Engine/Source/Editor/UnrealEd/Private/FileHelpers.cpp
@@ -894,7 +894,20 @@ static bool SaveWorld(UWorld* World,
 		// This makes UEditorEngine::Save's own call to InitializePhysicsSceneForSaveIfNecessary redundant but wasn't removed to avoid breaking other code paths
 		const bool bInitializedPhysicsSceneForSave = GEditor->InitializePhysicsSceneForSaveIfNecessary(SaveWorld, bForceInitializedWorld);
 				
-		if(!bAutosaving && !FEditorFileUtils::ShouldSkipExternalObjectSave())
+		// Save actual map
+		if (bSuccess)
+		{
+			const FString AutoSavingString = (bAutosaving || bPIESaving) ? TEXT("true") : TEXT("false");
+			const FString KeepDirtyString = bPIESaving ? TEXT("true") : TEXT("false");
+			FSaveErrorOutputDevice SaveErrors;
+
+			bSuccess = GEditor->Exec(NULL, *FString::Printf(TEXT("OBJ SAVEPACKAGE PACKAGE=\"%s\" FILE=\"%s\" SILENT=true AUTOSAVING=%s KEEPDIRTY=%s"), *Package->GetName(), *FinalFilename, *AutoSavingString, *KeepDirtyString), SaveErrors);
+			SaveErrors.Flush();
+		}
+
+		SlowTask.EnterProgressFrame(50);
+
+		if (!bAutosaving && !FEditorFileUtils::ShouldSkipExternalObjectSave())
 		{
 			if (bSuccess)
 			{
@@ -909,7 +922,7 @@ static bool SaveWorld(UWorld* World,
 
 				if (PackagesToSave.Num())
 				{
-					if (!UEditorLoadingAndSavingUtils::SavePackages(PackagesToSave, /*bCheckDirty*/!bNewlyCreated))
+					if (!UEditorLoadingAndSavingUtils::SavePackages(PackagesToSave, bCheckDirty && !bNewlyCreated))
 					{
 						FMessageDialog::Open(EAppMsgType::Ok, NSLOCTEXT("UnrealEd", "Error_FailedToSaveHLODLayersPackages", "Failed to save dependant map packages"));
 						bSuccess = false;
@@ -917,20 +930,7 @@ static bool SaveWorld(UWorld* World,
 				}
 			}
 		}
-				
-		SlowTask.EnterProgressFrame(50);
 
-		// Save actual map
-		if (bSuccess)
-		{
-			const FString AutoSavingString = (bAutosaving || bPIESaving) ? TEXT("true") : TEXT("false");
-			const FString KeepDirtyString = bPIESaving ? TEXT("true") : TEXT("false");
-			FSaveErrorOutputDevice SaveErrors;
-
-			bSuccess = GEditor->Exec(NULL, *FString::Printf(TEXT("OBJ SAVEPACKAGE PACKAGE=\"%s\" FILE=\"%s\" SILENT=true AUTOSAVING=%s KEEPDIRTY=%s"), *Package->GetName(), *FinalFilename, *AutoSavingString, *KeepDirtyString), SaveErrors);
-			SaveErrors.Flush();
-		}
-		
 		if (bSuccess)
 		{
 			// Force update before initializing World Partition
