Fix deadlock in cooker asset registry startup code if command line parameters can't be read or are missing

diff --git a/Engine/Source/Editor/UnrealEd/Private/CookOnTheFlyServer.cpp b/Engine/Source/Editor/UnrealEd/Private/CookOnTheFlyServer.cpp
--- a/Engine/Source/Editor/UnrealEd/Private/CookOnTheFlyServer.cpp
+++ b/Engine/Source/Editor/UnrealEd/Private/CookOnTheFlyServer.cpp
@@ -1212,6 +1212,10 @@ static FDelayedAutoRegisterHelper GPreloadARHelper(EDelayedRegisterRunPhase::Ear
 						GPreloadedARState.Serialize(*Reader.Get(), FAssetRegistrySerializationOptions());
 					}
 				}
+				else
+				{
+					GPreloadARInfoEvent->Trigger();
+				}
 
 				GPreloadAREvent->Trigger();
 			}
