WorldPartition - ResaveActorsBuilder: Silence warnings from LoadClass() that may be outputted when trying to find BP classes
* We output an error if the classes aren't found

diff --git a/Engine/Source/Editor/UnrealEd/Private/WorldPartition/WorldPartitionResaveActorsBuilder.cpp b/Engine/Source/Editor/UnrealEd/Private/WorldPartition/WorldPartitionResaveActorsBuilder.cpp
--- a/Engine/Source/Editor/UnrealEd/Private/WorldPartition/WorldPartitionResaveActorsBuilder.cpp
+++ b/Engine/Source/Editor/UnrealEd/Private/WorldPartition/WorldPartitionResaveActorsBuilder.cpp
@@ -120,14 +120,14 @@ TArray<TSubclassOf<AActor>> UWorldPartitionResaveActorsBuilder::GetActorClassesF
 		if (!ActorClass)
 		{
 			// Look for a fully qualified BP class
-			ActorClass = LoadClass<AActor>(nullptr, *ClassName, nullptr, LOAD_None, nullptr);
+			ActorClass = LoadClass<AActor>(nullptr, *ClassName, nullptr, LOAD_NoWarn, nullptr);
 
 			if (!ActorClass)
 			{
 				// Look for a package BP
 				if (FPackageName::DoesPackageExist(ClassName))
 				{
-					ActorClass = LoadClass<AActor>(nullptr, *FString::Printf(TEXT("%s.%s_C"), *ClassName, *FPackageName::GetLongPackageAssetName(ClassName)), nullptr, LOAD_None, nullptr);
+					ActorClass = LoadClass<AActor>(nullptr, *FString::Printf(TEXT("%s.%s_C"), *ClassName, *FPackageName::GetLongPackageAssetName(ClassName)), nullptr, LOAD_NoWarn, nullptr);
 				}
 			}
 		}
