Cooker: Fix suppression of "Cooking %s -> Skipped NeverCook" message for external actors to use a case insentive comparison for __externalactors__.
#rnx

diff --git a/Engine/Source/Editor/UnrealEd/Private/CookOnTheFlyServer.cpp b/Engine/Source/Editor/UnrealEd/Private/CookOnTheFlyServer.cpp
--- a/Engine/Source/Editor/UnrealEd/Private/CookOnTheFlyServer.cpp
+++ b/Engine/Source/Editor/UnrealEd/Private/CookOnTheFlyServer.cpp
@@ -2506,7 +2506,8 @@ void UCookOnTheFlyServer::DemoteToIdle(UE::Cook::FPackageData& PackageData, UE::
 
 				// ExternalActors: Do not send a message for every NeverCook external Actor package; too much spam
 				if (!(Reason == ESuppressCookReason::NeverCook &&
-					PackageNameStr.ToView().Contains(ULevel::GetExternalActorsFolderName())))
+					UE::String::FindFirst(PackageNameStr.ToView(),
+						ULevel::GetExternalActorsFolderName(), ESearchCase::IgnoreCase) != INDEX_NONE))
 				{
 					UE_CLOG((GCookProgressDisplay & (int32)ECookProgressDisplayMode::Instigators), LogCook, Display,
 						TEXT("Cooking %s, Instigator: { %s } -> Skipped %s"), *PackageNameStr,
