MPCook SoftGC: Reenable the change to include all public objects in referenced packages. It has been fixed to NOT include public objects from transient packages, which caused excessive hanging on to objects.
#rnx

diff --git a/Engine/Source/Editor/UnrealEd/Private/Commandlets/CookCommandlet.cpp b/Engine/Source/Editor/UnrealEd/Private/Commandlets/CookCommandlet.cpp
--- a/Engine/Source/Editor/UnrealEd/Private/Commandlets/CookCommandlet.cpp
+++ b/Engine/Source/Editor/UnrealEd/Private/Commandlets/CookCommandlet.cpp
@@ -641,6 +641,7 @@ void UCookCommandlet::ConditionalCollectGarbage(uint32 TickResults, UCookOnTheFl
 
 	UE_LOG(LogCookCommandlet, Display, TEXT("GarbageCollection...%s (%s)"), *GCType, *GCReason);
 	{
+		TGuardValue<bool> SoftGCGuard(UPackage::bSupportCookerSoftGC, true);
 		CollectGarbage(RF_NoFlags);
 	}
 
diff --git a/Engine/Source/Editor/UnrealEd/Private/CookOnTheFlyServer.cpp b/Engine/Source/Editor/UnrealEd/Private/CookOnTheFlyServer.cpp
--- a/Engine/Source/Editor/UnrealEd/Private/CookOnTheFlyServer.cpp
+++ b/Engine/Source/Editor/UnrealEd/Private/CookOnTheFlyServer.cpp
@@ -4772,7 +4772,13 @@ static void ConstructSoftGCPackageToObjectList(TArray<UObject*>& PackageToObject
 				UPackage* Package = Object->GetPackage();
 				if (!Package)
 				{
-					return;
+					continue;
+				}
+				if (Package->HasAnyFlags(RF_Transient) || Package->HasAnyPackageFlags(PKG_CompiledIn))
+				{
+					// Skip any Transient packages (e.g. /Engine/Transient) and script packages
+					// We only need to keep public objects alive in packages that could be saved.
+					continue;
 				}
 				if (Object == Package)
 				{
@@ -5000,25 +5006,20 @@ void UCookOnTheFlyServer::PreGarbageCollect()
 			});
 	}
 
-	// Add packages and all RF_Public objects outered to them to GCKeepObjects
+	// Add packages to GCKeepObjects. 
 	TArray<UObject*> ObjectsWithOuter;
 	for (UPackage* Package : GCKeepPackages)
 	{
 		GCKeepObjects.Add(Package);
-		ObjectsWithOuter.Reset();
-		GetObjectsWithOuter(Package, ObjectsWithOuter);
-		for (UObject* Obj : ObjectsWithOuter)
-		{
-			if (IsValidChecked(Obj) && Obj->HasAnyFlags(RF_Public))
-			{
-				GCKeepObjects.Add(Obj);
-			}
-		}
 	}
 	for (FPackageData* PackageData : GCKeepPackageDatas)
 	{
 		PackageData->SetKeepReferencedDuringGC(true);
 	}
+
+	// Add all public objects within every package in memory to the UPackage::SoftGCPackageToObjectList container,
+	// so they will be kept in memory if the package is kept in memory.
+	ConstructSoftGCPackageToObjectList(this->SoftGCPackageToObjectListBuffer);
 }
 
 void UCookOnTheFlyServer::CookerAddReferencedObjects(FReferenceCollector& Collector)
@@ -5060,6 +5061,8 @@ void UCookOnTheFlyServer::PostGarbageCollect()
 	check(!SavingPackageData || SavingPackageData->GetPackage() != nullptr);
 
 	GCKeepObjects.Empty();
+	UPackage::SoftGCPackageToObjectList.Empty();
+	SoftGCPackageToObjectListBuffer.Empty();
 
 	PackageDatas->LockAndEnumeratePackageDatas([](FPackageData* PackageData)
 	{
@@ -5221,7 +5224,14 @@ void UCookOnTheFlyServer::EvaluateGarbageCollectionResults(bool bWasDueToOOM, bo
 
 
 		UE_LOG(LogCook, Display, TEXT("See log for memory use information for UObject classes and LLM tags."));
-		UE::Cook::DumpObjClassList(CookByTheBookOptions->SessionStartupObjects);
+
+		{
+			TGuardValue<bool> SoftGCGuard(UPackage::bSupportCookerSoftGC, true);
+			ConstructSoftGCPackageToObjectList(this->SoftGCPackageToObjectListBuffer);
+			UE::Cook::DumpObjClassList(CookByTheBookOptions->SessionStartupObjects);
+			UPackage::SoftGCPackageToObjectList.Empty();
+			SoftGCPackageToObjectListBuffer.Empty();
+		}
 		GLog->Logf(TEXT("Memory Analysis: LLM Tags:"));
 #if ENABLE_LOW_LEVEL_MEM_TRACKER
 		if (FLowLevelMemTracker::Get().IsEnabled())
