Fix tooltip for DistanceFieldResolutionScale

#rnx

diff --git a/Engine/Source/Editor/UnrealEd/Classes/Factories/FbxStaticMeshImportData.h b/Engine/Source/Editor/UnrealEd/Classes/Factories/FbxStaticMeshImportData.h
--- a/Engine/Source/Editor/UnrealEd/Classes/Factories/FbxStaticMeshImportData.h
+++ b/Engine/Source/Editor/UnrealEd/Classes/Factories/FbxStaticMeshImportData.h
@@ -56,7 +56,10 @@ class UFbxStaticMeshImportData : public UFbxMeshImportData
 	UPROPERTY(EditAnywhere, BlueprintReadWrite, AdvancedDisplay, config, Category = Mesh, meta = (ToolTip = "If enabled, combines all meshes into a single mesh", ImportType = "StaticMesh"))
 	uint32 bCombineMeshes : 1;
 
-	/** Specify override color in the case that VertexColorImportOption is set to Override */
+	/**
+	 * Scale to apply to the mesh when allocating the distance field volume texture.
+	 * The default scale is 1, which is assuming that the mesh will be placed unscaled in the world.
+	 */
 	UPROPERTY(EditAnywhere, BlueprintReadWrite, AdvancedDisplay, config, Category= Mesh, meta=(ImportType="StaticMesh"))
 	float DistanceFieldResolutionScale;
 
