Dump Object References stack whenever the leak of a generated package is detected.
Add FScopeFindCookReferences to compute CookOnTheFlyServer' s references in a scope

diff --git a/Engine/Source/Editor/UnrealEd/Classes/CookOnTheSide/CookOnTheFlyServer.h b/Engine/Source/Editor/UnrealEd/Classes/CookOnTheSide/CookOnTheFlyServer.h
--- a/Engine/Source/Editor/UnrealEd/Classes/CookOnTheSide/CookOnTheFlyServer.h
+++ b/Engine/Source/Editor/UnrealEd/Classes/CookOnTheSide/CookOnTheFlyServer.h
@@ -1353,6 +1353,15 @@ private:
 
 		UCookOnTheFlyServer& COTFS;
 	};
+	/** Scoped type to call Set/Clear SoftGCPackageToObjectListBuffer. */
+	struct FScopeFindCookReferences
+	{
+		FScopeFindCookReferences(UCookOnTheFlyServer& InCOTFS);
+		~FScopeFindCookReferences();
+
+		UCookOnTheFlyServer& COTFS;
+		TGuardValue<bool> SoftGCGuard;
+	};
 	/** Callback for FGenericCrashContext; provides the current ActivePackage as context. */
 	void DumpCrashContext(FCrashContextExtendedWriter& Writer);
 	/** Callback for analytics when a new UPackage is loaded. */
diff --git a/Engine/Source/Editor/UnrealEd/Private/CookOnTheFlyServer.cpp b/Engine/Source/Editor/UnrealEd/Private/CookOnTheFlyServer.cpp
--- a/Engine/Source/Editor/UnrealEd/Private/CookOnTheFlyServer.cpp
+++ b/Engine/Source/Editor/UnrealEd/Private/CookOnTheFlyServer.cpp
@@ -4836,6 +4836,20 @@ static void ConstructSoftGCPackageToObjectList(TArray<UObject*>& PackageToObject
 	}
 }
 
+UCookOnTheFlyServer::FScopeFindCookReferences::FScopeFindCookReferences(UCookOnTheFlyServer& InCOTFS)
+	:COTFS(InCOTFS),
+	SoftGCGuard(UPackage::bSupportCookerSoftGC, true)
+{
+	check(COTFS.SoftGCPackageToObjectListBuffer.IsEmpty())
+	ConstructSoftGCPackageToObjectList(COTFS.SoftGCPackageToObjectListBuffer);
+}
+
+UCookOnTheFlyServer::FScopeFindCookReferences::~FScopeFindCookReferences()
+{
+	UPackage::SoftGCPackageToObjectList.Empty();
+	COTFS.SoftGCPackageToObjectListBuffer.Empty();
+}
+
 void UCookOnTheFlyServer::PreGarbageCollect()
 {
 	using namespace UE::Cook;
diff --git a/Engine/Source/Editor/UnrealEd/Private/Cooker/CookPackageData.cpp b/Engine/Source/Editor/UnrealEd/Private/Cooker/CookPackageData.cpp
--- a/Engine/Source/Editor/UnrealEd/Private/Cooker/CookPackageData.cpp
+++ b/Engine/Source/Editor/UnrealEd/Private/Cooker/CookPackageData.cpp
@@ -1773,6 +1773,14 @@ void FGeneratorPackage::PostGarbageCollect()
 			{
 				UE_LOG(LogCook, Warning, TEXT("PackageSplitter found a package it generated that was not removed from memory during garbage collection. This will cause errors later during population.")
 					TEXT("\n\tSplitter=%s, Generated=%s."), *GetSplitDataObjectName().ToString(), *Info.PackageData->GetPackageName().ToString());
+				
+				{
+					// Compute UCookOnTheFlyServer's references so they are gathered by OBJ REFS below 
+					UCookOnTheFlyServer::FScopeFindCookReferences(Owner.GetPackageDatas().GetCookOnTheFlyServer());
+
+					StaticExec(nullptr, *FString::Printf(TEXT("OBJ REFS NAME=%s"), *Info.PackageData->GetPackageName().ToString()));
+				}
+				
 				bHasIssuedWarning = true; // Only issue the warning once per GC
 			}
 		}
