Avoid GC'ing between every external package in ResavePackages commandlet if user didn't specify GCFREQ value.

diff --git a/Engine/Source/Editor/UnrealEd/Private/Commandlets/ContentCommandlets.cpp b/Engine/Source/Editor/UnrealEd/Private/Commandlets/ContentCommandlets.cpp
--- a/Engine/Source/Editor/UnrealEd/Private/Commandlets/ContentCommandlets.cpp
+++ b/Engine/Source/Editor/UnrealEd/Private/Commandlets/ContentCommandlets.cpp
@@ -1853,10 +1853,14 @@ void UResavePackagesCommandlet::PerformAdditionalOperations(class UWorld* World,
 
 	UWorldPartition* WorldPartition = World->GetWorldPartition();
 	const bool bResaveWorldPartitionExternalActors = !!WorldPartition;
+	const int32 DefaultExternalActorGCFreq = 2048;
 
 	// Load and Save Level's external packages
  	if (!bResaveWorldPartitionExternalActors)
 	{
+		// Use a default GC frequency for external actors if GarbageCollectionFrequency is 0.
+		TGuardValue<int32> ScopedGCFreq(GarbageCollectionFrequency, GarbageCollectionFrequency ? GarbageCollectionFrequency : DefaultExternalActorGCFreq);
+
 		World->AddToRoot();
 		for (UPackage* Package : World->PersistentLevel->GetPackage()->GetExternalPackages())
 		{
@@ -1887,6 +1891,9 @@ void UResavePackagesCommandlet::PerformAdditionalOperations(class UWorld* World,
 	// Load and Save world partition actor packages
 	if (bResaveWorldPartitionExternalActors && !bShouldBuildNavigationData)
 	{
+		// Use a default GC frequency for external actors if GarbageCollectionFrequency is 0.
+		TGuardValue<int32> ScopedGCFreq(GarbageCollectionFrequency, GarbageCollectionFrequency ? GarbageCollectionFrequency : DefaultExternalActorGCFreq);
+
 		FWorldPartitionHelpers::ForEachActorDesc(WorldPartition, [this, WorldPartition](const FWorldPartitionActorDesc* ActorDesc)
 		{
 			++TotalPackagesForResave;
