SourceControl: Replaced 'external' string checks with actual checks.

diff --git a/Engine/Source/Editor/UnrealEd/Private/FileHelpers.cpp b/Engine/Source/Editor/UnrealEd/Private/FileHelpers.cpp
--- a/Engine/Source/Editor/UnrealEd/Private/FileHelpers.cpp
+++ b/Engine/Source/Editor/UnrealEd/Private/FileHelpers.cpp
@@ -4554,27 +4554,18 @@ FEditorFileUtils::EPromptReturnCode FEditorFileUtils::PromptForCheckoutAndSave(
 		{
 			// Check if the world should be reloaded after the revert.
 			bool bReloadWorld = false;
-			if (UWorld* EditorWorld = GEditor->GetEditorWorldContext().World())
+
+			for (UPackage* Package : PackagesToRevert)
 			{
-				UPackage* EditorWorldPackage = EditorWorld->GetPackage();
-				if (PackagesToRevert.Contains(EditorWorldPackage))
+				if (UWorld* World = UWorld::FindWorldInPackage(Package))
 				{
-					// If the world file is reverted, the world should be reloaded.
 					bReloadWorld = true;
+					break;
 				}
-				else
+				if (AActor* Actor = AActor::FindActorInPackage(Package))
 				{
-					// If one of the external files is reverted, the world should be reloaded.
-					for (UPackage* Package : PackagesToRevert)
-					{
-						FString PackageName = Package->GetName();
-						if (PackageName.Contains(FPackagePath::GetExternalActorsFolderName()) ||
-							PackageName.Contains(FPackagePath::GetExternalObjectsFolderName()))
-						{
-							bReloadWorld = true;
-							break;
-						}
-					}
+					bReloadWorld = true;
+					break;
 				}
 			}
 
