Undo changelist 24537190.

diff --git a/Engine/Source/Editor/UnrealEd/Private/FileHelpers.cpp b/Engine/Source/Editor/UnrealEd/Private/FileHelpers.cpp
--- a/Engine/Source/Editor/UnrealEd/Private/FileHelpers.cpp
+++ b/Engine/Source/Editor/UnrealEd/Private/FileHelpers.cpp
@@ -4554,18 +4554,27 @@ FEditorFileUtils::EPromptReturnCode FEditorFileUtils::PromptForCheckoutAndSave(
 		{
 			// Check if the world should be reloaded after the revert.
 			bool bReloadWorld = false;
-
-			for (UPackage* Package : PackagesToRevert)
+			if (UWorld* EditorWorld = GEditor->GetEditorWorldContext().World())
 			{
-				if (UWorld* World = UWorld::FindWorldInPackage(Package))
+				UPackage* EditorWorldPackage = EditorWorld->GetPackage();
+				if (PackagesToRevert.Contains(EditorWorldPackage))
 				{
+					// If the world file is reverted, the world should be reloaded.
 					bReloadWorld = true;
-					break;
 				}
-				if (AActor* Actor = AActor::FindActorInPackage(Package))
+				else
 				{
-					bReloadWorld = true;
-					break;
+					// If one of the external files is reverted, the world should be reloaded.
+					for (UPackage* Package : PackagesToRevert)
+					{
+						FString PackageName = Package->GetName();
+						if (PackageName.Contains(FPackagePath::GetExternalActorsFolderName()) ||
+							PackageName.Contains(FPackagePath::GetExternalObjectsFolderName()))
+						{
+							bReloadWorld = true;
+							break;
+						}
+					}
 				}
 			}
 
