World Partition: make sure the minimap texture doesn't export to text.

diff --git a/Engine/Source/Editor/UnrealEd/Private/WorldPartition/WorldPartitionMiniMapBuilder.cpp b/Engine/Source/Editor/UnrealEd/Private/WorldPartition/WorldPartitionMiniMapBuilder.cpp
--- a/Engine/Source/Editor/UnrealEd/Private/WorldPartition/WorldPartitionMiniMapBuilder.cpp
+++ b/Engine/Source/Editor/UnrealEd/Private/WorldPartition/WorldPartitionMiniMapBuilder.cpp
@@ -63,7 +63,7 @@ bool UWorldPartitionMiniMapBuilder::PreRun(UWorld* World, FPackageSourceControlH
 	// Create minimap texture
 	{
 		TStrongObjectPtr<UTextureFactory> Factory(NewObject<UTextureFactory>());
-		WorldMiniMap->MiniMapTexture = Factory->CreateTexture2D(WorldMiniMap, TEXT("MinimapTexture"), RF_NoFlags);
+		WorldMiniMap->MiniMapTexture = Factory->CreateTexture2D(WorldMiniMap, TEXT("MinimapTexture"), RF_TextExportTransient);
 		WorldMiniMap->MiniMapTexture->Source.Init(MinimapImageSizeX, MinimapImageSizeY, 1, 1, TSF_BGRA8);
 		WorldMiniMap->MiniMapWorldBounds = IterativeWorldBounds;
 	}
diff --git a/Engine/Source/Runtime/Engine/Private/WorldPartition/WorldPartitionMiniMap.cpp b/Engine/Source/Runtime/Engine/Private/WorldPartition/WorldPartitionMiniMap.cpp
--- a/Engine/Source/Runtime/Engine/Private/WorldPartition/WorldPartitionMiniMap.cpp
+++ b/Engine/Source/Runtime/Engine/Private/WorldPartition/WorldPartitionMiniMap.cpp
@@ -40,6 +40,11 @@ void AWorldPartitionMiniMap::PostLoad()
 		WorldUnitsPerPixel = FMath::Clamp(WorldUnitsPerPixel, 10, 100000);
 		MiniMapTileSize_DEPRECATED = 0;
 	}
+
+	if (MiniMapTexture)
+	{
+		MiniMapTexture->SetFlags(RF_TextExportTransient);
+	}
 }
 
 #if WITH_EDITOR
