Do not run PSO cache commandlet (and stable key collection) when cooking DLC.

- Bundling PSO cache with the plugins isn't currently supported, and the commandlet picks up main game's cache instead. This also results in spurious warnings when the main game has PSO cache but the DLC doesn't.

#rnx

diff --git a/Engine/Source/Editor/UnrealEd/Private/CookOnTheFlyServer.cpp b/Engine/Source/Editor/UnrealEd/Private/CookOnTheFlyServer.cpp
--- a/Engine/Source/Editor/UnrealEd/Private/CookOnTheFlyServer.cpp
+++ b/Engine/Source/Editor/UnrealEd/Private/CookOnTheFlyServer.cpp
@@ -8429,6 +8429,10 @@ void UCookOnTheFlyServer::BeginCookStartShaderCodeLibrary(FBeginCookContext& Beg
 		GConfig->GetBool(TEXT("DevOptions.Shaders"), TEXT("NeedsShaderStableKeys"), bAllPlatformsNeedStableKeys, GEngineIni);
 		GConfig->GetBool(TEXT("DevOptions.Shaders"), TEXT("bNeedsShaderStableKeys"), bAllPlatformsNeedStableKeys, GEngineIni);
 
+		// PSO manual cache in DLC is not currently supported. Although stable keys can have other uses, disable this for DLC as it will
+		// also make it cook faster.
+		bAllPlatformsNeedStableKeys &= !IsCookingDLC();
+
  		for (const ITargetPlatform* TargetPlatform : BeginContext.TargetPlatforms)
  		{
 			// Find out if this platform requires stable shader keys, by reading the platform setting file.
