World Partition - HLOD: Allow the use of -RebuildHLODs with other build flags (such as -SetupHLODs)
* Prior to this change, -RebuildHLODs was clearing the build options

diff --git a/Engine/Source/Editor/UnrealEd/Private/WorldPartition/WorldPartitionHLODsBuilder.cpp b/Engine/Source/Editor/UnrealEd/Private/WorldPartition/WorldPartitionHLODsBuilder.cpp
--- a/Engine/Source/Editor/UnrealEd/Private/WorldPartition/WorldPartitionHLODsBuilder.cpp
+++ b/Engine/Source/Editor/UnrealEd/Private/WorldPartition/WorldPartitionHLODsBuilder.cpp
@@ -186,7 +186,7 @@ UWorldPartitionHLODsBuilder::UWorldPartitionHLODsBuilder(const FObjectInitialize
 
 	if (!HLODActorToBuild.IsNone() || !HLODLayerToBuild.IsNone())
 	{
-		BuildOptions = EHLODBuildStep::HLOD_Build;
+		BuildOptions |= EHLODBuildStep::HLOD_Build;
 		bForceBuild = bForceBuild || !HLODActorToBuild.IsNone();
 	}
 	
