OnLoadedActorRemoved: Avoid deselect of actors with a root selection parent

#rnx

diff --git a/Engine/Source/Editor/UnrealEd/Private/EditorEngine.cpp b/Engine/Source/Editor/UnrealEd/Private/EditorEngine.cpp
--- a/Engine/Source/Editor/UnrealEd/Private/EditorEngine.cpp
+++ b/Engine/Source/Editor/UnrealEd/Private/EditorEngine.cpp
@@ -313,15 +313,18 @@ private:
 	{
 		if (ElementSelectionSet)
 		{
-			if (FTypedElementHandle ActorHandle = UEngineElementsLibrary::AcquireEditorActorElementHandle(&InActor, /*bAllowCreate*/false))
+			if (!InActor.GetRootSelectionParent())
 			{
-				static const FTypedElementSelectionOptions SelectionOptions = FTypedElementSelectionOptions()
-					.SetAllowHidden(true)
-					.SetAllowGroups(false)
-					.SetWarnIfLocked(false)
-					.SetChildElementInclusionMethod(ETypedElementChildInclusionMethod::Recursive);
+				if (FTypedElementHandle ActorHandle = UEngineElementsLibrary::AcquireEditorActorElementHandle(&InActor, /*bAllowCreate*/false))
+				{
+					static const FTypedElementSelectionOptions SelectionOptions = FTypedElementSelectionOptions()
+						.SetAllowHidden(true)
+						.SetAllowGroups(false)
+						.SetWarnIfLocked(false)
+						.SetChildElementInclusionMethod(ETypedElementChildInclusionMethod::Recursive);
 
-				bSelectionChanged |= ElementSelectionSet->DeselectElement(ActorHandle, SelectionOptions);
+					bSelectionChanged |= ElementSelectionSet->DeselectElement(ActorHandle, SelectionOptions);
+				}
 			}
 		}
 	}
