Don't mark an outer level/OFPA asset as dirty when compiling its underlying Blueprint type or when reconstructing an Actor-based Blueprint instance.

diff --git a/Engine/Source/Editor/UnrealEd/Private/Kismet2/KismetReinstanceUtilities.cpp b/Engine/Source/Editor/UnrealEd/Private/Kismet2/KismetReinstanceUtilities.cpp
--- a/Engine/Source/Editor/UnrealEd/Private/Kismet2/KismetReinstanceUtilities.cpp
+++ b/Engine/Source/Editor/UnrealEd/Private/Kismet2/KismetReinstanceUtilities.cpp
@@ -2202,6 +2202,18 @@ void FBlueprintCompileReinstancer::ReplaceInstancesOfClass_Inner(const TMap<UCla
 	// Set global flag to let system know we are reconstructing blueprint instances
 	TGuardValue<bool> GuardTemplateNameFlag(GIsReconstructingBlueprintInstances, true);
 
+	// Keep track of non-dirty packages for objects about to be reinstanced, so we can clear the dirty state after reinstancing them
+	TSet<UPackage*> CleanPackageList;
+	auto CheckAndSaveOuterPackageToCleanList = [&CleanPackageList](const UObject* InObject)
+	{
+		check(InObject);
+		UPackage* ObjectPackage = InObject->GetPackage();
+		if (ObjectPackage && !ObjectPackage->IsDirty() && ObjectPackage != GetTransientPackage())
+		{
+			CleanPackageList.Add(ObjectPackage);
+		}
+	};
+
 	struct FObjectRemappingHelper
 	{
 		void OnObjectsReplaced(const TMap<UObject*, UObject*>& InReplacedObjects)
@@ -2255,6 +2267,13 @@ void FBlueprintCompileReinstancer::ReplaceInstancesOfClass_Inner(const TMap<UCla
 		if(GEditor && GEditor->GetSelectedActors())
 		{
 			SelectedActors = GEditor->GetSelectedActors();
+
+			// Note: For OFPA, each instance may be stored in its own external package.
+			for (int32 SelectionIdx = 0; SelectionIdx < SelectedActors->Num(); ++SelectionIdx)
+			{
+				CheckAndSaveOuterPackageToCleanList(SelectedActors->GetSelectedObject(SelectionIdx));
+			}
+			
 			SelectedActors->BeginBatchSelectOperation();
 			SelectedActors->Modify();
 		}
@@ -2322,6 +2341,8 @@ void FBlueprintCompileReinstancer::ReplaceInstancesOfClass_Inner(const TMap<UCla
 					// WARNING: This loop only handles non-actor objects, actor objects are handled below:
 					if (OldActor == nullptr)
 					{
+						CheckAndSaveOuterPackageToCleanList(OldObject);
+
 						UObject* NewUObject = nullptr;
 						ReplaceObjectHelper(OldObject, OldClass, NewUObject, NewClass, OldToNewInstanceMap, OldToNewNameMap, OldObjIndex, ObjectsToReplace, PotentialEditorsForRefreshing, OwnersToRerunConstructionScript, &FDirectAttachChildrenAccessor::Get, bIsComponent, bArchetypesAreUpToDate);
 						UpdateObjectBeingDebugged(OldObject, NewUObject);
@@ -2396,6 +2417,9 @@ void FBlueprintCompileReinstancer::ReplaceInstancesOfClass_Inner(const TMap<UCla
 					// handled above
 					if (OldActor != nullptr)
 					{
+						// Note: For OFPA, each instance may be stored in its own external package.
+						CheckAndSaveOuterPackageToCleanList(OldActor);
+
 						UObject* NewUObject = nullptr;
 						if (OldActor->GetLevel() && OldActor->GetWorld())
 						{
@@ -2547,6 +2571,13 @@ void FBlueprintCompileReinstancer::ReplaceInstancesOfClass_Inner(const TMap<UCla
 		GEditor->NoteSelectionChange();
 	}
 
+	// Clear the dirty flag on packages that didn't already have it set prior to objects being reinstanced.
+	for (UPackage* PackageToMarkAsClean : CleanPackageList)
+	{
+		check(PackageToMarkAsClean);
+		PackageToMarkAsClean->SetDirtyFlag(false);
+	}
+
 	TSet<UBlueprintGeneratedClass*> FixedSCS;
 
 	// in the case where we're replacing component instances, we need to make 
diff --git a/Engine/Source/Runtime/Engine/Private/ActorConstruction.cpp b/Engine/Source/Runtime/Engine/Private/ActorConstruction.cpp
--- a/Engine/Source/Runtime/Engine/Private/ActorConstruction.cpp
+++ b/Engine/Source/Runtime/Engine/Private/ActorConstruction.cpp
@@ -279,6 +279,21 @@ void AActor::RerunConstructionScripts()
 		// Temporarily suspend the undo buffer; we don't need to record reconstructed component objects into the current transaction
 		ITransaction* CurrentTransaction = GUndo;
 		GUndo = nullptr;
+
+		// Mark package as clean on exit if not transient and not already dirty.
+		bool bMarkPackageClean = false;
+		const UPackage* ActorPackage = GetPackage();
+		if (ActorPackage && !ActorPackage->IsDirty() && ActorPackage != GetTransientPackage())
+		{
+			bMarkPackageClean = true;
+		}
+		ON_SCOPE_EXIT
+		{
+			if (bMarkPackageClean)
+			{
+				GetPackage()->SetDirtyFlag(false);
+			}
+		};
 		
 		// Create cache to store component data across rerunning construction scripts
 		FComponentInstanceDataCache* InstanceDataCache;
