- Added Platform/OS and Architecture to editor About box

diff --git a/Engine/Source/Editor/UnrealEd/Private/AboutScreen.cpp b/Engine/Source/Editor/UnrealEd/Private/AboutScreen.cpp
--- a/Engine/Source/Editor/UnrealEd/Private/AboutScreen.cpp
+++ b/Engine/Source/Editor/UnrealEd/Private/AboutScreen.cpp
@@ -40,6 +40,17 @@ void SAboutScreen::Construct(const FArguments& InArgs)
 	FText Version = FText::Format( LOCTEXT("VersionLabel", "Version: {0}"), FText::FromString( FEngineVersion::Current().ToString( ) ) );
 	FText Title = FText::FromString(FApp::GetName());
 	
+	FString OSLabel, OSSubLabel;
+	FPlatformMisc::GetOSVersions(OSLabel, OSSubLabel);
+	FText OSName = FText::FromString(OSLabel);
+#if PLATFORM_CPU_ARM_FAMILY
+	FText ArchName = FText::FromString("arm64");
+#else
+	FText ArchName = FText::FromString("x86_64");
+#endif
+	
+	FText Platform = FText::Format( LOCTEXT("PlatformArchLabel", "Platform: {0} ({1})"), OSName, ArchName);
+	
 	ChildSlot
 	[
 		SNew(SBorder)
@@ -85,6 +96,13 @@ void SAboutScreen::Construct(const FArguments& InArgs)
 						.ColorAndOpacity(FStyleColors::ForegroundHover)
 						.Text( Version )
 					]
+					+SVerticalBox::Slot()
+					[
+						SNew(SEditableText)
+						.IsReadOnly(true)
+						.ColorAndOpacity(FStyleColors::ForegroundHover)
+						.Text(Platform)
+					]
 				]
 
 				+SHorizontalBox::Slot()
diff --git a/Engine/Source/Runtime/Core/Private/Mac/MacPlatformMisc.cpp b/Engine/Source/Runtime/Core/Private/Mac/MacPlatformMisc.cpp
--- a/Engine/Source/Runtime/Core/Private/Mac/MacPlatformMisc.cpp
+++ b/Engine/Source/Runtime/Core/Private/Mac/MacPlatformMisc.cpp
@@ -1619,6 +1619,7 @@ FGPUDriverInfo FMacPlatformMisc::GetGPUDriverInfo(const FString& DeviceDescripti
 void FMacPlatformMisc::GetOSVersions( FString& out_OSVersionLabel, FString& out_OSSubVersionLabel )
 {
 	MacPlatformGetOSProductVersion(out_OSVersionLabel);
+	out_OSVersionLabel = FString("macOS ") + out_OSVersionLabel;
 	MacPlatformGetOSVersion(out_OSSubVersionLabel);
 }
 
