Prevent non-server standalone instances from using ServerMapNameOverride

diff --git a/Engine/Source/Editor/UnrealEd/Private/PlayLevelNewProcess.cpp b/Engine/Source/Editor/UnrealEd/Private/PlayLevelNewProcess.cpp
--- a/Engine/Source/Editor/UnrealEd/Private/PlayLevelNewProcess.cpp
+++ b/Engine/Source/Editor/UnrealEd/Private/PlayLevelNewProcess.cpp
@@ -310,14 +310,17 @@ void UEditorEngine::LaunchNewProcess(const FRequestPlaySessionParams& InParams,
 	if (NetMode != EPlayNetMode::PIE_Client)
 	{
 		// If we're not a client, build a PlayWorld URL to load to.
-		FString ServerMapNameOverride;
-		InParams.EditorPlaySettings->GetServerMapNameOverride(ServerMapNameOverride);
-		
-		// Allow the user to override which map the server should load.
-		if (ServerMapNameOverride.Len() > 0)
+		if (NetMode != EPlayNetMode::PIE_Standalone)
 		{
-			UE_LOG(LogPlayLevel, Log, TEXT("Map Override specified in configuration, using %s instead of current map (%s)"), *ServerMapNameOverride, *MapName);
-			MapName = ServerMapNameOverride;
+			FString ServerMapNameOverride;
+			InParams.EditorPlaySettings->GetServerMapNameOverride(ServerMapNameOverride);
+		
+			// Allow the user to override which map the server should load.
+			if (ServerMapNameOverride.Len() > 0)
+			{
+				UE_LOG(LogPlayLevel, Log, TEXT("Map Override specified in configuration, using %s instead of current map (%s)"), *ServerMapNameOverride, *MapName);
+				MapName = ServerMapNameOverride;
+			}
 		}
 
 		NamedArguments.Add(TEXT("PlayWorldURL"), BuildPlayWorldURL(*MapName, false, UnrealURLParams));
