SourceControl: Removed the status update queueing for files that have been made writable as it's causing random crashes and without a solid repro rolling back the change is the way to go for now.

diff --git a/Engine/Source/Editor/UnrealEd/Private/FileHelpers.cpp b/Engine/Source/Editor/UnrealEd/Private/FileHelpers.cpp
--- a/Engine/Source/Editor/UnrealEd/Private/FileHelpers.cpp
+++ b/Engine/Source/Editor/UnrealEd/Private/FileHelpers.cpp
@@ -4544,12 +4544,6 @@ FEditorFileUtils::EPromptReturnCode FEditorFileUtils::PromptForCheckoutAndSave(
 			ReturnResponse = PR_Cancelled;
 		}
 
-		// Update SCC state for packages that were made writable outside of SCC
-		if (PackagesCheckedOutOrMadeWritable.Num())
-		{
-			ISourceControlModule::Get().QueueStatusUpdate(PackagesCheckedOutOrMadeWritable);
-		}
-
 		if (PackagesToRevert.Num() > 0)
 		{
 			// Check if the world should be reloaded after the revert.
