WorldInitialization: When calling CleanupWorld during a mapchange to detach components, call InitWorld again afterwards. Not calling InitWorld was skipped in the past only because it was not supported. It is supported now, and because we are reinitializating, we can remove the old workarounds of bCleanupResources=false and calling CreateFXSystem.
Remove the workaround we put in place for this case in SCS_Nodes that needed to reinitialize themselves in the Cleaned-up world.
#rnx

diff --git a/Engine/Source/Editor/UnrealEd/Private/UnrealEdMisc.cpp b/Engine/Source/Editor/UnrealEd/Private/UnrealEdMisc.cpp
--- a/Engine/Source/Editor/UnrealEd/Private/UnrealEdMisc.cpp
+++ b/Engine/Source/Editor/UnrealEd/Private/UnrealEdMisc.cpp
@@ -1189,16 +1189,12 @@ void FUnrealEdMisc::CB_MapChange( uint32 InFlags )
 		// Minor things like brush subtraction will set it to "0".
 		if (InFlags != MapChangeEventFlags::Default)
 		{
-			World->ClearWorldComponents();
-
-			// Note: CleanupWorld is being abused here to detach components and some other stuff
-			// CleanupWorld should only be called before destroying the world
-			// So bCleanupResources is being passed as false
-			World->CleanupWorld(true, false);
-
-			// CleanupWorld will have nulled the FXSystem, create a new one or else the dependent
-			// FXSystemComponents will be left unregistered and/or fail to activate.
-			World->CreateFXSystem();
+			if (World->IsInitialized())
+			{
+				// Call CleanupWorld/InitWorld to detach components and some other stuff
+				World->CleanupWorld();
+				World->InitWorld();
+			}
 		}
 
 		GEditor->EditorUpdateComponents();
diff --git a/Engine/Source/Runtime/Engine/Classes/Engine/World.h b/Engine/Source/Runtime/Engine/Classes/Engine/World.h
--- a/Engine/Source/Runtime/Engine/Classes/Engine/World.h
+++ b/Engine/Source/Runtime/Engine/Classes/Engine/World.h
@@ -3068,8 +3068,6 @@ public:
 	/** Returns whether InitWorld has ever been called since this World was created.  */
 	bool HasEverBeenInitialized() const { return bHasEverBeenInitialized; }
 #endif
-	UE_DEPRECATED(5.2, "Not for public use. This function is a workaround for UE-170919 and will be removed in 5.3")
-	bool HasEverBeenInitialized_DONOTUSE() const { return bHasEverBeenInitialized; }
 
 	/** Returns whether InitWorld has been called without yet calling CleanupWorld.  */
 	bool IsInitialized() const { return bIsWorldInitialized; }
diff --git a/Engine/Source/Runtime/Engine/Private/SCS_Node.cpp b/Engine/Source/Runtime/Engine/Private/SCS_Node.cpp
--- a/Engine/Source/Runtime/Engine/Private/SCS_Node.cpp
+++ b/Engine/Source/Runtime/Engine/Private/SCS_Node.cpp
@@ -120,9 +120,7 @@ UActorComponent* USCS_Node::ExecuteNodeOnActor(AActor* Actor, USceneComponent* P
 		{
 			// Only register scene components if the world is initialized
 			UWorld* World = Actor->GetWorld();
-			PRAGMA_DISABLE_DEPRECATION_WARNINGS
-			bool bRegisterComponent = World && World->HasEverBeenInitialized_DONOTUSE();
-			PRAGMA_ENABLE_DEPRECATION_WARNINGS
+			bool bRegisterComponent = World && World->bIsWorldInitialized;
 
 			// If NULL is passed in, we are the root, so set transform and assign as RootComponent on Actor, similarly if the 
 			// NewSceneComp is the ParentComponent then we are the root component. This happens when the root component is recycled
