World Partition Resave Actors Builder: added the option to resave all blueprint classes for the specified map, starting top-down from the class hierarchy.

diff --git a/Engine/Source/Editor/UnrealEd/Private/WorldPartition/WorldPartitionResaveActorsBuilder.cpp b/Engine/Source/Editor/UnrealEd/Private/WorldPartition/WorldPartitionResaveActorsBuilder.cpp
--- a/Engine/Source/Editor/UnrealEd/Private/WorldPartition/WorldPartitionResaveActorsBuilder.cpp
+++ b/Engine/Source/Editor/UnrealEd/Private/WorldPartition/WorldPartitionResaveActorsBuilder.cpp
@@ -3,6 +3,7 @@
 #include "WorldPartition/WorldPartitionResaveActorsBuilder.h"
 #include "WorldPartition/WorldPartition.h"
 #include "WorldPartition/WorldPartitionHelpers.h"
+#include "WorldPartition/WorldPartitionClassDescRegistry.h"
 #include "PackageSourceControlHelper.h"
 #include "SourceControlHelpers.h"
 #include "Engine/Level.h"
@@ -340,6 +341,58 @@ bool UWorldPartitionResaveActorsBuilder::RunInternal(UWorld* World, const FCellI
 	}
 	else
 	{
+		if (bResaveBlueprints)
+		{
+			// Resave blueprints top-down from the class hierarchy
+			FWorldPartitionClassDescRegistry& ClassDescRegistry = FWorldPartitionClassDescRegistry::Get();
+
+			TArray<FTopLevelAssetPath> ClassesToResave;
+			TMap<FTopLevelAssetPath, TArray<FTopLevelAssetPath>> ClassHierarchy;
+			for (auto& [ClassPath, ParentClassPath] : ClassDescRegistry.GetParentClassMap())
+			{
+				ClassHierarchy.FindOrAdd(ParentClassPath).Add(ClassPath);
+
+				if (!ClassPath.ToString().StartsWith(TEXT("/Script/")) && ParentClassPath.ToString().StartsWith(TEXT("/Script/")))
+				{
+					ClassesToResave.Add(ClassPath);
+				}
+			}
+
+			TArray<UPackage*> PackagesToSave;
+			while (!ClassesToResave.IsEmpty())
+			{
+				for (const FTopLevelAssetPath& ClassToResave : ClassesToResave)
+				{
+					UPackage* ClassPackage = LoadPackage(nullptr, *ClassToResave.GetPackageName().ToString(), LOAD_None);
+
+					if (!ClassPackage)
+					{
+						UE_LOG(LogWorldPartitionResaveActorsBuilder, Error, TEXT("Failed to load package for class '%s'"), **ClassToResave.ToString());
+						return false;
+					}
+
+					PackagesToSave.Add(ClassPackage);
+
+					if (FWorldPartitionHelpers::HasExceededMaxMemory())
+					{
+						UWorldPartitionBuilder::SavePackages(PackagesToSave, PackageHelper, true);
+						FWorldPartitionHelpers::DoCollectGarbage();
+						PackagesToSave.Empty();
+					}
+				}
+
+				const TArray<FTopLevelAssetPath> ParentClasses = MoveTemp(ClassesToResave);
+
+				for (const FTopLevelAssetPath& ParentClass : ParentClasses)
+				{
+					ClassesToResave += ClassHierarchy.FindRef(ParentClass);
+				}
+			}
+
+			UWorldPartitionBuilder::SavePackages(PackagesToSave, PackageHelper);
+			FWorldPartitionHelpers::DoCollectGarbage();
+		}
+
 		TArray<UPackage*> PackagesToSave;
 
 		FWorldPartitionHelpers::FForEachActorWithLoadingParams ForEachActorWithLoadingParams;
@@ -485,4 +538,4 @@ bool UWorldPartitionResaveActorsBuilder::RunInternal(UWorld* World, const FCellI
 	}
 
 	return true;
-}
\ No newline at end of file
+}
diff --git a/Engine/Source/Editor/UnrealEd/Public/WorldPartition/WorldPartitionResaveActorsBuilder.h b/Engine/Source/Editor/UnrealEd/Public/WorldPartition/WorldPartitionResaveActorsBuilder.h
--- a/Engine/Source/Editor/UnrealEd/Public/WorldPartition/WorldPartitionResaveActorsBuilder.h
+++ b/Engine/Source/Editor/UnrealEd/Public/WorldPartition/WorldPartitionResaveActorsBuilder.h
@@ -45,6 +45,9 @@ private:
 	UPROPERTY()
 	bool bEnableActorFolders;
 
+	UPROPERTY()
+	bool bResaveBlueprints;
+
 	UPROPERTY()
 	TSet<FName> ActorTags;
 
