Sequencer: Check for valid animation on skeletal animation fbx export

diff --git a/Engine/Source/Editor/UnrealEd/Private/Fbx/FbxMainExport.cpp b/Engine/Source/Editor/UnrealEd/Private/Fbx/FbxMainExport.cpp
--- a/Engine/Source/Editor/UnrealEd/Private/Fbx/FbxMainExport.cpp
+++ b/Engine/Source/Editor/UnrealEd/Private/Fbx/FbxMainExport.cpp
@@ -1766,7 +1766,10 @@ float FLevelSequenceAnimTrackAdapter::GetAnimTime(int32 LocalFrame) const
 	{
 		if (UMovieSceneSkeletalAnimationSection* AnimSection = Cast<UMovieSceneSkeletalAnimationSection>(Section))
 		{
-			return static_cast<float>(AnimSection->MapTimeToAnimation(LocalTime, TickResolution));
+			if (AnimSection->Params.Animation)
+			{
+				return static_cast<float>(AnimSection->MapTimeToAnimation(LocalTime, TickResolution));
+			}
 		}
 	}
 
