Fix cooker being blocked on Landscape material due to the same UMaterialInstanceConstant being saved in both a ExternalActor package and the generated package, with the saves overlapping their calls to BeginCacheForCookedPlatformData and IsCachedCookedPlatformDataLoaded.
#rnx

diff --git a/Engine/Source/Editor/UnrealEd/Private/CookOnTheFlyServer.cpp b/Engine/Source/Editor/UnrealEd/Private/CookOnTheFlyServer.cpp
--- a/Engine/Source/Editor/UnrealEd/Private/CookOnTheFlyServer.cpp
+++ b/Engine/Source/Editor/UnrealEd/Private/CookOnTheFlyServer.cpp
@@ -10716,6 +10716,8 @@ TArray<FName> UCookOnTheFlyServer::GetNeverCookPackageFileNames(TArrayView<const
 	FString ExternalActorsFolderName = ULevel::GetExternalActorsFolderName();
 	FString FullExternalActorsPath = FPaths::Combine(TEXT("/Game/"), ExternalActorsFolderName);
 	NeverCookDirectories.Add(MoveTemp(FullExternalActorsPath));
+	FullExternalActorsPath = FPaths::Combine(TEXT("/Engine/"), ExternalActorsFolderName);
+	NeverCookDirectories.Add(MoveTemp(FullExternalActorsPath));
 	for (TSharedRef<IPlugin>& Plugin : IPluginManager::Get().GetEnabledPluginsWithContent())
 	{
 		FullExternalActorsPath = FPaths::Combine(Plugin->GetMountedAssetPath(), ExternalActorsFolderName);
