ZenCook: Merge the oplog and package store manifest. Write an oplog instead of a package store manifest when cooking to loose files and read the oplog instead of the package store manifest to find cooked files in UAT when staging from Zen.
#rnx

diff --git a/Engine/Source/Developer/IoStoreUtilities/Private/IoStoreUtilities.cpp b/Engine/Source/Developer/IoStoreUtilities/Private/IoStoreUtilities.cpp
--- a/Engine/Source/Developer/IoStoreUtilities/Private/IoStoreUtilities.cpp
+++ b/Engine/Source/Developer/IoStoreUtilities/Private/IoStoreUtilities.cpp
@@ -70,6 +70,7 @@
 #include "IPlatformFileSandboxWrapper.h"
 #include "Misc/PathViews.h"
 #include "HAL/FileManagerGeneric.h"
+#include "Serialization/CompactBinarySerialization.h"
 
 //PRAGMA_DISABLE_OPTIMIZATION
 
@@ -646,78 +647,78 @@ class FCookedPackageStore
 {
 public:
 	FCookedPackageStore(const FString& InCookedDir)
-		: PackageStoreManifest(InCookedDir)
-		, CookedDir(InCookedDir)
+		: CookedDir(InCookedDir)
 	{
 	}
 
 	FIoStatus Load(const TCHAR* ManifestFilename)
 	{
 		IOSTORE_CPU_SCOPE(LoadCookedPackageStore);
-		FIoStatus Status = PackageStoreManifest.Load(ManifestFilename);
-		if (!Status.IsOk())
-		{
-			return Status;
-		}
 
-		if (const FPackageStoreManifest::FZenServerInfo* ZenServerInfo = PackageStoreManifest.ReadZenServerInfo())
+		TUniquePtr<FArchive> Ar(IFileManager::Get().CreateFileReader(ManifestFilename));
+		if (!Ar)
 		{
-			ZenStoreClient = MakeUnique<UE::FZenStoreHttpClient>(UE::Zen::FServiceSettings(ZenServerInfo->Settings));
-			ZenStoreClient->InitializeReadOnly(ZenServerInfo->ProjectId, ZenServerInfo->OplogId);
+			return FIoStatus(EIoErrorCode::NotFound);
 		}
-		
-		if (ZenStoreClient)
+		FCbObject ManifestObject = LoadCompactBinary(*Ar).AsObject();
+		FCbObject OplogObject;
+		if (FCbFieldView ZenServerField = ManifestObject["zenserver"])
 		{
-			TFuture<FIoStatus> FutureOplogStatus = ZenStoreClient->GetOplog().Next([this](TIoStatusOr<FCbObject> OplogStatus)
-				{
-					if (!OplogStatus.IsOk())
-					{
-						return OplogStatus.Status();
-					}
-					FCbObject Oplog = OplogStatus.ConsumeValueOrDie();
-					for (FCbField& OplogEntry : Oplog["entries"].AsArray())
-					{
-						FCbObject OplogObj = OplogEntry.AsObject();
-						FPackageStoreEntryResource Entry = FPackageStoreEntryResource::FromCbObject(OplogObj["packagestoreentry"].AsObject());
-						FPackageId PackageId = Entry.GetPackageId();
-						PackageIdToEntry.Add(PackageId, MoveTemp(Entry));
-					}
-					return FIoStatus::Ok;
-				});
+			UE::Zen::FServiceSettings ZenServiceSettings;
+			ZenServiceSettings.ReadFromCompactBinary(ZenServerField["settings"]);
+			FString ProjectId = FString(ZenServerField["projectid"].AsString());
+			FString OplogId = FString(ZenServerField["oplogid"].AsString());
 
-			if (!FutureOplogStatus.Get().IsOk())
+			ZenStoreClient = MakeUnique<UE::FZenStoreHttpClient>(MoveTemp(ZenServiceSettings));
+			ZenStoreClient->InitializeReadOnly(ProjectId, OplogId);
+
+			TIoStatusOr<FCbObject> OplogStatus = ZenStoreClient->GetOplog().Get();
+			if (!OplogStatus.IsOk())
 			{
-				return FutureOplogStatus.Get();
+				return OplogStatus.Status();
 			}
+			
+			OplogObject = OplogStatus.ConsumeValueOrDie();
 		}
-		
-		ManifestPackageInfos = PackageStoreManifest.GetPackages();
-		for (const FPackageStoreManifest::FPackageInfo& PackageInfo : ManifestPackageInfos)
+		else
 		{
-			FName PackageName = FName(PackageInfo.PackageName);
-			PackageNameToPackageInfoMap.Add(PackageName, &PackageInfo);
-			for (const FIoChunkId& ChunkId : PackageInfo.ExportBundleChunkIds)
-			{
-				ChunkIdToPackageNameMap.Add(ChunkId, PackageName);
-			}
-			for (const FIoChunkId& ChunkId : PackageInfo.BulkDataChunkIds)
-			{
-				ChunkIdToPackageNameMap.Add(ChunkId, PackageName);
-			}
+			OplogObject = ManifestObject["oplog"].AsObject();
 		}
 
-		for (const FPackageStoreManifest::FFileInfo& FileInfo : PackageStoreManifest.GetFiles())
+		for (FCbField& OplogEntry : OplogObject["entries"].AsArray())
 		{
-			FilenameToChunkIdMap.Add(FileInfo.FileName, FileInfo.ChunkId);
-			ChunkIdToFileNameMap.Add(FileInfo.ChunkId, FileInfo.FileName.RightChop(CookedDir.Len() + 1));
-			FName* FindPackageName = ChunkIdToPackageNameMap.Find(FileInfo.ChunkId);
-			if (FindPackageName)
+			FCbObject OplogObj = OplogEntry.AsObject();
+			FPackageStoreEntryResource PackageStoreEntry = FPackageStoreEntryResource::FromCbObject(OplogObj["packagestoreentry"].AsObject());
+
+			auto AddChunksFromOplog = [this, &OplogEntry, &PackageStoreEntry](const char* Field)
 			{
-				FString FilenameWithoutExtension = FPaths::ChangeExtension(FileInfo.FileName, FString());
-				FilenameToPackageNameMap.Add(MoveTemp(FilenameWithoutExtension), *FindPackageName);
-			}
-		}
+				for (FCbField& ChunkEntry : OplogEntry[Field].AsArray())
+				{
+					FCbObject ChunkObj = ChunkEntry.AsObject();
+					FIoChunkId ChunkId;
+					ChunkId.Set(ChunkObj["id"].AsObjectId().GetView());
+					ChunkIdToPackageNameMap.Add(ChunkId, PackageStoreEntry.PackageName);
+					if (ChunkObj["filename"])
+					{
+						TStringBuilder<1024> RelativeFilename;
+						RelativeFilename.Append(ChunkObj["filename"].AsString());
+						ChunkIdToRelativeFileNameMap.Add(ChunkId, *RelativeFilename);
+						TStringBuilder<1024> PathBuilder;
+						FPathViews::AppendPath(PathBuilder, CookedDir);
+						FPathViews::AppendPath(PathBuilder, RelativeFilename);
+						FPathViews::NormalizeFilename(PathBuilder);
+						FilenameToChunkIdMap.Add(*PathBuilder, ChunkId);
+						FString FilenameWithoutExtension = FPathViews::ChangeExtension(PathBuilder, FStringView());
+						FilenameToPackageNameMap.Add(MoveTemp(FilenameWithoutExtension), PackageStoreEntry.PackageName);
+					}
+				}
+			};
 
+			AddChunksFromOplog("packagedata");
+			AddChunksFromOplog("bulkdata");
+
+			PackageIdToEntry.Add(PackageStoreEntry.GetPackageId(), MoveTemp(PackageStoreEntry));
+		}
 		return FIoStatus::Ok;
 	}
 
@@ -726,9 +727,9 @@ public:
 		return FilenameToChunkIdMap.FindRef(*Filename);
 	}
 
-	FString GetFilenameFromChunkId(const FIoChunkId& ChunkId) const
+	FString GetRelativeFilenameFromChunkId(const FIoChunkId& ChunkId) const
 	{
-		return ChunkIdToFileNameMap.FindRef(ChunkId);
+		return ChunkIdToRelativeFileNameMap.FindRef(ChunkId);
 	}
 
 	FName GetPackageNameFromChunkId(const FIoChunkId& ChunkId) const
@@ -742,11 +743,6 @@ public:
 		return FilenameToPackageNameMap.FindRef(FilenameWithoutExtension);
 	}
 
-	const FPackageStoreManifest::FPackageInfo* GetPackageInfoFromPackageName(FName PackageName)
-	{
-		return PackageNameToPackageInfoMap.FindRef(PackageName);
-	}
-
 	const FPackageStoreEntryResource* GetPackageStoreEntry(FPackageId PackageId) const
 	{
 		return PackageIdToEntry.Find(PackageId);
@@ -836,15 +832,12 @@ public:
 
 private:
 	TUniquePtr<UE::FZenStoreHttpClient> ZenStoreClient;
-	FPackageStoreManifest PackageStoreManifest;
 	FString CookedDir;
-	TArray<FPackageStoreManifest::FPackageInfo> ManifestPackageInfos;
 	TMap<FPackageId, FPackageStoreEntryResource> PackageIdToEntry;
 	TMap<FString, FIoChunkId> FilenameToChunkIdMap;
-	TMap<FIoChunkId, FString> ChunkIdToFileNameMap;
+	TMap<FIoChunkId, FString> ChunkIdToRelativeFileNameMap;
 	TMap<FIoChunkId, FName> ChunkIdToPackageNameMap;
 	TMap<FString, FName> FilenameToPackageNameMap;
-	TMap<FName, const FPackageStoreManifest::FPackageInfo*> PackageNameToPackageInfoMap;
 };
 
 struct FFileOrderMap
@@ -3908,7 +3901,7 @@ int32 CreateTarget(const FIoStoreArguments& Arguments, const FIoStoreWriterSetti
 						FPackageId PackageId;
 						if (!ChunkInfo.bHasValidFileName)
 						{
-							FString FileName = Arguments.PackageStore->GetFilenameFromChunkId(ChunkInfo.Id);
+							FString FileName = Arguments.PackageStore->GetRelativeFilenameFromChunkId(ChunkInfo.Id);
 							if (FileName.Len() > 0)
 							{
 								ChunkInfo.FileName = MoveTemp(FileName);
@@ -7372,6 +7365,39 @@ int32 CreateIoStoreContainerFiles(const TCHAR* CmdLine)
 		}
 		return GenerateZenFileSystemManifest(TargetPlatform);
 	}
+	else if (FParse::Param(FCommandLine::Get(), TEXT("StartZenServerForStage")))
+	{
+		FString ManifestFilename;
+		if (!FParse::Value(FCommandLine::Get(), TEXT("PackageStoreManifest="), ManifestFilename))
+		{
+			UE_LOG(LogIoStore, Error, TEXT("Expected -PackageStoreManifest=<path to package store manifest>"));
+			return -1;
+		}
+
+		TUniquePtr<FArchive> Ar(IFileManager::Get().CreateFileReader(*ManifestFilename));
+		if (!Ar)
+		{
+			UE_LOG(LogIoStore, Error, TEXT("Failed reading package store manifest"));
+			return -1;
+		}
+
+		FCbObject ManifestObject = LoadCompactBinary(*Ar).AsObject();
+		FCbObject OplogObject;
+		FCbField ZenServerField = ManifestObject["zenserver"];
+		if (ZenServerField)
+		{
+			UE::Zen::FServiceSettings ZenServiceSettings;
+			ZenServiceSettings.ReadFromCompactBinary(ZenServerField["settings"]);
+			FString ProjectId = FString(ZenServerField["projectid"].AsString());
+			FString OplogId = FString(ZenServerField["oplogid"].AsString());
+
+			// We just want the auto launch functionality
+			UE::FZenStoreHttpClient ZenStoreClient(MoveTemp(ZenServiceSettings));
+			ZenStoreClient.InitializeReadOnly(ProjectId, OplogId);
+		}
+
+		return 0;
+	}
 	else if (FParse::Value(FCommandLine::Get(), TEXT("CreateDLCContainer="), Arguments.DLCPluginPath))
 	{
 		if (!ParseContainerGenerationArguments(Arguments, WriterSettings))
diff --git a/Engine/Source/Developer/IoStoreUtilities/Private/PackageStoreManifest.cpp b/Engine/Source/Developer/IoStoreUtilities/Private/PackageStoreManifest.cpp
deleted file mode 100644
diff --git a/Engine/Source/Developer/IoStoreUtilities/Private/ZenStoreWriter.cpp b/Engine/Source/Developer/IoStoreUtilities/Private/ZenStoreWriter.cpp
--- a/Engine/Source/Developer/IoStoreUtilities/Private/ZenStoreWriter.cpp
+++ b/Engine/Source/Developer/IoStoreUtilities/Private/ZenStoreWriter.cpp
@@ -27,6 +27,7 @@
 #include "Misc/Paths.h"
 #include "GenericPlatform/GenericPlatformFile.h"
 #include "UObject/SavePackage.h"
+#include "Misc/PathViews.h"
 
 DEFINE_LOG_CATEGORY_STATIC(LogZenStoreWriter, Log, All);
 
@@ -157,7 +158,6 @@ FZenStoreWriter::FZenStoreWriter(
 	, TargetPlatformFName(*InTargetPlatform->PlatformName())
 	, OutputPath(InOutputPath)
 	, MetadataDirectoryPath(InMetadataDirectoryPath)
-	, PackageStoreManifest(InOutputPath)
 	, PackageStoreOptimizer(new FPackageStoreOptimizer())
 	, CookMode(ICookedPackageWriter::FCookInfo::CookByTheBookMode)
 	, bInitialized(false)
@@ -196,15 +196,6 @@ FZenStoreWriter::FZenStoreWriter(
 
 	PackageStoreOptimizer->Initialize();
 
-	FPackageStoreManifest::FZenServerInfo& ZenServerInfo = PackageStoreManifest.EditZenServerInfo();
-
-#if UE_WITH_ZEN
-	const UE::Zen::FZenServiceInstance& ZenServiceInstance = HttpClient->GetZenServiceInstance();
-	ZenServerInfo.Settings = ZenServiceInstance.GetServiceSettings();
-#endif
-	ZenServerInfo.ProjectId = ProjectId;
-	ZenServerInfo.OplogId = OplogId;
-
 	ZenFileSystemManifest = MakeUnique<FZenFileSystemManifest>(TargetPlatform, OutputPath);
 	
 	CommitQueue = MakeUnique<FCommitQueue>();
@@ -245,7 +236,6 @@ void FZenStoreWriter::WritePackageData(const FPackageInfo& Info, FLargeMemoryWri
 		Region.Offset -= Info.HeaderSize;
 	}
 	FIoBuffer PackageBuffer = PackageStoreOptimizer->CreatePackageBuffer(Entry.OptimizedPackage.Get(), CookedExportsBuffer, &FileRegionsCopy);
-	PackageStoreManifest.AddPackageData(Info.PackageName, Info.LooseFilePath, Info.ChunkId);
 	for (FFileRegion& Region : FileRegionsCopy)
 	{
 		// Adjust regions once more so they are relative to the exports bundle buffer
@@ -273,7 +263,6 @@ void FZenStoreWriter::WriteIoStorePackageData(const FPackageInfo& Info, const FI
 
 	TRACE_CPUPROFILER_EVENT_SCOPE(WriteIoStorePackageData);
 
-	PackageStoreManifest.AddPackageData(Info.PackageName, Info.LooseFilePath, Info.ChunkId);
 	//WriteFileRegions(*FPaths::ChangeExtension(Info.LooseFilePath, FString(".uexp") + FFileRegion::RegionsFileExtension), FileRegionsCopy);
 
 	FCbObjectId ChunkOid = ToObjectId(Info.ChunkId);
@@ -315,8 +304,6 @@ void FZenStoreWriter::WriteBulkData(const FBulkDataInfo& Info, const FIoBuffer&
 	BulkEntry.ChunkId	= ChunkOid;
 	BulkEntry.IsValid	= true;
 
-	PackageStoreManifest.AddBulkData(Info.PackageName, Info.LooseFilePath, Info.ChunkId);
-
 	//	WriteFileRegions(*(Info.LooseFilePath + FFileRegion::RegionsFileExtension), FileRegions);
 }
 
@@ -517,13 +504,8 @@ void FZenStoreWriter::Initialize(const FCookInfo& Info)
 
 void FZenStoreWriter::BeginCook(const FCookInfo& Info)
 {
-	if (!Info.bWorkerOnSharedSandbox)
+	if (Info.bWorkerOnSharedSandbox)
 	{
-		PackageStoreManifest.Load(*(MetadataDirectoryPath / TEXT("packagestore.manifest")));
-	}
-	else
-	{
-		PackageStoreManifest.SetTrackPackageData(true);
 		bProvidePerPackageResults = true;
 	}
 	AllPackageHashes.Empty();
@@ -591,7 +573,30 @@ void FZenStoreWriter::EndCook(const FCookInfo& Info)
 		TIoStatusOr<uint64> Status = HttpClient->EndBuildPass(Pkg);
 		UE_CLOG(!Status.IsOk(), LogZenStoreWriter, Fatal, TEXT("Failed to append OpLog and end the build pass"));
 
-		PackageStoreManifest.Save(*(MetadataDirectoryPath / TEXT("packagestore.manifest")));
+		FCbWriter ManifestWriter;
+		ManifestWriter.BeginObject();
+		ManifestWriter.BeginObject("zenserver");
+#if UE_WITH_ZEN
+		ManifestWriter.BeginObject("settings");
+		const UE::Zen::FZenServiceInstance& ZenServiceInstance = HttpClient->GetZenServiceInstance();
+		ZenServiceInstance.GetServiceSettings().WriteToCompactBinary(ManifestWriter);
+		ManifestWriter.EndObject();
+#endif
+		ManifestWriter << "projectid" << ProjectId;
+		ManifestWriter << "oplogid" << OplogId;
+		ManifestWriter.EndObject();
+		ManifestWriter.EndObject();
+
+		FString PackageStoreManifestFilePath = FPaths::Combine(MetadataDirectoryPath, TEXT("packagestore.manifest"));
+		TUniquePtr<FArchive> Ar(IFileManager::Get().CreateFileWriter(*PackageStoreManifestFilePath));
+		if (Ar)
+		{
+			SaveCompactBinary(*Ar, ManifestWriter.Save());
+		}
+		else
+		{
+			UE_LOG(LogSavePackage, Error, TEXT("Failed saving package store manifest file '%s'"), *PackageStoreManifestFilePath);
+		}
 	}
 
 	UE_LOG(LogZenStoreWriter, Display, TEXT("Input:\t%d Packages"), PackageStoreOptimizer->GetTotalPackageCount());
@@ -606,8 +611,6 @@ void FZenStoreWriter::BeginPackage(const FBeginPackageInfo& Info)
 {
 	FPendingPackageState& State = AddPendingPackage(Info.PackageName);
 	State.PackageName = Info.PackageName;
-
-	PackageStoreManifest.BeginPackage(Info.PackageName);
 }
 
 bool FZenStoreWriter::IsReservedOplogKey(FUtf8StringView Key)
@@ -798,6 +801,11 @@ void FZenStoreWriter::CommitPackageInternal(FZenCommitInfo&& ZenCommitInfo)
 			OplogEntryDesc.BeginObject();
 			OplogEntryDesc << "id" << PkgData.ChunkId;
 			OplogEntryDesc << "data" << PkgDataAttachment;
+			FStringView RelativePathView;
+			if (FPathViews::TryMakeChildPathRelativeTo(PkgData.Info.LooseFilePath, OutputPath, RelativePathView))
+			{
+				OplogEntryDesc << "filename" << RelativePathView;
+			}
 			OplogEntryDesc.EndObject();
 		}
 
@@ -822,6 +830,11 @@ void FZenStoreWriter::CommitPackageInternal(FZenCommitInfo&& ZenCommitInfo)
 				OplogEntryDesc << "id" << Bulk.ChunkId;
 				OplogEntryDesc << "type" << LexToString(Bulk.Info.BulkDataType);
 				OplogEntryDesc << "data" << BulkAttachment;
+				FStringView RelativePathView;
+				if (FPathViews::TryMakeChildPathRelativeTo(Bulk.Info.LooseFilePath, OutputPath, RelativePathView))
+				{
+					OplogEntryDesc << "filename" << RelativePathView;
+				}
 				OplogEntryDesc.EndObject();
 			}
 
@@ -1107,10 +1120,6 @@ void FZenStoreWriter::MarkPackagesUpToDate(TArrayView<const FName> UpToDatePacka
 
 TFuture<FCbObject> FZenStoreWriter::WriteMPCookMessageForPackage(FName PackageName)
 {
-	FCbWriter ManifestWriter;
-	PackageStoreManifest.WritePackage(ManifestWriter, PackageName);
-	FCbFieldIterator ManifestField = ManifestWriter.Save();
-
 	TArray<FString> AdditionalFiles;
 	PackageAdditionalFiles.RemoveAndCopyValue(PackageName, AdditionalFiles);
 
@@ -1118,11 +1127,10 @@ TFuture<FCbObject> FZenStoreWriter::WriteMPCookMessageForPackage(FName PackageNa
 	AllPackageHashes.RemoveAndCopyValue(PackageName, PackageHashes);
 
 	auto ComposeMessage =
-	[ManifestField = MoveTemp(ManifestField), AdditionalFiles=MoveTemp(AdditionalFiles)](FPackageHashes* PackageHashes)
+	[AdditionalFiles=MoveTemp(AdditionalFiles)](FPackageHashes* PackageHashes)
 	{
 		FCbWriter Writer;
 		Writer.BeginObject();
-		Writer << "Manifest" << ManifestField;
 		if (!AdditionalFiles.IsEmpty())
 		{
 			Writer << "AdditionalFiles" << AdditionalFiles;
@@ -1157,12 +1165,6 @@ TFuture<FCbObject> FZenStoreWriter::WriteMPCookMessageForPackage(FName PackageNa
 
 bool FZenStoreWriter::TryReadMPCookMessageForPackage(FName PackageName, FCbObjectView Message)
 {
-	bool bOk = PackageStoreManifest.TryReadPackage(Message["Manifest"], PackageName);
-	if (!bOk)
-	{
-		return false;
-	}
-
 	TArray<FString> AdditionalFiles;
 	if (LoadFromCompactBinary(Message["AdditionalFiles"], AdditionalFiles))
 	{
@@ -1172,6 +1174,7 @@ bool FZenStoreWriter::TryReadMPCookMessageForPackage(FName PackageName, FCbObjec
 		}
 	}
 
+	bool bOk = true;
 	TRefCountPtr<FPackageHashes> ThisPackageHashes(new FPackageHashes());
 	if (LoadFromCompactBinary(Message["PackageHash"], ThisPackageHashes->PackageHash))
 	{
diff --git a/Engine/Source/Developer/IoStoreUtilities/Public/PackageStoreManifest.h b/Engine/Source/Developer/IoStoreUtilities/Public/PackageStoreManifest.h
deleted file mode 100644
diff --git a/Engine/Source/Developer/IoStoreUtilities/Public/ZenStoreWriter.h b/Engine/Source/Developer/IoStoreUtilities/Public/ZenStoreWriter.h
--- a/Engine/Source/Developer/IoStoreUtilities/Public/ZenStoreWriter.h
+++ b/Engine/Source/Developer/IoStoreUtilities/Public/ZenStoreWriter.h
@@ -20,7 +20,6 @@
 #include "Misc/AssertionMacros.h"
 #include "Misc/ScopeLock.h"
 #include "Misc/ScopeRWLock.h"
-#include "PackageStoreManifest.h"
 #include "PackageStoreWriter.h"
 #include "Serialization/AsyncLoading2.h"
 #include "Serialization/CompactBinary.h"
@@ -191,7 +190,6 @@ private:
 	FString								MetadataDirectoryPath;
 	TMap<FName, TRefCountPtr<FPackageHashes>> AllPackageHashes;
 
-	FPackageStoreManifest				PackageStoreManifest;
 	TUniquePtr<FPackageStoreOptimizer>	PackageStoreOptimizer;
 	
 	FRWLock								EntriesLock;
diff --git a/Engine/Source/Developer/Zen/Private/ZenServerInterface.cpp b/Engine/Source/Developer/Zen/Private/ZenServerInterface.cpp
--- a/Engine/Source/Developer/Zen/Private/ZenServerInterface.cpp
+++ b/Engine/Source/Developer/Zen/Private/ZenServerInterface.cpp
@@ -31,6 +31,7 @@
 #include "Serialization/JsonWriter.h"
 #include "String/LexFromString.h"
 #include "ZenVersion.h"
+#include "Serialization/CompactBinaryWriter.h"
 
 #if PLATFORM_WINDOWS
 #	include "Windows/AllowWindowsPlatformTypes.h"
@@ -593,43 +594,35 @@ FServiceSettings::ReadFromConfig()
 }
 
 void
-FServiceSettings::ReadFromJson(FJsonObject& JsonObject)
+FServiceSettings::ReadFromCompactBinary(FCbFieldView Field)
 {
-	if (TSharedPtr<FJsonValue> bAutoLaunchValue = JsonObject.Values.FindRef(TEXT("bAutoLaunch")))
+	if (bool bAutoLaunchValue = Field["bAutoLaunch"].AsBool())
 	{
-		if (bAutoLaunchValue->AsBool())
-		{
-			if (!TryApplyAutoLaunchOverride())
-			{
-				SettingsVariant.Emplace<FServiceAutoLaunchSettings>();
-				FServiceAutoLaunchSettings& AutoLaunchSettings = SettingsVariant.Get<FServiceAutoLaunchSettings>();
-
-				TSharedPtr<FJsonValue> AutoLaunchSettingsValue = JsonObject.Values.FindRef(TEXT("AutoLaunchSettings"));
-				if (AutoLaunchSettingsValue)
-				{
-					TSharedPtr<FJsonObject> AutoLaunchSettingsObject = AutoLaunchSettingsValue->AsObject();
-					AutoLaunchSettings.DataPath = AutoLaunchSettingsObject->Values.FindRef(TEXT("DataPath"))->AsString();
-					AutoLaunchSettings.ExtraArgs = AutoLaunchSettingsObject->Values.FindRef(TEXT("ExtraArgs"))->AsString();
-					AutoLaunchSettingsObject->Values.FindRef(TEXT("DesiredPort"))->TryGetNumber(AutoLaunchSettings.DesiredPort);
-					AutoLaunchSettingsObject->Values.FindRef(TEXT("ShowConsole"))->TryGetBool(AutoLaunchSettings.bShowConsole);
-					AutoLaunchSettingsObject->Values.FindRef(TEXT("LimitProcessLifetime"))->TryGetBool(AutoLaunchSettings.bLimitProcessLifetime);
-				}
-			}
-		}
-		else
+		if (!TryApplyAutoLaunchOverride())
 		{
-			SettingsVariant.Emplace<FServiceConnectSettings>();
-			FServiceConnectSettings& ConnectExistingSettings = SettingsVariant.Get<FServiceConnectSettings>();
+			SettingsVariant.Emplace<FServiceAutoLaunchSettings>();
+			FServiceAutoLaunchSettings& AutoLaunchSettings = SettingsVariant.Get<FServiceAutoLaunchSettings>();
 
-			TSharedPtr<FJsonValue> ConnectExistingSettingsValue = JsonObject.Values.FindRef(TEXT("ConnectExistingSettings"));
-			if (ConnectExistingSettingsValue)
+			if (FCbObjectView AutoLaunchSettingsObject = Field["AutoLaunchSettings"].AsObjectView())
 			{
-				TSharedPtr<FJsonObject> ConnectExistingSettingsObject = ConnectExistingSettingsValue->AsObject();
-				ConnectExistingSettings.HostName = ConnectExistingSettingsObject->Values.FindRef(TEXT("HostName"))->AsString();
-				ConnectExistingSettingsObject->Values.FindRef(TEXT("Port"))->TryGetNumber(ConnectExistingSettings.Port);
+				AutoLaunchSettings.DataPath = FString(AutoLaunchSettingsObject["DataPath"].AsString());
+				AutoLaunchSettings.ExtraArgs = FString(AutoLaunchSettingsObject["ExtraArgs"].AsString());
+				AutoLaunchSettings.DesiredPort = AutoLaunchSettingsObject["DesiredPort"].AsInt16();
+				AutoLaunchSettings.bShowConsole = AutoLaunchSettingsObject["ShowConsole"].AsBool();
+				AutoLaunchSettings.bLimitProcessLifetime = AutoLaunchSettingsObject["LimitProcessLifetime"].AsBool();
 			}
 		}
+	}
+	else
+	{
+		SettingsVariant.Emplace<FServiceConnectSettings>();
+		FServiceConnectSettings& ConnectExistingSettings = SettingsVariant.Get<FServiceConnectSettings>();
 
+		if (FCbObjectView ConnectExistingSettingsObject = Field["ConnectExistingSettings"].AsObjectView())
+		{
+			ConnectExistingSettings.HostName = FString(ConnectExistingSettingsObject["HostName"].AsString());
+			ConnectExistingSettings.Port = ConnectExistingSettingsObject["Port"].AsInt16();
+		}
 	}
 }
 
@@ -659,28 +652,28 @@ FServiceSettings::ReadFromURL(FStringView InstanceURL)
 }
 
 void
-FServiceSettings::WriteToJson(TJsonWriter<TCHAR, TPrettyJsonPrintPolicy<TCHAR>>& Writer) const
+FServiceSettings::WriteToCompactBinary(FCbWriter& Writer) const
 {
 	bool bAutoLaunch = IsAutoLaunch();
-	Writer.WriteValue(TEXT("bAutoLaunch"), bAutoLaunch);
+	Writer << "bAutoLaunch" << bAutoLaunch;
 	if (bAutoLaunch)
 	{
 		const FServiceAutoLaunchSettings& AutoLaunchSettings = SettingsVariant.Get<FServiceAutoLaunchSettings>();
-		Writer.WriteObjectStart(TEXT("AutoLaunchSettings"));
-		Writer.WriteValue(TEXT("DataPath"), AutoLaunchSettings.DataPath);
-		Writer.WriteValue(TEXT("ExtraArgs"), AutoLaunchSettings.ExtraArgs);
-		Writer.WriteValue(TEXT("DesiredPort"), AutoLaunchSettings.DesiredPort);
-		Writer.WriteValue(TEXT("ShowConsole"), AutoLaunchSettings.bShowConsole);
-		Writer.WriteValue(TEXT("LimitProcessLifetime"), AutoLaunchSettings.bLimitProcessLifetime);
-		Writer.WriteObjectEnd();
+		Writer.BeginObject("AutoLaunchSettings");
+		Writer << "DataPath" << AutoLaunchSettings.DataPath;
+		Writer << "ExtraArgs" <<AutoLaunchSettings.ExtraArgs;
+		Writer << "DesiredPort" << AutoLaunchSettings.DesiredPort;
+		Writer << "ShowConsole" << AutoLaunchSettings.bShowConsole;
+		Writer << "LimitProcessLifetime" << AutoLaunchSettings.bLimitProcessLifetime;
+		Writer.EndObject();
 	}
 	else
 	{
 		const FServiceConnectSettings& ConnectExistingSettings = SettingsVariant.Get<FServiceConnectSettings>();
-		Writer.WriteObjectStart(TEXT("ConnectExistingSettings"));
-		Writer.WriteValue(TEXT("HostName"), ConnectExistingSettings.HostName);
-		Writer.WriteValue(TEXT("Port"), ConnectExistingSettings.Port);
-		Writer.WriteObjectEnd();
+		Writer.BeginObject("ConnectExistingSettings");
+		Writer << "HostName" << ConnectExistingSettings.HostName;
+		Writer << "Port" << ConnectExistingSettings.Port;
+		Writer.EndObject();
 	}
 }
 
diff --git a/Engine/Source/Developer/Zen/Public/Experimental/ZenServerInterface.h b/Engine/Source/Developer/Zen/Public/Experimental/ZenServerInterface.h
--- a/Engine/Source/Developer/Zen/Public/Experimental/ZenServerInterface.h
+++ b/Engine/Source/Developer/Zen/Public/Experimental/ZenServerInterface.h
@@ -24,6 +24,7 @@
 #define UE_API ZEN_API
 
 struct FAnalyticsEventAttribute;
+class FCbFieldView;
 
 namespace UE::Zen
 {
@@ -51,10 +52,10 @@ struct FServiceSettings
 	inline bool IsConnectExisting() const { return SettingsVariant.IsType<FServiceConnectSettings>(); }
 
 	UE_API void ReadFromConfig();
-	UE_API void ReadFromJson(FJsonObject& JsonObject);
+	UE_API void ReadFromCompactBinary(FCbFieldView Field);
 	UE_API void ReadFromURL(FStringView InstanceURL);
 
-	UE_API void WriteToJson(TJsonWriter<TCHAR, TPrettyJsonPrintPolicy<TCHAR>>& Writer) const;
+	UE_API void WriteToCompactBinary(FCbWriter& Writer) const;
 
 private:
 	bool TryApplyAutoLaunchOverride();
diff --git a/Engine/Source/Editor/UnrealEd/Private/Cooker/LooseCookedPackageWriter.cpp b/Engine/Source/Editor/UnrealEd/Private/Cooker/LooseCookedPackageWriter.cpp
--- a/Engine/Source/Editor/UnrealEd/Private/Cooker/LooseCookedPackageWriter.cpp
+++ b/Engine/Source/Editor/UnrealEd/Private/Cooker/LooseCookedPackageWriter.cpp
@@ -55,6 +55,8 @@
 #include "UObject/Package.h"
 #include "UObject/SavePackage.h"
 
+LLM_DEFINE_TAG(Cooker_PackageStoreManifest);
+
 FLooseCookedPackageWriter::FLooseCookedPackageWriter(const FString& InOutputPath,
 	const FString& InMetadataDirectoryPath, const ITargetPlatform* InTargetPlatform, FAsyncIODelete& InAsyncIODelete,
 	UE::Cook::FPackageDatas& InPackageDatas, const TArray<TSharedRef<IPlugin>>& InPluginsToRemap)
@@ -62,7 +64,6 @@ FLooseCookedPackageWriter::FLooseCookedPackageWriter(const FString& InOutputPath
 	, MetadataDirectoryPath(InMetadataDirectoryPath)
 	, TargetPlatform(*InTargetPlatform)
 	, PackageDatas(InPackageDatas)
-	, PackageStoreManifest(InOutputPath)
 	, PluginsToRemap(InPluginsToRemap)
 	, AsyncIODelete(InAsyncIODelete)
 {
@@ -75,7 +76,12 @@ FLooseCookedPackageWriter::~FLooseCookedPackageWriter()
 void FLooseCookedPackageWriter::BeginPackage(const FBeginPackageInfo& Info)
 {
 	Super::BeginPackage(Info);
-	PackageStoreManifest.BeginPackage(Info.PackageName);
+	LLM_SCOPE_BYTAG(Cooker_PackageStoreManifest);
+	FScopeLock Lock(&OplogLock);
+	FOplogPackageInfo& PackageInfo = Oplog.FindOrAdd(Info.PackageName);
+	PackageInfo.PackageName = Info.PackageName;
+	PackageInfo.PackageDataChunks.Reset();
+	PackageInfo.BulkDataChunks.Reset();
 }
 
 void FLooseCookedPackageWriter::CommitPackageInternal(FPackageWriterRecords::FPackage&& BaseRecord,
@@ -489,13 +495,28 @@ void FLooseCookedPackageWriter::FWriteFileData::HashAndWrite(FMD5& AccumulatedHa
 
 void FLooseCookedPackageWriter::UpdateManifest(FRecord& Record)
 {
+	LLM_SCOPE_BYTAG(Cooker_PackageStoreManifest);
+	FScopeLock Lock(&OplogLock);
 	for (const FPackageWriterRecords::FWritePackage& Package : Record.Packages)
 	{
-		PackageStoreManifest.AddPackageData(Package.Info.PackageName, Package.Info.LooseFilePath, Package.Info.ChunkId);
+		FOplogPackageInfo* PackageInfo = Oplog.Find(Package.Info.PackageName);
+		check(PackageInfo);
+		FOplogChunkInfo& ChunkInfo = PackageInfo->PackageDataChunks.AddDefaulted_GetRef();
+		ChunkInfo.ChunkId = Package.Info.ChunkId;
+		FStringView RelativePathView;
+		FPathViews::TryMakeChildPathRelativeTo(Package.Info.LooseFilePath, OutputPath, RelativePathView);
+		ChunkInfo.RelativeFileName = RelativePathView;
+
 	}
 	for (const FPackageWriterRecords::FBulkData& BulkData : Record.BulkDatas)
 	{
-		PackageStoreManifest.AddBulkData(BulkData.Info.PackageName, BulkData.Info.LooseFilePath, BulkData.Info.ChunkId);
+		FOplogPackageInfo* PackageInfo = Oplog.Find(BulkData.Info.PackageName);
+		check(PackageInfo);
+		FOplogChunkInfo& ChunkInfo = PackageInfo->BulkDataChunks.AddDefaulted_GetRef();
+		ChunkInfo.ChunkId = BulkData.Info.ChunkId;
+		FStringView RelativePathView;
+		FPathViews::TryMakeChildPathRelativeTo(BulkData.Info.LooseFilePath, OutputPath, RelativePathView);
+		ChunkInfo.RelativeFileName = RelativePathView;
 	}
 }
 
@@ -534,15 +555,96 @@ void FLooseCookedPackageWriter::Initialize(const FCookInfo& Info)
 	}
 }
 
+void FLooseCookedPackageWriter::WriteOplogEntry(FCbWriter& Writer, const FOplogPackageInfo& PackageInfo)
+{
+	Writer.BeginObject();
+
+	Writer.BeginObject("packagestoreentry");
+	Writer << "packagename" << PackageInfo.PackageName;
+	Writer.EndObject();
+
+	Writer.BeginArray("packagedata");
+	for (const FOplogChunkInfo& ChunkInfo : PackageInfo.PackageDataChunks)
+	{
+		Writer.BeginObject();
+		Writer << "id" << ChunkInfo.ChunkId;
+		Writer << "filename" << ChunkInfo.RelativeFileName;
+		Writer.EndObject();
+	}
+	Writer.EndArray();
+
+	Writer.BeginArray("bulkdata");
+	for (const FOplogChunkInfo& ChunkInfo : PackageInfo.BulkDataChunks)
+	{
+		Writer.BeginObject();
+		Writer << "id" << ChunkInfo.ChunkId;
+		Writer << "filename" << ChunkInfo.RelativeFileName;
+		Writer.EndObject();
+	}
+	Writer.EndArray();
+
+	Writer.EndObject();
+}
+
+bool FLooseCookedPackageWriter::ReadOplogEntry(FOplogPackageInfo& PackageInfo, const FCbFieldView& Field)
+{
+	FCbObjectView PackageStoreEntryObjectView = Field["packagestoreentry"].AsObjectView();
+	if (!PackageStoreEntryObjectView)
+	{
+		return false;
+	}
+	PackageInfo.PackageName = FName(PackageStoreEntryObjectView["packagename"].AsString());
+
+	FCbArrayView PackageDataView = Field["packagedata"].AsArrayView();
+	PackageInfo.PackageDataChunks.Reset();
+	PackageInfo.PackageDataChunks.Reserve(PackageDataView.Num());
+	for (const FCbFieldView& ChunkEntry : PackageDataView)
+	{
+		FOplogChunkInfo& ChunkInfo = PackageInfo.PackageDataChunks.AddDefaulted_GetRef();
+		ChunkInfo.ChunkId.Set(ChunkEntry["id"].AsObjectId().GetView());
+		ChunkInfo.RelativeFileName = FString(ChunkEntry["filename"].AsString());
+	}
+
+	FCbArrayView BulkDataView = Field["bulkdata"].AsArrayView();
+	PackageInfo.BulkDataChunks.Reset();
+	PackageInfo.BulkDataChunks.Reserve(BulkDataView.Num());
+	for (const FCbFieldView& ChunkEntry : BulkDataView)
+	{
+		FOplogChunkInfo& ChunkInfo = PackageInfo.BulkDataChunks.AddDefaulted_GetRef();
+		ChunkInfo.ChunkId.Set(ChunkEntry["id"].AsObjectId().GetView());
+		ChunkInfo.RelativeFileName = FString(ChunkEntry["filename"].AsString());
+	}
+
+	return true;
+}
+
 void FLooseCookedPackageWriter::BeginCook(const FCookInfo& Info)
 {
 	if (!Info.bWorkerOnSharedSandbox)
 	{
-		PackageStoreManifest.Load(*(MetadataDirectoryPath / TEXT("packagestore.manifest")));
+		TRACE_CPUPROFILER_EVENT_SCOPE(LoadPackageStoreManifest);
+		FString PackageStoreManifestFilePath = FPaths::Combine(MetadataDirectoryPath, TEXT("packagestore.manifest"));
+		TUniquePtr<FArchive> Ar(IFileManager::Get().CreateFileReader(*PackageStoreManifestFilePath));
+		if (Ar)
+		{
+			FCbField ManifestField = LoadCompactBinary(*Ar);
+			FCbField OplogField = ManifestField["oplog"];
+			if (OplogField)
+			{
+				FCbArray EntriesArray = OplogField["entries"].AsArray();
+				FScopeLock Lock(&OplogLock);
+				Oplog.Reserve(EntriesArray.Num());
+				for (const FCbField& EntryField : EntriesArray)
+				{
+					FOplogPackageInfo PackageInfo;
+					ReadOplogEntry(PackageInfo, EntryField);
+					Oplog.Add(PackageInfo.PackageName, MoveTemp(PackageInfo));
+				}
+			}
+		}
 	}
 	else
 	{
-		PackageStoreManifest.SetTrackPackageData(true);
 		bProvidePerPackageResults = true;
 	}
 	AllPackageHashes.Empty();
@@ -552,7 +654,41 @@ void FLooseCookedPackageWriter::EndCook(const FCookInfo& Info)
 {
 	if (!Info.bWorkerOnSharedSandbox)
 	{
-		PackageStoreManifest.Save(*(MetadataDirectoryPath / TEXT("packagestore.manifest")));
+		TRACE_CPUPROFILER_EVENT_SCOPE(SavePackageStoreManifest);
+		FScopeLock Lock(&OplogLock);
+		// Sort packages for determinism
+		TArray<const FOplogPackageInfo*> SortedPackages;
+		SortedPackages.Reserve(Oplog.Num());
+		for (const auto& KV : Oplog)
+		{
+			SortedPackages.Add(&KV.Value);
+		}
+		Algo::Sort(SortedPackages, [](const FOplogPackageInfo* A, const FOplogPackageInfo* B)
+			{
+				return A->PackageName.LexicalLess(B->PackageName);
+			});
+		FCbWriter ManifestWriter;
+		ManifestWriter.BeginObject();
+		ManifestWriter.BeginObject("oplog");
+		ManifestWriter.BeginArray("entries");
+		for (const FOplogPackageInfo* Package : SortedPackages)
+		{
+			WriteOplogEntry(ManifestWriter, *Package);
+		}
+		ManifestWriter.EndArray();
+		ManifestWriter.EndObject();
+		ManifestWriter.EndObject();
+
+		FString PackageStoreManifestFilePath = FPaths::Combine(MetadataDirectoryPath, TEXT("packagestore.manifest"));
+		TUniquePtr<FArchive> Ar(IFileManager::Get().CreateFileWriter(*PackageStoreManifestFilePath));
+		if (Ar)
+		{
+			SaveCompactBinary(*Ar, ManifestWriter.Save());
+		}
+		else
+		{
+			UE_LOG(LogSavePackage, Error, TEXT("Failed saving package store manifest file '%s'"), *PackageStoreManifestFilePath);
+		}
 	}
 }
 
@@ -706,9 +842,22 @@ void FLooseCookedPackageWriter::MarkPackagesUpToDate(TArrayView<const FName> UpT
 
 TFuture<FCbObject> FLooseCookedPackageWriter::WriteMPCookMessageForPackage(FName PackageName)
 {
-	FCbWriter ManifestWriter;
-	PackageStoreManifest.WritePackage(ManifestWriter, PackageName);
-	FCbFieldIterator ManifestField = ManifestWriter.Save();
+	FCbFieldIterator OplogEntryField;
+	{
+		FOplogPackageInfo PackageInfo;
+		bool bValid = false;
+		{
+			FScopeLock Lock(&OplogLock);
+			bValid = Oplog.RemoveAndCopyValue(PackageName, PackageInfo);
+		}
+		if (bValid)
+		{
+			check(PackageName == PackageInfo.PackageName);
+			FCbWriter OplogEntryWriter;
+			WriteOplogEntry(OplogEntryWriter, PackageInfo);
+			OplogEntryField = OplogEntryWriter.Save();
+		}
+	}
 
 	TRefCountPtr<FPackageHashes> PackageHashes;
 	{
@@ -716,11 +865,14 @@ TFuture<FCbObject> FLooseCookedPackageWriter::WriteMPCookMessageForPackage(FName
 		AllPackageHashes.RemoveAndCopyValue(PackageName, PackageHashes);
 	}
 
-	auto ComposeMessage = [ManifestField = MoveTemp(ManifestField)](FPackageHashes* PackageHashes)
+	auto ComposeMessage = [OplogEntryField = MoveTemp(OplogEntryField)](FPackageHashes* PackageHashes)
 	{
 		FCbWriter Writer;
 		Writer.BeginObject();
-		Writer << "Manifest" << ManifestField;
+		if (OplogEntryField.HasValue())
+		{
+			Writer << "OplogEntry" << OplogEntryField;
+		}
 		if (PackageHashes)
 		{
 			Writer << "PackageHash" << PackageHashes->PackageHash;
@@ -752,9 +904,13 @@ TFuture<FCbObject> FLooseCookedPackageWriter::WriteMPCookMessageForPackage(FName
 
 bool FLooseCookedPackageWriter::TryReadMPCookMessageForPackage(FName PackageName, FCbObjectView Message)
 {
-	if (!PackageStoreManifest.TryReadPackage(Message["Manifest"], PackageName))
+	FOplogPackageInfo PackageInfo;
+	FCbFieldView OplogEntryField(Message["OplogEntry"]);
+	if (ReadOplogEntry(PackageInfo, OplogEntryField))
 	{
-		return false;
+		check(PackageName == PackageInfo.PackageName);
+		FScopeLock Lock(&OplogLock);
+		Oplog.Add(PackageName, MoveTemp(PackageInfo));
 	}
 
 	bool bOk = true;
diff --git a/Engine/Source/Editor/UnrealEd/Private/Cooker/LooseCookedPackageWriter.h b/Engine/Source/Editor/UnrealEd/Private/Cooker/LooseCookedPackageWriter.h
--- a/Engine/Source/Editor/UnrealEd/Private/Cooker/LooseCookedPackageWriter.h
+++ b/Engine/Source/Editor/UnrealEd/Private/Cooker/LooseCookedPackageWriter.h
@@ -14,7 +14,6 @@
 #include "Memory/SharedBuffer.h"
 #include "Misc/DateTime.h"
 #include "Misc/PackagePath.h"
-#include "PackageStoreManifest.h"
 #include "Serialization/CompactBinary.h"
 #include "Serialization/FileRegions.h"
 #include "Serialization/PackageWriter.h"
@@ -114,6 +113,19 @@ private:
 		TArray<FWriteFileData> OutputFiles;
 	};
 
+	struct FOplogChunkInfo
+	{
+		FString RelativeFileName;
+		FIoChunkId ChunkId;
+	};
+
+	struct FOplogPackageInfo
+	{
+		FName PackageName;
+		TArray<FOplogChunkInfo, TInlineAllocator<1>> PackageDataChunks;
+		TArray<FOplogChunkInfo> BulkDataChunks;
+	};
+
 	/* Delete the sandbox directory (asynchronously) in preparation for a clean cook */
 	void DeleteSandboxDirectory();
 	/** Searches the disk for all the cooked files in the sandbox. Stores results in PackageNameToCookedFiles. */
@@ -141,6 +153,8 @@ private:
 	void CollectForSaveExportsBuffers(FRecord& Record, FCommitContext& Context);
 	void AsyncSaveOutputFiles(FRecord& Record, FCommitContext& Context);
 	void UpdateManifest(FRecord& Record);
+	void WriteOplogEntry(FCbWriter& Writer, const FOplogPackageInfo& PackageInfo);
+	bool ReadOplogEntry(FOplogPackageInfo& PackageInfo, const FCbFieldView& Field);
 
 	TMap<FName, TRefCountPtr<FPackageHashes>>& GetPackageHashes() override
 	{
@@ -153,11 +167,12 @@ private:
 	TMap<FName, TArray<FString>> PackageNameToCookedFiles;
 	/** CommitPackage can be called in parallel if using recursive save, so we need a lock for shared containers used during CommitPackage */
 	FCriticalSection PackageHashesLock;
+	FCriticalSection OplogLock;
 	FString OutputPath;
 	FString MetadataDirectoryPath;
 	const ITargetPlatform& TargetPlatform;
 	UE::Cook::FPackageDatas& PackageDatas;
-	FPackageStoreManifest PackageStoreManifest;
+	TMap<FName, FOplogPackageInfo> Oplog;
 	const TArray<TSharedRef<IPlugin>>& PluginsToRemap;
 	FAsyncIODelete& AsyncIODelete;
 	bool bIterateSharedBuild = false;
diff --git a/Engine/Source/Programs/AutomationTool/AutomationUtils/DeploymentContext.cs b/Engine/Source/Programs/AutomationTool/AutomationUtils/DeploymentContext.cs
--- a/Engine/Source/Programs/AutomationTool/AutomationUtils/DeploymentContext.cs
+++ b/Engine/Source/Programs/AutomationTool/AutomationUtils/DeploymentContext.cs
@@ -96,22 +96,8 @@ public class FilesToStage
 
 public class PackageStoreManifest
 {
-	public class ZenServerInfo
-	{
-		public string HostName { get; set; }
-		public int Port { get; set; }
-		public string ProjectId { get; set; }
-		public string OplogId { get; set; }
-	}
-	public class FileInfo
-	{
-		public string Path { get; set; }
-	}
-
 	public string FullPath { get; set; }
-
-	public ZenServerInfo ZenServer { get; set; }
-	public IList<FileInfo> Files { get; set; }
+	public IList<string> ZenCookedFiles { get; set; }
 }
 
 public class DeploymentContext //: ProjectParams
diff --git a/Engine/Source/Programs/AutomationTool/Scripts/CopyBuildToStagingDirectory.Automation.cs b/Engine/Source/Programs/AutomationTool/Scripts/CopyBuildToStagingDirectory.Automation.cs
--- a/Engine/Source/Programs/AutomationTool/Scripts/CopyBuildToStagingDirectory.Automation.cs
+++ b/Engine/Source/Programs/AutomationTool/Scripts/CopyBuildToStagingDirectory.Automation.cs
@@ -17,6 +17,8 @@ using System.Diagnostics;
 using System.Collections.Concurrent;
 using UnrealBuildBase;
 using System.Text.Json;
+using System.Net.Http;
+using EpicGames.Serialization;
 using Microsoft.Extensions.Logging;
 
 using static AutomationTool.CommandUtils;
@@ -748,14 +750,100 @@ namespace AutomationScripts
 				return;
 			}
 			FileReference PackageStoreManifestFile = FileReference.Combine(SC.MetadataDir, "packagestore.manifest");
-			if (FileReference.Exists(PackageStoreManifestFile))
+			System.IO.FileInfo PackageStoreManifestFileInfo = PackageStoreManifestFile.ToFileInfo();
+			if (PackageStoreManifestFileInfo.Exists)
 			{
-				String JsonString = FileReference.ReadAllText(PackageStoreManifestFile);
-				SC.PackageStoreManifest = JsonSerializer.Deserialize<PackageStoreManifest>(JsonString);
+				SC.PackageStoreManifest = new PackageStoreManifest();
 				SC.PackageStoreManifest.FullPath = PackageStoreManifestFile.FullName;
-				foreach (PackageStoreManifest.FileInfo FileInfo in SC.PackageStoreManifest.Files)
+				if (PackageStoreManifestFileInfo.Length > (64 << 10))
 				{
-					FileInfo.Path = FileReference.Combine(SC.PlatformCookDir, FileInfo.Path).FullName;
+					// We assume that this is a loose file deploy and we only need to read zen manifests here
+					return;
+				}
+
+				byte[] ManifestData = File.ReadAllBytes(PackageStoreManifestFile.FullName);
+				CbObject ManifestObject = new CbField(ManifestData).AsObject();
+				CbObject ZenServerObject = ManifestObject["zenserver"].AsObject();
+				if (ZenServerObject != CbObject.Empty)
+				{
+					CommandUtils.LogInformation("Reading oplog from Zen...");
+					string Host = "localhost";
+					int Port = 1337;
+					string ProjectId = ZenServerObject["projectid"].AsString();
+					string OplogId = ZenServerObject["oplogid"].AsString();
+					bool bAutoLaunch = false;
+					CbObject SettingsObject = ZenServerObject["settings"].AsObject();
+					if (SettingsObject != CbObject.Empty)
+					{
+						bAutoLaunch = SettingsObject["bAutoLaunch"].AsBool();
+						CbObject AutoLaunchSettingsObject = SettingsObject["AutoLaunchSettings"].AsObject();
+						if (AutoLaunchSettingsObject != CbObject.Empty)
+						{
+							Port = AutoLaunchSettingsObject["DesiredPort"].AsInt16();
+						}
+						CbObject ConnectExistingSettingsObject = SettingsObject["ConnectExistingSettings"].AsObject();
+						if (ConnectExistingSettingsObject != CbObject.Empty)
+						{
+							Host = ConnectExistingSettingsObject["HostName"].AsString();
+							Port = ConnectExistingSettingsObject["Port"].AsInt16();
+						}
+					}
+
+					SC.PackageStoreManifest.ZenCookedFiles = new List<string>();
+
+					bool bAttemptAutoLaunchOnFailure = bAutoLaunch;
+					HttpResponseMessage HttpGetResult = null;
+					for (; ; )
+					{
+						HttpClient HttpClient = new HttpClient();
+						using var Request = new HttpRequestMessage(HttpMethod.Get, string.Format("http://{0}:{1}/prj/{2}/oplog/{3}/entries", Host, Port, ProjectId, OplogId));
+						Request.Headers.Add("Accept", "application/x-ue-cb");
+						try
+						{
+							HttpGetResult = HttpClient.Send(Request);
+							break;
+						}
+						catch
+						{
+							if (bAttemptAutoLaunchOnFailure)
+							{
+								string Arguments = String.Format("IoStore -StartZenServerForStage -PackageStoreManifest={0}", MakePathSafeToUseWithCommandLine(SC.PackageStoreManifest.FullPath));
+								if (Params.Unattended)
+								{
+									Arguments += " -unattended";
+								}
+
+								LogInformation("Running UnrealPak with arguments: {0}", Arguments);
+								RunAndLog(CmdEnv, GetUnrealPakLocation().FullName, Arguments, Options: ERunOptions.Default | ERunOptions.UTF8Output);
+								bAttemptAutoLaunchOnFailure = false;
+							}
+							else
+							{
+								throw new AutomationException(String.Format("Failed sending oplog request to Zen at {0}:{1}. Ensure that the server is running.", Host, Port));
+							}
+						}
+					}
+					if (!HttpGetResult.IsSuccessStatusCode)
+					{
+						throw new AutomationException(String.Format("Failed reading oplog {0}.{1} from Zen. Ensure that cooking was successful.", ProjectId, OplogId));
+					}
+					Task<byte[]> ReadOplogTask = HttpGetResult.Content.ReadAsByteArrayAsync();
+					ReadOplogTask.Wait();
+					byte[] OplogData = ReadOplogTask.Result;
+					CbObject OplogObject = new CbField(OplogData).AsObject();
+					foreach (CbField EntryField in OplogObject["entries"].AsArray())
+					{
+						foreach (CbField PackageDataField in EntryField["packagedata"].AsArray())
+						{
+							string RelativeFilename = PackageDataField["filename"].AsString();
+							SC.PackageStoreManifest.ZenCookedFiles.Add(FileReference.Combine(SC.PlatformCookDir, RelativeFilename).FullName);
+						}
+						foreach (CbField PackageDataField in EntryField["bulkdata"].AsArray())
+						{
+							string RelativeFilename = PackageDataField["filename"].AsString();
+							SC.PackageStoreManifest.ZenCookedFiles.Add(FileReference.Combine(SC.PlatformCookDir, RelativeFilename).FullName);
+						}
+					}
 				}
 			}
 			else
@@ -883,11 +971,11 @@ namespace AutomationScripts
 
 				// When cooking to Zen get the list of cooked package files from the manifest
 				LoadPackageStoreManifest(Params, SC);
-				if (SC.PackageStoreManifest != null && SC.PackageStoreManifest.ZenServer != null)
+				if (SC.PackageStoreManifest != null && SC.PackageStoreManifest.ZenCookedFiles != null)
 				{
-					foreach (PackageStoreManifest.FileInfo File in SC.PackageStoreManifest.Files)
+					foreach (string FilePath in SC.PackageStoreManifest.ZenCookedFiles)
 					{
-						CookedFiles.Add(new FileReference(File.Path));
+						CookedFiles.Add(new FileReference(FilePath));
 					}
 				}
 
@@ -1240,11 +1328,11 @@ namespace AutomationScripts
 
 						// When cooking to Zen get the list of cooked package files from the manifest
 						LoadPackageStoreManifest(Params, SC);
-						if (SC.PackageStoreManifest != null && SC.PackageStoreManifest.ZenServer != null)
+						if (SC.PackageStoreManifest != null && SC.PackageStoreManifest.ZenCookedFiles != null)
 						{
-							foreach (PackageStoreManifest.FileInfo File in SC.PackageStoreManifest.Files)
+							foreach (string FilePath in SC.PackageStoreManifest.ZenCookedFiles)
 							{
-								CookedFiles.Add(new FileReference(File.Path));
+								CookedFiles.Add(new FileReference(FilePath));
 							}
 						}
 
@@ -3649,7 +3737,7 @@ namespace AutomationScripts
 				MakePathSafeToUseWithCommandLine(SC.PlatformCookDir.ToString()),
 				MakePathSafeToUseWithCommandLine(SC.PackageStoreManifest.FullPath),
 				MakePathSafeToUseWithCommandLine(CommandsFileName));
-			if (SC.PackageStoreManifest.ZenServer == null)
+			if (SC.PackageStoreManifest.ZenCookedFiles == null)
 			{
 				FileReference ScriptObjectsFile = FileReference.Combine(SC.MetadataDir, "scriptobjects.bin");
 				if (!FileReference.Exists(ScriptObjectsFile))
