Make UEditorEngine::BroadcastPostUndoRedo() resilient to callbacks registering for PostUndoRedo event

The engine broadcasts undo/redo messages in UEditorEngine::BroadcastPostUndoRedo() to all registered FEditorUndoClient instances in a loop over a TSet<FEditorUndoClient*>. In the current code, if an FEditorUndoClient instance causes a registration with the undo system in either PostUndo() or PostRedo(), then the loop will never terminate, as a new element is added each iteration, causing an infinite loop.
This change opts to iterate over a copy of the TSet<FEditorUndoClient*> which therefore avoids the infinite loop without fundamentally restricting clients from unregistering/registering during the PostUndo()/PostRedo() calls.

#10149

diff --git a/Engine/Source/Editor/UnrealEd/Private/EditorServer.cpp b/Engine/Source/Editor/UnrealEd/Private/EditorServer.cpp
--- a/Engine/Source/Editor/UnrealEd/Private/EditorServer.cpp
+++ b/Engine/Source/Editor/UnrealEd/Private/EditorServer.cpp
@@ -5395,7 +5395,10 @@ void UEditorEngine::BroadcastPostUndoRedo(const FTransactionContext& UndoContext
 	// functions have been fixed. For the time being it improves editor stability, though:
 	UEdGraphPin::SanitizePinsPostUndoRedo();
 
-	for (auto UndoIt = UndoClients.CreateIterator(); UndoIt; ++UndoIt)
+	// Note that we use a copy here as clients can register/unregister with the undo system while in PostUndo()/PostRedo()
+	// which modifies UndoClients during the loop. This can cause an infinite loop where the iterator never finishes.
+	const TSet<FEditorUndoClient*> UndoClientsCopy = UndoClients;
+	for (auto UndoIt = UndoClientsCopy.CreateConstIterator(); UndoIt; ++UndoIt)
 	{
 		FEditorUndoClient* Client = *UndoIt;
 		if (Client && Client->MatchesContext(UndoContext, CurrentUndoRedoContext->TransactionObjects))
