MPCook: Move the lost connection error to occur after writing the logs from the CookWorker, so that the error on the cookworker (if any) is indicated as the first error rather than the followup.
Mark CookWorkers as unattended so that callstacks are logged.
#rnx

diff --git a/Engine/Source/Editor/UnrealEd/Private/Cooker/CookDirector.cpp b/Engine/Source/Editor/UnrealEd/Private/Cooker/CookDirector.cpp
--- a/Engine/Source/Editor/UnrealEd/Private/Cooker/CookDirector.cpp
+++ b/Engine/Source/Editor/UnrealEd/Private/Cooker/CookDirector.cpp
@@ -1230,7 +1230,8 @@ FString FCookDirector::GetWorkerCommandLine(FWorkerId WorkerId, int32 ProfileId)
 				Token.StartsWith(TEXT("-CoreLimit")) ||
 				Token.StartsWith(TEXT("-PhysicalCoreLimit")) ||
 				Token.StartsWith(TEXT("-CookProcessCount=")) ||
-				Token.StartsWith(TEXT("-abslog="))
+				Token.StartsWith(TEXT("-abslog=")) ||
+				Token.StartsWith(TEXT("-unattended"))
 				)
 			{
 				return;
@@ -1255,6 +1256,7 @@ FString FCookDirector::GetWorkerCommandLine(FWorkerId WorkerId, int32 ProfileId)
 	Tokens.Add(FString::Printf(TEXT("-CookDirectorHost=%s"), *WorkerConnectAuthority));
 	Tokens.Add(FString::Printf(TEXT("-MultiprocessId=%d"), WorkerId.GetRemoteIndex() + 1));
 	Tokens.Add(FString::Printf(TEXT("-CookProfileId=%d"), ProfileId));
+	Tokens.Add(TEXT("-unattended"));
 	Tokens.Add(FString::Printf(TEXT("-abslog=%s"), *GetWorkerLogFileName(ProfileId)));
 	if (CoreLimit > 0)
 	{
diff --git a/Engine/Source/Editor/UnrealEd/Private/Cooker/CookWorkerServer.cpp b/Engine/Source/Editor/UnrealEd/Private/Cooker/CookWorkerServer.cpp
--- a/Engine/Source/Editor/UnrealEd/Private/Cooker/CookWorkerServer.cpp
+++ b/Engine/Source/Editor/UnrealEd/Private/Cooker/CookWorkerServer.cpp
@@ -58,14 +58,13 @@ void FCookWorkerServer::DetachFromRemoteProcess()
 	if (bNeedCrashDiagnostics)
 	{
 		SendCrashDiagnostics();
-		bNeedCrashDiagnostics = false;
 	}
 }
 
 void FCookWorkerServer::SendCrashDiagnostics()
 {
 	FString LogFileName = Director.GetWorkerLogFileName(ProfileId);
-	UE_LOG(LogCook, Display, TEXT("CookWorker %d log messages written after communication loss:"), ProfileId);
+	UE_LOG(LogCook, Display, TEXT("LostConnection to CookWorker %d. Log messages written after communication loss:"), ProfileId);
 	FString LogText;
 	int32 ReadFlags = FILEREAD_AllowWrite; // To be able to open a file for read that might be open for write from another process, we have to specify FILEREAD_AllowWrite
 	if (!FFileHelper::LoadFileToString(LogText, *LogFileName, FFileHelper::EHashOptions::None, ReadFlags))
@@ -105,6 +104,13 @@ void FCookWorkerServer::SendCrashDiagnostics()
 			UE_LOG(LogCook, Display, TEXT("[CookWorker %d]: %s"), ProfileId, *Line);
 		}
 	}
+	if (!CrashDiagnosticsError.IsEmpty())
+	{
+		UE_LOG(LogCook, Error, TEXT("%s"), *CrashDiagnosticsError);
+	}
+
+	bNeedCrashDiagnostics = false;
+	CrashDiagnosticsError.Empty();
 }
 
 void FCookWorkerServer::ShutdownRemoteProcess()
@@ -415,7 +421,7 @@ void FCookWorkerServer::LaunchProcess()
 	else
 	{
 		// GetLastError information was logged by CreateProc
-		UE_LOG(LogCook, Error, TEXT("CookWorkerCrash: %d failed to create CookWorker process. Assigned packages will be returned to the director."),
+		CrashDiagnosticsError = FString::Printf(TEXT("CookWorkerCrash: Failed to create process for CookWorker %d. Assigned packages will be returned to the director."),
 			ProfileId);
 		bNeedCrashDiagnostics = true;
 		SendToState(EConnectStatus::LostConnection);
@@ -434,7 +440,7 @@ void FCookWorkerServer::TickWaitForConnect()
 	{
 		if (!FPlatformProcess::IsProcRunning(CookWorkerHandle))
 		{
-			UE_LOG(LogCook, Error, TEXT("CookWorkerCrash: %d process terminated before connecting. Assigned packages will be returned to the director."),
+			CrashDiagnosticsError = FString::Printf(TEXT("CookWorkerCrash: CookWorker %d process terminated before connecting. Assigned packages will be returned to the director."),
 				ProfileId);
 			bNeedCrashDiagnostics = true;
 			SendToState(EConnectStatus::LostConnection);
@@ -445,10 +451,10 @@ void FCookWorkerServer::TickWaitForConnect()
 
 	if (CurrentTime - ConnectStartTimeSeconds > WaitForConnectTimeout && !IsCookIgnoreTimeouts())
 	{
-		UE_LOG(LogCook, Error, TEXT("CookWorkerCrash: %d process failed to connect within %.0f seconds. Assigned packages will be returned to the director."),
+		CrashDiagnosticsError = FString::Printf(TEXT("CookWorkerCrash: CookWorker %d process failed to connect within %.0f seconds. Assigned packages will be returned to the director."),
 			ProfileId, WaitForConnectTimeout);
-		ShutdownRemoteProcess();
 		bNeedCrashDiagnostics = true;
+		ShutdownRemoteProcess();
 		SendToState(EConnectStatus::LostConnection);
 		return;
 	}
@@ -490,7 +496,7 @@ void FCookWorkerServer::PumpSendMessages()
 	UE::CompactBinaryTCP::EConnectionStatus Status = UE::CompactBinaryTCP::TryFlushBuffer(Socket, SendBuffer);
 	if (Status == UE::CompactBinaryTCP::Failed)
 	{
-		UE_LOG(LogCook, Error, TEXT("CookWorker Crash: %d failed to write to socket, we will shutdown the remote process. Assigned packages will be returned to the director."),
+		UE_LOG(LogCook, Error, TEXT("CookWorkerCrash: CookWorker %d failed to write to socket, we will shutdown the remote process. Assigned packages will be returned to the director."),
 			ProfileId);
 		bNeedCrashDiagnostics = true;
 		SendToState(EConnectStatus::WaitForDisconnect);
@@ -527,7 +533,7 @@ void FCookWorkerServer::PumpReceiveMessages()
 	EConnectionStatus SocketStatus = TryReadPacket(Socket, ReceiveBuffer, Messages);
 	if (SocketStatus != EConnectionStatus::Okay && SocketStatus != EConnectionStatus::Incomplete)
 	{
-		UE_LOG(LogCook, Error, TEXT("CookWorkerCrash: %d failed to read from socket, we will shutdown the remote process. Assigned packages will be returned to the director."),
+		CrashDiagnosticsError = FString::Printf(TEXT("CookWorkerCrash: CookWorker %d failed to read from socket, we will shutdown the remote process. Assigned packages will be returned to the director."),
 			ProfileId);
 		bNeedCrashDiagnostics = true;
 		SendToState(EConnectStatus::WaitForDisconnect);
@@ -558,7 +564,7 @@ void FCookWorkerServer::HandleReceiveMessagesInternal()
 			UE::CompactBinaryTCP::FMarshalledMessage Message = ReceiveMessages.PopFrontValue();
 			if (ConnectStatus != EConnectStatus::PumpingCookComplete && ConnectStatus != EConnectStatus::WaitForDisconnect)
 			{
-				UE_LOG(LogCook, Error, TEXT("CookWorkerCrash: %d remote process shut down unexpectedly. Assigned packages will be returned to the director."),
+				CrashDiagnosticsError = FString::Printf(TEXT("CookWorkerCrash: CookWorker %d remote process shut down unexpectedly. Assigned packages will be returned to the director."),
 					ProfileId);
 				bNeedCrashDiagnostics = true;
 			}
diff --git a/Engine/Source/Editor/UnrealEd/Private/Cooker/CookWorkerServer.h b/Engine/Source/Editor/UnrealEd/Private/Cooker/CookWorkerServer.h
--- a/Engine/Source/Editor/UnrealEd/Private/Cooker/CookWorkerServer.h
+++ b/Engine/Source/Editor/UnrealEd/Private/Cooker/CookWorkerServer.h
@@ -179,6 +179,7 @@ private:
 	UE::CompactBinaryTCP::FSendBuffer SendBuffer;
 	UE::CompactBinaryTCP::FReceiveBuffer ReceiveBuffer;
 	TRingBuffer<UE::CompactBinaryTCP::FMarshalledMessage> ReceiveMessages;
+	FString CrashDiagnosticsError;
 	FCookDirector& Director;
 	UCookOnTheFlyServer& COTFS;
 	FSocket* Socket = nullptr;
