Fixes Reimport basic curves menu item in the Curve Table Editor Asset menu

diff --git a/Engine/Source/Editor/UnrealEd/Private/Toolkits/AssetEditorToolkit.cpp b/Engine/Source/Editor/UnrealEd/Private/Toolkits/AssetEditorToolkit.cpp
--- a/Engine/Source/Editor/UnrealEd/Private/Toolkits/AssetEditorToolkit.cpp
+++ b/Engine/Source/Editor/UnrealEd/Private/Toolkits/AssetEditorToolkit.cpp
@@ -991,7 +991,7 @@ void FAssetEditorToolkit::FillDefaultAssetMenuCommands(FToolMenuSection& InSecti
 					FFormatNamedArguments ToolTipArguments;
 					ToolTipArguments.Add(TEXT("Type"), FText::FromString( EditingObject->GetClass()->GetName() ));
 					const FText ToolTipText = FText::Format( LOCTEXT("Reimport_ToolTip", "Reimports this {Type}"), ToolTipArguments );
-					const FName IconName = TEXT( "AssetEditor.Reimport" );
+					const FName IconName = TEXT( "AssetEditor.ReimportAsset" );
 					FUIAction UIAction;
 					UIAction.ExecuteAction.BindRaw( this, &FAssetEditorToolkit::Reimport_Execute, EditingObject );
 					ReimportEntryName.SetNumber(MenuEntryCount++);
