Guard against potential for NULL actor selections.

diff --git a/Engine/Source/Editor/UnrealEd/Private/Kismet2/KismetReinstanceUtilities.cpp b/Engine/Source/Editor/UnrealEd/Private/Kismet2/KismetReinstanceUtilities.cpp
--- a/Engine/Source/Editor/UnrealEd/Private/Kismet2/KismetReinstanceUtilities.cpp
+++ b/Engine/Source/Editor/UnrealEd/Private/Kismet2/KismetReinstanceUtilities.cpp
@@ -2271,7 +2271,10 @@ void FBlueprintCompileReinstancer::ReplaceInstancesOfClass_Inner(const TMap<UCla
 			// Note: For OFPA, each instance may be stored in its own external package.
 			for (int32 SelectionIdx = 0; SelectionIdx < SelectedActors->Num(); ++SelectionIdx)
 			{
-				CheckAndSaveOuterPackageToCleanList(SelectedActors->GetSelectedObject(SelectionIdx));
+				if (const UObject* SelectedActor = SelectedActors->GetSelectedObject(SelectionIdx))
+				{
+					CheckAndSaveOuterPackageToCleanList(SelectedActor);
+				}
 			}
 			
 			SelectedActors->BeginBatchSelectOperation();
