World Partition - Builder Commandlet: A collection name can now be provided as an argument instead of a map name.
* Builder will execute on each map of the collection

diff --git a/Engine/Build/Graph/Examples/BuildWorldPartitionHLODs.xml b/Engine/Build/Graph/Examples/BuildWorldPartitionHLODs.xml
--- a/Engine/Build/Graph/Examples/BuildWorldPartitionHLODs.xml
+++ b/Engine/Build/Graph/Examples/BuildWorldPartitionHLODs.xml
@@ -4,14 +4,14 @@
 	<!-- Options -->
 	<Option Name="ProjectName" DefaultValue="" Description="Name of the project."/>
 	<Option Name="EditorTarget" DefaultValue="" Description="Name of the editor target to use."/>
-	<Option Name="MapName" DefaultValue="" Description="The map for which HLODs will be generated."/>
+	<Option Name="MapName" DefaultValue="" Description="The map (or collection) for which HLODs will be generated."/>
 	<Option Name="BuilderCount" DefaultValue="10" Description="Number of machines to use"/>
 	<Option Name="SubmitHLODs" Restrict="true|false" DefaultValue="false" Description="Submit generated HLODs"/>
 	<Option Name="ShelveHLODs" Restrict="true|false" DefaultValue="false" Description="Shelve generated HLODs"/>
 	<Option Name="ShelveUser" DefaultValue="" Description="Shelve result to the provided user workspace"/>
 	<Option Name="ShelveWorkspace" DefaultValue="" Description="Shelve result to the provided user workspace"/>
 	
-	<Error Message="Missing map name. You must provide a map name using -MapName=..." If="'$(MapName)' == ''"/>
+	<Error Message="Missing map name. You must provide a map (or collection) name using -MapName=..." If="'$(MapName)' == ''"/>
 	<Error Message="Missing project name. You must provide the project name using -ProjectName=..." If="'$(ProjectName)' == ''"/>
 	<Error Message="Missing editor target. You must provide the target name using -EditorTarget=..." If="'$(EditorTarget)' == ''"/>
 	<Error Message="Conflicting options found: -SubmitHLODs and -ShelveHLODs" If="$(SubmitHLODs) And $(ShelveHLODs)"/>
diff --git a/Engine/Source/Editor/UnrealEd/Classes/Commandlets/WorldPartitionBuilderCommandlet.h b/Engine/Source/Editor/UnrealEd/Classes/Commandlets/WorldPartitionBuilderCommandlet.h
--- a/Engine/Source/Editor/UnrealEd/Classes/Commandlets/WorldPartitionBuilderCommandlet.h
+++ b/Engine/Source/Editor/UnrealEd/Classes/Commandlets/WorldPartitionBuilderCommandlet.h
@@ -3,6 +3,7 @@
 #pragma once
 
 #include "Commandlets/Commandlet.h"
+#include "Templates/SubclassOf.h"
 #include "PackageSourceControlHelper.h"
 #include "WorldPartitionBuilderCommandlet.generated.h"
 
@@ -17,5 +18,6 @@ class UWorldPartitionBuilderCommandlet : public UCommandlet
 	virtual int32 Main(const FString& Params) override;
 	//~ End UCommandlet Interface
 
-	UWorldPartitionBuilder* CreateBuilder(const FString& WorldConfigFilename);
+private:
+	bool RunBuilder(TSubclassOf<UWorldPartitionBuilder> InBuilderClass, const FString& InWorldPackageName);
 };
\ No newline at end of file
diff --git a/Engine/Source/Editor/UnrealEd/Private/Commandlets/WorldPartitionBuilderCommandlet.cpp b/Engine/Source/Editor/UnrealEd/Private/Commandlets/WorldPartitionBuilderCommandlet.cpp
--- a/Engine/Source/Editor/UnrealEd/Private/Commandlets/WorldPartitionBuilderCommandlet.cpp
+++ b/Engine/Source/Editor/UnrealEd/Private/Commandlets/WorldPartitionBuilderCommandlet.cpp
@@ -14,6 +14,11 @@
 #include "UObject/GCObjectScopeGuard.h"
 #include "Trace/Trace.h"
 
+#include "CollectionManagerModule.h"
+#include "ICollectionManager.h"
+#include "AssetRegistry/AssetData.h"
+#include "AssetRegistry/AssetRegistryModule.h"
+
 DEFINE_LOG_CATEGORY_STATIC(LogWorldPartitionBuilderCommandlet, All, All);
 
 UWorldPartitionBuilderCommandlet::UWorldPartitionBuilderCommandlet(const FObjectInitializer& ObjectInitializer)
@@ -48,11 +53,41 @@ int32 UWorldPartitionBuilderCommandlet::Main(const FString& Params)
 		FastExit = true;		// Faster exit which avoids crash during shutdown. The engine isn't shutdown cleanly.
 	}
 
-	// This will convert incomplete package name to a fully qualified path
-	FString WorldFilename;
-	if (!FPackageName::SearchForPackageOnDisk(Tokens[0], &Tokens[0], &WorldFilename))
+	ICollectionManager& CollectionManager = FModuleManager::LoadModuleChecked<FCollectionManagerModule>("CollectionManager").Get();
+	TArray<FString> MapPackagesNames;
+
+	FAssetRegistryModule* AssetRegistryModule = &FModuleManager::LoadModuleChecked<FAssetRegistryModule>(TEXT("AssetRegistry"));
+
+	// Parse map name or maps collection
+	if (FPackageName::SearchForPackageOnDisk(Tokens[0]))
+	{
+		MapPackagesNames = { Tokens[0] };
+	}	
+	else if (CollectionManager.CollectionExists(FName(Tokens[0]), ECollectionShareType::CST_All))
 	{
-		UE_LOG(LogWorldPartitionBuilderCommandlet, Error, TEXT("Unknown world '%s'"), *Tokens[0]);
+		TArray<FSoftObjectPath> AssetsPaths;
+		CollectionManager.GetAssetsInCollection(FName(Tokens[0]), ECollectionShareType::CST_All, AssetsPaths, ECollectionRecursionFlags::SelfAndChildren);
+
+		for (const auto& AssetPath : AssetsPaths)
+		{
+			const bool bIncludeOnlyOnDiskAssets = true;
+			FAssetData AssetData = AssetRegistryModule->Get().GetAssetByObjectPath(AssetPath, bIncludeOnlyOnDiskAssets);
+			const bool bIsWorldAsset = AssetData.IsValid() && (AssetData.AssetClassPath == UWorld::StaticClass()->GetClassPathName());
+			if (bIsWorldAsset)
+			{
+				MapPackagesNames.Add(AssetPath.GetAssetPathString());
+			}
+		}
+
+		if (MapPackagesNames.IsEmpty())
+		{
+			UE_LOG(LogWorldPartitionBuilderCommandlet, Warning, TEXT("Found no maps to process in collection %s), exiting"), *Tokens[0]);
+			return 0;
+		}
+	}
+	else
+	{
+		UE_LOG(LogWorldPartitionBuilderCommandlet, Error, TEXT("Missing world(s) as the first argument to the commandlet. Either supply the world name directly (WorldName or /Path/To/WorldName), or provide a collection name to have the builder operate on a set of maps."));
 		return 1;
 	}
 
@@ -65,27 +100,50 @@ int32 UWorldPartitionBuilderCommandlet::Main(const FString& Params)
 	}
 
 	// Find builder class
-	UClass* BuilderClass = FindFirstObject<UClass>(*BuilderClassName, EFindFirstObjectOptions::EnsureIfAmbiguous);
+	TSubclassOf<UWorldPartitionBuilder> BuilderClass = FindFirstObject<UClass>(*BuilderClassName, EFindFirstObjectOptions::EnsureIfAmbiguous);
 	if (!BuilderClass)
 	{
-		UE_LOG(LogWorldPartitionBuilderCommandlet, Error, TEXT("Unknown builder %s."), *BuilderClassName);
+		UE_LOG(LogWorldPartitionBuilderCommandlet, Error, TEXT("Unknown/invalid world partition builder class: %s."), *BuilderClassName);
 		return 1;
+	}	
+
+	// Run the builder on the provided map(s)
+	for (const FString& MapPackageName : MapPackagesNames)
+	{
+		if (!RunBuilder(BuilderClass, MapPackageName))
+		{
+			return false;
+		}
+	}
+
+	return true;
+}
+
+bool UWorldPartitionBuilderCommandlet::RunBuilder(TSubclassOf<UWorldPartitionBuilder> InBuilderClass, const FString& InWorldPackageName)
+{
+	// This will convert incomplete package name to a fully qualified path
+	FString WorldLongPackageName;
+	FString WorldFilename;
+	if (!FPackageName::SearchForPackageOnDisk(InWorldPackageName, &WorldLongPackageName, &WorldFilename))
+	{
+		UE_LOG(LogWorldPartitionBuilderCommandlet, Error, TEXT("Package '%s' not found"), *InWorldPackageName);
+		return false;
 	}
 
 	// Load the world package
-	UPackage* WorldPackage = LoadWorldPackageForEditor(Tokens[0]);
+	UPackage* WorldPackage = LoadWorldPackageForEditor(WorldLongPackageName);
 	if (!WorldPackage)
 	{
-		UE_LOG(LogWorldPartitionBuilderCommandlet, Error, TEXT("Couldn't load package %s."), *Tokens[0]);
-		return 1;
+		UE_LOG(LogWorldPartitionBuilderCommandlet, Error, TEXT("Couldn't load package %s."), *WorldLongPackageName);
+		return false;
 	}
 
 	// Find the world in the given package
 	UWorld* World = UWorld::FindWorldInPackage(WorldPackage);
 	if (!World)
 	{
-		UE_LOG(LogWorldPartitionBuilderCommandlet, Error, TEXT("No world in specified package %s."), *Tokens[0]);
-		return 1;
+		UE_LOG(LogWorldPartitionBuilderCommandlet, Error, TEXT("No world in specified package %s."), *WorldLongPackageName);
+		return false;
 	}
 
 	// Load configuration file
@@ -96,7 +154,7 @@ int32 UWorldPartitionBuilderCommandlet::Main(const FString& Params)
 	}
 
 	// Create builder instance
-	UWorldPartitionBuilder* Builder = NewObject<UWorldPartitionBuilder>(GetTransientPackage(), BuilderClass);
+	UWorldPartitionBuilder* Builder = NewObject<UWorldPartitionBuilder>(GetTransientPackage(), InBuilderClass);
 	if (!Builder)
 	{
 		UE_LOG(LogWorldPartitionBuilderCommandlet, Error, TEXT("Failed to create builder."));
@@ -114,7 +172,7 @@ int32 UWorldPartitionBuilderCommandlet::Main(const FString& Params)
 		!FPlatformFileManager::Get().GetPlatformFile().IsReadOnly(*WorldConfigFilename))
 	{
 		SaveConfig(CPF_Config, *WorldConfigFilename);
-	}	
+	}
 
-	return bResult ? 0 : 1;
+	return bResult;
 }
