World Partition - HLOD Builder: Support processing of multiple worlds when running in distributed mode (-DistributedBuild)
* Multiple worlds can now be provided to the builders using a collection
* Ensure each world writes its data to a different working directory

diff --git a/Engine/Source/Editor/UnrealEd/Private/WorldPartition/WorldPartitionBuilder.cpp b/Engine/Source/Editor/UnrealEd/Private/WorldPartition/WorldPartitionBuilder.cpp
--- a/Engine/Source/Editor/UnrealEd/Private/WorldPartition/WorldPartitionBuilder.cpp
+++ b/Engine/Source/Editor/UnrealEd/Private/WorldPartition/WorldPartitionBuilder.cpp
@@ -92,7 +92,7 @@ bool UWorldPartitionBuilder::RunBuilder(UWorld* World)
 	FPackageSourceControlHelper SCCHelper;
 
 	// Perform builder pre world initialisation
-	if (!PreWorldInitialization(SCCHelper))
+	if (!PreWorldInitialization(World, SCCHelper))
 	{
 		UE_LOG(LogWorldPartitionBuilder, Error, TEXT("PreWorldInitialization failed"));
 		return false;
diff --git a/Engine/Source/Editor/UnrealEd/Private/WorldPartition/WorldPartitionHLODsBuilder.cpp b/Engine/Source/Editor/UnrealEd/Private/WorldPartition/WorldPartitionHLODsBuilder.cpp
--- a/Engine/Source/Editor/UnrealEd/Private/WorldPartition/WorldPartitionHLODsBuilder.cpp
+++ b/Engine/Source/Editor/UnrealEd/Private/WorldPartition/WorldPartitionHLODsBuilder.cpp
@@ -163,8 +163,6 @@ UWorldPartitionHLODsBuilder::UWorldPartitionHLODsBuilder(const FObjectInitialize
 	: Super(ObjectInitializer)
 	, BuilderIdx(INDEX_NONE)
 	, BuilderCount(INDEX_NONE)
-	, DistributedBuildWorkingDir(FPaths::RootDir() / DistributedBuildWorkingDirName)
-	, DistributedBuildManifest(DistributedBuildWorkingDir / DistributedBuildManifestName)
 {
 	BuildOptions = FParse::Param(FCommandLine::Get(), TEXT("SetupHLODs")) ? EHLODBuildStep::HLOD_Setup : EHLODBuildStep::None;
 	BuildOptions |= FParse::Param(FCommandLine::Get(), TEXT("BuildHLODs")) ? EHLODBuildStep::HLOD_Build : EHLODBuildStep::None;
@@ -195,16 +193,6 @@ UWorldPartitionHLODsBuilder::UWorldPartitionHLODsBuilder(const FObjectInitialize
 	{
 		BuildOptions = EHLODBuildStep::HLOD_Setup | EHLODBuildStep::HLOD_Build;
 	}
-
-	if (bDistributedBuild)
-	{
-		if (!BuildManifest.IsEmpty())
-		{
-			UE_LOG(LogWorldPartitionHLODsBuilder, Warning, TEXT("Ignoring parameter -BuildManifest when a distributed build is performed"));
-		}
-
-		BuildManifest = DistributedBuildManifest;
-	}
 }
 
 bool UWorldPartitionHLODsBuilder::RequiresCommandletRendering() const
@@ -274,8 +262,22 @@ bool UWorldPartitionHLODsBuilder::ValidateParams() const
 	return true;
 }
 
-bool UWorldPartitionHLODsBuilder::PreWorldInitialization(FPackageSourceControlHelper& PackageHelper)
+bool UWorldPartitionHLODsBuilder::PreWorldInitialization(UWorld* World, FPackageSourceControlHelper& PackageHelper)
 {
+	if (bDistributedBuild)
+	{
+		uint32 WorldPackageHash = GetTypeHash(World->GetPackage()->GetFullName());
+		DistributedBuildWorkingDir = FString::Printf(TEXT("%s/%s/%08x"), *FPaths::RootDir(), *DistributedBuildWorkingDirName, WorldPackageHash);
+		DistributedBuildManifest = DistributedBuildWorkingDir / DistributedBuildManifestName;
+
+		if (!BuildManifest.IsEmpty())
+		{
+			UE_LOG(LogWorldPartitionHLODsBuilder, Warning, TEXT("Ignoring parameter -BuildManifest when a distributed build is performed"));
+		}
+
+		BuildManifest = DistributedBuildManifest;
+	}
+
 	if (!ValidateParams())
 	{
 		return false;
diff --git a/Engine/Source/Editor/UnrealEd/Public/WorldPartition/WorldPartitionBuilder.h b/Engine/Source/Editor/UnrealEd/Public/WorldPartition/WorldPartitionBuilder.h
--- a/Engine/Source/Editor/UnrealEd/Public/WorldPartition/WorldPartitionBuilder.h
+++ b/Engine/Source/Editor/UnrealEd/Public/WorldPartition/WorldPartitionBuilder.h
@@ -55,7 +55,7 @@ public:
 	
 	bool Run(UWorld* World, FPackageSourceControlHelper& PackageHelper);
 
-	virtual bool PreWorldInitialization(FPackageSourceControlHelper& PackageHelper) { return true; }
+	virtual bool PreWorldInitialization(UWorld* World, FPackageSourceControlHelper& PackageHelper) { return true; }
 
 	static bool SavePackages(const TArray<UPackage*>& Packages, FPackageSourceControlHelper& PackageHelper, bool bErrorsAsWarnings = false);
 	static bool DeletePackages(const TArray<UPackage*>& Packages, FPackageSourceControlHelper& PackageHelper, bool bErrorsAsWarnings = false);
@@ -91,6 +91,9 @@ protected:
 
 	virtual UWorld::InitializationValues GetWorldInitializationValues() const;
 
+	UE_DEPRECATED(5.3, "You must override the version that takes a World parameter")
+	virtual bool PreWorldInitialization(FPackageSourceControlHelper& PackageHelper) final { return true; }
+
 	int32 IterativeCellSize = 102400;
 	int32 IterativeCellOverlapSize = 0;
 	FBox  IterativeWorldBounds;
diff --git a/Engine/Source/Editor/UnrealEd/Public/WorldPartition/WorldPartitionHLODsBuilder.h b/Engine/Source/Editor/UnrealEd/Public/WorldPartition/WorldPartitionHLODsBuilder.h
--- a/Engine/Source/Editor/UnrealEd/Public/WorldPartition/WorldPartitionHLODsBuilder.h
+++ b/Engine/Source/Editor/UnrealEd/Public/WorldPartition/WorldPartitionHLODsBuilder.h
@@ -77,7 +77,7 @@ public:
 	// UWorldPartitionBuilder interface begin
 	virtual bool RequiresCommandletRendering() const override;
 	virtual ELoadingMode GetLoadingMode() const override { return ELoadingMode::Custom; }
-	virtual bool PreWorldInitialization(FPackageSourceControlHelper& PackageHelper) override;
+	virtual bool PreWorldInitialization(UWorld* World, FPackageSourceControlHelper& PackageHelper) override;
 protected:
 	virtual bool RunInternal(UWorld* World, const FCellInfo& InCellInfo, FPackageSourceControlHelper& PackageHelper) override;
 	// UWorldPartitionBuilder interface end
@@ -120,8 +120,8 @@ private:
 	FName HLODLayerToBuild;
 	FName HLODActorToBuild;
 
-	const FString DistributedBuildWorkingDir;
-	const FString DistributedBuildManifest;
+	FString DistributedBuildWorkingDir;
+	FString DistributedBuildManifest;
 	
 	FHLODModifiedFiles ModifiedFiles;
 };
\ No newline at end of file
