Remove virtual bone from import conflict dialog.

#rnx

diff --git a/Engine/Source/Editor/UnrealEd/Private/Fbx/FbxFactory.cpp b/Engine/Source/Editor/UnrealEd/Private/Fbx/FbxFactory.cpp
--- a/Engine/Source/Editor/UnrealEd/Private/Fbx/FbxFactory.cpp
+++ b/Engine/Source/Editor/UnrealEd/Private/Fbx/FbxFactory.cpp
@@ -1452,6 +1452,11 @@ namespace ImportCompareHelper
 		{
 			if (CurrentIndex == ReferenceSkeleton.GetParentIndex(ChildIndex))
 			{
+				//Skip virtual bones
+				if (ReferenceSkeleton.GetRequiredVirtualBones().Contains(ChildIndex))
+				{
+					return;
+				}
 				FSkeletonTreeNode& ChildNode = SkeletonTreeNode.Childrens.AddDefaulted_GetRef();
 				ChildNode.JointName = ReferenceSkeleton.GetBoneName(ChildIndex);
 				FillRecursivelySkeleton(ReferenceSkeleton, ChildIndex, ChildNode);
