Fix Regression causing SaveWorld to always save external actors without checking dirty flag (24511115)

diff --git a/Engine/Source/Editor/UnrealEd/Private/FileHelpers.cpp b/Engine/Source/Editor/UnrealEd/Private/FileHelpers.cpp
--- a/Engine/Source/Editor/UnrealEd/Private/FileHelpers.cpp
+++ b/Engine/Source/Editor/UnrealEd/Private/FileHelpers.cpp
@@ -922,7 +922,7 @@ static bool SaveWorld(UWorld* World,
 
 				if (PackagesToSave.Num())
 				{
-					if (!UEditorLoadingAndSavingUtils::SavePackages(PackagesToSave, bCheckDirty && !bNewlyCreated))
+					if (!UEditorLoadingAndSavingUtils::SavePackages(PackagesToSave, /*bCheckDirty=*/ !bNewlyCreated))
 					{
 						FMessageDialog::Open(EAppMsgType::Ok, NSLOCTEXT("UnrealEd", "Error_FailedToSaveHLODLayersPackages", "Failed to save dependant map packages"));
 						bSuccess = false;
