Remove standalone flags so that objects renamed by the reinstancer are still destroyed even wo/ pending kill

diff --git a/Engine/Source/Editor/UnrealEd/Private/Kismet2/KismetReinstanceUtilities.cpp b/Engine/Source/Editor/UnrealEd/Private/Kismet2/KismetReinstanceUtilities.cpp
--- a/Engine/Source/Editor/UnrealEd/Private/Kismet2/KismetReinstanceUtilities.cpp
+++ b/Engine/Source/Editor/UnrealEd/Private/Kismet2/KismetReinstanceUtilities.cpp
@@ -1952,12 +1952,17 @@ static void ReplaceObjectHelper(UObject*& OldObject, UClass* OldClass, UObject*&
 		}
 	}
 
+	OldObject->ClearFlags(RF_Standalone);
 	OldObject->RemoveFromRoot();
 	OldObject->MarkAsGarbage();
-	ForEachObjectWithOuter(OldObject, [](UObject* ObjectInOuter)
-	{
-		ObjectInOuter->MarkAsGarbage();
-	}, true, RF_NoFlags, EInternalObjectFlags::Garbage);
+	ForEachObjectWithOuter(OldObject, 
+		[](UObject* ObjectInOuter)
+		{
+			ObjectInOuter->ClearFlags(RF_Standalone);
+			ObjectInOuter->RemoveFromRoot();
+			ObjectInOuter->MarkAsGarbage();
+		}
+		, true, RF_NoFlags, EInternalObjectFlags::Garbage);
 
 	OldToNewInstanceMap.Add(OldObject, NewUObject);
 
