Cooker: bugfix for unnecessarily adding dependencies of packages that are referenced by the AssetManager but are also marked as NeverCook.
#rnx

diff --git a/Engine/Source/Editor/UnrealEd/Private/Cooker/CookRequestCluster.cpp b/Engine/Source/Editor/UnrealEd/Private/Cooker/CookRequestCluster.cpp
--- a/Engine/Source/Editor/UnrealEd/Private/Cooker/CookRequestCluster.cpp
+++ b/Engine/Source/Editor/UnrealEd/Private/Cooker/CookRequestCluster.cpp
@@ -952,7 +952,7 @@ FPackageData* FRequestCluster::FGraphSearch::FindOrAddVertex(FName PackageName,
 	OutNewVertex->bInitialRequest = bInitialRequest;
 	OutNewVertex->bCookable = bCookable;
 	OutNewVertex->SuppressCookReason = SuppressCookReason;
-	OutNewVertex->bExploreDependencies = true;
+	OutNewVertex->bExploreDependencies = bCookable || bHardDependency;
 	OutNewVertex->bExploreSoftDependencies = Cluster.bAllowSoftDependencies && bCookable;
 
 	return PackageData;
