Cooker: Add -cooklist mode to show which packages will be cooked (as far as we know at the beginning of the cook based on package dependencies).
#rnx

diff --git a/Engine/Source/Editor/UnrealEd/Classes/Commandlets/CookCommandlet.h b/Engine/Source/Editor/UnrealEd/Classes/Commandlets/CookCommandlet.h
--- a/Engine/Source/Editor/UnrealEd/Classes/Commandlets/CookCommandlet.h
+++ b/Engine/Source/Editor/UnrealEd/Classes/Commandlets/CookCommandlet.h
@@ -19,6 +19,7 @@
 class FSandboxPlatformFile;
 class ITargetPlatform;
 class UCookOnTheFlyServer;
+enum class ECookByTheBookOptions;
 
 UCLASS(config=Editor)
 class UCookCommandlet
@@ -94,4 +95,9 @@ public:
 	
 	//~ End UCommandlet Interface
 
+private:
+	void RunCookByTheBookList(UCookOnTheFlyServer* CookOnTheFlyServer, void* StartupOptionsAsVoid,
+		ECookByTheBookOptions CookOptions);
+	void RunCookByTheBookCook(UCookOnTheFlyServer* CookOnTheFlyServer, void* StartupOptionsAsVoid,
+		ECookByTheBookOptions CookOptions);
 };
diff --git a/Engine/Source/Editor/UnrealEd/Classes/CookOnTheSide/CookOnTheFlyServer.h b/Engine/Source/Editor/UnrealEd/Classes/CookOnTheSide/CookOnTheFlyServer.h
--- a/Engine/Source/Editor/UnrealEd/Classes/CookOnTheSide/CookOnTheFlyServer.h
+++ b/Engine/Source/Editor/UnrealEd/Classes/CookOnTheSide/CookOnTheFlyServer.h
@@ -96,6 +96,12 @@ enum class ECookByTheBookOptions
 };
 ENUM_CLASS_FLAGS(ECookByTheBookOptions);
 
+enum class ECookListOptions
+{
+	None =								0x00000000,
+	ShowRejected =						0x00000001,
+};
+ENUM_CLASS_FLAGS(ECookListOptions);
 
 UENUM()
 namespace ECookMode
@@ -440,6 +446,7 @@ private:
 	/* Update polled fields used by CookOnTheFly's network request handlers */
 	void TickNetwork();
 
+	/** Tick the cook for the current cook mode: request, load, save, finish. */
 	void TickMainCookLoop(UE::Cook::FTickStackData& StackData);
 
 	/** Execute operations that need to be done after each Scheduler task, such as checking for new external requests. */
@@ -679,6 +686,8 @@ public:
 	uint32 TickCookOnTheFly(const float TimeSlice, ECookTickFlags TickFlags = ECookTickFlags::None);
 	/** Tick CookWorker until it finishes or needs to yield. Should only be called when in CookWorker Mode. */
 	uint32 TickCookWorker();
+	/** Log a list of all of the transitively requested packages. */
+	void RunCookList(ECookListOptions CookListOptions);
 
 	/** Clear all the previously cooked data all cook requests from now on will be considered recook requests */
 	void ClearAllCookedData();
@@ -1424,6 +1433,8 @@ private:
 	bool bPollablesInTick = false;
 	/** Config value that specifies whether the Skip-Only-Editor-Only feature is enabled. */
 	bool bCanSkipEditorReferencedPackagesWhenCooking = false;
+	/** True if we're running cooklist; tweak output */
+	bool bCookListMode = false;
 
 	/** Timers for tracking how long we have been busy, to manage retries and warnings of deadlock */
 	double SaveBusyStartTimeSeconds = MAX_flt;
diff --git a/Engine/Source/Editor/UnrealEd/Private/Commandlets/CookCommandlet.cpp b/Engine/Source/Editor/UnrealEd/Private/Commandlets/CookCommandlet.cpp
--- a/Engine/Source/Editor/UnrealEd/Private/Commandlets/CookCommandlet.cpp
+++ b/Engine/Source/Editor/UnrealEd/Private/Commandlets/CookCommandlet.cpp
@@ -242,7 +242,7 @@ int32 UCookCommandlet::Main(const FString& CmdLineParams)
 	return 0;
 }
 
-bool UCookCommandlet::CookByTheBook( const TArray<ITargetPlatform*>& Platforms)
+bool UCookCommandlet::CookByTheBook(const TArray<ITargetPlatform*>& Platforms)
 {
 #if OUTPUT_COOKTIMING
 	TRACE_CPUPROFILER_EVENT_SCOPE_ON_CHANNEL(CookByTheBook, CookChannel);
@@ -269,19 +269,18 @@ bool UCookCommandlet::CookByTheBook( const TArray<ITargetPlatform*>& Platforms)
 	{
 		// Add shared build flag to method flag, and enable iterative
 		IterateFlags |= ECookInitializationFlags::IterateSharedBuild;
-		
+
 		bIterativeCooking = true;
 	}
-	
+
 	ECookInitializationFlags CookFlags = ECookInitializationFlags::IncludeServerMaps;
 	CookFlags |= bIterativeCooking ? IterateFlags : ECookInitializationFlags::None;
-	CookFlags |= bSkipEditorContent ? ECookInitializationFlags::SkipEditorContent : ECookInitializationFlags::None;	
+	CookFlags |= bSkipEditorContent ? ECookInitializationFlags::SkipEditorContent : ECookInitializationFlags::None;
 	CookFlags |= bUnversioned ? ECookInitializationFlags::Unversioned : ECookInitializationFlags::None;
 	CookFlags |= bCookEditorOptional ? ECookInitializationFlags::CookEditorOptional : ECookInitializationFlags::None;
 	CookFlags |= bVerboseCookerWarnings ? ECookInitializationFlags::OutputVerboseCookerWarnings : ECookInitializationFlags::None;
 	CookFlags |= bPartialGC ? ECookInitializationFlags::EnablePartialGC : ECookInitializationFlags::None;
-	bool bTestCook = Switches.Contains(TEXT("TestCook"));
-	CookFlags |= bTestCook ? ECookInitializationFlags::TestCook : ECookInitializationFlags::None;
+	CookFlags |= Switches.Contains(TEXT("TestCook")) ? ECookInitializationFlags::TestCook : ECookInitializationFlags::None;
 	CookFlags |= Switches.Contains(TEXT("LogDebugInfo")) ? ECookInitializationFlags::LogDebugInfo : ECookInitializationFlags::None;
 	CookFlags |= bIgnoreIniSettingsOutOfDate || CookerSettings->bIgnoreIniSettingsOutOfDateForIteration ? ECookInitializationFlags::IgnoreIniSettingsOutOfDate : ECookInitializationFlags::None;
 	CookFlags |= Switches.Contains(TEXT("IgnoreScriptPackagesOutOfDate")) || CookerSettings->bIgnoreScriptPackagesOutOfDateForIteration ? ECookInitializationFlags::IgnoreScriptPackagesOutOfDate : ECookInitializationFlags::None;
@@ -290,16 +289,16 @@ bool UCookCommandlet::CookByTheBook( const TArray<ITargetPlatform*>& Platforms)
 	// parse commandline options 
 
 	FString DLCName;
-	FParse::Value( *Params, TEXT("DLCNAME="), DLCName);
+	FParse::Value(*Params, TEXT("DLCNAME="), DLCName);
 
 	FString BasedOnReleaseVersion;
-	FParse::Value( *Params, TEXT("BasedOnReleaseVersion="), BasedOnReleaseVersion);
+	FParse::Value(*Params, TEXT("BasedOnReleaseVersion="), BasedOnReleaseVersion);
 
 	FString CreateReleaseVersion;
-	FParse::Value( *Params, TEXT("CreateReleaseVersion="), CreateReleaseVersion);
+	FParse::Value(*Params, TEXT("CreateReleaseVersion="), CreateReleaseVersion);
 
 	FString OutputDirectoryOverride;
-	FParse::Value( *Params, TEXT("OutputDir="), OutputDirectoryOverride);
+	FParse::Value(*Params, TEXT("OutputDir="), OutputDirectoryOverride);
 
 	TArray<FString> CmdLineMapEntries;
 	TArray<FString> CmdLineDirEntries;
@@ -415,6 +414,7 @@ bool UCookCommandlet::CookByTheBook( const TArray<ITargetPlatform*>& Platforms)
 	CookOptions |= Switches.Contains(TEXT("CookAgainstFixedBase")) ? ECookByTheBookOptions::CookAgainstFixedBase : ECookByTheBookOptions::None;
 	CookOptions |= (Switches.Contains(TEXT("DlcLoadMainAssetRegistry")) || !bErrorOnEngineContentUse) ? ECookByTheBookOptions::DlcLoadMainAssetRegistry : ECookByTheBookOptions::None;
 	CookOptions |= Switches.Contains(TEXT("DlcReevaluateUncookedAssets")) ? ECookByTheBookOptions::DlcReevaluateUncookedAssets : ECookByTheBookOptions::None;
+	bool bCookList = Switches.Contains(TEXT("CookList"));
 
 	if (bCookSinglePackage)
 	{
@@ -464,35 +464,69 @@ bool UCookCommandlet::CookByTheBook( const TArray<ITargetPlatform*>& Platforms)
 		return false;
 	}
 
-	Swap( StartupOptions.CookMaps, CmdLineMapEntries);
-	Swap( StartupOptions.CookDirectories, CmdLineDirEntries );
-	Swap( StartupOptions.NeverCookDirectories, CmdLineNeverCookDirEntries);
-	Swap( StartupOptions.CookCultures, CookCultures );
-	Swap( StartupOptions.DLCName, DLCName );
-	Swap( StartupOptions.BasedOnReleaseVersion, BasedOnReleaseVersion );
-	Swap( StartupOptions.CreateReleaseVersion, CreateReleaseVersion );
-	Swap( StartupOptions.IniMapSections, MapIniSections);
+	Swap(StartupOptions.CookMaps, CmdLineMapEntries);
+	Swap(StartupOptions.CookDirectories, CmdLineDirEntries);
+	Swap(StartupOptions.NeverCookDirectories, CmdLineNeverCookDirEntries);
+	Swap(StartupOptions.CookCultures, CookCultures);
+	Swap(StartupOptions.DLCName, DLCName);
+	Swap(StartupOptions.BasedOnReleaseVersion, BasedOnReleaseVersion);
+	Swap(StartupOptions.CreateReleaseVersion, CreateReleaseVersion);
+	Swap(StartupOptions.IniMapSections, MapIniSections);
 	StartupOptions.CookOptions = CookOptions;
 	StartupOptions.bErrorOnEngineContentUse = bErrorOnEngineContentUse;
 	StartupOptions.bGenerateDependenciesForMaps = Switches.Contains(TEXT("GenerateDependenciesForMaps"));
 	StartupOptions.bGenerateStreamingInstallManifests = bGenerateStreamingInstallManifests;
 
 	COOK_STAT(
-	{
-		for (const auto& Platform : Platforms)
-		{
-			DetailedCookStats::TargetPlatforms += Platform->PlatformName() + TEXT("+");
-		}
-		if (!DetailedCookStats::TargetPlatforms.IsEmpty())
 		{
-			DetailedCookStats::TargetPlatforms.RemoveFromEnd(TEXT("+"));
-		}
-	});
+			for (const auto& Platform : Platforms)
+			{
+				DetailedCookStats::TargetPlatforms += Platform->PlatformName() + TEXT("+");
+			}
+			if (!DetailedCookStats::TargetPlatforms.IsEmpty())
+			{
+				DetailedCookStats::TargetPlatforms.RemoveFromEnd(TEXT("+"));
+			}
+		});
+
+	// Cast to void as a workaround to support inability to foward-declare inner classes.
+	// TODO: Change FCookByTheBookStartupOptions to a global class.
+	void* StartupOptionsAsVoid = &StartupOptions;
+
+	if (bCookList)
+	{
+		RunCookByTheBookList(CookOnTheFlyServer, StartupOptionsAsVoid, CookOptions);
+	}
+	else
+	{
+		RunCookByTheBookCook(CookOnTheFlyServer, StartupOptionsAsVoid, CookOptions);
+	}
+	return true;
+}
+
+void UCookCommandlet::RunCookByTheBookList(UCookOnTheFlyServer* CookOnTheFlyServer, void* StartupOptionsAsVoid,
+	ECookByTheBookOptions CookOptions)
+{
+	UCookOnTheFlyServer::FCookByTheBookStartupOptions& StartupOptions =
+		*reinterpret_cast<UCookOnTheFlyServer::FCookByTheBookStartupOptions*>(StartupOptionsAsVoid);
+
+	ECookListOptions CookListOptions = ECookListOptions::None;
+	if (Switches.Contains(TEXT("showrejected"))) CookListOptions |= ECookListOptions::ShowRejected;
+
+	CookOnTheFlyServer->StartCookByTheBook(StartupOptions);
+	CookOnTheFlyServer->RunCookList(CookListOptions);
+}
+
+void UCookCommandlet::RunCookByTheBookCook(UCookOnTheFlyServer* CookOnTheFlyServer, void* StartupOptionsAsVoid,
+	ECookByTheBookOptions CookOptions)
+{
+	UCookOnTheFlyServer::FCookByTheBookStartupOptions& StartupOptions =
+		*reinterpret_cast<UCookOnTheFlyServer::FCookByTheBookStartupOptions*>(StartupOptionsAsVoid);
+	bool bTestCook = EnumHasAnyFlags(CookOnTheFlyServer->GetCookFlags(), ECookInitializationFlags::TestCook);
 
 #if ENABLE_LOW_LEVEL_MEM_TRACKER
 	FLowLevelMemTracker::Get().UpdateStatsPerFrame();
 #endif
-
 	do
 	{
 		{
@@ -520,8 +554,6 @@ bool UCookCommandlet::CookByTheBook( const TArray<ITargetPlatform*>& Platforms)
 		bool bFullReferencesExpected = !(CookOptions & ECookByTheBookOptions::SkipHardReferences);
 		UE::SavePackageUtilities::VerifyEDLCookInfo(bFullReferencesExpected);
 	}
-
-	return true;
 }
 
 bool UCookCommandlet::CookAsCookWorker()
diff --git a/Engine/Source/Editor/UnrealEd/Private/CookOnTheFlyServer.cpp b/Engine/Source/Editor/UnrealEd/Private/CookOnTheFlyServer.cpp
--- a/Engine/Source/Editor/UnrealEd/Private/CookOnTheFlyServer.cpp
+++ b/Engine/Source/Editor/UnrealEd/Private/CookOnTheFlyServer.cpp
@@ -125,6 +125,7 @@
 #include "String/ParseLines.h"
 #include "String/ParseTokens.h"
 #include "TargetDomain/TargetDomainUtils.h"
+#include "Templates/UnrealTemplate.h"
 #include "UnrealEdGlobals.h"
 #include "UObject/ArchiveCookContext.h"
 #include "UObject/Class.h"
@@ -145,6 +146,8 @@
 #define LOCTEXT_NAMESPACE "Cooker"
 
 DEFINE_LOG_CATEGORY(LogCook);
+DEFINE_LOG_CATEGORY_STATIC(LogCookList, Log, All);
+
 FName LogCookName(TEXT("LogCook"));
 LLM_DEFINE_TAG(Cooker);
 LLM_DEFINE_TAG(Cooker_SavePackage);
@@ -1287,6 +1290,64 @@ uint32 UCookOnTheFlyServer::TickCookByTheBook(const float TimeSlice, ECookTickFl
 	return StackData.ResultFlags;
 }
 
+void UCookOnTheFlyServer::RunCookList(ECookListOptions CookListOptions)
+{
+	using namespace UE::Cook;
+
+	TGuardValue<bool> SetRunCookListMode(bCookListMode, true);
+
+	FTickStackData StackData(MAX_flt, ECookTickFlags::None);
+	PumpExternalRequests(StackData.Timer);
+	ProcessUnsolicitedPackages();
+	FRequestQueue& RequestQueue = PackageDatas->GetRequestQueue();
+	while (RequestQueue.GetRequestClusters().Num() != 0 || RequestQueue.GetUnclusteredRequests().Num() != 0)
+	{
+		int32 NumPushed;
+		PumpRequests(StackData, NumPushed);
+	}
+
+	TArray<FPackageData*> ReportedDatas;
+	PackageDatas->LockAndEnumeratePackageDatas([&ReportedDatas, CookListOptions](FPackageData* PackageData)
+	{
+		bool bIncludePackage;
+		if (EnumHasAnyFlags(CookListOptions, ECookListOptions::ShowRejected))
+		{
+			bIncludePackage = PackageData->GetInstigator().Category != EInstigator::NotYetRequested;
+			if (bIncludePackage)
+			{
+				// Skip printing out a message for external actors
+				if (INDEX_NONE != UE::String::FindFirst(WriteToString<256>(PackageData->GetPackageName()), ULevel::GetExternalActorsFolderName()))
+				{
+					bIncludePackage = false;
+				}
+			}
+		}
+		else
+		{
+			bIncludePackage = PackageData->IsInProgress() || PackageData->HasAnyCookedPlatform();
+		}
+		if (bIncludePackage)
+		{
+			ReportedDatas.Add(PackageData);
+		}
+	});
+	Algo::Sort(ReportedDatas, [](FPackageData* A, FPackageData* B)
+	{
+		return A->GetPackageName().LexicalLess(B->GetPackageName());
+	});
+
+	bool bShowInstigators = (GCookProgressDisplay & (int32)ECookProgressDisplayMode::Instigators) != 0;
+	for (FPackageData* PackageData : ReportedDatas)
+	{
+		bool bRejected = !PackageData->IsInProgress() && !PackageData->HasAnyCookedPlatform();
+		UE_LOG(LogCookList, Display, TEXT("%s%s%s%s"),
+			bRejected ? TEXT("Rejected: ") : TEXT(""),
+			*WriteToString<256>(PackageData->GetPackageName()),
+			bShowInstigators ? TEXT(", Instigator: ") : TEXT(""),
+			bShowInstigators ? *PackageData->GetInstigator().ToString() : TEXT(""));
+	}
+}
+
 uint32 UCookOnTheFlyServer::TickCookOnTheFly(const float TimeSlice, ECookTickFlags TickFlags)
 {
 	check(IsCookOnTheFlyMode());
@@ -2497,7 +2558,8 @@ void UCookOnTheFlyServer::DemoteToIdle(UE::Cook::FPackageData& PackageData, UE::
 	{
 		WorkerRequests->ReportDemoteToIdle(PackageData, Reason);
 		// If per-package display is on, write a log statement explaining that the package was reachable but skipped.
-		if (GCookProgressDisplay & ((int32)ECookProgressDisplayMode::Instigators | (int32)ECookProgressDisplayMode::PackageNames))
+		if (!bCookListMode &
+			((GCookProgressDisplay & ((int32)ECookProgressDisplayMode::Instigators | (int32)ECookProgressDisplayMode::PackageNames)) != 0))
 		{
 			// Suppress the message in cases that cause large spam like NotInCurrentPlugin for DLC cooks.
 			if (Reason != ESuppressCookReason::NotInCurrentPlugin)
