CSV Import: allow reading from a file still open for write. This was already implemented for Reimport, but needed to be added for the initial import.
#rnx

diff --git a/Engine/Source/Editor/UnrealEd/Classes/Factories/CSVImportFactory.h b/Engine/Source/Editor/UnrealEd/Classes/Factories/CSVImportFactory.h
--- a/Engine/Source/Editor/UnrealEd/Classes/Factories/CSVImportFactory.h
+++ b/Engine/Source/Editor/UnrealEd/Classes/Factories/CSVImportFactory.h
@@ -58,9 +58,14 @@ class UNREALED_API UCSVImportFactory : public UFactory, public IImportSettingsPa
 {
 	GENERATED_UCLASS_BODY()
 
+public:
 	//~ Begin UFactory Interface
 	virtual FText GetDisplayName() const override;
-	virtual UObject* FactoryCreateText(UClass* InClass, UObject* InParent, FName InName, EObjectFlags Flags, UObject* Context, const TCHAR* Type, const TCHAR*& Buffer, const TCHAR* BufferEnd, FFeedbackContext* Warn, bool& bOutOperationCanceled) override;
+	virtual UObject* FactoryCreateFile(UClass* InClass, UObject* InParent, FName InName, EObjectFlags Flags,
+		const FString& Filename, const TCHAR* Parms, FFeedbackContext* Warn, bool& bOutOperationCanceled) override;
+	virtual UObject* FactoryCreateText(UClass* InClass, UObject* InParent, FName InName, EObjectFlags Flags,
+		UObject* Context, const TCHAR* Type, const TCHAR*& Buffer, const TCHAR* BufferEnd, FFeedbackContext* Warn,
+		bool& bOutOperationCanceled) override;
 	virtual bool DoesSupportClass(UClass * Class) override;
 	virtual bool FactoryCanImport(const FString& Filename) override;
 	virtual	IImportSettingsParser* GetImportSettingsParser() override;
diff --git a/Engine/Source/Editor/UnrealEd/Private/Factories/CSVImportFactory.cpp b/Engine/Source/Editor/UnrealEd/Private/Factories/CSVImportFactory.cpp
--- a/Engine/Source/Editor/UnrealEd/Private/Factories/CSVImportFactory.cpp
+++ b/Engine/Source/Editor/UnrealEd/Private/Factories/CSVImportFactory.cpp
@@ -111,6 +111,30 @@ void UCSVImportFactory::CleanUp()
 	DataTableImportOptions = nullptr;
 }
 
+UObject* UCSVImportFactory::FactoryCreateFile(UClass* InClass, UObject* InParent, FName InName, EObjectFlags Flags, 
+	const FString& Filename, const TCHAR* Parms, FFeedbackContext* Warn, bool& bOutOperationCanceled)
+{
+	// ScriptFactoryCreateFile not implemented. We do not support blueprint/python subclasses of CSVImportFactory
+	FString FileExtension = FPaths::GetExtension(Filename);
+
+	// load as text
+	check(bText); // Set in constructor, so we do not need to support load as binary
+	{
+		FString Data;
+		if (!FFileHelper::LoadFileToString(Data, *Filename, FFileHelper::EHashOptions::None, FILEREAD_AllowWrite))
+		{
+			UE_LOG(LogCSVImportFactory, Error, TEXT("Failed to load file '%s' to string"), *Filename);
+			return nullptr;
+		}
+
+		ParseParms(Parms);
+		const TCHAR* Ptr = *Data;
+
+		return FactoryCreateText(InClass, InParent, InName, Flags, nullptr, *FileExtension, Ptr, Ptr + Data.Len(),
+			Warn, bOutOperationCanceled);
+	}
+}
+
 UObject* UCSVImportFactory::FactoryCreateText(UClass* InClass, UObject* InParent, FName InName, EObjectFlags Flags, UObject* Context, const TCHAR* Type, const TCHAR*& Buffer, const TCHAR* BufferEnd, FFeedbackContext* Warn, bool& bOutOperationCanceled)
 {
 	GEditor->GetEditorSubsystem<UImportSubsystem>()->BroadcastAssetPreImport(this, InClass, InParent, InName, Type);
