Cooker: Fix static analysis error with & and !=. Use !! instead of != 0.
#rnx

diff --git a/Engine/Source/Editor/UnrealEd/Private/CookOnTheFlyServer.cpp b/Engine/Source/Editor/UnrealEd/Private/CookOnTheFlyServer.cpp
--- a/Engine/Source/Editor/UnrealEd/Private/CookOnTheFlyServer.cpp
+++ b/Engine/Source/Editor/UnrealEd/Private/CookOnTheFlyServer.cpp
@@ -2559,7 +2559,7 @@ void UCookOnTheFlyServer::DemoteToIdle(UE::Cook::FPackageData& PackageData, UE::
 		WorkerRequests->ReportDemoteToIdle(PackageData, Reason);
 		// If per-package display is on, write a log statement explaining that the package was reachable but skipped.
 		if (!bCookListMode &
-			((GCookProgressDisplay & ((int32)ECookProgressDisplayMode::Instigators | (int32)ECookProgressDisplayMode::PackageNames)) != 0))
+			!!((GCookProgressDisplay & ((int32)ECookProgressDisplayMode::Instigators | (int32)ECookProgressDisplayMode::PackageNames))))
 		{
 			// Suppress the message in cases that cause large spam like NotInCurrentPlugin for DLC cooks.
 			if (Reason != ESuppressCookReason::NotInCurrentPlugin)
