Change screenshot logic to use viewport's separate render target size when available, preventing crashes when viewport size is greater than render target size.
GetRenderTargetTextureSizeXY() should only vary from GetSizeXY() in the stereo case.

diff --git a/Engine/Source/Editor/LevelEditor/Private/HighResScreenshotUI.cpp b/Engine/Source/Editor/LevelEditor/Private/HighResScreenshotUI.cpp
--- a/Engine/Source/Editor/LevelEditor/Private/HighResScreenshotUI.cpp
+++ b/Engine/Source/Editor/LevelEditor/Private/HighResScreenshotUI.cpp
@@ -399,13 +399,13 @@ FReply SHighResScreenshotDialog::OnCaptureClicked()
 		auto ConfigViewport = Config.TargetViewport.Pin();
 		if (ConfigViewport.IsValid())
 		{
-			GScreenshotResolutionX = ConfigViewport->GetSizeXY().X * Config.ResolutionMultiplier;
-			GScreenshotResolutionY = ConfigViewport->GetSizeXY().Y * Config.ResolutionMultiplier;
+			GScreenshotResolutionX = ConfigViewport->GetRenderTargetTextureSizeXY().X * Config.ResolutionMultiplier;
+			GScreenshotResolutionY = ConfigViewport->GetRenderTargetTextureSizeXY().Y * Config.ResolutionMultiplier;
 			FIntRect ScaledCaptureRegion = Config.UnscaledCaptureRegion;
 
 			if (ScaledCaptureRegion.Area() > 0)
 			{
-				ScaledCaptureRegion.Clip(FIntRect(FIntPoint::ZeroValue, ConfigViewport->GetSizeXY()));
+				ScaledCaptureRegion.Clip(FIntRect(FIntPoint::ZeroValue, ConfigViewport->GetRenderTargetTextureSizeXY()));
 				ScaledCaptureRegion *= Config.ResolutionMultiplier;
 			}
 
diff --git a/Engine/Source/Editor/UnrealEd/Private/EditorViewportClient.cpp b/Engine/Source/Editor/UnrealEd/Private/EditorViewportClient.cpp
--- a/Engine/Source/Editor/UnrealEd/Private/EditorViewportClient.cpp
+++ b/Engine/Source/Editor/UnrealEd/Private/EditorViewportClient.cpp
@@ -5703,7 +5703,7 @@ void FEditorViewportClient::TakeScreenshot(FViewport* InViewport, bool bInValida
 	if (!bHdrEnabled)
 	{
 		TArray<FColor> RawPixels;
-		RawPixels.SetNum(InViewport->GetSizeXY().X * InViewport->GetSizeXY().Y);
+		RawPixels.SetNum(InViewport->GetRenderTargetTextureSizeXY().X * InViewport->GetRenderTargetTextureSizeXY().Y);
 		if( !InViewport->ReadPixels(RawPixels) )
 		{
 			// Failed to read the image from the viewport
@@ -5711,7 +5711,7 @@ void FEditorViewportClient::TakeScreenshot(FViewport* InViewport, bool bInValida
 			return;
 		}
 
-		TUniquePtr<TImagePixelData<FColor>> PixelData = MakeUnique<TImagePixelData<FColor>>(InViewport->GetSizeXY(), TArray64<FColor>(MoveTemp(RawPixels)));
+		TUniquePtr<TImagePixelData<FColor>> PixelData = MakeUnique<TImagePixelData<FColor>>(InViewport->GetRenderTargetTextureSizeXY(), TArray64<FColor>(MoveTemp(RawPixels)));
 
 		check(PixelData->IsDataWellFormed());
 		ImageTask->PixelData = MoveTemp(PixelData);
@@ -5722,7 +5722,7 @@ void FEditorViewportClient::TakeScreenshot(FViewport* InViewport, bool bInValida
 	else
 	{
 		TArray<FLinearColor> RawPixels;
-		RawPixels.SetNum(InViewport->GetSizeXY().X * InViewport->GetSizeXY().Y);
+		RawPixels.SetNum(InViewport->GetRenderTargetTextureSizeXY().X * InViewport->GetRenderTargetTextureSizeXY().Y);
 		if (!InViewport->ReadLinearColorPixels(RawPixels))
 		{
 			// Failed to read the image from the viewport
@@ -5732,7 +5732,7 @@ void FEditorViewportClient::TakeScreenshot(FViewport* InViewport, bool bInValida
 
 		ConvertPixelDataToSCRGB(RawPixels, InViewport->GetDisplayOutputFormat());
 
-		TUniquePtr<TImagePixelData<FLinearColor>> PixelData = MakeUnique<TImagePixelData<FLinearColor>>(InViewport->GetSizeXY(), TArray64<FLinearColor>(MoveTemp(RawPixels)));
+		TUniquePtr<TImagePixelData<FLinearColor>> PixelData = MakeUnique<TImagePixelData<FLinearColor>>(InViewport->GetRenderTargetTextureSizeXY(), TArray64<FLinearColor>(MoveTemp(RawPixels)));
 
 		check(PixelData->IsDataWellFormed());
 		ImageTask->PixelData = MoveTemp(PixelData);
@@ -5990,7 +5990,7 @@ bool FEditorViewportClient::ProcessScreenShots(FViewport* InViewport)
 		bool bHdrEnabled = InViewport->GetSceneHDREnabled();
 
 		// Determine the size of the captured viewport data.
-		FIntPoint BitmapSize = CaptureRect.Area() > 0 ? CaptureRect.Size() : InViewport->GetSizeXY();
+		FIntPoint BitmapSize = CaptureRect.Area() > 0 ? CaptureRect.Size() : InViewport->GetRenderTargetTextureSizeXY();
 		if (!bHdrEnabled)
 		{
 			TArray<FColor> Bitmap;
diff --git a/Engine/Source/Runtime/Engine/Private/GameViewportClient.cpp b/Engine/Source/Runtime/Engine/Private/GameViewportClient.cpp
--- a/Engine/Source/Runtime/Engine/Private/GameViewportClient.cpp
+++ b/Engine/Source/Runtime/Engine/Private/GameViewportClient.cpp
@@ -2002,7 +2002,7 @@ bool UGameViewportClient::ProcessScreenShots(FViewport* InViewport)
 
 		bool bScreenshotSuccessful = false;
 		bool bIsUI = false;
-		FIntVector Size(InViewport->GetSizeXY().X, InViewport->GetSizeXY().Y, 0);
+		FIntVector Size(InViewport->GetRenderTargetTextureSizeXY().X, InViewport->GetRenderTargetTextureSizeXY().Y, 0);
 
 		EDisplayOutputFormat ViewportOutputFormat = InViewport->GetDisplayOutputFormat();
 		bool bHdrEnabled = InViewport->GetSceneHDREnabled();
@@ -3759,8 +3759,8 @@ bool UGameViewportClient::HandleScreenshotCommand( const TCHAR* Cmd, FOutputDevi
 
 		FScreenshotRequest::RequestScreenshot(FileName, bShowUI, bAddFilenameSuffix, Viewport->GetSceneHDREnabled());
 
-		GScreenshotResolutionX = Viewport->GetSizeXY().X;
-		GScreenshotResolutionY = Viewport->GetSizeXY().Y;
+		GScreenshotResolutionX = Viewport->GetRenderTargetTextureSizeXY().X;
+		GScreenshotResolutionY = Viewport->GetRenderTargetTextureSizeXY().Y;
 	}
 	return true;
 }
diff --git a/Engine/Source/Runtime/Engine/Private/UnrealClient.cpp b/Engine/Source/Runtime/Engine/Private/UnrealClient.cpp
--- a/Engine/Source/Runtime/Engine/Private/UnrealClient.cpp
+++ b/Engine/Source/Runtime/Engine/Private/UnrealClient.cpp
@@ -62,7 +62,8 @@ bool FRenderTarget::ReadPixels(TArray< FColor >& OutImageData, FReadSurfaceDataF
 {
 	if(InRect == FIntRect(0, 0, 0, 0))
 	{
-		InRect = FIntRect(0, 0, GetSizeXY().X, GetSizeXY().Y);
+		FIntPoint RTTSize = RenderTargetTextureRHI->GetSizeXY(); // Actual RTT size may differ from viewport SizeXY for stereo render targets
+		InRect = FIntRect(0, 0, RTTSize.X, RTTSize.Y);
 	}
 
 	// Read the render target surface data back.	
@@ -198,7 +199,8 @@ bool FRenderTarget::ReadLinearColorPixels(TArray<FLinearColor> &OutImageData, FR
 {
 	if (InRect == FIntRect(0, 0, 0, 0))
 	{
-		InRect = FIntRect(0, 0, GetSizeXY().X, GetSizeXY().Y);
+		FIntPoint RTTSize = RenderTargetTextureRHI->GetSizeXY(); // Actual RTT size may differ from viewport SizeXY for stereo render targets
+		InRect = FIntRect(0, 0, RTTSize.X, RTTSize.Y);
 	}
 
 	// Read the render target surface data back.	
@@ -1449,8 +1451,8 @@ bool FViewport::TakeHighResScreenShot()
 {
 	if( GScreenshotResolutionX == 0 && GScreenshotResolutionY == 0 )
 	{
-		GScreenshotResolutionX = SizeX * GetHighResScreenshotConfig().ResolutionMultiplier;
-		GScreenshotResolutionY = SizeY * GetHighResScreenshotConfig().ResolutionMultiplier;
+		GScreenshotResolutionX = GetRenderTargetTextureSizeXY().X * GetHighResScreenshotConfig().ResolutionMultiplier;
+		GScreenshotResolutionY = GetRenderTargetTextureSizeXY().Y * GetHighResScreenshotConfig().ResolutionMultiplier;
 	}
 
 	uint32 MaxTextureDimension = GetMax2DTextureDimension();
@@ -2458,7 +2460,7 @@ ENGINE_API bool GetViewportScreenShot(FViewport* Viewport, TArray<FColor>& Bitma
 	// Read the contents of the viewport into an array.
 	if (Viewport->ReadPixels(Bitmap, FReadSurfaceDataFlags(), ViewRect))
 	{
-		check(Bitmap.Num() == ViewRect.Area() || (Bitmap.Num() == Viewport->GetSizeXY().X * Viewport->GetSizeXY().Y));
+		check(Bitmap.Num() == ViewRect.Area() || (Bitmap.Num() == Viewport->GetRenderTargetTextureSizeXY().X * Viewport->GetRenderTargetTextureSizeXY().Y));
 		return true;
 	}
 
@@ -2470,7 +2472,7 @@ ENGINE_API bool GetViewportScreenShotHDR(FViewport* Viewport, TArray<FLinearColo
 	// Read the contents of the viewport into an array.
 	if (Viewport->ReadLinearColorPixels(Bitmap, FReadSurfaceDataFlags(RCM_MinMax), ViewRect))
 	{
-		check(Bitmap.Num() == ViewRect.Area() || (Bitmap.Num() == Viewport->GetSizeXY().X * Viewport->GetSizeXY().Y));
+		check(Bitmap.Num() == ViewRect.Area() || (Bitmap.Num() == Viewport->GetRenderTargetTextureSizeXY().X * Viewport->GetRenderTargetTextureSizeXY().Y));
 		return true;
 	}
 
