Fixes extra separator at the bottom of the File menu of curve tables

diff --git a/Engine/Source/Editor/UnrealEd/Private/Toolkits/AssetEditorToolkit.cpp b/Engine/Source/Editor/UnrealEd/Private/Toolkits/AssetEditorToolkit.cpp
--- a/Engine/Source/Editor/UnrealEd/Private/Toolkits/AssetEditorToolkit.cpp
+++ b/Engine/Source/Editor/UnrealEd/Private/Toolkits/AssetEditorToolkit.cpp
@@ -946,7 +946,6 @@ void FAssetEditorToolkit::FillDefaultFileMenuCommands(FToolMenuSection& InSectio
 	{
 		InSection.AddMenuEntry(FAssetEditorCommonCommands::Get().SaveAssetAs, TAttribute<FText>(), TAttribute<FText>(), FSlateIcon(FAppStyle::GetAppStyleSetName(), "AssetEditor.SaveAssetAs")).InsertPosition = InsertPosition;
 	}
-	InSection.AddSeparator("DefaultFileMenuCommandsSeparator").InsertPosition = InsertPosition;;
 
 	if( IsWorldCentricAssetEditor() )
 	{
