Fixed editor world scene purge to happen later on when safer to do so. The previous location would trigger a component scene proxy reinit where some components (audio, UMG, etc..) would cache the (now) stale world or scene ptrs, and cause missing render target contents or audio. The new purge location is well after editor world is no longer assigned to GWorld, so any reinit will grab the right ptrs. There was a special case with Slate explicitly holding on to scene ptrs that I flush prior to purging the scene.

diff --git a/Engine/Source/Editor/UnrealEd/Private/Editor.cpp b/Engine/Source/Editor/UnrealEd/Private/Editor.cpp
--- a/Engine/Source/Editor/UnrealEd/Private/Editor.cpp
+++ b/Engine/Source/Editor/UnrealEd/Private/Editor.cpp
@@ -972,6 +972,19 @@ UWorld* SetPlayInEditorWorld( UWorld* PlayInEditorWorld )
 	GIsPlayInEditorWorld = true;
 	GWorld = PlayInEditorWorld;
 
+	// Purge the existing scene interface from the editor world to avoid 2x GPU allocations with the additional play-in-editor world
+	if (GEditor->EditorWorld != nullptr)
+	{
+		// Clear out Slate's active scenes list since these ptrs no longer reference valid scenes.
+		if (!GEditor->EditorWorld->bPurgedScene && FSlateApplication::IsInitialized())
+		{
+			FSlateApplication::Get().FlushRenderState();
+		}
+
+		// Tear down the scene interface for the editor world
+		GEditor->EditorWorld->PurgeScene();
+	}
+
 	if (FWorldContext* WorldContext = GEngine->GetWorldContextFromWorld(PlayInEditorWorld))
 	{
 		GPlayInEditorID = WorldContext->PIEInstance;
diff --git a/Engine/Source/Editor/UnrealEd/Private/PlayLevel.cpp b/Engine/Source/Editor/UnrealEd/Private/PlayLevel.cpp
--- a/Engine/Source/Editor/UnrealEd/Private/PlayLevel.cpp
+++ b/Engine/Source/Editor/UnrealEd/Private/PlayLevel.cpp
@@ -2703,9 +2703,6 @@ void UEditorEngine::StartPlayInEditorSession(FRequestPlaySessionParams& InReques
 	// Let navigation know PIE is starting so it can avoid any blueprint creation/deletion/instantiation affect editor map's navmesh changes
 	FNavigationSystem::OnPIEStart(*InWorld);
 
-	// Purge the existing scene interface from the editor world to avoid 2x GPU allocations with the additional play-in-editor world that is about to init
-	InWorld->PurgeScene();
-
 	ULevelEditorPlaySettings* EditorPlaySettings = InRequestParams.EditorPlaySettings;
 	check(EditorPlaySettings);
 
