Fix non-unity compile error in Editor.cpp

diff --git a/Engine/Source/Editor/UnrealEd/Private/Editor.cpp b/Engine/Source/Editor/UnrealEd/Private/Editor.cpp
--- a/Engine/Source/Editor/UnrealEd/Private/Editor.cpp
+++ b/Engine/Source/Editor/UnrealEd/Private/Editor.cpp
@@ -79,6 +79,8 @@
 #include "Subsystems/AssetEditorSubsystem.h"
 #endif
 
+#include "Framework/Application/SlateApplication.h"
+
 #define LOCTEXT_NAMESPACE "UnrealEd.Editor"
 
 PRAGMA_DISABLE_DEPRECATION_WARNINGS
