Added AutoSubmit feature to LogWorldPartitionNavigationDataBuilder.

diff --git a/Engine/Source/Editor/UnrealEd/Private/PackageSourceControlHelper.cpp b/Engine/Source/Editor/UnrealEd/Private/PackageSourceControlHelper.cpp
--- a/Engine/Source/Editor/UnrealEd/Private/PackageSourceControlHelper.cpp
+++ b/Engine/Source/Editor/UnrealEd/Private/PackageSourceControlHelper.cpp
@@ -224,13 +224,31 @@ bool FPackageSourceControlHelper::AddToSourceControl(UPackage* Package) const
 {
 	if (UseSourceControl())
 	{
-		FString PackageFilename = SourceControlHelpers::PackageFilename(Package);
+		const FString PackageFilename = SourceControlHelpers::PackageFilename(Package);
 		return AddToSourceControl({ PackageFilename });
 	}
 
 	return true;
 }
 
+bool FPackageSourceControlHelper::AddToSourceControl(const TArray<UPackage*>& Packages, bool bErrorsAsWarnings) const
+{
+	if (!UseSourceControl())
+	{
+		return true;
+	}
+
+	TArray<FString> PackageFilenames;
+	PackageFilenames.Reserve(Packages.Num());
+
+	for (const UPackage* Package : Packages)
+	{
+		FString PackageFilename = SourceControlHelpers::PackageFilename(Package);
+		PackageFilenames.Add(MoveTemp(PackageFilename));
+	}
+	return AddToSourceControl(PackageFilenames, bErrorsAsWarnings);
+}
+
 bool FPackageSourceControlHelper::AddToSourceControl(const TArray<FString>& PackageNames, bool bErrorsAsWarnings) const
 {
 	if (!UseSourceControl())
@@ -483,4 +501,33 @@ bool FPackageSourceControlHelper::GetDesiredStatesForModification(const TArray<F
 	}
 
 	return bSuccess;
-}
\ No newline at end of file
+}
+
+bool FPackageSourceControlHelper::GetMarkedForDeleteFiles(const TArray<FString>& Filenames, TArray<FString>& OutFilenames, bool bErrorsAsWarnings) const
+{
+	if (!UseSourceControl())
+	{
+		return true;
+	}
+
+	TArray<FSourceControlStateRef> SourceControlStates;
+	if (GetSourceControlProvider().GetState(Filenames, SourceControlStates, EStateCacheUsage::ForceUpdate) != ECommandResult::Succeeded)
+	{
+		// Nothing we can do if SCCStates fail
+		FPackageSourceControlHelperLog::Error(false, TEXT("Could not get revision control state for packages"));
+		return false;
+	}
+
+	for (const FSourceControlStateRef& SourceControlState : SourceControlStates)
+	{
+		const FString& PackageFilename = SourceControlState->GetFilename();
+
+		if (SourceControlState->IsDeleted())
+		{
+			OutFilenames.Add(PackageFilename);
+		}
+	}
+
+	return true;
+}
+
diff --git a/Engine/Source/Editor/UnrealEd/Private/WorldPartition/WorldPartitionNavigationDataBuilder.cpp b/Engine/Source/Editor/UnrealEd/Private/WorldPartition/WorldPartitionNavigationDataBuilder.cpp
--- a/Engine/Source/Editor/UnrealEd/Private/WorldPartition/WorldPartitionNavigationDataBuilder.cpp
+++ b/Engine/Source/Editor/UnrealEd/Private/WorldPartition/WorldPartitionNavigationDataBuilder.cpp
@@ -246,26 +246,98 @@ bool UWorldPartitionNavigationDataBuilder::RunInternal(UWorld* World, const FCel
 			WorldPartition->OnPackageDeleted(Package);
 		}
 		
+		if (PackageHelper.UseSourceControl())
 		{
-			// Add new packages to source control
-			TRACE_CPUPROFILER_EVENT_SCOPE(AddingToSourceControl);
-			UE_LOG(LogWorldPartitionNavigationDataBuilder, Log, TEXT("Adding packages to revision control."));
+			{
+				// Add new packages to source control
+				TRACE_CPUPROFILER_EVENT_SCOPE(AddingToSourceControl);
+				UE_LOG(LogWorldPartitionNavigationDataBuilder, Log, TEXT("Adding packages to revision control."));
+
+				if (!PackageHelper.AddToSourceControl(PackagesToAdd))
+				{
+					UE_LOG(LogWorldPartitionNavigationDataBuilder, Error, TEXT("Error adding packages."));
+					return true;
+				}
+			}
+
+			// Calculate AddedPackagesToSubmit, also adding an entry to AddedPackagesToSubmitMap as when we process later WP cells they may have PackagesToDelete
+			// that have already been added to AddedPackagesToSubmit (which we need to remove) and we want to avoid lots of string compares.
+			AddedPackagesToSubmit.Reserve(AddedPackagesToSubmit.Num() + PackagesToAdd.Num());
+			AddedPackagesToSubmitMap.Reserve(AddedPackagesToSubmitMap.Num() + PackagesToAdd.Num());
 
-			TArray<FString> PackageNamesToAdd;
-			PackageNamesToAdd.Reserve(PackagesToAdd.Num());
 			for (const UPackage* Package : PackagesToAdd)
 			{
-				PackageNamesToAdd.Add(Package->GetName());
+				FString PackageFilename = SourceControlHelpers::PackageFilename(Package);
+			
+				checkSlow(AddedPackagesToSubmit.Find(PackageFilename) == INDEX_NONE);
+				checkSlow(DeletedPackagesToSubmit.Find(PackageFilename) == INDEX_NONE);
+
+				const int32 Idx = AddedPackagesToSubmit.Add(PackageFilename);
+				AddedPackagesToSubmitMap.Add(MoveTemp(PackageFilename), Idx);
 			}
+		}
+	}
+
+	if (PackageHelper.UseSourceControl())
+	{
+		DeletedPackagesToSubmit.Reserve(DeletedPackagesToSubmit.Num() + PackagesToDelete.Num());
+
+		for (const UPackage* Package : PackagesToDelete)
+		{
+			const FString PackageFilename = SourceControlHelpers::PackageFilename(Package);
+
+			checkSlow(DeletedPackagesToSubmit.Find(PackageFilename) == INDEX_NONE);
+
+			const int32* AddedIdx = AddedPackagesToSubmitMap.Find(PackageFilename);
 			
-			if (!PackageHelper.AddToSourceControl(PackageNamesToAdd))
+			if (AddedIdx != nullptr)
 			{
-				UE_LOG(LogWorldPartitionNavigationDataBuilder, Error, TEXT("Error adding packages."));
-				return true;
+				// Dont remove the entry completely or the Indices in AddedPackagesToSubmitMap will point to the wrong place.
+				AddedPackagesToSubmit[*AddedIdx] = FString(TEXT(""));
+				const int32 NumRemoved = AddedPackagesToSubmitMap.Remove(PackageFilename);
+				ensure(NumRemoved == 1);
+			}
+			else
+			{
+				DeletedPackagesToSubmit.Add(PackageFilename);
 			}
 		}
+	}
+
+	UPackage::WaitForAsyncFileWrites();
+
+	return true;
+}
+
+bool UWorldPartitionNavigationDataBuilder::PostRun(UWorld* World, FPackageSourceControlHelper& PackageHelper, const bool bInRunSuccess)
+{
+	Super::PostRun(World, PackageHelper, bInRunSuccess);
+	
+	const FString ChangeDescription = FString::Printf(TEXT("Rebuilt navigation data for %s"), *World->GetName());
+	TArray<FString> PackagesToSubmit;
+
+	PackagesToSubmit.Reserve(AddedPackagesToSubmit.Num() + DeletedPackagesToSubmit.Num());
+
+	// DeletedPackages may never have been added to source control so don't attempt to check them in!
+	PackageHelper.GetMarkedForDeleteFiles(DeletedPackagesToSubmit, PackagesToSubmit);
+
+	for (const FString& Filename : AddedPackagesToSubmit)
+	{
+		// Empty Filenames indicate the file was added then deleted / reverted.
+		if (Filename != FString(TEXT("")))
+		{
+			PackagesToSubmit.Add(Filename);
+		}
+	}
+
+	DeletedPackagesToSubmit.Empty();
+	AddedPackagesToSubmit.Empty();
+	AddedPackagesToSubmitMap.Empty();
 
-		UPackage::WaitForAsyncFileWrites();
+	if (!AutoSubmitFiles(PackagesToSubmit, ChangeDescription))
+	{
+		UE_LOG(LogWorldPartitionNavigationDataBuilder, Error, TEXT("Error auto-submitting packages."));
+		return true;
 	}
 
 	return true;
diff --git a/Engine/Source/Editor/UnrealEd/Public/PackageSourceControlHelper.h b/Engine/Source/Editor/UnrealEd/Public/PackageSourceControlHelper.h
--- a/Engine/Source/Editor/UnrealEd/Public/PackageSourceControlHelper.h
+++ b/Engine/Source/Editor/UnrealEd/Public/PackageSourceControlHelper.h
@@ -17,10 +17,12 @@ public:
 	bool Delete(UPackage* Package) const;
 	bool Delete(const TArray<UPackage*>& Packages) const;
 	bool AddToSourceControl(UPackage* Package) const;
+	bool AddToSourceControl(const TArray<UPackage*>& Packages, bool bErrorsAsWarnings = false) const;
 	bool AddToSourceControl(const TArray<FString>& PackageNames, bool bErrorsAsWarnings = false) const;
 	bool Checkout(UPackage* Package) const;
 	bool Checkout(const TArray<FString>& PackageNames, bool bErrorsAsWarnings = false) const;
 	bool GetDesiredStatesForModification(const TArray<FString>& PackageNames, TArray<FString>& OutPackagesToCheckout, TArray<FString>& OutPackagesToAdd, bool bErrorsAsWarnings = false) const;
+	bool GetMarkedForDeleteFiles(const TArray<FString>& PackageFilenames, TArray<FString>& OutPackageFilenames, bool bErrorsAsWarnings = false) const;
 private:
 	ISourceControlProvider& GetSourceControlProvider() const;
 	FScopedSourceControl SourceControl;	
diff --git a/Engine/Source/Editor/UnrealEd/Public/WorldPartition/WorldPartitionNavigationDataBuilder.h b/Engine/Source/Editor/UnrealEd/Public/WorldPartition/WorldPartitionNavigationDataBuilder.h
--- a/Engine/Source/Editor/UnrealEd/Public/WorldPartition/WorldPartitionNavigationDataBuilder.h
+++ b/Engine/Source/Editor/UnrealEd/Public/WorldPartition/WorldPartitionNavigationDataBuilder.h
@@ -20,6 +20,7 @@ public:
 protected:
 	virtual bool PreRun(UWorld* World, FPackageSourceControlHelper& PackageHelper) override;
 	virtual bool RunInternal(UWorld* World, const FCellInfo& InCellInfo, FPackageSourceControlHelper& PackageHelper) override;
+	virtual bool PostRun(UWorld* World, FPackageSourceControlHelper& PackageHelper, const bool bInRunSuccess) override;
 	// UWorldPartitionBuilder interface end
 
 	bool GenerateNavigationData(UWorldPartition* WorldPartition, const FBox& LoadedBounds, const FBox& GeneratingBounds) const;
@@ -28,4 +29,8 @@ protected:
 	bool DeletePackages(const FPackageSourceControlHelper& PackageHelper, const TArray<UPackage*>& PackagesToDelete) const;
 
 	bool bCleanBuilderPackages = false;
+
+	TMap<FString, int32> AddedPackagesToSubmitMap;
+	TArray<FString> AddedPackagesToSubmit;
+	TArray<FString> DeletedPackagesToSubmit;
 };
