MPCook: When logging messages from the logs of a crashed cookworker, parse the severity and category and issue the log on the CookDirector with that severity and category.
#rnx

diff --git a/Engine/Source/Editor/UnrealEd/Private/Cooker/CookWorkerServer.cpp b/Engine/Source/Editor/UnrealEd/Private/Cooker/CookWorkerServer.cpp
--- a/Engine/Source/Editor/UnrealEd/Private/Cooker/CookWorkerServer.cpp
+++ b/Engine/Source/Editor/UnrealEd/Private/Cooker/CookWorkerServer.cpp
@@ -13,6 +13,7 @@
 #include "Interfaces/ITargetPlatformManagerModule.h"
 #include "Math/NumericLimits.h"
 #include "Misc/AssertionMacros.h"
+#include "Misc/Char.h"
 #include "Misc/FileHelper.h"
 #include "Misc/Parse.h"
 #include "Misc/ScopeLock.h"
@@ -61,6 +62,76 @@ void FCookWorkerServer::DetachFromRemoteProcess()
 	}
 }
 
+bool TryParseLogCategoryVerbosityMessage(FStringView Line, FName& OutCategory, ELogVerbosity::Type& OutVerbosity, FStringView& OutMessage)
+{
+	TPair<FStringView, ELogVerbosity::Type> VerbosityMarkers[]{
+		{ TEXTVIEW(": Fatal:"), ELogVerbosity::Fatal },
+		{ TEXTVIEW(": Error:"), ELogVerbosity::Error },
+		{ TEXTVIEW(": Warning:"), ELogVerbosity::Warning},
+		{ TEXTVIEW(": Display:"), ELogVerbosity::Display },
+		{ TEXTVIEW(":"), ELogVerbosity::Log },
+	};
+
+
+	// Find the first colon not in brackets and look for ": <Verbosity>:". This is complicated by Log verbosity not printing out the Verbosity:
+	// [2023.03.20-16.32.48:878][  0]LogCook: MessageText
+	// [2023.03.20-16.32.48:878][  0]LogCook: Display: MessageText
+
+	int32 FirstColon = INDEX_NONE;
+	int32 SubExpressionLevel = 0;
+	for (int32 Index = 0; Index < Line.Len(); ++Index)
+	{
+		switch (Line[Index])
+		{
+		case '[':
+			++SubExpressionLevel;
+			break;
+		case ']':
+			if (SubExpressionLevel > 0)
+			{
+				--SubExpressionLevel;
+			}
+			break;
+		case ':':
+			if (SubExpressionLevel == 0)
+			{
+				FirstColon = Index;
+			}
+			break;
+		default:
+			break;
+		}
+		if (FirstColon != INDEX_NONE)
+		{
+			break;
+		}
+	}
+	if (FirstColon == INDEX_NONE)
+	{
+		return false;
+	}
+
+	FStringView RestOfLine = FStringView(Line).RightChop(FirstColon);
+	for (TPair<FStringView, ELogVerbosity::Type>& VerbosityPair : VerbosityMarkers)
+	{
+		if (RestOfLine.StartsWith(VerbosityPair.Key, ESearchCase::IgnoreCase))
+		{
+			int32 CategoryEndIndex = FirstColon;
+			while (CategoryEndIndex > 0 && FChar::IsWhitespace(Line[CategoryEndIndex - 1])) --CategoryEndIndex;
+			int32 CategoryStartIndex = CategoryEndIndex > 0 ? CategoryEndIndex - 1 : CategoryEndIndex;
+			while (CategoryStartIndex > 0 && FChar::IsAlnum(Line[CategoryStartIndex - 1])) --CategoryStartIndex;
+			int32 MessageStartIndex = CategoryEndIndex + VerbosityPair.Key.Len();
+			while (MessageStartIndex < Line.Len() && FChar::IsWhitespace(Line[MessageStartIndex])) ++MessageStartIndex;
+
+			OutCategory = FName(FStringView(Line).SubStr(CategoryStartIndex, CategoryEndIndex - CategoryStartIndex));
+			OutVerbosity = VerbosityPair.Value;
+			OutMessage = FStringView(Line).SubStr(MessageStartIndex, Line.Len() - MessageStartIndex);
+			return true;
+		}
+	}
+	return false;
+}
+
 void FCookWorkerServer::SendCrashDiagnostics()
 {
 	FString LogFileName = Director.GetWorkerLogFileName(ProfileId);
@@ -101,7 +172,24 @@ void FCookWorkerServer::SendCrashDiagnostics()
 
 		while (FParse::Line(&StartText, Line))
 		{
-			UE_LOG(LogCook, Display, TEXT("[CookWorker %d]: %s"), ProfileId, *Line);
+			// Get the Category,Severity,Message out of each line and log it with that Category and Severity
+			// TODO: Change the CookWorkers to write out structured logs rather than interpreting their text logs
+			FName Category;
+			ELogVerbosity::Type Verbosity;
+			FStringView Message;
+			if (!TryParseLogCategoryVerbosityMessage(Line, Category, Verbosity, Message))
+			{
+				Category = LogCook.GetCategoryName();
+				Verbosity = ELogVerbosity::Display;
+				Message = Line;
+			}
+			// Downgrade Fatals in our local verbosity from Fatal to Error to avoid crashing the CookDirector
+			if (Verbosity == ELogVerbosity::Fatal)
+			{
+				Verbosity = ELogVerbosity::Error;
+			}
+			FMsg::Logf(__FILE__, __LINE__, Category, Verbosity, TEXT("[CookWorker %d]: %.*s"),
+				ProfileId, Message.Len(), Message.GetData());
 		}
 	}
 	if (!CrashDiagnosticsError.IsEmpty())
