Reduce log size when using Display Mode 2 during cooks by not printing out the cooked package file path.

diff --git a/Engine/Source/Editor/UnrealEd/Private/CookOnTheFlyServer.cpp b/Engine/Source/Editor/UnrealEd/Private/CookOnTheFlyServer.cpp
--- a/Engine/Source/Editor/UnrealEd/Private/CookOnTheFlyServer.cpp
+++ b/Engine/Source/Editor/UnrealEd/Private/CookOnTheFlyServer.cpp
@@ -5782,7 +5782,7 @@ void FSaveCookedPackageContext::SetupPlatform(const ITargetPlatform* InTargetPla
 	else if (!UAssetManager::Get().ShouldCookForPlatform(Package, TargetPlatform))
 	{
 		SavePackageResult = ESavePackageResult::ContainsEditorOnlyData;
-		UE_LOG(LogCook, Display, TEXT("Excluding %s -> %s"), *PackageName, *PlatFilename);
+		UE_LOG(LogCook, Display, TEXT("Excluding %s"), *PackageName);
 		return;
 	}
 	// check if this package is unsupported for the target platform (typically plugin content)
@@ -5796,7 +5796,7 @@ void FSaveCookedPackageContext::SetupPlatform(const ITargetPlatform* InTargetPla
 				(Generator && NeverCookPackages->Find(Generator->GetOwner().GetPackageName())))
 			{
 				SavePackageResult = ESavePackageResult::ContainsEditorOnlyData;
-				UE_LOG(LogCook, Display, TEXT("Excluding %s -> %s"), *PackageName, *PlatFilename);
+				UE_LOG(LogCook, Display, TEXT("Excluding %s"), *PackageName);
 				return;				
 			}
 		}
@@ -5826,7 +5826,7 @@ void FSaveCookedPackageContext::SetupPlatform(const ITargetPlatform* InTargetPla
 	UE_CLOG((GCookProgressDisplay & (int32)ECookProgressDisplayMode::Instigators) && bFirstPlatform, LogCook, Display,
 		TEXT("Cooking %s, Instigator: { %s }"), *PackageName, *(PackageData.GetInstigator().ToString()));
 	UE_CLOG(GCookProgressDisplay & (int32)ECookProgressDisplayMode::PackageNames, LogCook, Display,
-		TEXT("Cooking %s -> %s"), *PackageName, *PlatFilename);
+		TEXT("Cooking %s"), *PackageName);
 
 	bEndianSwap = (!TargetPlatform->IsLittleEndian()) ^ (!PLATFORM_LITTLE_ENDIAN);
 
