- Add support for delegating editor relaunch behaviour

diff --git a/Engine/Source/Editor/UnrealEd/Private/Editor.cpp b/Engine/Source/Editor/UnrealEd/Private/Editor.cpp
--- a/Engine/Source/Editor/UnrealEd/Private/Editor.cpp
+++ b/Engine/Source/Editor/UnrealEd/Private/Editor.cpp
@@ -169,6 +169,7 @@ FEditorDelegates::FOnOpenReferenceViewer				FEditorDelegates::OnOpenReferenceVie
 FEditorDelegates::FOnViewAssetIdentifiers				FEditorDelegates::OnOpenSizeMap;
 FEditorDelegates::FOnViewAssetIdentifiers				FEditorDelegates::OnOpenAssetAudit;
 FEditorDelegates::FOnViewAssetIdentifiers				FEditorDelegates::OnEditAssetIdentifiers;
+FEditorDelegates::FOnRestartRequested					FEditorDelegates::OnRestartRequested;
 
 PRAGMA_ENABLE_DEPRECATION_WARNINGS
 
diff --git a/Engine/Source/Editor/UnrealEd/Private/UnrealEdMisc.cpp b/Engine/Source/Editor/UnrealEd/Private/UnrealEdMisc.cpp
--- a/Engine/Source/Editor/UnrealEd/Private/UnrealEdMisc.cpp
+++ b/Engine/Source/Editor/UnrealEd/Private/UnrealEdMisc.cpp
@@ -944,15 +944,14 @@ bool FUnrealEdMisc::EnableWorldComposition(UWorld* InWorld, bool bEnable)
 	return true;
 }
 
-/** Build and return the path to the current project (used for relaunching the editor.)	 */
-FString CreateProjectPath()
+FString FUnrealEdMisc::GetProjectEditorBinaryPath()
 {
 #if PLATFORM_WINDOWS
 	return FPlatformProcess::ExecutablePath();
 #elif PLATFORM_MAC
 	@autoreleasepool
 	{
-		return UTF8_TO_TCHAR([[[NSBundle mainBundle] executablePath] fileSystemRepresentation]);
+		return UTF8_TO_TCHAR([[[NSBundle mainBundle]executablePath] fileSystemRepresentation] );
 	}
 #elif PLATFORM_LINUX
 	const TCHAR* PlatformConfig = FPlatformMisc::GetUBTPlatform();
@@ -963,9 +962,19 @@ FString CreateProjectPath()
 #endif
 }
 
-FString FUnrealEdMisc::GetProjectEditorBinaryPath()
+bool FUnrealEdMisc::SpawnEditorInstance(const FString& ProjectName)
 {
-	return CreateProjectPath();
+	// Use the same command line parameters that were used for this editor instance.
+	const FString Cmd = FString::Printf(TEXT("%s %s"), *ProjectName, *PendingCommandLine.Get(FCommandLine::Get()));
+
+	const FString ExeFilename = GetProjectEditorBinaryPath();
+	FProcHandle Handle = FPlatformProcess::CreateProc(*ExeFilename, *Cmd, true, false, false, NULL, 0, NULL, NULL);
+	const bool bSuccess = Handle.IsValid();
+	if (bSuccess)
+	{
+		FPlatformProcess::CloseProc(Handle);
+	}
+	return bSuccess;
 }
 
 void FUnrealEdMisc::OnExit()
@@ -1080,24 +1089,28 @@ void FUnrealEdMisc::OnExit()
 	const FString& PendingProjName = FUnrealEdMisc::Get().GetPendingProjectName();
 	if( PendingProjName.Len() > 0 )
 	{
-		// If there is a pending project switch, spawn that process now and use the same command line parameters that were used for this editor instance.
-		FString Cmd = FString::Printf(TEXT("%s %s"), *PendingProjName, *PendingCommandLine.Get(FCommandLine::Get()));
-
-		FString ExeFilename = CreateProjectPath();
-		FProcHandle Handle = FPlatformProcess::CreateProc( *ExeFilename, *Cmd, true, false, false, NULL, 0, NULL, NULL );
-		if( !Handle.IsValid() )
+		bool bSuccess = false;
+		if (FEditorDelegates::OnRestartRequested.IsBound())
+		{
+			bSuccess = FEditorDelegates::OnRestartRequested.Execute(PendingProjName);
+		}
+		else
+		{
+			bSuccess = SpawnEditorInstance(PendingProjName);
+		}
+		
+		if (!bSuccess)
 		{
 			// We were not able to spawn the new project exe.
 			// Its likely that the exe doesn't exist.
 			// Skip shutting down the editor if this happens
-			UE_LOG(LogUnrealEdMisc, Warning, TEXT("Could not restart the editor") );
+			UE_LOG(LogUnrealEdMisc, Warning, TEXT("Could not restart the editor"));
 
 			// Clear the pending project to ensure the editor can still be shut down normally
-			FUnrealEdMisc::Get().ClearPendingProjectName();
+			ClearPendingProjectName();
 
 			return;
 		}
-		FPlatformProcess::CloseProc(Handle);
 	}
 
 	// Unregister the command executor
diff --git a/Engine/Source/Editor/UnrealEd/Public/Editor.h b/Engine/Source/Editor/UnrealEd/Public/Editor.h
--- a/Engine/Source/Editor/UnrealEd/Public/Editor.h
+++ b/Engine/Source/Editor/UnrealEd/Public/Editor.h
@@ -167,6 +167,8 @@ struct UNREALED_API FEditorDelegates
 	DECLARE_MULTICAST_DELEGATE_OneParam(FOnViewAssetIdentifiers, TArray<FAssetIdentifier>);
 	/** delegate type to handle viewing/editing a set of asset identifiers (which are packages or ids) in the reference viewer */
 	DECLARE_MULTICAST_DELEGATE_TwoParams(FOnOpenReferenceViewer, const TArray<FAssetIdentifier>, const FReferenceViewerParams);
+	/** delegate type for when the editor requests a restart, enables overriding how a restart is performed */
+	DECLARE_DELEGATE_RetVal_OneParam(bool /*bSuccess*/, FOnRestartRequested, const FString& /*ProjectName*/);
 
 	/** Called when the CurrentLevel is switched to a new level.  Note that this event won't be fired for temporary
 		changes to the current level, such as when copying/pasting actors. */
@@ -345,6 +347,8 @@ struct UNREALED_API FEditorDelegates
 	static FOnViewAssetIdentifiers OnOpenAssetAudit;
 	/** Called to try and edit an asset identifier, which could be a package or searchable name */
 	static FOnViewAssetIdentifiers OnEditAssetIdentifiers;
+	/** Called when the editor requests a restart */
+	static FOnRestartRequested OnRestartRequested;
 };
 
 /**
diff --git a/Engine/Source/Editor/UnrealEd/Public/UnrealEdMisc.h b/Engine/Source/Editor/UnrealEd/Public/UnrealEdMisc.h
--- a/Engine/Source/Editor/UnrealEd/Public/UnrealEdMisc.h
+++ b/Engine/Source/Editor/UnrealEd/Public/UnrealEdMisc.h
@@ -27,7 +27,6 @@ enum class EMapChangeType : uint8
 	TearDownWorld,
 };
 
-
 /** The public interface for the unreal editor misc singleton. */
 class UNREALED_API FUnrealEdMisc
 {
@@ -280,6 +279,9 @@ public:
 	/** Finds a map using only the map name, no extension, no path, also caches it for faster lookup next time. */
 	FString FindMapFileFromPartialName(const FString& PartialMapName);
 
+	/** Launch an editor instance, passing ProjectName and the command line args from this running instance */
+	bool SpawnEditorInstance(const FString& ProjectName);
+
 private:
 	void SelectActorFromMessageToken(AActor* InActor);
 
