Avoid manipulating selected actors from ALT during reinstanciation

#rnx

diff --git a/Engine/Source/Editor/UnrealEd/Private/Kismet2/KismetReinstanceUtilities.cpp b/Engine/Source/Editor/UnrealEd/Private/Kismet2/KismetReinstanceUtilities.cpp
--- a/Engine/Source/Editor/UnrealEd/Private/Kismet2/KismetReinstanceUtilities.cpp
+++ b/Engine/Source/Editor/UnrealEd/Private/Kismet2/KismetReinstanceUtilities.cpp
@@ -2269,21 +2269,26 @@ void FBlueprintCompileReinstancer::ReplaceInstancesOfClass_Inner(const TMap<UCla
 		TArray<UObject*> ObjectsToReplace;
 
 		BP_SCOPED_COMPILER_EVENT_STAT(EKismetReinstancerStats_ReplaceInstancesOfClass);
-		if(GEditor && GEditor->GetSelectedActors())
-		{
-			SelectedActors = GEditor->GetSelectedActors();
 
-			// Note: For OFPA, each instance may be stored in its own external package.
-			for (int32 SelectionIdx = 0; SelectionIdx < SelectedActors->Num(); ++SelectionIdx)
+		// Reinstantiation can happen on the asyncloading thread and should not interact with GEditor in this case.
+		if (IsInGameThread())
+		{
+			if(GEditor && GEditor->GetSelectedActors())
 			{
-				if (const UObject* SelectedActor = SelectedActors->GetSelectedObject(SelectionIdx))
+				SelectedActors = GEditor->GetSelectedActors();
+
+				// Note: For OFPA, each instance may be stored in its own external package.
+				for (int32 SelectionIdx = 0; SelectionIdx < SelectedActors->Num(); ++SelectionIdx)
 				{
-					CheckAndSaveOuterPackageToCleanList(SelectedActor);
+					if (const UObject* SelectedActor = SelectedActors->GetSelectedObject(SelectionIdx))
+					{
+						CheckAndSaveOuterPackageToCleanList(SelectedActor);
+					}
 				}
-			}
 			
-			SelectedActors->BeginBatchSelectOperation();
-			SelectedActors->Modify();
+				SelectedActors->BeginBatchSelectOperation();
+				SelectedActors->Modify();
+			}
 		}
 
 		// WARNING: for (TPair<UClass*, UClass*> OldToNewClass : InOldToNewClassMap) duplicated below 
