Added missing include in CookPackageSplitter.h

diff --git a/Engine/Source/Editor/UnrealEd/Public/CookPackageSplitter.h b/Engine/Source/Editor/UnrealEd/Public/CookPackageSplitter.h
--- a/Engine/Source/Editor/UnrealEd/Public/CookPackageSplitter.h
+++ b/Engine/Source/Editor/UnrealEd/Public/CookPackageSplitter.h
@@ -4,6 +4,7 @@
 
 #if WITH_EDITOR
 
+#include "AssetRegistry/IAssetRegistry.h"
 #include "Containers/Array.h"
 #include "Containers/List.h"
 #include "Containers/UnrealString.h"
