World Partition - Builder: Fix inverted exit code for commandlet success vs error

diff --git a/Engine/Source/Editor/UnrealEd/Private/Commandlets/WorldPartitionBuilderCommandlet.cpp b/Engine/Source/Editor/UnrealEd/Private/Commandlets/WorldPartitionBuilderCommandlet.cpp
--- a/Engine/Source/Editor/UnrealEd/Private/Commandlets/WorldPartitionBuilderCommandlet.cpp
+++ b/Engine/Source/Editor/UnrealEd/Private/Commandlets/WorldPartitionBuilderCommandlet.cpp
@@ -112,11 +112,11 @@ int32 UWorldPartitionBuilderCommandlet::Main(const FString& Params)
 	{
 		if (!RunBuilder(BuilderClass, MapPackageName))
 		{
-			return false;
+			return 1;
 		}
 	}
 
-	return true;
+	return 0;
 }
 
 bool UWorldPartitionBuilderCommandlet::RunBuilder(TSubclassOf<UWorldPartitionBuilder> InBuilderClass, const FString& InWorldPackageName)
