Fix to support running distributed builds on multiple maps
* Build products list was getting overriden for each map, causing only the last one to be considered
* Now appending building products to the same build product file.

diff --git a/Engine/Source/Editor/UnrealEd/Private/WorldPartition/WorldPartitionHLODsBuilder.cpp b/Engine/Source/Editor/UnrealEd/Private/WorldPartition/WorldPartitionHLODsBuilder.cpp
--- a/Engine/Source/Editor/UnrealEd/Private/WorldPartition/WorldPartitionHLODsBuilder.cpp
+++ b/Engine/Source/Editor/UnrealEd/Private/WorldPartition/WorldPartitionHLODsBuilder.cpp
@@ -438,11 +438,8 @@ bool UWorldPartitionHLODsBuilder::SetupHLODActors()
 			BuildProducts.Add(BuildManifest);
 
 			// Write build products to a file
-			FString BuildProductsFile = DistributedBuildWorkingDir / BuildProductsFileName;
-			bool bRet = FFileHelper::SaveStringArrayToFile(BuildProducts, *BuildProductsFile);
-			if (!bRet)
+			if (!AddBuildProducts(BuildProducts))
 			{
-				UE_LOG(LogWorldPartitionHLODsBuilder, Error, TEXT("Error writing build product file %s"), *BuildProductsFile);
 				return false;
 			}
 
@@ -532,11 +529,8 @@ bool UWorldPartitionHLODsBuilder::BuildHLODActors()
 		}
 
 		// Write build products to a file
-		FString BuildProductsFile = DistributedBuildWorkingDir / BuildProductsFileName;
-		bool bRet = FFileHelper::SaveStringArrayToFile(BuildProducts, *BuildProductsFile);
-		if (!bRet)
+		if (!AddBuildProducts(BuildProducts))
 		{
-			UE_LOG(LogWorldPartitionHLODsBuilder, Error, TEXT("Error writing build product file %s"), *BuildProductsFile);
 			return false;
 		}
 
@@ -1228,3 +1222,15 @@ bool UWorldPartitionHLODsBuilder::CopyFilesFromWorkingDir(const FString& SourceD
 
 	return true;
 }
+
+bool UWorldPartitionHLODsBuilder::AddBuildProducts(const TArray<FString>& BuildProducts) const
+{
+	// Write build products to a file
+	FString BuildProductsFile = FString::Printf(TEXT("%s/%s/%s"), *FPaths::RootDir(), *DistributedBuildWorkingDirName, *BuildProductsFileName);
+	bool bRet = FFileHelper::SaveStringArrayToFile(BuildProducts, *BuildProductsFile, FFileHelper::EEncodingOptions::AutoDetect, &IFileManager::Get(), EFileWrite::FILEWRITE_Append);
+	if (!bRet)
+	{
+		UE_LOG(LogWorldPartitionHLODsBuilder, Error, TEXT("Error writing build product file %s"), *BuildProductsFile);
+	}
+	return bRet;
+}
diff --git a/Engine/Source/Editor/UnrealEd/Public/WorldPartition/WorldPartitionHLODsBuilder.h b/Engine/Source/Editor/UnrealEd/Public/WorldPartition/WorldPartitionHLODsBuilder.h
--- a/Engine/Source/Editor/UnrealEd/Public/WorldPartition/WorldPartitionHLODsBuilder.h
+++ b/Engine/Source/Editor/UnrealEd/Public/WorldPartition/WorldPartitionHLODsBuilder.h
@@ -103,6 +103,8 @@ protected:
 
 	bool ShouldRunStep(const EHLODBuildStep BuildStep) const;
 
+	bool AddBuildProducts(const TArray<FString>& BuildProducts) const;
+
 private:
 	class UWorldPartition* WorldPartition;
 	class FSourceControlHelper* SourceControlHelper;
