World Partition - Builder: Improved collection processing logging & error reporting

diff --git a/Engine/Source/Editor/UnrealEd/Classes/Commandlets/WorldPartitionBuilderCommandlet.h b/Engine/Source/Editor/UnrealEd/Classes/Commandlets/WorldPartitionBuilderCommandlet.h
--- a/Engine/Source/Editor/UnrealEd/Classes/Commandlets/WorldPartitionBuilderCommandlet.h
+++ b/Engine/Source/Editor/UnrealEd/Classes/Commandlets/WorldPartitionBuilderCommandlet.h
@@ -19,5 +19,6 @@ class UWorldPartitionBuilderCommandlet : public UCommandlet
 	//~ End UCommandlet Interface
 
 private:
+	TArray<FString> GatherMapsFromCollection(const FString& CollectionName) const;
 	bool RunBuilder(TSubclassOf<UWorldPartitionBuilder> InBuilderClass, const FString& InWorldPackageName);
 };
\ No newline at end of file
diff --git a/Engine/Source/Editor/UnrealEd/Private/Commandlets/WorldPartitionBuilderCommandlet.cpp b/Engine/Source/Editor/UnrealEd/Private/Commandlets/WorldPartitionBuilderCommandlet.cpp
--- a/Engine/Source/Editor/UnrealEd/Private/Commandlets/WorldPartitionBuilderCommandlet.cpp
+++ b/Engine/Source/Editor/UnrealEd/Private/Commandlets/WorldPartitionBuilderCommandlet.cpp
@@ -56,8 +56,6 @@ int32 UWorldPartitionBuilderCommandlet::Main(const FString& Params)
 	ICollectionManager& CollectionManager = FModuleManager::LoadModuleChecked<FCollectionManagerModule>("CollectionManager").Get();
 	TArray<FString> MapPackagesNames;
 
-	FAssetRegistryModule* AssetRegistryModule = &FModuleManager::LoadModuleChecked<FAssetRegistryModule>(TEXT("AssetRegistry"));
-
 	// Parse map name or maps collection
 	if (FPackageName::SearchForPackageOnDisk(Tokens[0]))
 	{
@@ -65,25 +63,12 @@ int32 UWorldPartitionBuilderCommandlet::Main(const FString& Params)
 	}	
 	else if (CollectionManager.CollectionExists(FName(Tokens[0]), ECollectionShareType::CST_All))
 	{
-		TArray<FSoftObjectPath> AssetsPaths;
-		CollectionManager.GetAssetsInCollection(FName(Tokens[0]), ECollectionShareType::CST_All, AssetsPaths, ECollectionRecursionFlags::SelfAndChildren);
-
-		for (const auto& AssetPath : AssetsPaths)
-		{
-			const bool bIncludeOnlyOnDiskAssets = true;
-			FAssetData AssetData = AssetRegistryModule->Get().GetAssetByObjectPath(AssetPath, bIncludeOnlyOnDiskAssets);
-			const bool bIsWorldAsset = AssetData.IsValid() && (AssetData.AssetClassPath == UWorld::StaticClass()->GetClassPathName());
-			if (bIsWorldAsset)
-			{
-				MapPackagesNames.Add(AssetPath.GetAssetPathString());
-			}
-		}
-
+		MapPackagesNames = GatherMapsFromCollection(Tokens[0]);
 		if (MapPackagesNames.IsEmpty())
-		{
-			UE_LOG(LogWorldPartitionBuilderCommandlet, Warning, TEXT("Found no maps to process in collection %s), exiting"), *Tokens[0]);
-			return 0;
-		}
+	    {
+		    UE_LOG(LogWorldPartitionBuilderCommandlet, Warning, TEXT("Found no maps to process in collection %s, exiting"), *Tokens[0]);
+		    return 0;
+	    }
 	}
 	else
 	{
@@ -119,6 +104,45 @@ int32 UWorldPartitionBuilderCommandlet::Main(const FString& Params)
 	return 0;
 }
 
+TArray<FString> UWorldPartitionBuilderCommandlet::GatherMapsFromCollection(const FString& CollectionName) const
+{
+	TArray<FString> MapPackagesNames;
+
+	ICollectionManager& CollectionManager = FModuleManager::LoadModuleChecked<FCollectionManagerModule>("CollectionManager").Get();
+	IAssetRegistry& AssetRegistry = FModuleManager::LoadModuleChecked<FAssetRegistryModule>("AssetRegistry").Get();
+
+	TArray<FSoftObjectPath> AssetsPaths;
+	CollectionManager.GetAssetsInCollection(FName(CollectionName), ECollectionShareType::CST_All, AssetsPaths, ECollectionRecursionFlags::SelfAndChildren);
+
+	UE_LOG(LogWorldPartitionBuilderCommandlet, Display, TEXT("Processing collection %s (%d items)"), *CollectionName, AssetsPaths.Num());
+	for (const auto& AssetPath : AssetsPaths)
+	{
+		FString PackageName = AssetPath.GetLongPackageName();
+		UE_LOG(LogWorldPartitionBuilderCommandlet, Display, TEXT("* %s"), *PackageName);
+
+		const bool bIncludeOnlyOnDiskAssets = true;
+		FAssetData AssetData = AssetRegistry.GetAssetByObjectPath(AssetPath, bIncludeOnlyOnDiskAssets);
+
+		if (AssetData.IsValid())
+		{
+			if (AssetData.AssetClassPath == UWorld::StaticClass()->GetClassPathName())
+			{
+				MapPackagesNames.Add(PackageName);
+			}
+			else
+			{
+				UE_LOG(LogWorldPartitionBuilderCommandlet, Warning, TEXT("%s is not a map asset"), *PackageName);
+			}
+		}
+		else
+		{
+			UE_LOG(LogWorldPartitionBuilderCommandlet, Warning, TEXT("%s was not found"), *PackageName);
+		}
+	}
+
+	return MapPackagesNames;
+}
+
 bool UWorldPartitionBuilderCommandlet::RunBuilder(TSubclassOf<UWorldPartitionBuilder> InBuilderClass, const FString& InWorldPackageName)
 {
 	// This will convert incomplete package name to a fully qualified path
