Cooker: Fix detection of externalactors by naming convention to use case-insensitive compare.
#rnx

diff --git a/Engine/Source/Editor/UnrealEd/Private/Commandlets/AssetRegistryGenerator.cpp b/Engine/Source/Editor/UnrealEd/Private/Commandlets/AssetRegistryGenerator.cpp
--- a/Engine/Source/Editor/UnrealEd/Private/Commandlets/AssetRegistryGenerator.cpp
+++ b/Engine/Source/Editor/UnrealEd/Private/Commandlets/AssetRegistryGenerator.cpp
@@ -45,6 +45,7 @@
 #include "Serialization/JsonWriter.h"
 #include "Settings/ProjectPackagingSettings.h"
 #include "Stats/StatsMisc.h"
+#include "String/Find.h"
 #include "String/ParseTokens.h"
 #include "Templates/UniquePtr.h"
 #include "UObject/SoftObjectPath.h"
@@ -1431,7 +1432,7 @@ void FAssetRegistryGenerator::ComputePackageDifferences(const FComputeDifference
 				// Workaround: Detect external actors by naming convention and suppress their reference to the map package.
 				// See also UAssetManager::ShouldSetManager
 				TStringBuilder<256> ReferencerPackageNameStr(InPlace, ReferencerPackageName);
-				int32 ExternalActorsFolderIndex = ReferencerPackageNameStr.ToView().Find(ExternalActorsFolderName);
+				int32 ExternalActorsFolderIndex = UE::String::FindFirst(ReferencerPackageNameStr, ExternalActorsFolderName, ESearchCase::IgnoreCase);
 				if (ExternalActorsFolderIndex != INDEX_NONE)
 				{
 					if (ModifiedPackageLeafName.IsEmpty())
