Remove call to flush RenderCommand and block from EditorViewportClient::Draw()

Realtime viewports are ones which are redrawn each game tick.  When PIE or SIE starts, all viewports are forced to be non-realtime, the reason for this is to give PIE more resources to render.
This flush was added to reduce the latency between user input at rendering, however it has the drawback of significantly reducing performance of PIE for viewports that push frames each tick anyways (AnimationBlueprintEditor + AnimationEditorViewportClient). Since the engine is pipelined already for performance reasons, the small increase in latency between input and rendering is the price to be paid for this architecture. Removal of this flush is justified to be consistent with the architecture and to maintain framerate.

diff --git a/Engine/Source/Editor/UnrealEd/Private/EditorViewportClient.cpp b/Engine/Source/Editor/UnrealEd/Private/EditorViewportClient.cpp
--- a/Engine/Source/Editor/UnrealEd/Private/EditorViewportClient.cpp
+++ b/Engine/Source/Editor/UnrealEd/Private/EditorViewportClient.cpp
@@ -4240,12 +4240,6 @@ void FEditorViewportClient::Draw(FViewport* InViewport, FCanvas* Canvas)
 #endif
 	}
 
-	if(!IsRealtime())
-	{
-		// Wait for the rendering thread to finish drawing the view before returning.
-		// This reduces the apparent latency of dragging the viewport around.
-		FlushRenderingCommands();
-	}
 
 	Viewport = ViewportBackup;
 }
