Don't call ResetLoaders on newly created world packages as this will prevent future loading of external actor packages to work propertly

diff --git a/Engine/Source/Editor/UnrealEd/Private/FileHelpers.cpp b/Engine/Source/Editor/UnrealEd/Private/FileHelpers.cpp
--- a/Engine/Source/Editor/UnrealEd/Private/FileHelpers.cpp
+++ b/Engine/Source/Editor/UnrealEd/Private/FileHelpers.cpp
@@ -5,6 +5,7 @@
 #include "HAL/PlatformFileManager.h"
 #include "Misc/MessageDialog.h"
 #include "HAL/FileManager.h"
+#include "Algo/Copy.h"
 #include "Misc/CommandLine.h"
 #include "Misc/Paths.h"
 #include "Misc/ConfigCacheIni.h"
@@ -3293,7 +3294,17 @@ static void PrepareSavePackages(const TArray<UPackage*>& PackagesToSave)
 		}
 	}
 
-	ResetLoaders(MakeArrayView<UObject*>((UObject**)PackagesToSave.GetData(), PackagesToSave.Num()));
+	// Don't call ResetLoaders on newly created world packages as this will prevent future loading of external actor packages to work propertly
+	// Linker will fail to resolve SourceLinker of external actor's world package import (see GetPackageLinker test for PKG_InMemoryOnly on TargetPackage's Package Flag)
+	TArray<UPackage*> PackagesToResetLoaders;
+	PackagesToResetLoaders.Reserve(PackagesToSave.Num());
+	Algo::CopyIf(PackagesToSave, PackagesToResetLoaders, [&](UPackage* Package)
+	{
+		const bool bIsNewlyCreatedWorldPackage = Package->HasAnyPackageFlags(PKG_NewlyCreated) && UWorld::FindWorldInPackage(Package);
+		return !bIsNewlyCreatedWorldPackage;
+	});
+
+	ResetLoaders(MakeArrayView<UObject*>((UObject**)PackagesToResetLoaders.GetData(), PackagesToResetLoaders.Num()));
 }
 
 /**
