Make sure AssetEditorSubsystem exists before emitting notification on shutdown

diff --git a/Engine/Source/Editor/UnrealEd/Private/Toolkits/AssetEditorToolkit.cpp b/Engine/Source/Editor/UnrealEd/Private/Toolkits/AssetEditorToolkit.cpp
--- a/Engine/Source/Editor/UnrealEd/Private/Toolkits/AssetEditorToolkit.cpp
+++ b/Engine/Source/Editor/UnrealEd/Private/Toolkits/AssetEditorToolkit.cpp
@@ -327,7 +327,10 @@ FAssetEditorToolkit::~FAssetEditorToolkit()
 	// We're no longer editing this object, so let the editor know
 	if (GEditor)
 	{
-		GEditor->GetEditorSubsystem<UAssetEditorSubsystem>()->NotifyEditorClosed(this);
+		if (UAssetEditorSubsystem* AssetEditorSubsystem = GEditor->GetEditorSubsystem<UAssetEditorSubsystem>())
+		{
+			AssetEditorSubsystem->NotifyEditorClosed(this);
+		}
 	}
 
 	EditorModeManager.Reset();
