Fix highlighting of multiline hlsl comments

diff --git a/Engine/Plugins/Animation/DeformerGraph/Source/OptimusEditor/Private/OptimusHLSLSyntaxHighlighter.cpp b/Engine/Plugins/Animation/DeformerGraph/Source/OptimusEditor/Private/OptimusHLSLSyntaxHighlighter.cpp
--- a/Engine/Plugins/Animation/DeformerGraph/Source/OptimusEditor/Private/OptimusHLSLSyntaxHighlighter.cpp
+++ b/Engine/Plugins/Animation/DeformerGraph/Source/OptimusEditor/Private/OptimusHLSLSyntaxHighlighter.cpp
@@ -37,9 +37,9 @@ void FOptimusHLSLSyntaxHighlighter::ParseTokens(const FString& SourceString, FTe
 	TargetTextLayout.SetLineHighlights(LineHighlightsToAdd);
 }
 
-FTextLayout::FNewLineData FOptimusHLSLSyntaxHighlighter::ProcessTokenizedLine(const ISyntaxTokenizer::FTokenizedLine& TokenizedLine, const int32& LineNumber, const FString& SourceString)
+FTextLayout::FNewLineData FOptimusHLSLSyntaxHighlighter::ProcessTokenizedLine(const ISyntaxTokenizer::FTokenizedLine& TokenizedLine, const int32& LineNumber, const FString& SourceString, EParseState& CurrentParseState)
 {
-	FTextLayout::FNewLineData LineData = FHLSLSyntaxHighlighterMarshaller::ProcessTokenizedLine(TokenizedLine, LineNumber, SourceString);
+	FTextLayout::FNewLineData LineData = FHLSLSyntaxHighlighterMarshaller::ProcessTokenizedLine(TokenizedLine, LineNumber, SourceString, CurrentParseState);
 
 	FTextBlockStyle ErrorTextStyle = SyntaxTextStyle.ErrorTextStyle;
 	TSharedPtr<FSlateTextUnderlineLineHighlighter> UnderlineLineHighlighter = FSlateTextUnderlineLineHighlighter::Create(ErrorTextStyle.UnderlineBrush, ErrorTextStyle.Font, ErrorTextStyle.ColorAndOpacity, ErrorTextStyle.ShadowOffset, ErrorTextStyle.ShadowColorAndOpacity);
diff --git a/Engine/Plugins/Animation/DeformerGraph/Source/OptimusEditor/Private/OptimusHLSLSyntaxHighlighter.h b/Engine/Plugins/Animation/DeformerGraph/Source/OptimusEditor/Private/OptimusHLSLSyntaxHighlighter.h
--- a/Engine/Plugins/Animation/DeformerGraph/Source/OptimusEditor/Private/OptimusHLSLSyntaxHighlighter.h
+++ b/Engine/Plugins/Animation/DeformerGraph/Source/OptimusEditor/Private/OptimusHLSLSyntaxHighlighter.h
@@ -27,7 +27,7 @@ protected:
 	FOptimusHLSLSyntaxHighlighter(TSharedPtr<ISyntaxTokenizer> InTokenizer, const FSyntaxTextStyle& InSyntaxTextStyle);
 
 	virtual void ParseTokens(const FString& SourceString, FTextLayout& TargetTextLayout, TArray<ISyntaxTokenizer::FTokenizedLine> TokenizedLines) override;
-	virtual FTextLayout::FNewLineData ProcessTokenizedLine(const ISyntaxTokenizer::FTokenizedLine& TokenizedLine, const int32& LineNumber, const FString& SourceString) override;
+	virtual FTextLayout::FNewLineData ProcessTokenizedLine(const ISyntaxTokenizer::FTokenizedLine& TokenizedLine, const int32& LineNumber, const FString& SourceString, EParseState& CurrentParseState) override;
 
 	TMultiMap<int32 /*Line*/, FOptimusCompilerDiagnostic> CompilerMessages;
 
diff --git a/Engine/Source/Editor/UnrealEd/Private/Text/HLSLSyntaxHighlighterMarshaller.cpp b/Engine/Source/Editor/UnrealEd/Private/Text/HLSLSyntaxHighlighterMarshaller.cpp
--- a/Engine/Source/Editor/UnrealEd/Private/Text/HLSLSyntaxHighlighterMarshaller.cpp
+++ b/Engine/Source/Editor/UnrealEd/Private/Text/HLSLSyntaxHighlighterMarshaller.cpp
@@ -505,26 +505,18 @@ void FHLSLSyntaxHighlighterMarshaller::ParseTokens(const FString& SourceString,
 
 	// Parse the tokens, generating the styled runs for each line
 	int32 LineNo = 0;
+	EParseState ParseState = EParseState::None;
 	for(const ISyntaxTokenizer::FTokenizedLine& TokenizedLine : TokenizedLines)
 	{
-		LinesToAdd.Add(ProcessTokenizedLine(TokenizedLine, LineNo, SourceString));
+		LinesToAdd.Add(ProcessTokenizedLine(TokenizedLine, LineNo, SourceString, ParseState));
 		LineNo++;
 	}
 
 	TargetTextLayout.AddLines(LinesToAdd);
 }
 
-FTextLayout::FNewLineData FHLSLSyntaxHighlighterMarshaller::ProcessTokenizedLine(const ISyntaxTokenizer::FTokenizedLine& TokenizedLine, const int32& LineNumber, const FString& SourceString)
+FTextLayout::FNewLineData FHLSLSyntaxHighlighterMarshaller::ProcessTokenizedLine(const ISyntaxTokenizer::FTokenizedLine& TokenizedLine, const int32& LineNumber, const FString& SourceString, EParseState& ParseState)
 {
-	enum class EParseState : uint8
-	{
-		None,
-		LookingForString,
-		LookingForCharacter,
-		LookingForSingleLineComment,
-		LookingForMultiLineComment,
-	};
-	EParseState ParseState = EParseState::None;
 	TSharedRef<FString> ModelString = MakeShareable(new FString());
 	TArray< TSharedRef< IRun > > Runs;
 
@@ -645,6 +637,10 @@ FTextLayout::FNewLineData FHLSLSyntaxHighlighterMarshaller::ProcessTokenizedLine
 		}
 	}
 
+	if (ParseState != EParseState::LookingForMultiLineComment)
+	{
+		ParseState = EParseState::None;
+	}
 	return FTextLayout::FNewLineData(MoveTemp(ModelString), MoveTemp(Runs));
 }
 
diff --git a/Engine/Source/Editor/UnrealEd/Public/Text/HLSLSyntaxHighlighterMarshaller.h b/Engine/Source/Editor/UnrealEd/Public/Text/HLSLSyntaxHighlighterMarshaller.h
--- a/Engine/Source/Editor/UnrealEd/Public/Text/HLSLSyntaxHighlighterMarshaller.h
+++ b/Engine/Source/Editor/UnrealEd/Public/Text/HLSLSyntaxHighlighterMarshaller.h
@@ -1,8 +1,6 @@
 ﻿// Copyright Epic Games, Inc. All Rights Reserved.
 #pragma once
 
-#include "Containers/Array.h"
-#include "Containers/UnrealString.h"
 #include "CoreMinimal.h"
 #include "Framework/Text/SyntaxHighlighterTextLayoutMarshaller.h"
 #include "Framework/Text/SyntaxTokenizer.h"
@@ -57,11 +55,20 @@ public:
 
 protected:
 
+	enum class EParseState : uint8
+	{
+		None,
+		LookingForString,
+		LookingForCharacter,
+		LookingForSingleLineComment,
+		LookingForMultiLineComment,
+	};
+
 	static TSharedPtr<ISyntaxTokenizer> CreateTokenizer();
 	
 	virtual void ParseTokens(const FString& SourceString, FTextLayout& TargetTextLayout, TArray<ISyntaxTokenizer::FTokenizedLine> TokenizedLines) override;
 
-	virtual FTextLayout::FNewLineData ProcessTokenizedLine(const ISyntaxTokenizer::FTokenizedLine& TokenizedLine, const int32& LineNumber, const FString& SourceString);
+	virtual FTextLayout::FNewLineData ProcessTokenizedLine(const ISyntaxTokenizer::FTokenizedLine& TokenizedLine, const int32& LineNumber, const FString& SourceString, EParseState& CurrentParseState);
 
 	FHLSLSyntaxHighlighterMarshaller(TSharedPtr<ISyntaxTokenizer> InTokenizer, const FSyntaxTextStyle& InSyntaxTextStyle);
 
