TypedBlockAllocator: Add FreeList ability to skip calling destructor/constructor when an allocation is freed and then reused. Split class into subclasses to support the different types of FreeList functionality.
#rnx

diff --git a/Engine/Source/Editor/UnrealEd/Private/Cooker/CookPackageData.h b/Engine/Source/Editor/UnrealEd/Private/Cooker/CookPackageData.h
--- a/Engine/Source/Editor/UnrealEd/Private/Cooker/CookPackageData.h
+++ b/Engine/Source/Editor/UnrealEd/Private/Cooker/CookPackageData.h
@@ -1330,7 +1330,7 @@ private:
 	static FName LookupPackageNameOnDisk(FName NormalizedFileName, bool bExactMatchRequired, FName& FoundFileName);
 
 	/** Allocator for PackageDatas Guarded by ExistenceLock. */
-	TTypedBlockAllocator<FPackageData> Allocator;
+	TTypedBlockAllocatorFreeList<FPackageData> Allocator;
 	FPackageDataMonitor Monitor;
 	/** Guarded by ExistenceLock */
 	TMap<FName, FPackageData*> PackageNameToPackageData;
diff --git a/Engine/Source/Editor/UnrealEd/Private/Cooker/TypedBlockAllocator.h b/Engine/Source/Editor/UnrealEd/Private/Cooker/TypedBlockAllocator.h
--- a/Engine/Source/Editor/UnrealEd/Private/Cooker/TypedBlockAllocator.h
+++ b/Engine/Source/Editor/UnrealEd/Private/Cooker/TypedBlockAllocator.h
@@ -18,9 +18,12 @@
  * Block size is adaptive; it doubles with each new block, up to MaxBlockSize.
  * 
  * Allocated memory is not released until Empty is called.
+ * 
+ * This is the base version and does not keep a free list. Freed elements will never be
+ * reused and will only be freed by Empty() or destructor.
  */
 template <typename ElementType>
-class TTypedBlockAllocator
+class TTypedBlockAllocatorBase
 {
 public:
 	/** Set the MinimumBlockSize. Default is 1024. All new blocks will be this size or larger. */
@@ -70,7 +73,7 @@ public:
 	 */
 	void Empty();
 
-private:
+protected:
 	struct FAllocationBlock
 	{
 		FAllocationBlock(int32 InCapacity);
@@ -79,56 +82,123 @@ private:
 		int32 NextIndex;
 		int32 Capacity;
 	};
+
+	ElementType* BaseAlloc();
+	template <typename CallbackType>
+	void BaseEnumerateAllocations(CallbackType&& Callback);
+	void BaseReserveDelta(int32 AllocationCount, int32 InMaxBlockSize, int32 NumAvailableOnFreeList);
+	void BaseEmpty();
+	void CallDestructorInternal(ElementType* Element);
+
+protected:
 	TArray<FAllocationBlock> Blocks;
-	ElementType* FreeList = nullptr;
 	int32 NextBlock = 0;
 	int32 NumAllocations = 0;
-	int32 NumFreeList = 0;
+	int32 NumFreed = 0;
 	int32 MinBlockSize = 1024;
 	int32 MaxBlockSize = 65536;
 };
 
+/**
+ * A TTypedAllocatorBase with a freelist. Freed elements are destructed, but their memory
+ * is reused and the constructor called again on it in future allocation.
+ */
+template <typename ElementType>
+class TTypedBlockAllocatorFreeList : public TTypedBlockAllocatorBase<ElementType>
+{
+public:
+	ElementType* Alloc();
+	void Free(ElementType* Element);
+	template <typename... ArgsType> ElementType* NewElement(ArgsType&&... Args);
+	void DeleteElement(ElementType* Element);
+	void ReserveDelta(int32 AllocationCount, int32 InMaxBlockSize = 0);
+	void Empty();
+
+protected:
+	using TTypedBlockAllocatorBase<ElementType>::BaseAlloc;
+	using TTypedBlockAllocatorBase<ElementType>::BaseEnumerateAllocations;
+	using TTypedBlockAllocatorBase<ElementType>::BaseReserveDelta;
+	using TTypedBlockAllocatorBase<ElementType>::BaseEmpty;
+	using TTypedBlockAllocatorBase<ElementType>::CallDestructorInternal;
+	using TTypedBlockAllocatorBase<ElementType>::NumAllocations;
+	using TTypedBlockAllocatorBase<ElementType>::NumFreed;
 
+protected:
+	ElementType* FreeList = nullptr;
+};
+
+/**
+ * A TTypedAllocatorBase with a resetlist. Freed elements are not destructed until empty is called. When NewElement is
+ * called and the reset list is available, a previously-freed element is returned without having had the destructor and
+ * a second call to the constructor called on it.
+ */
 template <typename ElementType>
-inline void TTypedBlockAllocator<ElementType>::SetMinBlockSize(int32 BlockSize)
+class TTypedBlockAllocatorResetList : public TTypedBlockAllocatorBase<ElementType>
+{
+public:
+	~TTypedBlockAllocatorResetList();
+
+	void Free(ElementType* Element);
+	template <typename... ArgsType> ElementType* NewElement(ArgsType&&... Args);
+	void ReserveDelta(int32 AllocationCount, int32 InMaxBlockSize = 0);
+	void Empty();
+
+protected:
+	ElementType* Alloc(); // Not public on this subclass; every allocation must have constructor called
+	using TTypedBlockAllocatorBase<ElementType>::DeleteElement; // Not supported on this subclass
+	void ReallocFreeList(int32 NewCapacity);
+
+protected:
+	using TTypedBlockAllocatorBase<ElementType>::BaseAlloc;
+	using TTypedBlockAllocatorBase<ElementType>::BaseEnumerateAllocations;
+	using TTypedBlockAllocatorBase<ElementType>::BaseReserveDelta;
+	using TTypedBlockAllocatorBase<ElementType>::BaseEmpty;
+	using TTypedBlockAllocatorBase<ElementType>::CallDestructorInternal;
+	using TTypedBlockAllocatorBase<ElementType>::NumAllocations;
+	using TTypedBlockAllocatorBase<ElementType>::NumFreed;
+
+protected:
+	TUniquePtr<ElementType* []> FreeList;
+	int32 CapacityFreed = 0;
+};
+
+
+template <typename ElementType>
+void TTypedBlockAllocatorBase<ElementType>::SetMinBlockSize(int32 BlockSize)
 {
 	check(BlockSize >= 0);
 	MinBlockSize = FMath::Max(BlockSize, 1);
 }
 
 template <typename ElementType>
-inline int32 TTypedBlockAllocator<ElementType>::GetMinBlockSize() const
+int32 TTypedBlockAllocatorBase<ElementType>::GetMinBlockSize() const
 {
 	return MinBlockSize;
 }
 
 template <typename ElementType>
-inline void TTypedBlockAllocator<ElementType>::SetMaxBlockSize(int32 BlockSize)
+void TTypedBlockAllocatorBase<ElementType>::SetMaxBlockSize(int32 BlockSize)
 {
 	check(BlockSize >= 0);
 	MaxBlockSize = BlockSize;
 }
 
 template <typename ElementType>
-inline int32 TTypedBlockAllocator<ElementType>::GetMaxBlockSize() const
+int32 TTypedBlockAllocatorBase<ElementType>::GetMaxBlockSize() const
 {
 	return MaxBlockSize;
 }
 
 template <typename ElementType>
-ElementType* TTypedBlockAllocator<ElementType>::Alloc()
+ElementType* TTypedBlockAllocatorBase<ElementType>::Alloc()
 {
-	ElementType* Result = nullptr;
-	if (FreeList)
-	{
-		check(NumFreeList > 0);
-		Result = FreeList;
-		FreeList = *reinterpret_cast<ElementType**>(FreeList);
-		++NumAllocations;
-		--NumFreeList;
-		return Result;
-	}
+	return BaseAlloc();
+}
 
+template <typename ElementType>
+ElementType* TTypedBlockAllocatorBase<ElementType>::BaseAlloc()
+{
+	ElementType* Result = nullptr;
 	FAllocationBlock* Block = nullptr;
 	for (; NextBlock < Blocks.Num(); ++NextBlock)
 	{
@@ -154,48 +224,55 @@ ElementType* TTypedBlockAllocator<ElementType>::Alloc()
 	Result = Block->Elements[Block->NextIndex].GetTypedPtr();
 	++Block->NextIndex;
 	++NumAllocations;
+
 	return Result;
 }
 
 template <typename ElementType>
-void TTypedBlockAllocator<ElementType>::Free(ElementType* Element)
+void TTypedBlockAllocatorBase<ElementType>::Free(ElementType* Element)
 {
-	if constexpr (sizeof(ElementType) <= sizeof(ElementType*))
-	{
-		static_assert(sizeof(ElementType) >= sizeof(ElementType*),
-			"The FreeList is implemented by storing pointers within each freed element. To use the FreeList, elementsize must be >= pointer size.");
-	}
 	check(NumAllocations > 0);
-
-	*reinterpret_cast<ElementType**>(Element) = FreeList;
-	FreeList = Element;
 	--NumAllocations;
-	++NumFreeList;
+	++NumFreed;
+	// The Element will remain in memory, in its block, until Empty is called.
 }
 
 template <typename ElementType>
 template <typename... ArgsType>
-ElementType* TTypedBlockAllocator<ElementType>::NewElement(ArgsType&&... Args)
+ElementType* TTypedBlockAllocatorBase<ElementType>::NewElement(ArgsType&&... Args)
 {
 	return new(Alloc()) ElementType(Forward<ArgsType>(Args)...);
 }
 
 template <typename ElementType>
-void TTypedBlockAllocator<ElementType>::DeleteElement(ElementType* Element)
+void TTypedBlockAllocatorBase<ElementType>::DeleteElement(ElementType* Element)
+{
+	CallDestructorInternal(Element);
+	Free(Element);
+}
+
+template <typename ElementType>
+void TTypedBlockAllocatorBase<ElementType>::CallDestructorInternal(ElementType* Element)
 {
 	// We need a typedef here because VC won't compile the destructor call below if ElementType itself has a member
 	// called ElementType
 	typedef ElementType TTypedBlockAllocatorDestructElementType;
 	((TTypedBlockAllocatorDestructElementType*)Element)->
 		TTypedBlockAllocatorDestructElementType::~TTypedBlockAllocatorDestructElementType();
-	Free(Element);
 }
 
 template <typename ElementType>
 template <typename CallbackType>
-void TTypedBlockAllocator<ElementType>::EnumerateAllocations(CallbackType&& Callback)
+void TTypedBlockAllocatorBase<ElementType>::EnumerateAllocations(CallbackType&& Callback)
+{
+	checkf(NumFreed == 0, TEXT("It is invalid to call EnumerateAllocations after calling Free or Destruct."));
+	BaseEnumerateAllocations(Forward<CallbackType>(Callback));
+}
+
+template <typename ElementType>
+template <typename CallbackType>
+void TTypedBlockAllocatorBase<ElementType>::BaseEnumerateAllocations(CallbackType&& Callback)
 {
-	checkf(FreeList == nullptr, TEXT("It is invalid to call EnumerateAllocations after calling Free or Destruct."));
 	for (const FAllocationBlock& Block : Blocks)
 	{
 		for (TTypeCompatibleBytes<ElementType>& Bytes :
@@ -208,10 +285,17 @@ void TTypedBlockAllocator<ElementType>::EnumerateAllocations(CallbackType&& Call
 }
 
 template <typename ElementType>
-void TTypedBlockAllocator<ElementType>::ReserveDelta(int32 AllocationCount, int32 InMaxBlockSize)
+void TTypedBlockAllocatorBase<ElementType>::ReserveDelta(int32 AllocationCount, int32 InMaxBlockSize)
+{
+	BaseReserveDelta(AllocationCount, InMaxBlockSize, 0);
+}
+
+template <typename ElementType>
+void TTypedBlockAllocatorBase<ElementType>::BaseReserveDelta(int32 AllocationCount,
+	int32 InMaxBlockSize, int32 NumAvailableOnFreeList)
 {
-	int32 DeltaAllocationCount = AllocationCount - NumFreeList;
-	for (const FAllocationBlock& Block: TConstArrayView<FAllocationBlock>(Blocks).RightChop(NextBlock))
+	int32 DeltaAllocationCount = AllocationCount - NumAvailableOnFreeList;
+	for (const FAllocationBlock& Block : TConstArrayView<FAllocationBlock>(Blocks).RightChop(NextBlock))
 	{
 		DeltaAllocationCount -= Block.Capacity - Block.NextIndex;
 	}
@@ -250,19 +334,154 @@ void TTypedBlockAllocator<ElementType>::ReserveDelta(int32 AllocationCount, int3
 }
 
 template <typename ElementType>
-void TTypedBlockAllocator<ElementType>::Empty()
+void TTypedBlockAllocatorBase<ElementType>::Empty()
+{
+	BaseEmpty();
+}
+
+template <typename ElementType>
+void TTypedBlockAllocatorBase<ElementType>::BaseEmpty()
 {
 	Blocks.Empty();
-	FreeList = nullptr;
 	NextBlock = 0;
 	NumAllocations = 0;
-	NumFreeList = 0;
+	NumFreed = 0;
 }
 
 template <typename ElementType>
-TTypedBlockAllocator<ElementType>::FAllocationBlock::FAllocationBlock(int32 InCapacity)
+TTypedBlockAllocatorBase<ElementType>::FAllocationBlock::FAllocationBlock(int32 InCapacity)
 	: Elements(InCapacity > 0 ? new TTypeCompatibleBytes<ElementType>[InCapacity] : nullptr)
 	, NextIndex(0)
 	, Capacity(InCapacity)
 {
 }
+
+template <typename ElementType>
+ElementType* TTypedBlockAllocatorFreeList<ElementType>::Alloc()
+{
+	ElementType* Result = nullptr;
+	if (NumFreed > 0)
+	{
+		check(FreeList);
+		Result = FreeList;
+		FreeList = *reinterpret_cast<ElementType**>(FreeList);
+		++NumAllocations;
+		--NumFreed;
+		return Result;
+	}
+	return BaseAlloc();
+}
+
+template <typename ElementType>
+void TTypedBlockAllocatorFreeList<ElementType>::Free(ElementType* Element)
+{
+	if constexpr (sizeof(ElementType) <= sizeof(ElementType*))
+	{
+		static_assert(sizeof(ElementType) >= sizeof(ElementType*),
+			"The FreeList is implemented by storing pointers within each freed element. To use the FreeList, elementsize must be >= pointer size.");
+	}
+	check(NumAllocations > 0);
+
+	*reinterpret_cast<ElementType**>(Element) = FreeList;
+	FreeList = Element;
+	--NumAllocations;
+	++NumFreed;
+}
+
+template <typename ElementType>
+template <typename... ArgsType>
+ElementType* TTypedBlockAllocatorFreeList<ElementType>::NewElement(ArgsType&&... Args)
+{
+	return new(Alloc()) ElementType(Forward<ArgsType>(Args)...);
+}
+
+template <typename ElementType>
+void TTypedBlockAllocatorFreeList<ElementType>::DeleteElement(ElementType* Element)
+{
+	CallDestructorInternal(Element);
+	Free(Element);
+}
+
+template <typename ElementType>
+void TTypedBlockAllocatorFreeList<ElementType>::ReserveDelta(int32 AllocationCount, int32 InMaxBlockSize)
+{
+	BaseReserveDelta(AllocationCount, InMaxBlockSize, NumFreed);
+}
+
+template <typename ElementType>
+void TTypedBlockAllocatorFreeList<ElementType>::Empty()
+{
+	BaseEmpty();
+	FreeList = nullptr;
+}
+
+
+template <typename ElementType>
+TTypedBlockAllocatorResetList<ElementType>::~TTypedBlockAllocatorResetList()
+{
+	Empty();
+}
+
+template <typename ElementType>
+ElementType* TTypedBlockAllocatorResetList<ElementType>::Alloc()
+{
+	ElementType* Result = nullptr;
+	if (NumFreed > 0)
+	{
+		check(FreeList);
+		Result = FreeList[--NumFreed];
+		if (8 < NumFreed && NumFreed*2 < CapacityFreed)
+		{
+			ReallocFreeList(NumFreed * 3 / 2);
+		}
+		++NumAllocations;
+		return Result;
+	}
+	return BaseAlloc();
+}
+
+template <typename ElementType>
+void TTypedBlockAllocatorResetList<ElementType>::ReallocFreeList(int32 NewCapacity)
+{
+	TUniquePtr<ElementType* []> Copy = MakeUnique<ElementType* []>(NewCapacity);
+	FMemory::Memcpy(Copy.Get(), FreeList.Get(), NumFreed * sizeof(ElementType*));
+	FreeList = MoveTemp(Copy);
+	CapacityFreed = NewCapacity;
+}
+
+template <typename ElementType>
+void TTypedBlockAllocatorResetList<ElementType>::Free(ElementType* Element)
+{
+	check(NumAllocations > 0);
+	if (CapacityFreed == NumFreed)
+	{
+		ReallocFreeList(NumFreed < 8 ? 8 : NumFreed * 3 / 2);
+	}
+	FreeList[NumFreed++] = Element;
+	--NumAllocations;
+}
+
+template <typename ElementType>
+template <typename... ArgsType>
+ElementType* TTypedBlockAllocatorResetList<ElementType>::NewElement(ArgsType&&... Args)
+{
+	return new(Alloc()) ElementType(Forward<ArgsType>(Args)...);
+}
+
+template <typename ElementType>
+void TTypedBlockAllocatorResetList<ElementType>::ReserveDelta(int32 AllocationCount, int32 InMaxBlockSize)
+{
+	BaseReserveDelta(AllocationCount, InMaxBlockSize, NumFreed);
+}
+
+template <typename ElementType>
+void TTypedBlockAllocatorResetList<ElementType>::Empty()
+{
+	BaseEnumerateAllocations([this](ElementType* Element)
+		{
+			CallDestructorInternal(Element);
+		});
+	BaseEmpty();
+	FreeList.Reset();
+	CapacityFreed = 0;
+}
