Allow multi-process cook to pass along TargetPlatform argument to the cook workers and remove use of CookDirectorTargetPlatform.

diff --git a/Engine/Source/Editor/UnrealEd/Private/Cooker/CookDirector.cpp b/Engine/Source/Editor/UnrealEd/Private/Cooker/CookDirector.cpp
--- a/Engine/Source/Editor/UnrealEd/Private/Cooker/CookDirector.cpp
+++ b/Engine/Source/Editor/UnrealEd/Private/Cooker/CookDirector.cpp
@@ -1217,8 +1217,6 @@ FString FCookDirector::GetWorkerCommandLine(FWorkerId WorkerId, int32 ProfileId)
 	TArray<FString> Tokens;
 	UE::String::ParseTokensMultiple(CommandLine, { ' ', '\t', '\r', '\n' }, [&Tokens](FStringView Token)
 		{
-			const FStringView TargetPlatformToken(TEXT("-TargetPlatform"));
-
 			if (Token.StartsWith(TEXT("-run=")) ||
 				Token == TEXT("-CookOnTheFly") ||
 				Token == TEXT("-CookWorker") ||
@@ -1236,14 +1234,6 @@ FString FCookDirector::GetWorkerCommandLine(FWorkerId WorkerId, int32 ProfileId)
 			{
 				return;
 			}
-			else if (Token.StartsWith(TargetPlatformToken))
-			{
-				// TargetPlatform gets transformed into CookDirectorTargetPlatform
-				FString ModifiedTargetPlatformArgument(TEXT("-CookDirectorTargetPlatform"));
-				ModifiedTargetPlatformArgument.Append(Token.RightChop(TargetPlatformToken.Len()));
-				Tokens.Add(ModifiedTargetPlatformArgument);
-				return;
-			}
 			Tokens.Add(FString(Token));
 		}, UE::String::EParseTokensOptions::SkipEmpty);
 	if (Tokens[0] != ProjectName)
diff --git a/Engine/Source/Editor/UnrealEd/Private/Cooker/CookWorkerClient.cpp b/Engine/Source/Editor/UnrealEd/Private/Cooker/CookWorkerClient.cpp
--- a/Engine/Source/Editor/UnrealEd/Private/Cooker/CookWorkerClient.cpp
+++ b/Engine/Source/Editor/UnrealEd/Private/Cooker/CookWorkerClient.cpp
@@ -11,6 +11,7 @@
 #include "CookWorkerServer.h"
 #include "CoreGlobals.h"
 #include "HAL/PlatformTime.h"
+#include "Misc/CoreMisc.h"
 #include "PackageResultsMessage.h"
 #include "PackageTracker.h"
 #include "Sockets.h"
@@ -465,6 +466,61 @@ void FCookWorkerClient::PollReceiveConfigMessage()
 	}
 	DirectorCookMode = InitialConfigMessage->GetDirectorCookMode();
 	OrderedSessionPlatforms = InitialConfigMessage->GetOrderedSessionPlatforms();
+	const TArray<ITargetPlatform*>& ActiveTargetPlatforms = GetTargetPlatformManagerRef().GetActiveTargetPlatforms();
+
+	auto GetPlatformDetails = [this, &ActiveTargetPlatforms](TStringBuilder<512>& StringBuilder)
+	{
+		StringBuilder << TEXT("ActiveTargetPlatforms(") << ActiveTargetPlatforms.Num() << TEXT("): ");
+		for (int32 PlatformIndex = 0; PlatformIndex < ActiveTargetPlatforms.Num(); ++PlatformIndex)
+		{
+			ITargetPlatform* Platform = ActiveTargetPlatforms[PlatformIndex];
+			StringBuilder << Platform->PlatformName();
+			if (PlatformIndex < (ActiveTargetPlatforms.Num() - 1))
+			{
+				StringBuilder << TEXT(", ");
+			}
+		}
+		StringBuilder << TEXT("\n");
+
+		StringBuilder << TEXT("OrderedSessionPlatforms(") << OrderedSessionPlatforms.Num() << TEXT("): ");
+		for (int32 PlatformIndex = 0; PlatformIndex < OrderedSessionPlatforms.Num(); ++PlatformIndex)
+		{
+			ITargetPlatform* Platform = OrderedSessionPlatforms[PlatformIndex];
+			StringBuilder << Platform->PlatformName();
+			if (PlatformIndex < (OrderedSessionPlatforms.Num() - 1))
+			{
+				StringBuilder << TEXT(", ");
+			}
+		}
+	};
+
+	if (OrderedSessionPlatforms.Num() != ActiveTargetPlatforms.Num())
+	{
+		TStringBuilder<512> StringBuilder;
+		GetPlatformDetails(StringBuilder);
+		UE_LOG(LogCook, Error, TEXT("CookWorker initialization failure: Director sent a mismatch in session platform quantity.\n%s"), *StringBuilder);
+		SendToState(EConnectStatus::LostConnection);
+		return;
+	}
+
+	bool bPlatformMismatch = false;
+	for (ITargetPlatform* Platform : ActiveTargetPlatforms)
+	{
+		if (!OrderedSessionPlatforms.Contains(Platform))
+		{
+			bPlatformMismatch = true;
+			break;
+		}
+	}
+
+	if (bPlatformMismatch)
+	{
+		TStringBuilder<512> StringBuilder;
+		GetPlatformDetails(StringBuilder);
+		UE_LOG(LogCook, Error, TEXT("CookWorker initialization failure: Director sent a mismatch in session platform contents.\n%s"), *StringBuilder);
+		SendToState(EConnectStatus::LostConnection);
+		return;
+	}
 
 	UE_LOG(LogCook, Display, TEXT("Initialization from CookDirector complete."));
 	SendToState(EConnectStatus::Connected);
